Fix for a crash in Async Loading Graph's CheckCycles when GC kicks in on the game thread and forces ALT to exit early

diff --git a/Engine/Source/Runtime/CoreUObject/Private/Serialization/AsyncLoading.cpp b/Engine/Source/Runtime/CoreUObject/Private/Serialization/AsyncLoading.cpp
--- a/Engine/Source/Runtime/CoreUObject/Private/Serialization/AsyncLoading.cpp
+++ b/Engine/Source/Runtime/CoreUObject/Private/Serialization/AsyncLoading.cpp
@@ -5043,7 +5043,7 @@ void FAsyncLoadingThread::CheckForCycles()
 }
 
 
-EAsyncPackageState::Type  FAsyncLoadingThread::TickAsyncThread(bool bUseTimeLimit, bool bUseFullTimeLimit, float TimeLimit, bool& bDidSomething, FFlushTree* FlushTree)
+EAsyncPackageState::Type FAsyncLoadingThread::TickAsyncThread(bool bUseTimeLimit, bool bUseFullTimeLimit, float TimeLimit, bool& bDidSomething, FFlushTree* FlushTree)
 {
 	check(!IsInGameThread() || !IsMultithreaded());
 	EAsyncPackageState::Type Result = EAsyncPackageState::Complete;
@@ -5074,7 +5074,7 @@ EAsyncPackageState::Type  FAsyncLoadingThread::TickAsyncThread(bool bUseTimeLimi
 				bDidSomething = bDidSomething || ProcessedRequests > 0;
 			}
 		}
-		if (ProcessedRequests == 0 && IsMultithreaded())
+		if (ProcessedRequests == 0 && IsMultithreaded() && Result == EAsyncPackageState::Complete)
 		{
 			uint32 WaitTime = 30;
 			if (IsEventDrivenLoaderEnabled())
