Removing unused async loading stat

diff --git a/Engine/Source/Runtime/CoreUObject/Private/Serialization/AsyncLoading.cpp b/Engine/Source/Runtime/CoreUObject/Private/Serialization/AsyncLoading.cpp
--- a/Engine/Source/Runtime/CoreUObject/Private/Serialization/AsyncLoading.cpp
+++ b/Engine/Source/Runtime/CoreUObject/Private/Serialization/AsyncLoading.cpp
@@ -58,7 +58,6 @@ DECLARE_MEMORY_STAT(TEXT("Streaming Memory Used"),STAT_StreamingAllocSize,STATGR
 DECLARE_STATS_GROUP_VERBOSE(TEXT("Async Load"), STATGROUP_AsyncLoad, STATCAT_Advanced);
 
 DECLARE_CYCLE_STAT(TEXT("Tick AsyncPackage"),STAT_FAsyncPackage_Tick,STATGROUP_AsyncLoad);
-DECLARE_FLOAT_ACCUMULATOR_STAT(TEXT("Tick AsyncPackage Time"), STAT_FAsyncPackage_TickTime, STATGROUP_AsyncLoad);
 
 DECLARE_CYCLE_STAT(TEXT("CreateLinker AsyncPackage"),STAT_FAsyncPackage_CreateLinker,STATGROUP_AsyncLoad);
 DECLARE_CYCLE_STAT(TEXT("FinishLinker AsyncPackage"),STAT_FAsyncPackage_FinishLinker,STATGROUP_AsyncLoad);
