PR #4779: Reason fails as the supplied variable is incorrect (Contributed by [REDACTED_PERSON])

diff --git a/Engine/Source/Runtime/CoreUObject/Private/Misc/PackageName.cpp b/Engine/Source/Runtime/CoreUObject/Private/Misc/PackageName.cpp
--- a/Engine/Source/Runtime/CoreUObject/Private/Misc/PackageName.cpp
+++ b/Engine/Source/Runtime/CoreUObject/Private/Misc/PackageName.cpp
@@ -818,9 +818,8 @@ bool FPackageName::DoesPackageExist(const FString& LongPackageName, const FGuid*
 
 	if (!FPackageName::TryConvertFilenameToLongPackageName(LongPackageName, PackageName))
 	{
-		verify(!FPackageName::IsValidLongPackageName(PackageName, true, &Reason));
+		verify(!FPackageName::IsValidLongPackageName(LongPackageName, true, &Reason));
 		UE_LOG(LogPackageName, Error, TEXT("Illegal call to DoesPackageExist: '%s' is not a standard unreal filename or a long path name. Reason: %s"), *LongPackageName, *Reason.ToString());
-		ensureMsgf(false, TEXT("Illegal call to DoesPackageExist: '%s' is not a standard unreal filename or a long path name. Reason: %s"), *LongPackageName, *Reason.ToString());
 		return false;
 	}
 	// Once we have the real Package Name, we can exit early if it's a script package - they exist only in memory.
