Fix for a potential crash when terminating the engine while it's being initialized

diff --git a/Engine/Source/Runtime/CoreUObject/Private/UObject/Obj.cpp b/Engine/Source/Runtime/CoreUObject/Private/UObject/Obj.cpp
--- a/Engine/Source/Runtime/CoreUObject/Private/UObject/Obj.cpp
+++ b/Engine/Source/Runtime/CoreUObject/Private/UObject/Obj.cpp
@@ -1893,7 +1893,8 @@ void UObject::LoadConfig( UClass* ConfigClass/*=NULL*/, const TCHAR* InFilename/
 	// because ProcessNewlyLoadedUObjects hasn't happened yet
 	checkf(ConfigClass->PropertyLink != nullptr
 		|| (ConfigClass->GetSuperStruct() && ConfigClass->PropertiesSize == ConfigClass->GetSuperStruct()->PropertiesSize)
-		|| ConfigClass->PropertiesSize == 0,
+		|| ConfigClass->PropertiesSize == 0
+		|| GIsRequestingExit, // Ignore this check when exiting as we may have requested exit during init when not everything is initialized
 		TEXT("class %s has uninitialized properties. Accessed too early?"), *ConfigClass->GetName());
 #endif
 
