Fixing another creash when terminating the engine while initializing.

diff --git a/Engine/Source/Runtime/CoreUObject/Private/UObject/UObjectArray.cpp b/Engine/Source/Runtime/CoreUObject/Private/UObject/UObjectArray.cpp
--- a/Engine/Source/Runtime/CoreUObject/Private/UObject/UObjectArray.cpp
+++ b/Engine/Source/Runtime/CoreUObject/Private/UObject/UObjectArray.cpp
@@ -66,7 +66,10 @@ void FUObjectArray::CloseDisregardForGC()
 	check(OpenForDisregardForGC);
 
 	// Make sure all classes that have been loaded/created so far are properly initialized
-	ProcessNewlyLoadedUObjects();
+	if (!GIsRequestingExit)
+	{
+		ProcessNewlyLoadedUObjects();
+	}
 
 	UClass::AssembleReferenceTokenStreams();
 
