The warning 'Calling StaticLoadObject during PostLoad may result in hitches during streaming' will now also report the object which had the PostLoad called on it when StaticLoadObject call happened.

diff --git a/Engine/Source/Runtime/CoreUObject/Private/Serialization/AsyncLoading.cpp b/Engine/Source/Runtime/CoreUObject/Private/Serialization/AsyncLoading.cpp
--- a/Engine/Source/Runtime/CoreUObject/Private/Serialization/AsyncLoading.cpp
+++ b/Engine/Source/Runtime/CoreUObject/Private/Serialization/AsyncLoading.cpp
@@ -6382,9 +6382,10 @@ EAsyncPackageState::Type FAsyncPackage::PostLoadObjects()
 	// GC can't run in here
 	FGCScopeGuard GCGuard;
 
-	TGuardValue<bool> GuardIsRoutingPostLoad(FUObjectThreadContext::Get().IsRoutingPostLoad, true);
+	FUObjectThreadContext& ThreadContext = FUObjectThreadContext::Get();
+	TGuardValue<bool> GuardIsRoutingPostLoad(ThreadContext.IsRoutingPostLoad, true);
 
-	TArray<UObject*>& ThreadObjLoaded = FUObjectThreadContext::Get().ObjLoaded;
+	TArray<UObject*>& ThreadObjLoaded = ThreadContext.ObjLoaded;
 	if (ThreadObjLoaded.Num())
 	{
 		// New objects have been loaded. They need to go through PreLoad first so exit now and come back after they've been preloaded.
@@ -6416,7 +6417,9 @@ EAsyncPackageState::Type FAsyncPackage::PostLoadObjects()
 				// We want this check only with EDL enabled
 				check(!GEventDrivenLoaderEnabled || !Object->HasAnyFlags(RF_NeedLoad));
 
+				ThreadContext.CurrentlyPostLoadedObjectByALT = Object;
 				Object->ConditionalPostLoad();
+				ThreadContext.CurrentlyPostLoadedObjectByALT = nullptr;
 
 				LastObjectWorkWasPerformedOn = Object;
 				LastTypeOfWorkPerformed = TEXT("postloading_async");
@@ -6478,7 +6481,9 @@ EAsyncPackageState::Type FAsyncPackage::PostLoadDeferredObjects(double InTickSta
 
 		FScopeCycleCounterUObject ConstructorScope(Object, GET_STATID(STAT_FAsyncPackage_PostLoadObjectsGameThread));
 
+		PackageScope.ThreadContext.CurrentlyPostLoadedObjectByALT = Object;
 		Object->ConditionalPostLoad();
+		PackageScope.ThreadContext.CurrentlyPostLoadedObjectByALT = nullptr;
 
 		if (ObjLoadedInPostLoad.Num())
 		{
diff --git a/Engine/Source/Runtime/CoreUObject/Private/UObject/UObjectGlobals.cpp b/Engine/Source/Runtime/CoreUObject/Private/UObject/UObjectGlobals.cpp
--- a/Engine/Source/Runtime/CoreUObject/Private/UObject/UObjectGlobals.cpp
+++ b/Engine/Source/Runtime/CoreUObject/Private/UObject/UObjectGlobals.cpp
@@ -978,7 +978,11 @@ UObject* StaticLoadObjectInternal(UClass* ObjectClass, UObject* InOuter, const T
 
 UObject* StaticLoadObject(UClass* ObjectClass, UObject* InOuter, const TCHAR* InName, const TCHAR* Filename, uint32 LoadFlags, UPackageMap* Sandbox, bool bAllowObjectReconciliation )
 {
-	UE_CLOG(FUObjectThreadContext::Get().IsRoutingPostLoad && IsInAsyncLoadingThread(), LogUObjectGlobals, Warning, TEXT("Calling StaticLoadObject during PostLoad may result in hitches during streaming."));
+	FUObjectThreadContext& ThreadContext = FUObjectThreadContext::Get();
+	if (ThreadContext.IsRoutingPostLoad && IsInAsyncLoadingThread())
+	{
+		UE_LOG(LogUObjectGlobals, Warning, TEXT("Calling StaticLoadObject during PostLoad of %s may result in hitches during streaming."), *GetFullNameSafe(ThreadContext.CurrentlyPostLoadedObjectByALT));
+	}
 
 	UObject* Result = StaticLoadObjectInternal(ObjectClass, InOuter, InName, Filename, LoadFlags, Sandbox, bAllowObjectReconciliation);
 	if (!Result)
diff --git a/Engine/Source/Runtime/CoreUObject/Private/UObject/UObjectThreadContext.cpp b/Engine/Source/Runtime/CoreUObject/Private/UObject/UObjectThreadContext.cpp
--- a/Engine/Source/Runtime/CoreUObject/Private/UObject/UObjectThreadContext.cpp
+++ b/Engine/Source/Runtime/CoreUObject/Private/UObject/UObjectThreadContext.cpp
@@ -13,6 +13,7 @@ FUObjectThreadContext::FUObjectThreadContext()
 , ForcedExportCount(0)
 , ObjBeginLoadCount(0)
 , IsRoutingPostLoad(false)
+, CurrentlyPostLoadedObjectByALT(nullptr)
 , IsDeletingLinkers(false)
 , IsInConstructor(0)
 , ConstructedObject(nullptr)
diff --git a/Engine/Source/Runtime/CoreUObject/Public/UObject/UObjectThreadContext.h b/Engine/Source/Runtime/CoreUObject/Public/UObject/UObjectThreadContext.h
--- a/Engine/Source/Runtime/CoreUObject/Public/UObject/UObjectThreadContext.h
+++ b/Engine/Source/Runtime/CoreUObject/Public/UObject/UObjectThreadContext.h
@@ -74,6 +74,8 @@ public:
 	TArray<FLinkerLoad*> DelayedLinkerClosePackages;
 	/** true when we are routing ConditionalPostLoad/PostLoad to objects										*/
 	bool IsRoutingPostLoad;
+	/** The object we are routing PostLoad from the Async Loading code for */
+	UObject* CurrentlyPostLoadedObjectByALT;
 	/** true when FLinkerManager deletes linkers */
 	bool IsDeletingLinkers;
 	/* Global flag so that FObjectFinders know if they are called from inside the UObject constructors or not. */
