Added extra checks to catch scenarios where the EDL Precache Buffer is flushed before a package header is fully read

diff --git a/Engine/Source/Runtime/CoreUObject/Private/Serialization/AsyncLoading.cpp b/Engine/Source/Runtime/CoreUObject/Private/Serialization/AsyncLoading.cpp
--- a/Engine/Source/Runtime/CoreUObject/Private/Serialization/AsyncLoading.cpp
+++ b/Engine/Source/Runtime/CoreUObject/Private/Serialization/AsyncLoading.cpp
@@ -7395,6 +7395,7 @@ void FArchiveAsync2::FlushPrecacheBlock()
 #if USE_DETAILED_FARCHIVEASYNC2_MEMORY_TRACKING
 		GArchiveAsync2MemTracker.Deallocate(FileName, PrecacheEndPos - PrecacheStartPos);
 #endif
+		check(!GEventDrivenLoaderEnabled || LoadPhase > ELoadPhase::WaitingForHeader);
 	}
 	PrecacheBuffer = nullptr;
 	PrecacheStartPos = 0;
@@ -8017,6 +8018,7 @@ void FArchiveAsync2::Serialize(void* Data, int64 Count)
 			// no before block and head of desired block is in the cache
 			int64 CopyLen = FMath::Min(PrecacheEndPos - CurrentPos, Count);
 			check(CopyLen > 0);
+			check(PrecacheBuffer);
 			FMemory::Memcpy(Data, PrecacheBuffer + CurrentPos - PrecacheStartPos, CopyLen);
 			AfterBlockSize = Count - CopyLen;
 			check(AfterBlockSize >= 0);
@@ -8033,6 +8035,7 @@ void FArchiveAsync2::Serialize(void* Data, int64 Count)
 				// tail of desired block is in the cache
 				int64 CopyLen = FMath::Min(PrecacheEndPos - CurrentPos - BeforeBlockSize, Count - BeforeBlockSize);
 				check(CopyLen > 0);
+				check(PrecacheBuffer);
 				FMemory::Memcpy(((uint8*)Data) + BeforeBlockSize, PrecacheBuffer, CopyLen);
 				AfterBlockSize = Count - CopyLen - BeforeBlockSize;
 				check(AfterBlockSize >= 0);
@@ -8053,6 +8056,7 @@ void FArchiveAsync2::Serialize(void* Data, int64 Count)
 			return;
 		}
 		check(BeforeBlockOffset >= PrecacheStartPos && BeforeBlockOffset + BeforeBlockSize <= PrecacheEndPos);
+		check(PrecacheBuffer);
 		FMemory::Memcpy(Data, PrecacheBuffer + BeforeBlockOffset - PrecacheStartPos, BeforeBlockSize);
 	}
 	if (AfterBlockSize)
@@ -8098,6 +8102,7 @@ void FArchiveAsync2::Serialize(void* Data, int64 Count)
 			return;
 		}
 		checkf(AfterBlockOffset >= PrecacheStartPos && AfterBlockOffset + AfterBlockSize <= PrecacheEndPos, TEXT("Sync After Block ????   %lld %lld %lld %lld"), AfterBlockOffset, AfterBlockSize, PrecacheStartPos, PrecacheEndPos);
+		check(PrecacheBuffer);
 		FMemory::Memcpy(((uint8*)Data) + Count - AfterBlockSize, PrecacheBuffer + AfterBlockOffset - PrecacheStartPos, AfterBlockSize);
 	}
 #if DEVIRTUALIZE_FLinkerLoad_Serialize
