Made sure the Async Loading Thread doesn't continue running after creating new async packages when garbage collector wants to run on the game thread

diff --git a/Engine/Source/Runtime/CoreUObject/Private/Serialization/AsyncLoading.cpp b/Engine/Source/Runtime/CoreUObject/Private/Serialization/AsyncLoading.cpp
--- a/Engine/Source/Runtime/CoreUObject/Private/Serialization/AsyncLoading.cpp
+++ b/Engine/Source/Runtime/CoreUObject/Private/Serialization/AsyncLoading.cpp
@@ -5064,7 +5064,7 @@ EAsyncPackageState::Type FAsyncLoadingThread::TickAsyncThread(bool bUseTimeLimit
 			}
 			float TimeUsed = (float)(FPlatformTime::Seconds() - TickStartTime);
 			const float RemainingTimeLimit = FMath::Max(0.0f, TimeLimit - TimeUsed);
-			if (RemainingTimeLimit <= 0.0f && bUseTimeLimit && !IsMultithreaded())
+			if (IsGarbageCollectionWaiting() || (RemainingTimeLimit <= 0.0f && bUseTimeLimit && !IsMultithreaded()))
 			{
 				Result = EAsyncPackageState::TimeOut;
 			}
