Creating new objects from within ForEachObjectWithOuter will now result in a fatal error as it's unsafe to change internal UObject hash tables when iterating over them.

diff --git a/Engine/Source/Runtime/CoreUObject/Private/UObject/UObjectHash.cpp b/Engine/Source/Runtime/CoreUObject/Private/UObject/UObjectHash.cpp
--- a/Engine/Source/Runtime/CoreUObject/Private/UObject/UObjectHash.cpp
+++ b/Engine/Source/Runtime/CoreUObject/Private/UObject/UObjectHash.cpp
@@ -240,7 +240,10 @@ struct FHashBucketIterator
 
 class FUObjectHashTables
 {
+	/** Critical section that guards against concurrent adds from multiple threads */
 	FCriticalSection CriticalSection;
+	/** Hard lock flag, this is used when iterating over the hash table - we don't want any adds when this is happening */
+	int32 HardLock;
 
 public:
 
@@ -254,6 +257,7 @@ public:
 	TMap<UClass*, TSet<UClass*> > ClassToChildListMap;
 
 	FUObjectHashTables()
+		: HardLock(0)
 	{
 	}
 
@@ -271,6 +275,10 @@ public:
 	/** Adds the Hash/Object pair to the FName hash table */
 	FORCEINLINE void AddToHash(int32 InHash, UObjectBase* Object)
 	{
+		UE_CLOG(HardLock > 0, LogUObjectHash, Fatal,
+			TEXT("Trying to add new object %s to UObject hashtables when they are in hard lockdown mode. Are you trying to create a new object in ForEachObjectWithOuter?"),
+			*GetFullNameSafe(static_cast<UObjectBaseUtility*>(Object)));
+
 		FHashBucket& Bucket = Hash.FindOrAdd(InHash);
 		Bucket.Add(Object);
 	}
@@ -290,13 +298,22 @@ public:
 		return NumRemoved;
 	}
 
-	FORCEINLINE void Lock()
+	FORCEINLINE void Lock(bool bHard = false)
 	{
 		CriticalSection.Lock();
+		if (bHard)
+		{
+			HardLock++;
+		}
 	}
 
-	FORCEINLINE void Unlock()
+	FORCEINLINE void Unlock(bool bHard = false)
 	{
+		if (bHard)
+		{
+			HardLock--;
+			check(HardLock >= 0);
+		}
 		CriticalSection.Unlock();
 	}
 
@@ -311,15 +328,17 @@ class FHashTableLock
 {
 #if THREADSAFE_UOBJECTS
 	FUObjectHashTables* Tables;
+	bool bHardLock;
 #endif
 public:
-	FORCEINLINE FHashTableLock(FUObjectHashTables& InTables)
+	FORCEINLINE FHashTableLock(FUObjectHashTables& InTables, bool bHard = false)
 	{
 #if THREADSAFE_UOBJECTS
+		bHardLock = bHard;
 		if (!(IsGarbageCollecting() && IsInGameThread()))
 		{
 			Tables = &InTables;
-			InTables.Lock();
+			InTables.Lock(bHardLock);
 		}
 		else
 		{
@@ -334,7 +353,7 @@ public:
 #if THREADSAFE_UOBJECTS
 		if (Tables)
 		{
-			Tables->Unlock();
+			Tables->Unlock(bHardLock);
 		}
 #endif
 	}
@@ -691,7 +710,7 @@ void ForEachObjectWithOuter(const class UObjectBase* Outer, TFunctionRef<void(UO
 		ExclusionInternalFlags |= EInternalObjectFlags::AsyncLoading;
 	}
 	FUObjectHashTables& ThreadHash = FUObjectHashTables::Get();
-	FHashTableLock HashLock(ThreadHash);
+	FHashTableLock HashLock(ThreadHash, true);
 	TArray<FHashBucket*, TInlineAllocator<1> > AllInners;
 
 	if (FHashBucket* Inners = ThreadHash.ObjectOuterMap.Find(Outer))
