Added the option to filter clusters listed with gc.ListClusters by objects within them.

Usage:

gc.ListClusters Hierachy With=ObjectName1,ObjectName2...

diff --git a/Engine/Source/Runtime/CoreUObject/Private/UObject/UObjectClusters.cpp b/Engine/Source/Runtime/CoreUObject/Private/UObject/UObjectClusters.cpp
--- a/Engine/Source/Runtime/CoreUObject/Private/UObject/UObjectClusters.cpp
+++ b/Engine/Source/Runtime/CoreUObject/Private/UObject/UObjectClusters.cpp
@@ -18,6 +18,7 @@
 #include "UObject/ReferenceChainSearch.h"
 #include "UObject/LinkerLoad.h"
 #include "UObject/FastReferenceCollector.h"
+#include "UObject/Package.h"
 
 int32 GCreateGCClusters = 1;
 static FAutoConsoleVariableRef CCreateGCClusters(
@@ -294,6 +295,25 @@ void DumpClusterToLog(const FUObjectCluster& Cluster, bool bHierarchy, bool bInd
 	}
 }
 
+static bool DoesClusterContainObjects(FUObjectCluster& Cluster, const TArray<int32>& Objects)
+{
+	for (int32 ObjectIndex : Objects)
+	{
+		if (Cluster.RootIndex == ObjectIndex)
+		{
+			return true;
+		}
+		if (Cluster.Objects.Contains(ObjectIndex))
+		{
+			return true;
+		}
+		else if (Cluster.MutableObjects.Contains(ObjectIndex))
+		{
+			return true;
+		}
+	}
+	return false;
+}
 
 // Dumps all clusters to log.
 void ListClusters(const TArray<FString>& Args)
@@ -313,35 +333,60 @@ void ListClusters(const TArray<FString>& Args)
 		}
 	}
 
-	if (Args.Contains(TEXT("SortByName")))
+	TArray<int32> WithObjects;
+	for (const FString& Arg : Args)
 	{
-		Algo::SortBy(AllClusters, [](FUObjectCluster* A)
+		if (Arg == TEXT("SortByName"))
 		{
-			return GUObjectArray.IndexToObject(A->RootIndex)->Object->GetFName();
-		});
-	}
-	else if (Args.Contains(TEXT("SortByObjectCount")))
-	{
-		Algo::SortBy(AllClusters, [](FUObjectCluster* A)
+			Algo::SortBy(AllClusters, [](FUObjectCluster* A)
+			{
+				return GUObjectArray.IndexToObject(A->RootIndex)->Object->GetFName();
+			});
+		}
+		else if (Arg == TEXT("SortByObjectCount"))
 		{
-			return A->Objects.Num();
-		});
-	}
-	else if (Args.Contains(TEXT("SortByMutableObjectCount")))
-	{
-		Algo::SortBy(AllClusters, [](FUObjectCluster* A)
+			Algo::SortBy(AllClusters, [](FUObjectCluster* A)
+			{
+				return A->Objects.Num();
+			});
+		}
+		else if (Arg == TEXT("SortByMutableObjectCount"))
 		{
-			return A->MutableObjects.Num();
-		});
-	}
-	else if (Args.Contains(TEXT("SortByReferencedClustersCount")))
-	{
-		Algo::SortBy(AllClusters, [](FUObjectCluster* A)
+			Algo::SortBy(AllClusters, [](FUObjectCluster* A)
+			{
+				return A->MutableObjects.Num();
+			});
+		}
+		else if (Arg == TEXT("SortByReferencedClustersCount"))
+		{
+			Algo::SortBy(AllClusters, [](FUObjectCluster* A)
+			{
+				return A->ReferencedClusters.Num();
+			});
+		}
+		else if (Arg.StartsWith("With="))
 		{
-			return A->ReferencedClusters.Num();
-		});
+			FString ObjectsList = Arg.Mid(5);
+			TArray<FString> ObjectNames;
+			ObjectsList.ParseIntoArray(ObjectNames, TEXT(","));
+			for (const FString& ObjectName : ObjectNames)
+			{
+				UObject* Res = StaticFindObject(UObject::StaticClass(), ANY_PACKAGE, *ObjectName);
+				if (Res)
+				{
+					int32 ObjectIndex = GUObjectArray.ObjectToIndex(Res);
+					WithObjects.Add(ObjectIndex);
+				}
+				else
+				{
+					UE_LOG(LogObj, Warning, TEXT("ListClusters can't find object \"%s\""), *ObjectName);
+				}
+			}
+		}
 	}
 
+	int32 NumberOfClustersPrinted = 0;
+
 	for (FUObjectCluster* Cluster : AllClusters)
 	{
 		check(Cluster->RootIndex != INDEX_NONE);
@@ -351,9 +396,19 @@ void ListClusters(const TArray<FString>& Args)
 		MaxClusterSize = FMath::Max(MaxClusterSize, Cluster->Objects.Num());
 		TotalClusterObjects += Cluster->Objects.Num();
 
-		DumpClusterToLog(*Cluster, bHierarchy, false);
+		bool bListCluster = true;
+		if (WithObjects.Num())
+		{
+			bListCluster = DoesClusterContainObjects(*Cluster, WithObjects);
+		}
+		if (bListCluster)
+		{
+			DumpClusterToLog(*Cluster, bHierarchy, false);
+			NumberOfClustersPrinted++;
+		}
 	}
-	UE_LOG(LogObj, Display, TEXT("Number of clusters: %d"), AllClusters.Num());
+	UE_LOG(LogObj, Display, TEXT("Displayed %d clusters"), NumberOfClustersPrinted);
+	UE_LOG(LogObj, Display, TEXT("Total number of clusters: %d"), AllClusters.Num());
 	UE_LOG(LogObj, Display, TEXT("Maximum cluster size: %d"), MaxClusterSize);
 	UE_LOG(LogObj, Display, TEXT("Average cluster size: %d"), AllClusters.Num() ? (TotalClusterObjects / AllClusters.Num()) : 0);
 	UE_LOG(LogObj, Display, TEXT("Number of objects in GC clusters: %d"), TotalClusterObjects);
