It's now possible to filter clusters that get logged with verbose cluster logging enabled (UE_GCCLUSTER_VERBOSE_LOGGING=1) by objects within them by specifying -DumpClustersWithObjects=ObjectName1,ObjectName2 in the command line

diff --git a/Engine/Source/Runtime/CoreUObject/Private/UObject/UObjectClusters.cpp b/Engine/Source/Runtime/CoreUObject/Private/UObject/UObjectClusters.cpp
--- a/Engine/Source/Runtime/CoreUObject/Private/UObject/UObjectClusters.cpp
+++ b/Engine/Source/Runtime/CoreUObject/Private/UObject/UObjectClusters.cpp
@@ -19,6 +19,7 @@
 #include "UObject/LinkerLoad.h"
 #include "UObject/FastReferenceCollector.h"
 #include "UObject/Package.h"
+#include "Misc/CommandLine.h"
 
 int32 GCreateGCClusters = 1;
 static FAutoConsoleVariableRef CCreateGCClusters(
@@ -229,8 +230,77 @@ int32 FUObjectClusterContainer::GetMinClusterSize() const
 
 #if !UE_BUILD_SHIPPING
 
+static bool DoesClusterContainObjects(const FUObjectCluster& Cluster, const TArray<int32>& Objects)
+{
+	for (int32 ObjectIndex : Objects)
+	{
+		if (Cluster.RootIndex == ObjectIndex)
+		{
+			return true;
+		}
+		if (Cluster.Objects.Contains(ObjectIndex))
+		{
+			return true;
+		}
+		else if (Cluster.MutableObjects.Contains(ObjectIndex))
+		{
+			return true;
+		}
+	}
+	return false;
+}
+
+static void ParseObjectNameArrayForClusters(TArray<int32>& OutIndexArray, const TArray<FString>& InNameArray, bool bWarn = true)
+{
+	for (const FString& ObjectName : InNameArray)
+	{
+		UObject* Res = StaticFindObject(UObject::StaticClass(), ANY_PACKAGE, *ObjectName);
+		if (Res)
+		{
+			int32 ObjectIndex = GUObjectArray.ObjectToIndex(Res);
+			OutIndexArray.Add(ObjectIndex);
+		}
+		else
+		{
+			UE_CLOG(bWarn, LogObj, Warning, TEXT("ParseObjectNameArrayForClusters can't find object \"%s\""), *ObjectName);
+		}
+	}
+}
+
 void DumpClusterToLog(const FUObjectCluster& Cluster, bool bHierarchy, bool bIndexOnly)
 {
+#if UE_GCCLUSTER_VERBOSE_LOGGING
+	static struct FVerboseClusterLoggingSettings
+	{
+		TArray<FString> WithObjects;
+		FVerboseClusterLoggingSettings()
+		{
+			FString ObjectsList;
+			FParse::Value(FCommandLine::Get(), TEXT("DumpClustersWithObjects="), ObjectsList);
+			ObjectsList.ParseIntoArray(WithObjects, TEXT(","));
+		}
+		bool DoesClusterContainRequestedObjects(const FUObjectCluster& InCluster)
+		{
+			bool bContainsObjects = true;
+			if (WithObjects.Num())
+			{
+				// We need to process the object name list each time we check it against a cluster because
+				// the objects may get loaded in and out as we create new clusters
+				TArray<int32> ObjectIndices;
+				ParseObjectNameArrayForClusters(ObjectIndices, WithObjects, false);
+				// If none of the objects is currently loaded and ObjectIndices is empty we will properly reject the cluster
+				bContainsObjects = DoesClusterContainObjects(InCluster, ObjectIndices);
+			}
+			return bContainsObjects;
+		}
+	} VerboseClusterLoggingSettings;
+
+	if (!VerboseClusterLoggingSettings.DoesClusterContainRequestedObjects(Cluster))
+	{
+		return;
+	}
+#endif
+
 	FUObjectItem* RootItem = GUObjectArray.IndexToObjectUnsafeForGC(Cluster.RootIndex);
 	UObject* RootObject = static_cast<UObject*>(RootItem->Object);
 	UE_LOG(LogObj, Display, TEXT("%s (Index: %d), Size: %d, ReferencedClusters: %d"), *RootObject->GetFullName(), Cluster.RootIndex, Cluster.Objects.Num(), Cluster.ReferencedClusters.Num());
@@ -295,26 +365,6 @@ void DumpClusterToLog(const FUObjectCluster& Cluster, bool bHierarchy, bool bInd
 	}
 }
 
-static bool DoesClusterContainObjects(FUObjectCluster& Cluster, const TArray<int32>& Objects)
-{
-	for (int32 ObjectIndex : Objects)
-	{
-		if (Cluster.RootIndex == ObjectIndex)
-		{
-			return true;
-		}
-		if (Cluster.Objects.Contains(ObjectIndex))
-		{
-			return true;
-		}
-		else if (Cluster.MutableObjects.Contains(ObjectIndex))
-		{
-			return true;
-		}
-	}
-	return false;
-}
-
 // Dumps all clusters to log.
 void ListClusters(const TArray<FString>& Args)
 {
@@ -369,19 +419,7 @@ void ListClusters(const TArray<FString>& Args)
 			FString ObjectsList = Arg.Mid(5);
 			TArray<FString> ObjectNames;
 			ObjectsList.ParseIntoArray(ObjectNames, TEXT(","));
-			for (const FString& ObjectName : ObjectNames)
-			{
-				UObject* Res = StaticFindObject(UObject::StaticClass(), ANY_PACKAGE, *ObjectName);
-				if (Res)
-				{
-					int32 ObjectIndex = GUObjectArray.ObjectToIndex(Res);
-					WithObjects.Add(ObjectIndex);
-				}
-				else
-				{
-					UE_LOG(LogObj, Warning, TEXT("ListClusters can't find object \"%s\""), *ObjectName);
-				}
-			}
+			ParseObjectNameArrayForClusters(WithObjects, ObjectNames);
 		}
 	}
 
