When iterating with ForEachObjectWithouter, don't lock the entire has table but only the hash bucket that is currently being iterated

diff --git a/Engine/Source/Runtime/CoreUObject/Private/UObject/UObjectHash.cpp b/Engine/Source/Runtime/CoreUObject/Private/UObject/UObjectHash.cpp
--- a/Engine/Source/Runtime/CoreUObject/Private/UObject/UObjectHash.cpp
+++ b/Engine/Source/Runtime/CoreUObject/Private/UObject/UObjectHash.cpp
@@ -50,6 +50,19 @@ struct FHashBucket
 	* If the first one is not null, then it is a uobject ptr, and the second ptr is either null or a second element
 	*/
 	void *ElementsOrSetPtr[2];
+	/** If true this bucket is being iterated over and no Add or Remove operations are allowed */
+	int32 ReadOnlyLock;
+
+	FORCEINLINE void Lock()
+	{
+		ReadOnlyLock++;
+	}
+
+	FORCEINLINE void Unlock()
+	{
+		ReadOnlyLock--;
+		check(ReadOnlyLock >= 0);
+	}
 
 	FORCEINLINE TSet<UObjectBase*>* GetSet()
 	{
@@ -74,6 +87,7 @@ struct FHashBucket
 	{
 		ElementsOrSetPtr[0] = nullptr;
 		ElementsOrSetPtr[1] = nullptr;
+		ReadOnlyLock = 0;
 	}
 	FORCEINLINE ~FHashBucket()
 	{
@@ -82,6 +96,9 @@ struct FHashBucket
 	/** Adds an Object to the bucket */
 	FORCEINLINE void Add(UObjectBase* Object)
 	{
+		UE_CLOG(ReadOnlyLock != 0, LogObj, Fatal, TEXT("Trying to add %s to a hash bucket that is currently being iterated over which is not allowed and may lead to undefined behavior!"),
+			*static_cast<UObject*>(Object)->GetFullName());
+
 		TSet<UObjectBase*>* Items = GetSet();
 		if (Items)
 		{
@@ -109,6 +126,9 @@ struct FHashBucket
 	/** Removes an Object from the bucket */
 	FORCEINLINE int32 Remove(UObjectBase* Object)
 	{
+		UE_CLOG(ReadOnlyLock != 0, LogObj, Fatal, TEXT("Trying to remove %s from a hash bucket that is currently being iterated over which is not allowed and may lead to undefined behavior!"),
+			*static_cast<UObject*>(Object)->GetFullName());
+
 		int32 Result = 0;
 		TSet<UObjectBase*>* Items = GetSet();
 		if (Items)
@@ -242,8 +262,6 @@ class FUObjectHashTables
 {
 	/** Critical section that guards against concurrent adds from multiple threads */
 	FCriticalSection CriticalSection;
-	/** Hard lock flag, this is used when iterating over the hash table - we don't want any adds when this is happening */
-	int32 HardLock;
 
 public:
 
@@ -257,7 +275,6 @@ public:
 	TMap<UClass*, TSet<UClass*> > ClassToChildListMap;
 
 	FUObjectHashTables()
-		: HardLock(0)
 	{
 	}
 
@@ -275,10 +292,6 @@ public:
 	/** Adds the Hash/Object pair to the FName hash table */
 	FORCEINLINE void AddToHash(int32 InHash, UObjectBase* Object)
 	{
-		UE_CLOG(HardLock > 0, LogUObjectHash, Fatal,
-			TEXT("Trying to add new object %s to UObject hashtables when they are in hard lockdown mode. Are you trying to create a new object in ForEachObjectWithOuter?"),
-			*GetFullNameSafe(static_cast<UObjectBaseUtility*>(Object)));
-
 		FHashBucket& Bucket = Hash.FindOrAdd(InHash);
 		Bucket.Add(Object);
 	}
@@ -298,22 +311,13 @@ public:
 		return NumRemoved;
 	}
 
-	FORCEINLINE void Lock(bool bHard = false)
+	FORCEINLINE void Lock()
 	{
 		CriticalSection.Lock();
-		if (bHard)
-		{
-			HardLock++;
-		}
 	}
 
-	FORCEINLINE void Unlock(bool bHard = false)
+	FORCEINLINE void Unlock()
 	{
-		if (bHard)
-		{
-			HardLock--;
-			check(HardLock >= 0);
-		}
 		CriticalSection.Unlock();
 	}
 
@@ -328,17 +332,15 @@ class FHashTableLock
 {
 #if THREADSAFE_UOBJECTS
 	FUObjectHashTables* Tables;
-	bool bHardLock;
 #endif
 public:
-	FORCEINLINE FHashTableLock(FUObjectHashTables& InTables, bool bHard = false)
+	FORCEINLINE FHashTableLock(FUObjectHashTables& InTables)
 	{
 #if THREADSAFE_UOBJECTS
-		bHardLock = bHard;
 		if (!(IsGarbageCollecting() && IsInGameThread()))
 		{
 			Tables = &InTables;
-			InTables.Lock(bHardLock);
+			InTables.Lock();
 		}
 		else
 		{
@@ -353,7 +355,7 @@ public:
 #if THREADSAFE_UOBJECTS
 		if (Tables)
 		{
-			Tables->Unlock(bHardLock);
+			Tables->Unlock();
 		}
 #endif
 	}
@@ -710,7 +712,7 @@ void ForEachObjectWithOuter(const class UObjectBase* Outer, TFunctionRef<void(UO
 		ExclusionInternalFlags |= EInternalObjectFlags::AsyncLoading;
 	}
 	FUObjectHashTables& ThreadHash = FUObjectHashTables::Get();
-	FHashTableLock HashLock(ThreadHash, true);
+	FHashTableLock HashLock(ThreadHash);
 	TArray<FHashBucket*, TInlineAllocator<1> > AllInners;
 
 	if (FHashBucket* Inners = ThreadHash.ObjectOuterMap.Find(Outer))
@@ -720,6 +722,7 @@ void ForEachObjectWithOuter(const class UObjectBase* Outer, TFunctionRef<void(UO
 	while (AllInners.Num())
 	{
 		FHashBucket* Inners = AllInners.Pop();
+		Inners->Lock();
 		for (FHashBucketIterator It(*Inners); It; ++It)
 		{
 			UObject *Object = static_cast<UObject*>(*It);
@@ -735,6 +738,7 @@ void ForEachObjectWithOuter(const class UObjectBase* Outer, TFunctionRef<void(UO
 				}
 			}
 		}
+		Inners->Unlock();
 	}
 }
 
