Changed package probing logic so it probes for binary formats first, text second

This improves performance for the current common case by reducing needless syscalls

diff --git a/Engine/Source/Runtime/CoreUObject/Private/Misc/PackageName.cpp b/Engine/Source/Runtime/CoreUObject/Private/Misc/PackageName.cpp
--- a/Engine/Source/Runtime/CoreUObject/Private/Misc/PackageName.cpp
+++ b/Engine/Source/Runtime/CoreUObject/Private/Misc/PackageName.cpp
@@ -744,29 +744,52 @@ bool FPackageName::FindPackageFileWithoutExtension(const FString& InPackageFilen
 {
 	auto& FileManager = IFileManager::Get();
 
-	static const FString* PackageExtensions[] =
 	{
-		&TextAssetPackageExtension,
-		&TextMapPackageExtension,
-		&AssetPackageExtension,
-		&MapPackageExtension
-	};
+		static const FString* PackageExtensions[] =
+		{
+			&AssetPackageExtension,
+			&MapPackageExtension
+		};
 
-	// Loop through all known extensions and check if the file exist.
-	const bool bDebuggingTextFormatEnable = true;
-	const int32 BaseIndex = (bDebuggingTextFormatEnable && InAllowTextFormats && WITH_TEXT_ARCHIVE_SUPPORT) ? 0 : 2;
-	for (int32 ExtensionIndex = BaseIndex; ExtensionIndex < ARRAY_COUNT(PackageExtensions); ++ExtensionIndex)
+		// Loop through all known extensions and check if the file exists
+
+		for (int32 ExtensionIndex = 0; ExtensionIndex < ARRAY_COUNT(PackageExtensions); ++ExtensionIndex)
+		{
+			FString   PackageFilename = InPackageFilename + *PackageExtensions[ExtensionIndex];
+			FDateTime Timestamp       = FileManager.GetTimeStamp(*PackageFilename);
+			if (Timestamp != FDateTime::MinValue())
+			{
+				// The package exists so exit. From now on InPackageFilename can be equal to OutFilename so
+				// don't attempt to use it anymore (case where &InPackageFilename == &OutFilename).
+				OutFilename = MoveTemp(PackageFilename);
+				return true;
+			}
+		}
+	}
+
+#if WITH_TEXT_ARCHIVE_SUPPORT
+	if (InAllowTextFormats)
 	{
-		FString   PackageFilename = InPackageFilename + *PackageExtensions[ExtensionIndex];
-		FDateTime Timestamp       = FileManager.GetTimeStamp(*PackageFilename);
-		if (Timestamp != FDateTime::MinValue())
+		static const FString* TextPackageExtensions[] =
 		{
-			// The package exists so exit. From now on InPackageFilename can be equal to OutFilename so
-			// don't attempt to use it anymore (case where &InPackageFilename == &OutFilename).
-			OutFilename = MoveTemp(PackageFilename);
-			return true;
+			&TextAssetPackageExtension,
+			&TextMapPackageExtension
+		};
+
+		for (int32 ExtensionIndex = 0; ExtensionIndex < ARRAY_COUNT(TextPackageExtensions); ++ExtensionIndex)
+		{
+			FString   PackageFilename = InPackageFilename + *TextPackageExtensions[ExtensionIndex];
+			FDateTime Timestamp		  = FileManager.GetTimeStamp(*PackageFilename);
+			if (Timestamp != FDateTime::MinValue())
+			{
+				// The package exists so exit. From now on InPackageFilename can be equal to OutFilename so
+				// don't attempt to use it anymore (case where &InPackageFilename == &OutFilename).
+				OutFilename = MoveTemp(PackageFilename);
+				return true;
+			}
 		}
 	}
+#endif
 
 	return false;
 }
@@ -818,8 +841,9 @@ bool FPackageName::DoesPackageExist(const FString& LongPackageName, const FGuid*
 
 	if (!FPackageName::TryConvertFilenameToLongPackageName(LongPackageName, PackageName))
 	{
-		verify(!FPackageName::IsValidLongPackageName(LongPackageName, true, &Reason));
+		verify(!FPackageName::IsValidLongPackageName(PackageName, true, &Reason));
 		UE_LOG(LogPackageName, Error, TEXT("Illegal call to DoesPackageExist: '%s' is not a standard unreal filename or a long path name. Reason: %s"), *LongPackageName, *Reason.ToString());
+		ensureMsgf(false, TEXT("Illegal call to DoesPackageExist: '%s' is not a standard unreal filename or a long path name. Reason: %s"), *LongPackageName, *Reason.ToString());
 		return false;
 	}
 	// Once we have the real Package Name, we can exit early if it's a script package - they exist only in memory.
