'Fix' for TMap comparison in PIE, because UObjectBaseProperty::Identical() always returns false.

diff --git a/Engine/Source/Runtime/CoreUObject/Private/UObject/PropertyMap.cpp b/Engine/Source/Runtime/CoreUObject/Private/UObject/PropertyMap.cpp
--- a/Engine/Source/Runtime/CoreUObject/Private/UObject/PropertyMap.cpp
+++ b/Engine/Source/Runtime/CoreUObject/Private/UObject/PropertyMap.cpp
@@ -76,12 +76,12 @@ namespace UE4MapProperty_Private
 
 			const uint8* PairA = MapHelperA.GetPairPtr(IndexA);
 			const uint8* PairB = MapHelperB.GetPairPtr(IndexB);
-			if (KeyProp->Identical(PairA, PairToCompare, PortFlags) && ValueProp->Identical(PairA + ValueOffset, PairToCompare + ValueOffset, PortFlags))
+			if (PairA == PairToCompare || (KeyProp->Identical(PairA, PairToCompare, PortFlags) && ValueProp->Identical(PairA + ValueOffset, PairToCompare + ValueOffset, PortFlags)))
 			{
 				++CountA;
 			}
 
-			if (KeyProp->Identical(PairB, PairToCompare, PortFlags) && ValueProp->Identical(PairB + ValueOffset, PairToCompare + ValueOffset, PortFlags))
+			if (PairB == PairToCompare || (KeyProp->Identical(PairB, PairToCompare, PortFlags) && ValueProp->Identical(PairB + ValueOffset, PairToCompare + ValueOffset, PortFlags)))
 			{
 				++CountB;
 			}
