Changing the UObjectArray chunk size to 64K (down from 65K which was a mistake)

diff --git a/Engine/Source/Runtime/CoreUObject/Public/UObject/UObjectArray.h b/Engine/Source/Runtime/CoreUObject/Public/UObject/UObjectArray.h
--- a/Engine/Source/Runtime/CoreUObject/Public/UObject/UObjectArray.h
+++ b/Engine/Source/Runtime/CoreUObject/Public/UObject/UObjectArray.h
@@ -322,7 +322,7 @@ class FChunkedFixedUObjectArray
 {
 	enum
 	{
-		NumElementsPerChunk = 65 * 1024,
+		NumElementsPerChunk = 64 * 1024,
 	};
 
 	/** Master table to chunks of pointers **/
