Disable delta serialization for duplication.

diff --git a/Engine/Source/Runtime/CoreUObject/Private/Serialization/DuplicateDataReader.cpp b/Engine/Source/Runtime/CoreUObject/Private/Serialization/DuplicateDataReader.cpp
--- a/Engine/Source/Runtime/CoreUObject/Private/Serialization/DuplicateDataReader.cpp
+++ b/Engine/Source/Runtime/CoreUObject/Private/Serialization/DuplicateDataReader.cpp
@@ -27,6 +27,7 @@ FDuplicateDataReader::FDuplicateDataReader( class FUObjectAnnotationSparse<FDupl
 {
 	this->SetIsLoading(true);
 	this->SetIsPersistent(true);
+	this->ArNoDelta = true;
 	ArPortFlags |= PPF_Duplicate | InPortFlags;
 
 #if USE_STABLE_LOCALIZATION_KEYS
diff --git a/Engine/Source/Runtime/CoreUObject/Private/Serialization/DuplicateDataWriter.cpp b/Engine/Source/Runtime/CoreUObject/Private/Serialization/DuplicateDataWriter.cpp
--- a/Engine/Source/Runtime/CoreUObject/Private/Serialization/DuplicateDataWriter.cpp
+++ b/Engine/Source/Runtime/CoreUObject/Private/Serialization/DuplicateDataWriter.cpp
@@ -33,6 +33,7 @@ FDuplicateDataWriter::FDuplicateDataWriter( FUObjectAnnotationSparse<FDuplicated
 {
 	this->SetIsSaving(true);
 	this->SetIsPersistent(true);
+	this->ArNoDelta = true;
 	ArAllowLazyLoading	= false;
 	ArPortFlags |= PPF_Duplicate | InPortFlags;
 
diff --git a/Engine/Source/Runtime/CoreUObject/Private/UObject/PropertyMap.cpp b/Engine/Source/Runtime/CoreUObject/Private/UObject/PropertyMap.cpp
--- a/Engine/Source/Runtime/CoreUObject/Private/UObject/PropertyMap.cpp
+++ b/Engine/Source/Runtime/CoreUObject/Private/UObject/PropertyMap.cpp
@@ -244,6 +244,12 @@ void UMapProperty::SerializeItem(FStructuredArchive::FSlot Slot, void* Value, co
 	FArchive& UnderlyingArchive = Slot.GetUnderlyingArchive();
 	FStructuredArchive::FRecord Record = Slot.EnterRecord();
 
+	// If we're doing delta serialization, act as if there are no defaults
+	if (!UnderlyingArchive.DoDelta())
+	{
+		Defaults = nullptr;
+	}
+
 	// Ar related calls in this function must be mirrored in UMapProperty::ConvertFromType
 	checkSlow(KeyProp);
 	checkSlow(ValueProp);
