Fix for a potential crash (reported on UDN) when flushing linker cashe when async loading

diff --git a/Engine/Source/Runtime/CoreUObject/Private/Serialization/AsyncLoading.cpp b/Engine/Source/Runtime/CoreUObject/Private/Serialization/AsyncLoading.cpp
--- a/Engine/Source/Runtime/CoreUObject/Private/Serialization/AsyncLoading.cpp
+++ b/Engine/Source/Runtime/CoreUObject/Private/Serialization/AsyncLoading.cpp
@@ -5442,10 +5442,13 @@ void FAsyncPackage::FlushObjectLinkerCache()
 {
 	for (UObject* Obj : PackageObjLoaded)
 	{
-		FLinkerLoad* ObjLinker = Obj->GetLinker();
-		if (ObjLinker)
+		if (Obj)
 		{
-			ObjLinker->FlushCache();
+			FLinkerLoad* ObjLinker = Obj->GetLinker();
+			if (ObjLinker)
+			{
+				ObjLinker->FlushCache();
+			}
 		}
 	}
 }
