The warning 'Calling StaticLoadObject during PostLoad of ...' will now also include the parameters passed to StaticLoadObject for easier debugging.

diff --git a/Engine/Source/Runtime/CoreUObject/Private/UObject/UObjectGlobals.cpp b/Engine/Source/Runtime/CoreUObject/Private/UObject/UObjectGlobals.cpp
--- a/Engine/Source/Runtime/CoreUObject/Private/UObject/UObjectGlobals.cpp
+++ b/Engine/Source/Runtime/CoreUObject/Private/UObject/UObjectGlobals.cpp
@@ -981,7 +981,11 @@ UObject* StaticLoadObject(UClass* ObjectClass, UObject* InOuter, const TCHAR* In
 	FUObjectThreadContext& ThreadContext = FUObjectThreadContext::Get();
 	if (ThreadContext.IsRoutingPostLoad && IsInAsyncLoadingThread())
 	{
-		UE_LOG(LogUObjectGlobals, Warning, TEXT("Calling StaticLoadObject during PostLoad of %s may result in hitches during streaming."), *GetFullNameSafe(ThreadContext.CurrentlyPostLoadedObjectByALT));
+		UE_LOG(LogUObjectGlobals, Warning, TEXT("Calling StaticLoadObject(\"%s\", \"%s\", \"%s\") during PostLoad of %s may result in hitches during streaming."), 
+			*GetFullNameSafe(ObjectClass),
+			*GetFullNameSafe(InOuter),
+			InName,
+			*GetFullNameSafe(ThreadContext.CurrentlyPostLoadedObjectByALT));
 	}
 
 	UObject* Result = StaticLoadObjectInternal(ObjectClass, InOuter, InName, Filename, LoadFlags, Sandbox, bAllowObjectReconciliation);
