Removed code that calculates and compares CRC of a package's name but does nothing with the result.

diff --git a/Engine/Source/Runtime/CoreUObject/Private/UObject/LinkerLoad.cpp b/Engine/Source/Runtime/CoreUObject/Private/UObject/LinkerLoad.cpp
--- a/Engine/Source/Runtime/CoreUObject/Private/UObject/LinkerLoad.cpp
+++ b/Engine/Source/Runtime/CoreUObject/Private/UObject/LinkerLoad.cpp
@@ -2038,12 +2038,6 @@ FLinkerLoad::ELinkerStatus FLinkerLoad::FinalizeCreation()
 		// Add this linker to the object manager's linker array.
 		FLinkerManager::Get().AddLoader(this);
 
-		// check if the package source matches the package filename's CRC (if it doesn't match, a user saved this package)
-		if (Summary.PackageSource != FCrc::StrCrc_DEPRECATED(*FPaths::GetBaseFilename(Filename).ToUpper()))
-		{
-//			UE_LOG(LogLinker, Log, TEXT("Found a user created package (%s)"), *(FPaths::GetBaseFilename(Filename)));
-		}
-
 		if (GEventDrivenLoaderEnabled && AsyncRoot)
 		{
 			for (int32 ExportIndex = 0; ExportIndex < ExportMap.Num(); ++ExportIndex)
