Disabled the checks against adding or removing objects from the UObject hash tables when iterating over them in shipping builds to reduce their memory overhead

diff --git a/Engine/Source/Runtime/CoreUObject/Private/UObject/UObjectHash.cpp b/Engine/Source/Runtime/CoreUObject/Private/UObject/UObjectHash.cpp
--- a/Engine/Source/Runtime/CoreUObject/Private/UObject/UObjectHash.cpp
+++ b/Engine/Source/Runtime/CoreUObject/Private/UObject/UObjectHash.cpp
@@ -50,6 +50,8 @@ struct FHashBucket
 	* If the first one is not null, then it is a uobject ptr, and the second ptr is either null or a second element
 	*/
 	void *ElementsOrSetPtr[2];
+
+#if !UE_BUILD_SHIPPING
 	/** If true this bucket is being iterated over and no Add or Remove operations are allowed */
 	int32 ReadOnlyLock;
 
@@ -63,6 +65,7 @@ struct FHashBucket
 		ReadOnlyLock--;
 		check(ReadOnlyLock >= 0);
 	}
+#endif // !UE_BUILD_SHIPPING
 
 	FORCEINLINE TSet<UObjectBase*>* GetSet()
 	{
@@ -87,7 +90,9 @@ struct FHashBucket
 	{
 		ElementsOrSetPtr[0] = nullptr;
 		ElementsOrSetPtr[1] = nullptr;
+#if !UE_BUILD_SHIPPING
 		ReadOnlyLock = 0;
+#endif // !UE_BUILD_SHIPPING
 	}
 	FORCEINLINE ~FHashBucket()
 	{
@@ -96,8 +101,10 @@ struct FHashBucket
 	/** Adds an Object to the bucket */
 	FORCEINLINE void Add(UObjectBase* Object)
 	{
+#if !UE_BUILD_SHIPPING
 		UE_CLOG(ReadOnlyLock != 0, LogObj, Fatal, TEXT("Trying to add %s to a hash bucket that is currently being iterated over which is not allowed and may lead to undefined behavior!"),
 			*static_cast<UObject*>(Object)->GetFullName());
+#endif // !UE_BUILD_SHIPPING
 
 		TSet<UObjectBase*>* Items = GetSet();
 		if (Items)
@@ -126,8 +133,10 @@ struct FHashBucket
 	/** Removes an Object from the bucket */
 	FORCEINLINE int32 Remove(UObjectBase* Object)
 	{
+#if !UE_BUILD_SHIPPING
 		UE_CLOG(ReadOnlyLock != 0, LogObj, Fatal, TEXT("Trying to remove %s from a hash bucket that is currently being iterated over which is not allowed and may lead to undefined behavior!"),
 			*static_cast<UObject*>(Object)->GetFullName());
+#endif // !UE_BUILD_SHIPPING
 
 		int32 Result = 0;
 		TSet<UObjectBase*>* Items = GetSet();
@@ -722,7 +731,9 @@ void ForEachObjectWithOuter(const class UObjectBase* Outer, TFunctionRef<void(UO
 	while (AllInners.Num())
 	{
 		FHashBucket* Inners = AllInners.Pop();
+#if !UE_BUILD_SHIPPING
 		Inners->Lock();
+#endif // !UE_BUILD_SHIPPING
 		for (FHashBucketIterator It(*Inners); It; ++It)
 		{
 			UObject *Object = static_cast<UObject*>(*It);
@@ -738,7 +749,9 @@ void ForEachObjectWithOuter(const class UObjectBase* Outer, TFunctionRef<void(UO
 				}
 			}
 		}
+#if !UE_BUILD_SHIPPING
 		Inners->Unlock();
+#endif // !UE_BUILD_SHIPPING
 	}
 }
 
