Removed ArchiveAsync.h since it didn't contain anything useful anymore

diff --git a/Engine/Source/Runtime/AssetRegistry/Private/PackageReader.cpp b/Engine/Source/Runtime/AssetRegistry/Private/PackageReader.cpp
--- a/Engine/Source/Runtime/AssetRegistry/Private/PackageReader.cpp
+++ b/Engine/Source/Runtime/AssetRegistry/Private/PackageReader.cpp
@@ -3,7 +3,6 @@
 #include "PackageReader.h"
 #include "HAL/FileManager.h"
 #include "UObject/Class.h"
-#include "Serialization/ArchiveAsync.h"
 #include "Misc/PackageName.h"
 #include "AssetRegistryPrivate.h"
 #include "AssetData.h"
diff --git a/Engine/Source/Runtime/CoreUObject/Private/Misc/WorldCompositionUtility.cpp b/Engine/Source/Runtime/CoreUObject/Private/Misc/WorldCompositionUtility.cpp
--- a/Engine/Source/Runtime/CoreUObject/Private/Misc/WorldCompositionUtility.cpp
+++ b/Engine/Source/Runtime/CoreUObject/Private/Misc/WorldCompositionUtility.cpp
@@ -7,7 +7,6 @@
 #include "HAL/FileManager.h"
 #include "Templates/ScopedPointer.h"
 #include "UObject/PropertyPortFlags.h"
-#include "Serialization/ArchiveAsync.h"
 #include "UObject/PackageFileSummary.h"
 #include "UObject/Linker.h"
 #include "Templates/UniquePtr.h"
diff --git a/Engine/Source/Runtime/CoreUObject/Private/Serialization/AsyncLoading.cpp b/Engine/Source/Runtime/CoreUObject/Private/Serialization/AsyncLoading.cpp
--- a/Engine/Source/Runtime/CoreUObject/Private/Serialization/AsyncLoading.cpp
+++ b/Engine/Source/Runtime/CoreUObject/Private/Serialization/AsyncLoading.cpp
@@ -16,7 +16,6 @@
 #include "Misc/CoreDelegates.h"
 #include "Misc/CommandLine.h"
 #include "Misc/App.h"
-#include "Serialization/ArchiveAsync.h"
 #include "Misc/PackageName.h"
 #include "UObject/PackageFileSummary.h"
 #include "UObject/Linker.h"
diff --git a/Engine/Source/Runtime/CoreUObject/Private/UObject/LinkerLoad.cpp b/Engine/Source/Runtime/CoreUObject/Private/UObject/LinkerLoad.cpp
--- a/Engine/Source/Runtime/CoreUObject/Private/UObject/LinkerLoad.cpp
+++ b/Engine/Source/Runtime/CoreUObject/Private/UObject/LinkerLoad.cpp
@@ -14,7 +14,6 @@
 #include "UObject/Package.h"
 #include "UObject/MetaData.h"
 #include "UObject/ObjectRedirector.h"
-#include "Serialization/ArchiveAsync.h"
 #include "Misc/PackageName.h"
 #include "Blueprint/BlueprintSupport.h"
 #include "Misc/SecureHash.h"
diff --git a/Engine/Source/Runtime/CoreUObject/Public/CoreUObject.h b/Engine/Source/Runtime/CoreUObject/Public/CoreUObject.h
--- a/Engine/Source/Runtime/CoreUObject/Public/CoreUObject.h
+++ b/Engine/Source/Runtime/CoreUObject/Public/CoreUObject.h
@@ -72,7 +72,6 @@ MONOLITHIC_HEADER_BOILERPLATE()
 #include "Serialization/ArchiveReplaceOrClearExternalReferences.h"
 #include "Serialization/ArchiveObjectPropertyMapper.h"
 #include "Serialization/ArchiveReferenceMarker.h"
-#include "Serialization/ArchiveAsync.h"
 #include "Serialization/ArchiveObjectCrc32.h"
 #include "Misc/PackageName.h"
 #include "UObject/ConstructorHelpers.h"
diff --git a/Engine/Source/Runtime/CoreUObject/Public/Serialization/ArchiveAsync.h b/Engine/Source/Runtime/CoreUObject/Public/Serialization/ArchiveAsync.h
deleted file mode 100644
