PR #5161: TaggedBinaryArchiveOutputFormatter Typo uint16->uint32 (Contributed by [REDACTED_PERSON])

diff --git a/Engine/Source/Runtime/CoreUObject/Private/Serialization/Formatters/TaggedBinaryArchiveOutputFormatter.cpp b/Engine/Source/Runtime/CoreUObject/Private/Serialization/Formatters/TaggedBinaryArchiveOutputFormatter.cpp
--- a/Engine/Source/Runtime/CoreUObject/Private/Serialization/Formatters/TaggedBinaryArchiveOutputFormatter.cpp
+++ b/Engine/Source/Runtime/CoreUObject/Private/Serialization/Formatters/TaggedBinaryArchiveOutputFormatter.cpp
@@ -371,7 +371,7 @@ void FTaggedBinaryArchiveOutputFormatter::WriteSize(uint64 Size)
 		uint8 FirstByte = 254;
 		Inner << FirstByte;
 
-		uint16 Size32 = Size;
+		uint32 Size32 = Size;
 		Inner << Size32;
 	}
 	else
