Moved UObject serialization state out of FUObjectThreadContext to a local-only (to the current FAsyncPackage/LoadPackage call) FUObjectSerializeContext. This will allow more advanced changes to the old serialization paths and prevent objects from being added to random FAsyncPackages when async loading (and that will allow package prioritization).

diff --git a/Engine/Plugins/Experimental/ControlRig/Source/ControlRigEditor/Private/ControlRigBlueprintCompiler.cpp b/Engine/Plugins/Experimental/ControlRig/Source/ControlRigEditor/Private/ControlRigBlueprintCompiler.cpp
--- a/Engine/Plugins/Experimental/ControlRig/Source/ControlRigEditor/Private/ControlRigBlueprintCompiler.cpp
+++ b/Engine/Plugins/Experimental/ControlRig/Source/ControlRigEditor/Private/ControlRigBlueprintCompiler.cpp
@@ -22,9 +22,9 @@ bool FControlRigBlueprintCompiler::CanCompile(const UBlueprint* Blueprint)
 	return false;
 }
 
-void FControlRigBlueprintCompiler::Compile(UBlueprint* Blueprint, const FKismetCompilerOptions& CompileOptions, FCompilerResultsLog& Results, TArray<UObject*>* ObjLoaded)
+void FControlRigBlueprintCompiler::Compile(UBlueprint* Blueprint, const FKismetCompilerOptions& CompileOptions, FCompilerResultsLog& Results)
 {
-	FControlRigBlueprintCompilerContext Compiler(Blueprint, Results, CompileOptions, ObjLoaded);
+	FControlRigBlueprintCompilerContext Compiler(Blueprint, Results, CompileOptions);
 	Compiler.Compile();
 }
 
diff --git a/Engine/Plugins/Experimental/ControlRig/Source/ControlRigEditor/Private/ControlRigBlueprintCompiler.h b/Engine/Plugins/Experimental/ControlRig/Source/ControlRigEditor/Private/ControlRigBlueprintCompiler.h
--- a/Engine/Plugins/Experimental/ControlRig/Source/ControlRigEditor/Private/ControlRigBlueprintCompiler.h
+++ b/Engine/Plugins/Experimental/ControlRig/Source/ControlRigEditor/Private/ControlRigBlueprintCompiler.h
@@ -13,14 +13,14 @@ class FControlRigBlueprintCompiler : public IBlueprintCompiler
 public:
 	/** IBlueprintCompiler interface */
 	virtual bool CanCompile(const UBlueprint* Blueprint) override;
-	virtual void Compile(UBlueprint* Blueprint, const FKismetCompilerOptions& CompileOptions, FCompilerResultsLog& Results, TArray<UObject*>* ObjLoaded) override;
+	virtual void Compile(UBlueprint* Blueprint, const FKismetCompilerOptions& CompileOptions, FCompilerResultsLog& Results) override;
 };
 
 class FControlRigBlueprintCompilerContext : public FKismetCompilerContext
 {
 public:
-	FControlRigBlueprintCompilerContext(UBlueprint* SourceSketch, FCompilerResultsLog& InMessageLog, const FKismetCompilerOptions& InCompilerOptions, TArray<UObject*>* InObjLoaded)
-		: FKismetCompilerContext(SourceSketch, InMessageLog, InCompilerOptions, InObjLoaded)
+	FControlRigBlueprintCompilerContext(UBlueprint* SourceSketch, FCompilerResultsLog& InMessageLog, const FKismetCompilerOptions& InCompilerOptions)
+		: FKismetCompilerContext(SourceSketch, InMessageLog, InCompilerOptions)
 		, NewControlRigBlueprintGeneratedClass(nullptr)
 	{
 	}
diff --git a/Engine/Plugins/Experimental/ControlRig/Source/ControlRigEditor/Private/ControlRigEditorModule.cpp b/Engine/Plugins/Experimental/ControlRig/Source/ControlRigEditor/Private/ControlRigEditorModule.cpp
--- a/Engine/Plugins/Experimental/ControlRig/Source/ControlRigEditor/Private/ControlRigEditorModule.cpp
+++ b/Engine/Plugins/Experimental/ControlRig/Source/ControlRigEditor/Private/ControlRigEditorModule.cpp
@@ -703,7 +703,7 @@ TSharedRef<IControlRigEditor> FControlRigEditorModule::CreateControlRigEditor(co
 
 TSharedPtr<FKismetCompilerContext> FControlRigEditorModule::GetControlRigCompiler(UBlueprint* BP, FCompilerResultsLog& InMessageLog, const FKismetCompilerOptions& InCompileOptions)
 {
-	return TSharedPtr<FKismetCompilerContext>(new FControlRigBlueprintCompilerContext(BP, InMessageLog, InCompileOptions, nullptr));
+	return TSharedPtr<FKismetCompilerContext>(new FControlRigBlueprintCompilerContext(BP, InMessageLog, InCompileOptions));
 }
 
 void FControlRigEditorModule::RegisterRigUnitEditorClass(FName RigUnitClassName, TSubclassOf<URigUnitEditor_Base> InClass)
diff --git a/Engine/Plugins/Runtime/GameplayAbilities/Source/GameplayAbilities/Private/AttributeSet.cpp b/Engine/Plugins/Runtime/GameplayAbilities/Source/GameplayAbilities/Private/AttributeSet.cpp
--- a/Engine/Plugins/Runtime/GameplayAbilities/Source/GameplayAbilities/Private/AttributeSet.cpp
+++ b/Engine/Plugins/Runtime/GameplayAbilities/Source/GameplayAbilities/Private/AttributeSet.cpp
@@ -221,8 +221,8 @@ void FGameplayAttribute::PostSerialize(const FArchive& Ar)
 
 			if (!Attribute)
 			{
-				FUObjectThreadContext& ThreadContext = FUObjectThreadContext::Get();
-				FString AssetName = ThreadContext.SerializedObject ? ThreadContext.SerializedObject->GetPathName() : TEXT("Unknown Object");
+				FUObjectSerializeContext* LoadContext = const_cast<FArchive*>(&Ar)->GetSerializeContext();
+				FString AssetName = (LoadContext && LoadContext->SerializedObject) ? LoadContext->SerializedObject->GetPathName() : TEXT("Unknown Object");
 
 				FString OwnerName = AttributeOwner ? AttributeOwner->GetName() : TEXT("NONE");
 				ABILITY_LOG(Warning, TEXT("FGameplayAttribute::PostSerialize called on an invalid attribute with owner %s and name %s. (Asset: %s)"), *OwnerName, *AttributeName, *AssetName);
diff --git a/Engine/Plugins/ScriptPlugin/Source/ScriptEditorPlugin/Private/ScriptBlueprintCompiler.cpp b/Engine/Plugins/ScriptPlugin/Source/ScriptEditorPlugin/Private/ScriptBlueprintCompiler.cpp
--- a/Engine/Plugins/ScriptPlugin/Source/ScriptEditorPlugin/Private/ScriptBlueprintCompiler.cpp
+++ b/Engine/Plugins/ScriptPlugin/Source/ScriptEditorPlugin/Private/ScriptBlueprintCompiler.cpp
@@ -15,8 +15,8 @@
 
 ///-------------------------------------------------------------
 
-FScriptBlueprintCompiler::FScriptBlueprintCompiler(UScriptBlueprint* SourceSketch, FCompilerResultsLog& InMessageLog, const FKismetCompilerOptions& InCompilerOptions, TArray<UObject*>* InObjLoaded)
-	: Super(SourceSketch, InMessageLog, InCompilerOptions, InObjLoaded)
+FScriptBlueprintCompiler::FScriptBlueprintCompiler(UScriptBlueprint* SourceSketch, FCompilerResultsLog& InMessageLog, const FKismetCompilerOptions& InCompilerOptions)
+	: Super(SourceSketch, InMessageLog, InCompilerOptions)
 	, NewScriptBlueprintClass(NULL)
 	, ContextProperty(NULL)
 {
diff --git a/Engine/Plugins/ScriptPlugin/Source/ScriptEditorPlugin/Private/ScriptBlueprintCompiler.h b/Engine/Plugins/ScriptPlugin/Source/ScriptEditorPlugin/Private/ScriptBlueprintCompiler.h
--- a/Engine/Plugins/ScriptPlugin/Source/ScriptEditorPlugin/Private/ScriptBlueprintCompiler.h
+++ b/Engine/Plugins/ScriptPlugin/Source/ScriptEditorPlugin/Private/ScriptBlueprintCompiler.h
@@ -18,7 +18,7 @@ protected:
 	typedef FKismetCompilerContext Super;
 
 public:
-	FScriptBlueprintCompiler(UScriptBlueprint* SourceSketch, FCompilerResultsLog& InMessageLog, const FKismetCompilerOptions& InCompilerOptions, TArray<UObject*>* InObjLoaded);
+	FScriptBlueprintCompiler(UScriptBlueprint* SourceSketch, FCompilerResultsLog& InMessageLog, const FKismetCompilerOptions& InCompilerOptions);
 	virtual ~FScriptBlueprintCompiler();
 
 protected:
diff --git a/Engine/Plugins/ScriptPlugin/Source/ScriptEditorPlugin/Private/ScriptEditorPlugin.cpp b/Engine/Plugins/ScriptPlugin/Source/ScriptEditorPlugin/Private/ScriptEditorPlugin.cpp
--- a/Engine/Plugins/ScriptPlugin/Source/ScriptEditorPlugin/Private/ScriptEditorPlugin.cpp
+++ b/Engine/Plugins/ScriptPlugin/Source/ScriptEditorPlugin/Private/ScriptEditorPlugin.cpp
@@ -60,7 +60,7 @@ public:
 	virtual bool CanCompile(const UBlueprint* Blueprint) override;
 
 	virtual void PreCompile(UBlueprint* Blueprint) override;
-	virtual void Compile(UBlueprint* Blueprint, const FKismetCompilerOptions& CompileOptions, FCompilerResultsLog& Results, TArray<UObject*>* ObjLoaded) override;
+	virtual void Compile(UBlueprint* Blueprint, const FKismetCompilerOptions& CompileOptions, FCompilerResultsLog& Results) override;
 	virtual void PostCompile(UBlueprint* Blueprint) override;
 };
 
@@ -90,11 +90,11 @@ void FScriptEditorPlugin::PreCompile(UBlueprint* Blueprint)
 
 }
 
-void FScriptEditorPlugin::Compile(UBlueprint* Blueprint, const FKismetCompilerOptions& CompileOptions, FCompilerResultsLog& Results, TArray<UObject*>* ObjLoaded)
+void FScriptEditorPlugin::Compile(UBlueprint* Blueprint, const FKismetCompilerOptions& CompileOptions, FCompilerResultsLog& Results)
 {
 	if ( UScriptBlueprint* ScriptBlueprint = Cast<UScriptBlueprint>(Blueprint) )
 	{
-		FScriptBlueprintCompiler Compiler(ScriptBlueprint, Results, CompileOptions, ObjLoaded);
+		FScriptBlueprintCompiler Compiler(ScriptBlueprint, Results, CompileOptions);
 		Compiler.Compile();
 		check(Compiler.NewClass);
 	}
diff --git a/Engine/Source/Editor/Kismet/Private/BlueprintCompilationManager.cpp b/Engine/Source/Editor/Kismet/Private/BlueprintCompilationManager.cpp
--- a/Engine/Source/Editor/Kismet/Private/BlueprintCompilationManager.cpp
+++ b/Engine/Source/Editor/Kismet/Private/BlueprintCompilationManager.cpp
@@ -89,12 +89,12 @@ struct FBlueprintCompilationManagerImpl : public FGCObject
 
 	void QueueForCompilation(const FBPCompileRequest& CompileJob);
 	void CompileSynchronouslyImpl(const FBPCompileRequest& Request);
-	void FlushCompilationQueueImpl(TArray<UObject*>* ObjLoaded, bool bSuppressBroadcastCompiled, TArray<UBlueprint*>* BlueprintsCompiled, TArray<UBlueprint*>* BlueprintsCompiledOrSkeletonCompiled);
+	void FlushCompilationQueueImpl(bool bSuppressBroadcastCompiled, TArray<UBlueprint*>* BlueprintsCompiled, TArray<UBlueprint*>* BlueprintsCompiledOrSkeletonCompiled);
 	void FlushReinstancingQueueImpl();
 	bool HasBlueprintsToCompile() const;
 	bool IsGeneratedClassLayoutReady() const;
 	void GetDefaultValue(const UClass* ForClass, const UProperty* Property, FString& OutDefaultValueAsString) const;
-	static void ReinstanceBatch(TArray<FReinstancingJob>& Reinstancers, TMap< UClass*, UClass* >& InOutOldToNewClassMap, TArray<UObject*>* ObjLoaded);
+	static void ReinstanceBatch(TArray<FReinstancingJob>& Reinstancers, TMap< UClass*, UClass* >& InOutOldToNewClassMap);
 	static UClass* FastGenerateSkeletonClass(UBlueprint* BP, FKismetCompilerContext& CompilerContext, bool bIsSkeletonOnly);
 	static bool IsQueuedForCompilation(UBlueprint* BP);
 	
@@ -198,7 +198,7 @@ void FBlueprintCompilationManagerImpl::CompileSynchronouslyImpl(const FBPCompile
 	const bool bSuppressBroadcastCompiled = true;
 	TArray<UBlueprint*> CompiledBlueprints;
 	TArray<UBlueprint*> SkeletonCompiledBlueprints;
-	FlushCompilationQueueImpl(nullptr, bSuppressBroadcastCompiled, &CompiledBlueprints, &SkeletonCompiledBlueprints);
+	FlushCompilationQueueImpl(bSuppressBroadcastCompiled, &CompiledBlueprints, &SkeletonCompiledBlueprints);
 	FlushReinstancingQueueImpl();
 	
 	if (FBlueprintEditorUtils::IsLevelScriptBlueprint(Request.BPToCompile) && !bRegenerateSkeletonOnly)
@@ -354,7 +354,7 @@ struct FReinstancingJob
 	TSharedPtr<FKismetCompilerContext> Compiler;
 };
 
-void FBlueprintCompilationManagerImpl::FlushCompilationQueueImpl(TArray<UObject*>* ObjLoaded, bool bSuppressBroadcastCompiled, TArray<UBlueprint*>* BlueprintsCompiled, TArray<UBlueprint*>* BlueprintsCompiledOrSkeletonCompiled)
+void FBlueprintCompilationManagerImpl::FlushCompilationQueueImpl(bool bSuppressBroadcastCompiled, TArray<UBlueprint*>* BlueprintsCompiled, TArray<UBlueprint*>* BlueprintsCompiledOrSkeletonCompiled)
 {
 	TGuardValue<bool> GuardTemplateNameFlag(GCompilingBlueprint, true);
 	ensure(bGeneratedClassLayoutReady);
@@ -1087,7 +1087,7 @@ void FBlueprintCompilationManagerImpl::FlushCompilationQueueImpl(TArray<UObject*
 			}
 
 			FScopedDurationTimer ReinstTimer(GTimeReinstancing);
-			ReinstanceBatch(Reinstancers, ClassesToReinstance, ObjLoaded);
+			ReinstanceBatch(Reinstancers, ClassesToReinstance);
 
 			// We purposefully do not remove the OldCDOs yet, need to keep them in memory past first GC
 		}
@@ -1325,7 +1325,7 @@ void FBlueprintCompilationManagerImpl::GetDefaultValue(const UClass* ForClass, c
 	}
 }
 
-void FBlueprintCompilationManagerImpl::ReinstanceBatch(TArray<FReinstancingJob>& Reinstancers, TMap< UClass*, UClass* >& InOutOldToNewClassMap, TArray<UObject*>* ObjLoaded)
+void FBlueprintCompilationManagerImpl::ReinstanceBatch(TArray<FReinstancingJob>& Reinstancers, TMap< UClass*, UClass* >& InOutOldToNewClassMap)
 {
 	const auto FilterOutOfDateClasses = [](TArray<UClass*>& ClassList)
 	{
@@ -1508,7 +1508,7 @@ void FBlueprintCompilationManagerImpl::ReinstanceBatch(TArray<FReinstancingJob>&
 			BPGClass->UpdateCustomPropertyListForPostConstruction();
 
 			// patch new cdo into linker table:
-			if(OldCDO && ObjLoaded)
+			if(OldCDO)
 			{
 				UBlueprint* CurrentBP = CastChecked<UBlueprint>(BPGClass->ClassGeneratedBy);
 				if(FLinkerLoad* CurrentLinker = CurrentBP->GetLinker())
@@ -1527,7 +1527,7 @@ void FBlueprintCompilationManagerImpl::ReinstanceBatch(TArray<FReinstancingJob>&
 
 					if(OldCDOIndex != INDEX_NONE)
 					{
-						FBlueprintEditorUtils::PatchNewCDOIntoLinker(CurrentBP->GeneratedClass->ClassDefaultObject, CurrentLinker, OldCDOIndex, *ObjLoaded);
+						FBlueprintEditorUtils::PatchNewCDOIntoLinker(CurrentBP->GeneratedClass->ClassDefaultObject, CurrentLinker, OldCDOIndex);
 						FBlueprintEditorUtils::PatchCDOSubobjectsIntoExport(OldCDO, CurrentBP->GeneratedClass->ClassDefaultObject);
 					}
 				}
@@ -2348,11 +2348,11 @@ void FBlueprintCompilationManager::Shutdown()
 }
 
 // Forward to impl:
-void FBlueprintCompilationManager::FlushCompilationQueue(TArray<UObject*>* ObjLoaded)
+void FBlueprintCompilationManager::FlushCompilationQueue()
 {
 	if(BPCMImpl)
 	{
-		BPCMImpl->FlushCompilationQueueImpl(ObjLoaded, false, nullptr, nullptr);
+		BPCMImpl->FlushCompilationQueueImpl(false, nullptr, nullptr);
 
 		// We can't support save on compile or keeping old CDOs from GCing when reinstancing is deferred:
 		BPCMImpl->CompiledBlueprintsToSave.Empty();
@@ -2364,7 +2364,7 @@ void FBlueprintCompilationManager::FlushCompilationQueueAndReinstance()
 {
 	if(BPCMImpl)
 	{
-		BPCMImpl->FlushCompilationQueueImpl(nullptr, false, nullptr, nullptr);
+		BPCMImpl->FlushCompilationQueueImpl(false, nullptr, nullptr);
 		BPCMImpl->FlushReinstancingQueueImpl();
 
 		BPCMImpl->OldCDOs.Empty();
diff --git a/Engine/Source/Editor/Kismet/Public/BlueprintCompilationManager.h b/Engine/Source/Editor/Kismet/Public/BlueprintCompilationManager.h
--- a/Engine/Source/Editor/Kismet/Public/BlueprintCompilationManager.h
+++ b/Engine/Source/Editor/Kismet/Public/BlueprintCompilationManager.h
@@ -41,7 +41,7 @@ struct KISMET_API FBlueprintCompilationManager
 	 * of objects the compilation manager has loaded. The compilation manager
 	 * will not load data while processing the compilation queue)
 	 */
-	static void FlushCompilationQueue(TArray<UObject*>* ObjLoaded = nullptr);
+	static void FlushCompilationQueue();
 	
 	/**
 	 * Flushes the compilation queue and finishes reinstancing
diff --git a/Engine/Source/Editor/KismetCompiler/Private/AnimBlueprintCompiler.cpp b/Engine/Source/Editor/KismetCompiler/Private/AnimBlueprintCompiler.cpp
--- a/Engine/Source/Editor/KismetCompiler/Private/AnimBlueprintCompiler.cpp
+++ b/Engine/Source/Editor/KismetCompiler/Private/AnimBlueprintCompiler.cpp
@@ -102,8 +102,8 @@ bool FAnimBlueprintCompilerContext::FEffectiveConstantRecord::Apply(UObject* Obj
 //////////////////////////////////////////////////////////////////////////
 // FAnimBlueprintCompiler
 
-FAnimBlueprintCompilerContext::FAnimBlueprintCompilerContext(UAnimBlueprint* SourceSketch, FCompilerResultsLog& InMessageLog, const FKismetCompilerOptions& InCompileOptions, TArray<UObject*>* InObjLoaded)
-	: FKismetCompilerContext(SourceSketch, InMessageLog, InCompileOptions, InObjLoaded)
+FAnimBlueprintCompilerContext::FAnimBlueprintCompilerContext(UAnimBlueprint* SourceSketch, FCompilerResultsLog& InMessageLog, const FKismetCompilerOptions& InCompileOptions)
+	: FKismetCompilerContext(SourceSketch, InMessageLog, InCompileOptions)
 	, AnimBlueprint(SourceSketch)
 	, bIsDerivedAnimBlueprint(false)
 {
diff --git a/Engine/Source/Editor/KismetCompiler/Private/AnimBlueprintCompiler.h b/Engine/Source/Editor/KismetCompiler/Private/AnimBlueprintCompiler.h
--- a/Engine/Source/Editor/KismetCompiler/Private/AnimBlueprintCompiler.h
+++ b/Engine/Source/Editor/KismetCompiler/Private/AnimBlueprintCompiler.h
@@ -36,7 +36,7 @@ class FAnimBlueprintCompilerContext : public FKismetCompilerContext
 protected:
 	typedef FKismetCompilerContext Super;
 public:
-	FAnimBlueprintCompilerContext(UAnimBlueprint* SourceSketch, FCompilerResultsLog& InMessageLog, const FKismetCompilerOptions& InCompileOptions, TArray<UObject*>* InObjLoaded);
+	FAnimBlueprintCompilerContext(UAnimBlueprint* SourceSketch, FCompilerResultsLog& InMessageLog, const FKismetCompilerOptions& InCompileOptions);
 	virtual ~FAnimBlueprintCompilerContext();
 
 	virtual void PostCompile() override;
diff --git a/Engine/Source/Editor/KismetCompiler/Private/KismetCompiler.cpp b/Engine/Source/Editor/KismetCompiler/Private/KismetCompiler.cpp
--- a/Engine/Source/Editor/KismetCompiler/Private/KismetCompiler.cpp
+++ b/Engine/Source/Editor/KismetCompiler/Private/KismetCompiler.cpp
@@ -124,11 +124,10 @@ namespace
 //////////////////////////////////////////////////////////////////////////
 // FKismetCompilerContext
 
-FKismetCompilerContext::FKismetCompilerContext(UBlueprint* SourceSketch, FCompilerResultsLog& InMessageLog, const FKismetCompilerOptions& InCompilerOptions, TArray<UObject*>* InObjLoaded)
+FKismetCompilerContext::FKismetCompilerContext(UBlueprint* SourceSketch, FCompilerResultsLog& InMessageLog, const FKismetCompilerOptions& InCompilerOptions)
 	: FGraphCompilerContext(InMessageLog)
 	, Schema(NULL)
 	, CompileOptions(InCompilerOptions)
-	, ObjLoaded(InObjLoaded)
 	, Blueprint(SourceSketch)
 	, NewClass(NULL)
 	, ConsolidatedEventGraph(NULL)
@@ -4167,17 +4166,10 @@ void FKismetCompilerContext::CompileFunctions(EInternalCompilerFlags InternalFla
 				// Propagate the old CDO's properties to the new
 				if( OldCDO )
 				{
-					if (ObjLoaded)
+					if (OldLinker && OldGenLinkerIdx != INDEX_NONE)
 					{
-						if (OldLinker && OldGenLinkerIdx != INDEX_NONE)
-						{
-							// If we have a list of objects that are loading, patch our export table. This also fixes up load flags
-							FBlueprintEditorUtils::PatchNewCDOIntoLinker(Blueprint->GeneratedClass->GetDefaultObject(), OldLinker, OldGenLinkerIdx, *ObjLoaded);
-						}
-						else
-						{
-							UE_LOG(LogK2Compiler, Warning, TEXT("Failed to patch linker table for blueprint CDO %s"), *NewCDO->GetName());
-						}
+						// If we have a list of objects that are loading, patch our export table. This also fixes up load flags
+						FBlueprintEditorUtils::PatchNewCDOIntoLinker(Blueprint->GeneratedClass->GetDefaultObject(), OldLinker, OldGenLinkerIdx);
 					}
 
 					UEditorEngine::FCopyPropertiesForUnrelatedObjectsParams CopyDetails;
@@ -4684,7 +4676,7 @@ TSharedPtr<FKismetCompilerContext> FKismetCompilerContext::GetCompilerForBP(UBlu
 	// so I have simply hard-coded it:
 	if(UAnimBlueprint* AnimBP = Cast<UAnimBlueprint>(BP))
 	{
-		return TSharedPtr<FKismetCompilerContext>(new FAnimBlueprintCompilerContext(AnimBP, InMessageLog, InCompileOptions, nullptr));
+		return TSharedPtr<FKismetCompilerContext>(new FAnimBlueprintCompilerContext(AnimBP, InMessageLog, InCompileOptions));
 	}
 	else if(CompilerContextFactoryFunction* FactoryFunction = CustomCompilerMap.Find(BP->GetClass()))
 	{
@@ -4692,7 +4684,7 @@ TSharedPtr<FKismetCompilerContext> FKismetCompilerContext::GetCompilerForBP(UBlu
 	}
 	else
 	{
-		return TSharedPtr<FKismetCompilerContext>(new FKismetCompilerContext(BP, InMessageLog, InCompileOptions, nullptr));
+		return TSharedPtr<FKismetCompilerContext>(new FKismetCompilerContext(BP, InMessageLog, InCompileOptions));
 	}
 }
 
diff --git a/Engine/Source/Editor/KismetCompiler/Private/KismetCompilerModule.cpp b/Engine/Source/Editor/KismetCompiler/Private/KismetCompilerModule.cpp
--- a/Engine/Source/Editor/KismetCompiler/Private/KismetCompilerModule.cpp
+++ b/Engine/Source/Editor/KismetCompiler/Private/KismetCompilerModule.cpp
@@ -44,7 +44,7 @@ class FKismet2CompilerModule : public IKismetCompilerInterface
 {
 public:
 	// Implementation of the IKismetCompilerInterface
-	virtual void CompileBlueprint(class UBlueprint* Blueprint, const FKismetCompilerOptions& CompileOptions, FCompilerResultsLog& Results, TSharedPtr<class FBlueprintCompileReinstancer> ParentReinstancer = NULL, TArray<UObject*>* ObjLoaded = NULL) override;
+	virtual void CompileBlueprint(class UBlueprint* Blueprint, const FKismetCompilerOptions& CompileOptions, FCompilerResultsLog& Results, TSharedPtr<class FBlueprintCompileReinstancer> ParentReinstancer = NULL) override;
 	virtual void RefreshVariables(UBlueprint* Blueprint) override final;
 	virtual void CompileStructure(class UUserDefinedStruct* Struct, FCompilerResultsLog& Results) override;
 	virtual void RecoverCorruptedBlueprint(class UBlueprint* Blueprint) override;
@@ -56,7 +56,7 @@ public:
 	virtual FString GenerateCppWrapper(UBlueprintGeneratedClass* BPGC, const FCompilerNativizationOptions& NativizationOptions) override;
 	// End implementation
 private:
-	void CompileBlueprintInner(class UBlueprint* Blueprint, const FKismetCompilerOptions& CompileOptions, FCompilerResultsLog& Results, TSharedPtr<FBlueprintCompileReinstancer> Reinstancer, TArray<UObject*>* ObjLoaded);
+	void CompileBlueprintInner(class UBlueprint* Blueprint, const FKismetCompilerOptions& CompileOptions, FCompilerResultsLog& Results, TSharedPtr<FBlueprintCompileReinstancer> Reinstancer);
 
 	TArray<IBlueprintCompiler*> Compilers;
 };
@@ -82,7 +82,7 @@ public:
 };
 
 // Compiles a blueprint.
-void FKismet2CompilerModule::CompileBlueprintInner(class UBlueprint* Blueprint, const FKismetCompilerOptions& CompileOptions, FCompilerResultsLog& Results, TSharedPtr<FBlueprintCompileReinstancer> Reinstancer, TArray<UObject*>* ObjLoaded)
+void FKismet2CompilerModule::CompileBlueprintInner(class UBlueprint* Blueprint, const FKismetCompilerOptions& CompileOptions, FCompilerResultsLog& Results, TSharedPtr<FBlueprintCompileReinstancer> Reinstancer)
 {
 	FBlueprintIsBeingCompiledHelper BeingCompiled(Blueprint);
 
@@ -112,7 +112,7 @@ void FKismet2CompilerModule::CompileBlueprintInner(class UBlueprint* Blueprint,
 			if ( Compiler->CanCompile(Blueprint) )
 			{
 				Compiled = true;
-				Compiler->Compile(Blueprint, CompileOptions, Results, ObjLoaded);
+				Compiler->Compile(Blueprint, CompileOptions, Results);
 				break;
 			}
 		}
@@ -122,13 +122,13 @@ void FKismet2CompilerModule::CompileBlueprintInner(class UBlueprint* Blueprint,
 		{
 			if ( UAnimBlueprint* AnimBlueprint = Cast<UAnimBlueprint>(Blueprint) )
 			{
-				FAnimBlueprintCompilerContext Compiler(AnimBlueprint, Results, CompileOptions, ObjLoaded);
+				FAnimBlueprintCompilerContext Compiler(AnimBlueprint, Results, CompileOptions);
 				Compiler.Compile();
 				check(Compiler.NewClass);
 			}
 			else
 			{
-				FKismetCompilerContext Compiler(Blueprint, Results, CompileOptions, ObjLoaded);
+				FKismetCompilerContext Compiler(Blueprint, Results, CompileOptions);
 				Compiler.Compile();
 				check(Compiler.NewClass);
 			}
@@ -188,7 +188,7 @@ FString FKismet2CompilerModule::GenerateCppWrapper(UBlueprintGeneratedClass* BPG
 extern UNREALED_API FSecondsCounterData BlueprintCompileAndLoadTimerData;
 
 // Compiles a blueprint.
-void FKismet2CompilerModule::CompileBlueprint(class UBlueprint* Blueprint, const FKismetCompilerOptions& CompileOptions, FCompilerResultsLog& Results, TSharedPtr<FBlueprintCompileReinstancer> ParentReinstancer, TArray<UObject*>* ObjLoaded)
+void FKismet2CompilerModule::CompileBlueprint(class UBlueprint* Blueprint, const FKismetCompilerOptions& CompileOptions, FCompilerResultsLog& Results, TSharedPtr<FBlueprintCompileReinstancer> ParentReinstancer)
 {
 	FSecondsCounterScope Timer(BlueprintCompileAndLoadTimerData);
 	BP_SCOPED_COMPILER_EVENT_STAT(EKismetCompilerStats_CompileTime);
@@ -213,7 +213,7 @@ void FKismet2CompilerModule::CompileBlueprint(class UBlueprint* Blueprint, const
 		SkeletonResults.bSilentMode = true;
 		FKismetCompilerOptions SkeletonCompileOptions;
 		SkeletonCompileOptions.CompileType = EKismetCompileType::SkeletonOnly;
-		CompileBlueprintInner(Blueprint, SkeletonCompileOptions, SkeletonResults, ParentReinstancer, ObjLoaded);
+		CompileBlueprintInner(Blueprint, SkeletonCompileOptions, SkeletonResults, ParentReinstancer);
 
 		// Only when doing full compiles do we want to compile all skeletons before continuing to compile 
 		if (CompileOptions.CompileType == EKismetCompileType::Full)
@@ -230,7 +230,7 @@ void FKismet2CompilerModule::CompileBlueprint(class UBlueprint* Blueprint, const
 		FBlueprintCompileReinstancer::OptionallyRefreshNodes(Blueprint);
 
 		// Perform the full compile
-		CompileBlueprintInner(Blueprint, CompileOptions, Results, ParentReinstancer, ObjLoaded);
+		CompileBlueprintInner(Blueprint, CompileOptions, Results, ParentReinstancer);
 
 		if (Results.NumErrors == 0)
 		{
@@ -275,7 +275,7 @@ void FKismet2CompilerModule::CompileBlueprint(class UBlueprint* Blueprint, const
 				FKismetCompilerOptions StubCompileOptions(CompileOptions);
 				StubCompileOptions.CompileType = EKismetCompileType::StubAfterFailure;
 				{
-					CompileBlueprintInner(Blueprint, StubCompileOptions, StubResults, StubReinstancer, ObjLoaded);
+					CompileBlueprintInner(Blueprint, StubCompileOptions, StubResults, StubReinstancer);
 				}
 
 				StubReinstancer->UpdateBytecodeReferences();
diff --git a/Engine/Source/Editor/KismetCompiler/Public/KismetCompiler.h b/Engine/Source/Editor/KismetCompiler/Public/KismetCompiler.h
--- a/Engine/Source/Editor/KismetCompiler/Public/KismetCompiler.h
+++ b/Engine/Source/Editor/KismetCompiler/Public/KismetCompiler.h
@@ -90,9 +90,6 @@ protected:
 	UScriptStruct* TransformStruct;
 	UScriptStruct* LinearColorStruct;
 
-	// If set, this is a list of all the objects that are currently loading
-	TArray<UObject*>* ObjLoaded;
-
 public:
 	UBlueprint* Blueprint;
 	UBlueprintGeneratedClass* NewClass;
@@ -138,7 +135,7 @@ public:
 	static FSimpleMulticastDelegate OnPostCompile;
 
 public:
-	FKismetCompilerContext(UBlueprint* SourceSketch, FCompilerResultsLog& InMessageLog, const FKismetCompilerOptions& InCompilerOptions, TArray<UObject*>* InObjLoaded);
+	FKismetCompilerContext(UBlueprint* SourceSketch, FCompilerResultsLog& InMessageLog, const FKismetCompilerOptions& InCompilerOptions);
 	virtual ~FKismetCompilerContext();
 	
 	/** Compile the class layout of the blueprint */
diff --git a/Engine/Source/Editor/KismetCompiler/Public/KismetCompilerModule.h b/Engine/Source/Editor/KismetCompiler/Public/KismetCompilerModule.h
--- a/Engine/Source/Editor/KismetCompiler/Public/KismetCompilerModule.h
+++ b/Engine/Source/Editor/KismetCompiler/Public/KismetCompilerModule.h
@@ -28,7 +28,7 @@ public:
 	}
 
 	virtual bool CanCompile(const UBlueprint* Blueprint) = 0;
-	virtual void Compile(UBlueprint* Blueprint, const FKismetCompilerOptions& CompileOptions, FCompilerResultsLog& Results, TArray<UObject*>* ObjLoaded) = 0;
+	virtual void Compile(UBlueprint* Blueprint, const FKismetCompilerOptions& CompileOptions, FCompilerResultsLog& Results) = 0;
 	
 	virtual void PostCompile(UBlueprint* Blueprint, const FKismetCompilerOptions& CompileOptions)
 	{
@@ -57,7 +57,7 @@ public:
 	 * @param	Results  	The results log for warnings and errors.
 	 * @param	Options		Compiler options.
 	 */
-	virtual void CompileBlueprint(class UBlueprint* Blueprint, const FKismetCompilerOptions& CompileOptions, FCompilerResultsLog& Results, TSharedPtr<class FBlueprintCompileReinstancer> ParentReinstancer = NULL, TArray<UObject*>* ObjLoaded = NULL)=0;
+	virtual void CompileBlueprint(class UBlueprint* Blueprint, const FKismetCompilerOptions& CompileOptions, FCompilerResultsLog& Results, TSharedPtr<class FBlueprintCompileReinstancer> ParentReinstancer = NULL)=0;
 
 	/**
 	 * Synchronizes Blueprint's GeneratedClass's properties with the NewVariable declarations in the blueprint
diff --git a/Engine/Source/Editor/UMGEditor/Private/UMGEditorModule.cpp b/Engine/Source/Editor/UMGEditor/Private/UMGEditorModule.cpp
--- a/Engine/Source/Editor/UMGEditor/Private/UMGEditorModule.cpp
+++ b/Engine/Source/Editor/UMGEditor/Private/UMGEditorModule.cpp
@@ -148,11 +148,11 @@ public:
 		CompileCount++;
 	}
 
-	void Compile(UBlueprint* Blueprint, const FKismetCompilerOptions& CompileOptions, FCompilerResultsLog& Results, TArray<UObject*>* ObjLoaded) override
+	void Compile(UBlueprint* Blueprint, const FKismetCompilerOptions& CompileOptions, FCompilerResultsLog& Results) override
 	{
 		if ( UWidgetBlueprint* WidgetBlueprint = CastChecked<UWidgetBlueprint>(Blueprint) )
 		{
-			FWidgetBlueprintCompiler Compiler(WidgetBlueprint, Results, CompileOptions, ObjLoaded);
+			FWidgetBlueprintCompiler Compiler(WidgetBlueprint, Results, CompileOptions);
 			Compiler.Compile();
 			check(Compiler.NewClass);
 		}
diff --git a/Engine/Source/Editor/UMGEditor/Private/WidgetBlueprint.cpp b/Engine/Source/Editor/UMGEditor/Private/WidgetBlueprint.cpp
--- a/Engine/Source/Editor/UMGEditor/Private/WidgetBlueprint.cpp
+++ b/Engine/Source/Editor/UMGEditor/Private/WidgetBlueprint.cpp
@@ -737,7 +737,7 @@ bool UWidgetBlueprint::ValidateGeneratedClass(const UClass* InClass)
 
 TSharedPtr<FKismetCompilerContext> UWidgetBlueprint::GetCompilerForWidgetBP(UBlueprint* BP, FCompilerResultsLog& InMessageLog, const FKismetCompilerOptions& InCompileOptions)
 {
-	return TSharedPtr<FKismetCompilerContext>(new FWidgetBlueprintCompiler(CastChecked<UWidgetBlueprint>(BP), InMessageLog, InCompileOptions, nullptr));
+	return TSharedPtr<FKismetCompilerContext>(new FWidgetBlueprintCompiler(CastChecked<UWidgetBlueprint>(BP), InMessageLog, InCompileOptions));
 }
 
 void UWidgetBlueprint::GetReparentingRules(TSet< const UClass* >& AllowedChildrenOfClasses, TSet< const UClass* >& DisallowedChildrenOfClasses) const
diff --git a/Engine/Source/Editor/UMGEditor/Private/WidgetBlueprintCompiler.cpp b/Engine/Source/Editor/UMGEditor/Private/WidgetBlueprintCompiler.cpp
--- a/Engine/Source/Editor/UMGEditor/Private/WidgetBlueprintCompiler.cpp
+++ b/Engine/Source/Editor/UMGEditor/Private/WidgetBlueprintCompiler.cpp
@@ -27,8 +27,8 @@
 
 extern COREUOBJECT_API bool GMinimalCompileOnLoad;
 
-FWidgetBlueprintCompiler::FWidgetBlueprintCompiler(UWidgetBlueprint* SourceSketch, FCompilerResultsLog& InMessageLog, const FKismetCompilerOptions& InCompilerOptions, TArray<UObject*>* InObjLoaded)
-	: Super(SourceSketch, InMessageLog, InCompilerOptions, InObjLoaded)
+FWidgetBlueprintCompiler::FWidgetBlueprintCompiler(UWidgetBlueprint* SourceSketch, FCompilerResultsLog& InMessageLog, const FKismetCompilerOptions& InCompilerOptions)
+	: Super(SourceSketch, InMessageLog, InCompilerOptions)
 	, NewWidgetBlueprintClass(nullptr)
 	, WidgetSchema(nullptr)
 {
diff --git a/Engine/Source/Editor/UMGEditor/Private/WidgetBlueprintCompiler.h b/Engine/Source/Editor/UMGEditor/Private/WidgetBlueprintCompiler.h
--- a/Engine/Source/Editor/UMGEditor/Private/WidgetBlueprintCompiler.h
+++ b/Engine/Source/Editor/UMGEditor/Private/WidgetBlueprintCompiler.h
@@ -16,7 +16,7 @@ protected:
 	typedef FKismetCompilerContext Super;
 
 public:
-	FWidgetBlueprintCompiler(UWidgetBlueprint* SourceSketch, FCompilerResultsLog& InMessageLog, const FKismetCompilerOptions& InCompilerOptions, TArray<UObject*>* InObjLoaded);
+	FWidgetBlueprintCompiler(UWidgetBlueprint* SourceSketch, FCompilerResultsLog& InMessageLog, const FKismetCompilerOptions& InCompilerOptions);
 	virtual ~FWidgetBlueprintCompiler();
 
 	static bool CanAllowTemplate(FCompilerResultsLog& MessageLog, UWidgetBlueprintGeneratedClass* InClass);
diff --git a/Engine/Source/Editor/UnrealEd/Private/Commandlets/ContentCommandlets.cpp b/Engine/Source/Editor/UnrealEd/Private/Commandlets/ContentCommandlets.cpp
--- a/Engine/Source/Editor/UnrealEd/Private/Commandlets/ContentCommandlets.cpp
+++ b/Engine/Source/Editor/UnrealEd/Private/Commandlets/ContentCommandlets.cpp
@@ -70,6 +70,7 @@ DEFINE_LOG_CATEGORY(LogContentCommandlet);
 #include "ShaderCompiler.h"
 #include "ICollectionManager.h"
 #include "CollectionManagerModule.h"
+#include "UObject/UObjectThreadContext.h"
 
 /**-----------------------------------------------------------------------------
  *	UResavePackages commandlet.
@@ -571,9 +572,10 @@ void UResavePackagesCommandlet::LoadAndSaveOnePackage(const FString& Filename)
 		// Get the package linker.
 		VerboseMessage(TEXT("Pre GetPackageLinker"));
 
-		BeginLoad();
-		auto Linker = GetPackageLinker(NULL,*Filename,LOAD_NoVerify,NULL,NULL);
-		EndLoad();
+		TRefCountPtr<FUObjectSerializeContext> LoadContext(new FUObjectSerializeContext());
+		BeginLoad(LoadContext);
+		FLinkerLoad* Linker = GetPackageLinker(NULL, *Filename, LOAD_NoVerify, nullptr, nullptr, nullptr, LoadContext);
+		EndLoad(Linker->GetSerializeContext());
 	
 		// Bail early if we don't have a valid linker (package was out of date, etc)
 		if( !Linker )
@@ -2169,11 +2171,12 @@ int32 UWrangleContentCommandlet::Main( const FString& Params )
 	// or make sure that the following redirects marks the package dirty (which maybe it shouldn't do in the editor?)
 
 					// load the package fully
-					UPackage* Package = LoadPackage(NULL, *PackageFilename, LOAD_None);
+					UPackage* Package = LoadPackage(nullptr, *PackageFilename, LOAD_None);
 
-					BeginLoad();
-					auto Linker = GetPackageLinker( NULL, *PackageFilename, LOAD_Quiet|LOAD_NoWarn|LOAD_NoVerify, NULL, NULL );
-					EndLoad();
+					TRefCountPtr<FUObjectSerializeContext> LoadContext(new FUObjectSerializeContext());
+					BeginLoad(LoadContext);
+					FLinkerLoad* Linker = GetPackageLinker(nullptr, *PackageFilename, LOAD_Quiet | LOAD_NoWarn | LOAD_NoVerify, nullptr, nullptr, nullptr, LoadContext);
+					EndLoad(Linker->GetSerializeContext());
 
 					// look for special package types
 					bool bIsMap = Linker->ContainsMap();
@@ -2325,9 +2328,10 @@ int32 UWrangleContentCommandlet::Main( const FString& Params )
 				}
 			}
 
-			BeginLoad();
-			auto Linker = GetPackageLinker( NULL, *PackageFilename, LOAD_Quiet|LOAD_NoWarn|LOAD_NoVerify, NULL, NULL );
-			EndLoad();
+			TRefCountPtr<FUObjectSerializeContext> LoadContext(new FUObjectSerializeContext());
+			BeginLoad(LoadContext);
+			FLinkerLoad* Linker = GetPackageLinker(nullptr, *PackageFilename, LOAD_Quiet | LOAD_NoWarn | LOAD_NoVerify, nullptr, nullptr, nullptr, LoadContext);
+			EndLoad(Linker->GetSerializeContext());
 
 			// go through the exports in the package, looking for public objects
 			for (int32 ExportIndex = 0; ExportIndex < Linker->ExportMap.Num(); ExportIndex++)
diff --git a/Engine/Source/Editor/UnrealEd/Private/Commandlets/PackageUtilities.cpp b/Engine/Source/Editor/UnrealEd/Private/Commandlets/PackageUtilities.cpp
--- a/Engine/Source/Editor/UnrealEd/Private/Commandlets/PackageUtilities.cpp
+++ b/Engine/Source/Editor/UnrealEd/Private/Commandlets/PackageUtilities.cpp
@@ -55,6 +55,7 @@
 #include "Serialization/ArchiveStackTrace.h"
 #include "Misc/OutputDeviceHelper.h"
 #include "Misc/OutputDeviceFile.h"
+#include "UObject/UObjectThreadContext.h"
 
 DEFINE_LOG_CATEGORY(LogPackageHelperFunctions);
 DEFINE_LOG_CATEGORY_STATIC(LogPackageUtilities, Log, All);
@@ -780,9 +781,10 @@ int32 ULoadPackageCommandlet::Main( const FString& Params )
 
 		if (bCheckForLegacyPackages)
 		{
-			BeginLoad();
-			auto Linker = GetPackageLinker(nullptr, *Filename,LOAD_NoVerify,NULL,NULL);
-			EndLoad();
+			TRefCountPtr<FUObjectSerializeContext> LoadContext(new FUObjectSerializeContext());
+			BeginLoad(LoadContext);
+			auto Linker = GetPackageLinker(nullptr, *Filename, LOAD_NoVerify, nullptr, nullptr, nullptr, LoadContext);
+			EndLoad(Linker->GetSerializeContext());
 			MinVersion = FMath::Min<int32>(MinVersion, Linker->Summary.GetFileVersionUE4());
 		}
 		else
@@ -901,7 +903,7 @@ namespace
 }
 
 /** Given a package filename, creates a linker and a temporary package. The filename does not need to point to a package under the current project content folder */
-FLinkerLoad* CreateLinkerForFilename(const FString& InFilename)
+FLinkerLoad* CreateLinkerForFilename(FUObjectSerializeContext* LoadContext, const FString& InFilename)
 {
 	FString TempPackageName;
 	TempPackageName = FPaths::Combine(TEXT("/Temp"), *FPaths::GetPath(InFilename.Mid(InFilename.Find(TEXT(":"), ESearchCase::CaseSensitive) + 1)), *FPaths::GetBaseFilename(InFilename));
@@ -910,7 +912,7 @@ FLinkerLoad* CreateLinkerForFilename(const FString& InFilename)
 	{
 		Package = CreatePackage(nullptr, *TempPackageName);
 	}
-	FLinkerLoad* Linker = FLinkerLoad::CreateLinker(Package, *InFilename, LOAD_NoVerify);
+	FLinkerLoad* Linker = FLinkerLoad::CreateLinker(LoadContext, Package, *InFilename, LOAD_NoVerify);
 	return Linker;
 }
 
@@ -1554,9 +1556,10 @@ int32 UPkgInfoCommandlet::Main( const FString& Params )
 		if (!bDumpProperties)
 		{
 			TGuardValue<bool> GuardAllowUnversionedContentInEditor(GAllowUnversionedContentInEditor, true);
-			BeginLoad();
-			Linker = CreateLinkerForFilename(Filename);
-			EndLoad();
+			TRefCountPtr<FUObjectSerializeContext> LoadContext(new FUObjectSerializeContext());
+			BeginLoad(LoadContext);
+			Linker = CreateLinkerForFilename(LoadContext, Filename);
+			EndLoad(Linker->GetSerializeContext());
 		}
 		else
 		{
diff --git a/Engine/Source/Editor/UnrealEd/Private/Kismet2/BlueprintEditorUtils.cpp b/Engine/Source/Editor/UnrealEd/Private/Kismet2/BlueprintEditorUtils.cpp
--- a/Engine/Source/Editor/UnrealEd/Private/Kismet2/BlueprintEditorUtils.cpp
+++ b/Engine/Source/Editor/UnrealEd/Private/Kismet2/BlueprintEditorUtils.cpp
@@ -115,6 +115,7 @@
 #include "Styling/SlateIconFinder.h"
 #include "BaseWidgetBlueprint.h"
 #include "Components/Widget.h"
+#include "UObject/UObjectThreadContext.h"
 
 extern COREUOBJECT_API bool GBlueprintUseCompilationManager;
 
@@ -715,7 +716,7 @@ void FBlueprintEditorUtils::PreloadConstructionScript(UBlueprint* Blueprint)
 	}
 }
 
-void FBlueprintEditorUtils::PatchNewCDOIntoLinker(UObject* CDO, FLinkerLoad* Linker, int32 ExportIndex, TArray<UObject*>& ObjLoaded)
+void FBlueprintEditorUtils::PatchNewCDOIntoLinker(UObject* CDO, FLinkerLoad* Linker, int32 ExportIndex)
 {
 	if( (CDO != nullptr) && (Linker != nullptr) && (ExportIndex != INDEX_NONE) )
 	{
@@ -730,6 +731,9 @@ void FBlueprintEditorUtils::PatchNewCDOIntoLinker(UObject* CDO, FLinkerLoad* Lin
 			// Copy flags from the old CDO.
 			CDO->SetFlags(OldObjectFlags);
 
+			FUObjectSerializeContext* LoadContext = Linker->GetSerializeContext();
+			TArray<UObject*>& ObjLoaded = LoadContext->GetObjectsLoaded();
+
 			// Make sure the new CDO gets PostLoad called on it, so either add it to ObjLoaded list, or replace it if already present.
 			int32 ObjLoadedIdx = ObjLoaded.Find(OldCDO);
 			if (ObjLoadedIdx != INDEX_NONE)
@@ -1074,7 +1078,7 @@ struct FRegenerationHelper
 	 * 
 	 * @param  Blueprint	The blueprint whose implemented interfaces you want loaded.
 	 */
-	static void PreloadInterfaces(UBlueprint* Blueprint, TArray<UObject*>& ObjLoaded)
+	static void PreloadInterfaces(UBlueprint* Blueprint)
 	{
 #if WITH_EDITORONLY_DATA // ImplementedInterfaces is wrapped WITH_EDITORONLY_DATA 
 		for (FBPInterfaceDescription const& InterfaceDesc : Blueprint->ImplementedInterfaces)
@@ -1086,14 +1090,14 @@ struct FRegenerationHelper
 				UBlueprint::ForceLoadMembers(InterfaceBlueprint);
 				if (InterfaceBlueprint->HasAnyFlags(RF_BeingRegenerated))
 				{
-					InterfaceBlueprint->RegenerateClass(InterfaceClass, InterfaceClass->ClassDefaultObject, ObjLoaded);
+					InterfaceBlueprint->RegenerateClass(InterfaceClass, InterfaceClass->ClassDefaultObject);
 				}
 			}
 		}
 #endif // #if WITH_EDITORONLY_DATA
 	}
 
-	static void LinkExternalDependencies(UBlueprint* Blueprint, TArray<UObject*>& ObjLoaded)
+	static void LinkExternalDependencies(UBlueprint* Blueprint)
 	{
 		check(Blueprint);
 		const UEdGraphSchema_K2* Schema = GetDefault<UEdGraphSchema_K2>();
@@ -1194,7 +1198,7 @@ struct FRegenerationHelper
 		}
 		PreloadMacroSources(MacroSources);
 
-		PreloadInterfaces(Blueprint, ObjLoaded);
+		PreloadInterfaces(Blueprint);
 	}
 };
 
@@ -1264,7 +1268,7 @@ void FBlueprintEditorUtils::PreloadBlueprintSpecificData(UBlueprint* Blueprint)
 	}
 }
 
-UClass* FBlueprintEditorUtils::RegenerateBlueprintClass(UBlueprint* Blueprint, UClass* ClassToRegenerate, UObject* PreviousCDO, TArray<UObject*>& ObjLoaded)
+UClass* FBlueprintEditorUtils::RegenerateBlueprintClass(UBlueprint* Blueprint, UClass* ClassToRegenerate, UObject* PreviousCDO)
 {
 	bool bRegenerated = false;
 
@@ -1342,7 +1346,7 @@ UClass* FBlueprintEditorUtils::RegenerateBlueprintClass(UBlueprint* Blueprint, U
 		const bool bHasCode = !FBlueprintEditorUtils::IsDataOnlyBlueprint(Blueprint) && !bIsMacro;
 
 		// Make sure all used external classes/functions/structures/macros/etc are loaded and linked
-		FRegenerationHelper::LinkExternalDependencies(Blueprint, ObjLoaded);
+		FRegenerationHelper::LinkExternalDependencies(Blueprint);
 
 		bool bSkeletonUpToDate = FKismetEditorUtilities::GenerateBlueprintSkeleton(Blueprint);
 
@@ -1453,11 +1457,11 @@ UClass* FBlueprintEditorUtils::RegenerateBlueprintClass(UBlueprint* Blueprint, U
 		// Patch the new CDOs to the old indices in the linker
 		if( Blueprint->SkeletonGeneratedClass )
 		{
-			PatchNewCDOIntoLinker(Blueprint->SkeletonGeneratedClass->GetDefaultObject(), OldLinker, OldSkelLinkerIdx, ObjLoaded);
+			PatchNewCDOIntoLinker(Blueprint->SkeletonGeneratedClass->GetDefaultObject(), OldLinker, OldSkelLinkerIdx);
 		}
 		if( Blueprint->GeneratedClass )
 		{
-			PatchNewCDOIntoLinker(Blueprint->GeneratedClass->GetDefaultObject(), OldLinker, OldGenLinkerIdx, ObjLoaded);
+			PatchNewCDOIntoLinker(Blueprint->GeneratedClass->GetDefaultObject(), OldLinker, OldGenLinkerIdx);
 		}
 
 		// Success or failure, there's no point in trying to recompile this class again when other objects reference it
@@ -1589,8 +1593,7 @@ UClass* FBlueprintEditorUtils::RegenerateBlueprintClass(UBlueprint* Blueprint, U
 
 void FBlueprintEditorUtils::LinkExternalDependencies(UBlueprint* Blueprint)
 {
-	TArray<UObject*> Unused;
-	FRegenerationHelper::LinkExternalDependencies(Blueprint, Unused);
+	FRegenerationHelper::LinkExternalDependencies(Blueprint);
 }
 
 void FBlueprintEditorUtils::RecreateClassMetaData(UBlueprint* Blueprint, UClass* Class, bool bRemoveExistingMetaData)
diff --git a/Engine/Source/Editor/UnrealEd/Private/Kismet2/Kismet2.cpp b/Engine/Source/Editor/UnrealEd/Private/Kismet2/Kismet2.cpp
--- a/Engine/Source/Editor/UnrealEd/Private/Kismet2/Kismet2.cpp
+++ b/Engine/Source/Editor/UnrealEd/Private/Kismet2/Kismet2.cpp
@@ -1016,7 +1016,7 @@ bool FKismetEditorUtilities::GenerateBlueprintSkeleton(UBlueprint* BlueprintObj,
 }
 
 /** Recompiles the bytecode of a blueprint only.  Should only be run for recompiling dependencies during compile on load */
-void FKismetEditorUtilities::RecompileBlueprintBytecode(UBlueprint* BlueprintObj, TArray<UObject*>* ObjLoaded,  EBlueprintBytecodeRecompileOptions Flags)
+void FKismetEditorUtilities::RecompileBlueprintBytecode(UBlueprint* BlueprintObj,  EBlueprintBytecodeRecompileOptions Flags)
 {
 	FSecondsCounterScope Timer(BlueprintCompileAndLoadTimerData); 
 
@@ -1049,7 +1049,7 @@ void FKismetEditorUtilities::RecompileBlueprintBytecode(UBlueprint* BlueprintObj
 	{
 		FRecreateUberGraphFrameScope RecreateUberGraphFrameScope(BlueprintObj->GeneratedClass, true);
 		FCompilerResultsLog Results;
-		Compiler.CompileBlueprint(BlueprintObj, CompileOptions, Results, NULL, ObjLoaded);
+		Compiler.CompileBlueprint(BlueprintObj, CompileOptions, Results, NULL);
 	}
 	
 	if(!bSkipReinstancing)
diff --git a/Engine/Source/Editor/UnrealEd/Private/Kismet2/KismetReinstanceUtilities.cpp b/Engine/Source/Editor/UnrealEd/Private/Kismet2/KismetReinstanceUtilities.cpp
--- a/Engine/Source/Editor/UnrealEd/Private/Kismet2/KismetReinstanceUtilities.cpp
+++ b/Engine/Source/Editor/UnrealEd/Private/Kismet2/KismetReinstanceUtilities.cpp
@@ -882,7 +882,7 @@ void FBlueprintCompileReinstancer::ReinstanceObjects(bool bForceAlwaysReinstance
 				for (int I = 0; I != OrderedBytecodeRecompile.Num(); ++I)
 				{
 					UBlueprint* BP = OrderedBytecodeRecompile[I];
-					FKismetEditorUtilities::RecompileBlueprintBytecode(BP, nullptr, EBlueprintBytecodeRecompileOptions::BatchCompile);
+					FKismetEditorUtilities::RecompileBlueprintBytecode(BP, EBlueprintBytecodeRecompileOptions::BatchCompile);
 					ensure(0 == DependentBlueprintsToRecompile.Num());
 					CompiledBlueprints.Add(BP);
 
diff --git a/Engine/Source/Editor/UnrealEd/Public/Kismet2/BlueprintEditorUtils.h b/Engine/Source/Editor/UnrealEd/Public/Kismet2/BlueprintEditorUtils.h
--- a/Engine/Source/Editor/UnrealEd/Public/Kismet2/BlueprintEditorUtils.h
+++ b/Engine/Source/Editor/UnrealEd/Public/Kismet2/BlueprintEditorUtils.h
@@ -183,7 +183,7 @@ public:
 	/** 
 	 * Helper function to patch the new CDO into the linker where the old one existed 
 	 */
-	static void PatchNewCDOIntoLinker(UObject* CDO, FLinkerLoad* Linker, int32 ExportIndex, TArray<UObject*>& ObjLoaded);
+	static void PatchNewCDOIntoLinker(UObject* CDO, FLinkerLoad* Linker, int32 ExportIndex);
 
 	/** 
 	 * Procedure used to remove old function implementations and child properties from data only blueprints.
@@ -201,7 +201,7 @@ public:
 	/**
 	 * Regenerates the class at class load time, and refreshes the blueprint
 	 */
-	static UClass* RegenerateBlueprintClass(UBlueprint* Blueprint, UClass* ClassToRegenerate, UObject* PreviousCDO, TArray<UObject*>& ObjLoaded);
+	static UClass* RegenerateBlueprintClass(UBlueprint* Blueprint, UClass* ClassToRegenerate, UObject* PreviousCDO);
 	
 	/**
 	 * Links external dependencies
diff --git a/Engine/Source/Editor/UnrealEd/Public/Kismet2/KismetEditorUtilities.h b/Engine/Source/Editor/UnrealEd/Public/Kismet2/KismetEditorUtilities.h
--- a/Engine/Source/Editor/UnrealEd/Public/Kismet2/KismetEditorUtilities.h
+++ b/Engine/Source/Editor/UnrealEd/Public/Kismet2/KismetEditorUtilities.h
@@ -150,7 +150,7 @@ public:
 	static bool GenerateBlueprintSkeleton(UBlueprint* BlueprintObj, bool bForceRegeneration = false);
 
 	/** Recompiles the bytecode of a blueprint only.  Should only be run for recompiling dependencies during compile on load */
-	static void RecompileBlueprintBytecode(UBlueprint* BlueprintObj, TArray<UObject*>* ObjLoaded = nullptr, EBlueprintBytecodeRecompileOptions Flags = EBlueprintBytecodeRecompileOptions::None);
+	static void RecompileBlueprintBytecode(UBlueprint* BlueprintObj, EBlueprintBytecodeRecompileOptions Flags = EBlueprintBytecodeRecompileOptions::None);
 
 	/** Tries to make sure that a data-only blueprint is conformed to its native parent, in case any native class flags have changed */
 	static void ConformBlueprintFlagsAndComponents(UBlueprint* BlueprintObj);
diff --git a/Engine/Source/Runtime/Core/Public/Serialization/Archive.h b/Engine/Source/Runtime/Core/Public/Serialization/Archive.h
--- a/Engine/Source/Runtime/Core/Public/Serialization/Archive.h
+++ b/Engine/Source/Runtime/Core/Public/Serialization/Archive.h
@@ -28,7 +28,7 @@ struct FUntypedBulkData;
 struct FArchiveSerializedPropertyChain;
 template<class TEnum> class TEnumAsByte;
 typedef TFunction<bool (double RemainingTime)> FExternalReadCallback;
-
+struct FUObjectSerializeContext;
 
 // Temporary while we shake out the EDL at boot
 #define USE_EVENT_DRIVEN_ASYNC_LOAD_AT_BOOT_TIME (1)
@@ -1273,6 +1273,12 @@ PRAGMA_ENABLE_DEPRECATION_WARNINGS
 	virtual bool IsEditorOnlyPropertyOnTheStack() const;
 #endif
 
+	/* Sets the current UObject serialization context for this archive */
+	virtual void SetSerializeContext(FUObjectSerializeContext* InLoadContext) {}
+
+	/* Gets the current UObject serialization context for this archive */
+	virtual FUObjectSerializeContext* GetSerializeContext() { return nullptr; }
+
 	/** 
 	 * Adds external read dependency 
 	 *
diff --git a/Engine/Source/Runtime/CoreUObject/Private/Blueprint/BlueprintSupport.cpp b/Engine/Source/Runtime/CoreUObject/Private/Blueprint/BlueprintSupport.cpp
--- a/Engine/Source/Runtime/CoreUObject/Private/Blueprint/BlueprintSupport.cpp
+++ b/Engine/Source/Runtime/CoreUObject/Private/Blueprint/BlueprintSupport.cpp
@@ -432,7 +432,7 @@ FScopedClassDependencyGather::~FScopedClassDependencyGather()
 		// BatchMasterClass and BatchClassDependencies entries
 		auto RecompileClassLambda = [&DependencyIter](UClass* Class)
 		{
-			Class->ConditionalRecompileClass(&FUObjectThreadContext::Get().ObjLoaded);
+			Class->ConditionalRecompileClass();
 
 			// because of the above call to ConditionalRecompileClass(), the 
 			// specified Class gets "cleaned and sanitized" (meaning its old 
@@ -475,7 +475,7 @@ FScopedClassDependencyGather::~FScopedClassDependencyGather()
 		}
 		else
 		{
-			BatchMasterClass->ConditionalRecompileClass(&FUObjectThreadContext::Get().ObjLoaded);
+			BatchMasterClass->ConditionalRecompileClass();
 		}
 
 		BatchMasterClass = NULL;
@@ -727,7 +727,7 @@ bool FLinkerLoad::RegenerateBlueprintClass(UClass* LoadClass, UObject* ExportObj
 			// Preload the blueprint to make sure it has all the data the class needs for regeneration
 			FPreloadMembersHelper::PreloadObject(BlueprintObject);
 
-			UClass* RegeneratedClass = BlueprintObject->RegenerateClass(LoadClass, CurrentCDO, FUObjectThreadContext::Get().ObjLoaded);
+			UClass* RegeneratedClass = BlueprintObject->RegenerateClass(LoadClass, CurrentCDO);
 			if (RegeneratedClass)
 			{
 				BlueprintObject->ClearFlags(RF_BeingRegenerated);
@@ -1261,7 +1261,7 @@ int32 FLinkerLoad::FindCDOExportIndex(UClass* LoadClass)
 	return INDEX_NONE;
 }
 
-UPackage* LoadPackageInternal(UPackage* InOuter, const TCHAR* InLongPackageName, uint32 LoadFlags, FLinkerLoad* ImportLinker, FArchive* InReaderOverride);
+UPackage* LoadPackageInternal(UPackage* InOuter, const TCHAR* InLongPackageName, uint32 LoadFlags, FLinkerLoad* ImportLinker, FArchive* InReaderOverride, FUObjectSerializeContext* InLoadContext);
 
 void FLinkerLoad::ResolveDeferredDependencies(UStruct* LoadStruct)
 {
@@ -1381,7 +1381,7 @@ void FLinkerLoad::ResolveDeferredDependencies(UStruct* LoadStruct)
 			{
 				uint32 InternalLoadFlags = LoadFlags & (LOAD_NoVerify | LOAD_NoWarn | LOAD_Quiet);
 				// make sure LoadAllObjects() is called for this package
-				LoadPackageInternal(/*Outer =*/nullptr, *SourceLinker->Filename, InternalLoadFlags, this, nullptr); //-V595
+				LoadPackageInternal(/*Outer =*/nullptr, *SourceLinker->Filename, InternalLoadFlags, this, nullptr, GetSerializeContext()); //-V595
 			}
 
 			DEFERRED_DEPENDENCY_CHECK(ResolvingPlaceholderStack.Num() == 0);
diff --git a/Engine/Source/Runtime/CoreUObject/Private/Serialization/ArchiveStackTrace.cpp b/Engine/Source/Runtime/CoreUObject/Private/Serialization/ArchiveStackTrace.cpp
--- a/Engine/Source/Runtime/CoreUObject/Private/Serialization/ArchiveStackTrace.cpp
+++ b/Engine/Source/Runtime/CoreUObject/Private/Serialization/ArchiveStackTrace.cpp
@@ -323,7 +323,8 @@ void FArchiveStackTrace::Serialize(void* InData, int64 Num)
 			if (CallstackAtOffsetMap.Num() == 0 || CurrentOffset > CallstackAtOffsetMap.Last().Offset)
 			{
 				// New data serialized at the end of archive buffer
-				LastSerializeCallstack = AddUniqueCallstack(ThreadContext.SerializedObject, GetSerializedProperty(), CallstackCRC);
+				check(LoadContext);
+				LastSerializeCallstack = AddUniqueCallstack(LoadContext->SerializedObject, GetSerializedProperty(), CallstackCRC);
 				CallstackAtOffsetMap.Add(FCallstactAtOffset(CurrentOffset, CallstackCRC, GIgnoreDiffManager.ShouldIgnoreDiff()));
 			}
 			else
@@ -331,8 +332,9 @@ void FArchiveStackTrace::Serialize(void* InData, int64 Num)
 				// This happens usually after Seek() so we need to find the exiting offset or insert a new one
 				int32 CallstackToUpdateIndex = GetCallstackAtOffset(CurrentOffset, 0);
 				check(CallstackToUpdateIndex != -1);
+				check(LoadContext);
 				FCallstactAtOffset& CallstackToUpdate = CallstackAtOffsetMap[CallstackToUpdateIndex];
-				LastSerializeCallstack = AddUniqueCallstack(ThreadContext.SerializedObject, GetSerializedProperty(), CallstackCRC);
+				LastSerializeCallstack = AddUniqueCallstack(LoadContext->SerializedObject, GetSerializedProperty(), CallstackCRC);
 				if (CallstackToUpdate.Offset == CurrentOffset)
 				{
 					CallstackToUpdate.Callstack = CallstackCRC;
@@ -355,6 +357,16 @@ void FArchiveStackTrace::Serialize(void* InData, int64 Num)
 	FLargeMemoryWriter::Serialize(InData, Num);
 }
 
+void FArchiveStackTrace::SetSerializeContext(FUObjectSerializeContext* InLoadContext)
+{
+	LoadContext = InLoadContext;
+}
+
+FUObjectSerializeContext* FArchiveStackTrace::GetSerializeContext()
+{
+	return LoadContext;
+}
+
 int32 FArchiveStackTrace::GetCallstackAtOffset(int64 InOffset, int32 MinOffsetIndex)
 {
 	if (InOffset < 0 || InOffset > TotalSize() || MinOffsetIndex < 0 || MinOffsetIndex >= CallstackAtOffsetMap.Num())
@@ -914,7 +926,7 @@ bool FArchiveStackTrace::IsIdentical(const TCHAR* InFilename, int64 BufferSize,
 	return bIdentical;
 }
 
-FLinkerLoad* FArchiveStackTrace::CreateLinkerForPackage(const FString& InPackageName, const FString& InFilename, const FPackageData& PackageData)
+FLinkerLoad* FArchiveStackTrace::CreateLinkerForPackage(FUObjectSerializeContext* LoadContext, const FString& InPackageName, const FString& InFilename, const FPackageData& PackageData)
 {
 	// First create a temp package to associate the linker with
 	UPackage* Package = FindObjectFast<UPackage>(nullptr, *InPackageName);
@@ -924,7 +936,7 @@ FLinkerLoad* FArchiveStackTrace::CreateLinkerForPackage(const FString& InPackage
 	}
 	// Create an archive for the linker. The linker will take ownership of it.
 	FLargeMemoryReader* PackageReader = new FLargeMemoryReader(PackageData.Data, PackageData.Size, ELargeMemoryReaderFlags::None, *InPackageName);	
-	FLinkerLoad* Linker = FLinkerLoad::CreateLinker(Package, *InFilename, LOAD_NoVerify, PackageReader);
+	FLinkerLoad* Linker = FLinkerLoad::CreateLinker(LoadContext, Package, *InFilename, LOAD_NoVerify, PackageReader);
 
 	if (Linker && Package)
 	{
@@ -1286,14 +1298,22 @@ void FArchiveStackTrace::DumpPackageHeaderDiffs(const FPackageData& SourcePackag
 	TGuardValue<bool> GuardIsSavingPackage(GIsSavingPackage, false);
 	TGuardValue<bool> GuardAllowUnversionedContentInEditor(GAllowUnversionedContentInEditor, true);
 
+	FLinkerLoad* SourceLinker = nullptr;
+	FLinkerLoad* DestLinker = nullptr;
 	// Create linkers. Note there's no need to clean them up here since they will be removed by the package associated with them
-	BeginLoad();
-	FLinkerLoad* SourceLinker = CreateLinkerForPackage(SourceAssetPackageName, AssetFilename, SourcePackage);
-	EndLoad();
-
-	BeginLoad();
-	FLinkerLoad* DestLinker = CreateLinkerForPackage(DestAssetPackageName, AssetFilename, DestPackage);
-	EndLoad();
+	{
+		TRefCountPtr<FUObjectSerializeContext> LinkerLoadContext(new FUObjectSerializeContext());
+		BeginLoad(LinkerLoadContext);
+		SourceLinker = CreateLinkerForPackage(LinkerLoadContext, SourceAssetPackageName, AssetFilename, SourcePackage);
+		EndLoad(SourceLinker->GetSerializeContext());
+	}
+	
+	{
+		TRefCountPtr<FUObjectSerializeContext> LinkerLoadContext(new FUObjectSerializeContext());
+		BeginLoad(LinkerLoadContext);
+		DestLinker = CreateLinkerForPackage(LinkerLoadContext, DestAssetPackageName, AssetFilename, DestPackage);
+		EndLoad(DestLinker->GetSerializeContext());
+	}
 
 	if (SourceLinker && DestLinker)
 	{
@@ -1334,7 +1354,8 @@ FArchiveStackTraceReader::FArchiveStackTraceReader(const TCHAR* InFilename, cons
 void FArchiveStackTraceReader::Serialize(void* OutData, int64 Num)
 {
 	bool bAddData = true;
-	FSerializeData NewData(Tell(), Num, ThreadContext.SerializedObject, GetSerializedProperty());
+	check(GetSerializeContext());
+	FSerializeData NewData(Tell(), Num, GetSerializeContext()->SerializedObject, GetSerializedProperty());
 	if (SerializeTrace.Num())
 	{
 		FSerializeData& Last = SerializeTrace.Last();
diff --git a/Engine/Source/Runtime/CoreUObject/Private/Serialization/AsyncLoading.cpp b/Engine/Source/Runtime/CoreUObject/Private/Serialization/AsyncLoading.cpp
--- a/Engine/Source/Runtime/CoreUObject/Private/Serialization/AsyncLoading.cpp
+++ b/Engine/Source/Runtime/CoreUObject/Private/Serialization/AsyncLoading.cpp
@@ -438,8 +438,6 @@ void FAsyncLoadingThread::CancelAsyncLoadingInternal()
 
 	NotifyAsyncLoadingStateHasMaybeChanged();
 
-	FUObjectThreadContext::Get().ObjLoaded.Empty();
-
 	// Notify everyone streaming is canceled.
 	CancelLoadingEvent->Trigger();
 }
@@ -3124,11 +3122,10 @@ UObject* FAsyncPackage::EventDrivenIndexToObject(FPackageIndex Index, bool bChec
 		MyDependentNode.Phase = EEventLoadNode::ImportOrExport_Create;
 		if (EventNodeArray.GetNode(MyDependentNode, false).bAddedToGraph || !EventNodeArray.GetNode(MyDependentNode, false).bFired)
 		{
-			FUObjectThreadContext& ThreadContext = FUObjectThreadContext::Get();
-			UClass* SerClass = Cast<UClass>(ThreadContext.SerializedObject);
+			UClass* SerClass = Cast<UClass>(LoadContext->SerializedObject);
 			if (!SerClass || Linker->ImpExp(Index).ObjectName != SerClass->GetDefaultObjectName())
 			{
-				DumpDependencies(TEXT("Dependencies"), ThreadContext.SerializedObject);
+				DumpDependencies(TEXT("Dependencies"), LoadContext->SerializedObject);
 				UE_LOG(LogStreaming, Fatal, TEXT("Missing Dependency, request for %s but it was still waiting for creation."), *Linker->GetPathName(Index));
 			}
 		}
@@ -3142,8 +3139,7 @@ UObject* FAsyncPackage::EventDrivenIndexToObject(FPackageIndex Index, bool bChec
 
 		if (DumpIndex.IsNull())
 		{
-			FUObjectThreadContext& ThreadContext = FUObjectThreadContext::Get();
-			DumpDependencies(TEXT("Dependencies"), ThreadContext.SerializedObject);
+			DumpDependencies(TEXT("Dependencies"), LoadContext->SerializedObject);
 		}
 		else
 		{
@@ -3200,7 +3196,7 @@ void FAsyncPackage::EventDrivenCreateExport(int32 LocalExportIndex)
 			LastTypeOfWorkPerformed = TEXT("EventDrivenCreateExport");
 			LastObjectWorkWasPerformedOn = nullptr;
 			check(Export.ObjectName != NAME_None || !(Export.ObjectFlags&RF_Public));
-			check(IsLoading());
+			check(LoadContext->HasStartedLoading());
 			if (Export.DynamicType == FObjectExport::EDynamicType::DynamicType)
 			{
 //native blueprint 
@@ -3558,9 +3554,8 @@ void FAsyncPackage::EventDrivenSerializeExport(int32 LocalExportIndex)
 
 		Object->ClearFlags(RF_NeedLoad);
 
-		FUObjectThreadContext& ThreadContext = FUObjectThreadContext::Get();
-		UObject* PrevSerializedObject = ThreadContext.SerializedObject;
-		ThreadContext.SerializedObject = Object;
+		UObject* PrevSerializedObject = LoadContext->SerializedObject;
+		LoadContext->SerializedObject = Object;
 		Linker->bForceSimpleIndexToObject = true;
 
 		// Find the Archetype object for the one we are loading. This is piped to GetArchetypeFromLoader
@@ -3585,7 +3580,7 @@ void FAsyncPackage::EventDrivenSerializeExport(int32 LocalExportIndex)
 		Linker->TemplateForGetArchetypeFromLoader = nullptr;
 
 		Object->SetFlags(RF_LoadCompleted);
-		ThreadContext.SerializedObject = PrevSerializedObject;
+		LoadContext->SerializedObject = PrevSerializedObject;
 		Linker->bForceSimpleIndexToObject = false;
 
 		if (FAA2->Tell() - Export.SerialOffset != Export.SerialSize)
@@ -4101,7 +4096,7 @@ void FAsyncPackage::Event_StartPostload()
 	AsyncPackageLoadingState = EAsyncPackageLoadingState::PostLoad_Etc;
 	EventDrivenLoadingComplete();
 	{
-		TArray<UObject*>& ObjLoaded = FUObjectThreadContext::Get().ObjLoaded;
+		TArray<UObject*>& ObjLoaded = LoadContext->GetObjectsLoaded();
 		ObjLoaded.Reserve(ObjLoaded.Num() + Linker->ExportMap.Num());
 		for (int32 LocalExportIndex = 0; LocalExportIndex < Linker->ExportMap.Num(); LocalExportIndex++)
 		{
@@ -5384,6 +5379,7 @@ FAsyncPackage::FAsyncPackage(const FAsyncPackageDesc& InDesc)
 , FinishObjectsTime(0.0)
 #endif // PERF_TRACK_DETAILED_ASYNC_STATS
 {
+	LoadContext = new FUObjectSerializeContext();
 	AddRequestID(InDesc.RequestID);
 }
 
@@ -5573,7 +5569,7 @@ void FAsyncPackage::BeginAsyncLoad()
 	}
 
 	// this won't do much during async loading except increase the load count which causes IsLoading to return true
-	BeginLoad();
+	BeginLoad(LoadContext);
 }
 
 /**
@@ -5585,7 +5581,7 @@ void FAsyncPackage::EndAsyncLoad()
 	check(IsAsyncLoading());
 
 	// this won't do much during async loading except decrease the load count which causes IsLoading to return false
-	EndLoad();
+	EndLoad(LoadContext);
 
 	if (IsInGameThread())
 	{
@@ -5931,7 +5927,7 @@ EAsyncPackageState::Type FAsyncPackage::CreateLinker()
 			{
 				FWeakAsyncPackagePtr WeakPtr(this);
 				check(Package);
-				Linker = FLinkerLoad::CreateLinkerAsync(Package, *PackageFileName, LinkerFlags
+				Linker = FLinkerLoad::CreateLinkerAsync(LoadContext, Package, *PackageFileName, LinkerFlags
 					, TFunction<void()>(
 						[WeakPtr]()
 				{
@@ -5951,7 +5947,7 @@ EAsyncPackageState::Type FAsyncPackage::CreateLinker()
 			}
 			else
 			{
-				Linker = FLinkerLoad::CreateLinkerAsync(Package, *PackageFileName, LinkerFlags, TFunction<void()>([]() {}));
+				Linker = FLinkerLoad::CreateLinkerAsync(LoadContext, Package, *PackageFileName, LinkerFlags, TFunction<void()>([]() {}));
 			}
 		}
 
@@ -6378,7 +6374,7 @@ EAsyncPackageState::Type FAsyncPackage::PreLoadObjects()
 	// GC can't run in here
 	FGCScopeGuard GCGuard;
 
-	TArray<UObject*>& ThreadObjLoaded = FUObjectThreadContext::Get().ObjLoaded;
+	TArray<UObject*>& ThreadObjLoaded = LoadContext->GetObjectsLoaded();
 	PackageObjLoaded.Append(ThreadObjLoaded);
 	ThreadObjLoaded.Reset();
 
@@ -6477,7 +6473,7 @@ EAsyncPackageState::Type FAsyncPackage::PostLoadObjects()
 	FUObjectThreadContext& ThreadContext = FUObjectThreadContext::Get();
 	TGuardValue<bool> GuardIsRoutingPostLoad(ThreadContext.IsRoutingPostLoad, true);
 
-	TArray<UObject*>& ThreadObjLoaded = ThreadContext.ObjLoaded;
+	TArray<UObject*>& ThreadObjLoaded = LoadContext->GetObjectsLoaded();
 	if (ThreadObjLoaded.Num())
 	{
 		// New objects have been loaded. They need to go through PreLoad first so exit now and come back after they've been preloaded.
@@ -6557,7 +6553,7 @@ EAsyncPackageState::Type FAsyncPackage::PostLoadDeferredObjects(double InTickSta
 	TGuardValue<bool> GuardIsRoutingPostLoad(PackageScope.ThreadContext.IsRoutingPostLoad, true);
 	FAsyncLoadingTickScope InAsyncLoadingTick;
 
-	TArray<UObject*>& ObjLoadedInPostLoad = PackageScope.ThreadContext.ObjLoaded;
+	TArray<UObject*>& ObjLoadedInPostLoad = LoadContext->GetObjectsLoaded();
 	TArray<UObject*> ObjLoadedInPostLoadLocal;
 
 	STAT(double PostLoadStartTime = FPlatformTime::Seconds());
@@ -6745,8 +6741,9 @@ EAsyncPackageState::Type FAsyncPackage::FinishObjects()
 	SCOPE_CYCLE_COUNTER(STAT_FAsyncPackage_FinishObjects);
 	LastObjectWorkWasPerformedOn	= nullptr;
 	LastTypeOfWorkPerformed			= TEXT("finishing all objects");
-		
-	TArray<UObject*>& ThreadObjLoaded = FUObjectThreadContext::Get().ObjLoaded;
+
+	check(!Linker || LoadContext == Linker->GetSerializeContext());		
+	TArray<UObject*>& ThreadObjLoaded = LoadContext->GetObjectsLoaded();
 
 	EAsyncLoadingResult::Type LoadingResult;
 	if (!bLoadHasFailed)
@@ -6793,7 +6790,7 @@ EAsyncPackageState::Type FAsyncPackage::FinishObjects()
 	FinishExternalReadDependenciesIndex = 0;
 
 	// Keep the linkers to close until we finish loading and it's safe to close them too
-	DelayedLinkerClosePackages = MoveTemp(FUObjectThreadContext::Get().DelayedLinkerClosePackages);
+	LoadContext->MoveDelayedLinkerClosePackages(DelayedLinkerClosePackages);
 
 	if (Linker)
 	{
diff --git a/Engine/Source/Runtime/CoreUObject/Private/Serialization/DuplicateDataReader.cpp b/Engine/Source/Runtime/CoreUObject/Private/Serialization/DuplicateDataReader.cpp
--- a/Engine/Source/Runtime/CoreUObject/Private/Serialization/DuplicateDataReader.cpp
+++ b/Engine/Source/Runtime/CoreUObject/Private/Serialization/DuplicateDataReader.cpp
@@ -40,8 +40,9 @@ FDuplicateDataReader::FDuplicateDataReader( class FUObjectAnnotationSparse<FDupl
 
 void FDuplicateDataReader::SerializeFail()
 {
-	FUObjectThreadContext& ThreadContext = FUObjectThreadContext::Get();
-	UE_LOG(LogObj, Fatal, TEXT("FDuplicateDataReader Overread. SerializedObject = %s SerializedProperty = %s"), *GetFullNameSafe(ThreadContext.SerializedObject), *GetFullNameSafe(GetSerializedProperty()));
+	FUObjectSerializeContext* LoadContext = GetSerializeContext();
+	check(LoadContext);
+	UE_LOG(LogObj, Fatal, TEXT("FDuplicateDataReader Overread. SerializedObject = %s SerializedProperty = %s"), *GetFullNameSafe(LoadContext->SerializedObject), *GetFullNameSafe(GetSerializedProperty()));
 }
 
 FArchive& FDuplicateDataReader::operator<<(FName& N)
diff --git a/Engine/Source/Runtime/CoreUObject/Private/UObject/Class.cpp b/Engine/Source/Runtime/CoreUObject/Private/UObject/Class.cpp
--- a/Engine/Source/Runtime/CoreUObject/Private/UObject/Class.cpp
+++ b/Engine/Source/Runtime/CoreUObject/Private/UObject/Class.cpp
@@ -40,6 +40,7 @@
 #include "Serialization/ArchiveUObjectFromStructuredArchive.h"
 #include "UObject/FrameworkObjectVersion.h"
 #include "UObject/GarbageCollection.h"
+#include "UObject/UObjectThreadContext.h"
 
 // This flag enables some expensive class tree validation that is meant to catch mutations of 
 // the class tree outside of SetSuperStruct. It has been disabled because loading blueprints 
@@ -579,16 +580,19 @@ void UStruct::Link(FArchive& Ar, bool bRelinkExistingProperties)
 	{
 		// Preload everything before we calculate size, as the preload may end up recursively linking things
 		UStruct* InheritanceSuper = GetInheritanceSuper();
-		if (InheritanceSuper)
+		if (Ar.IsLoading())
 		{
-			Ar.Preload(InheritanceSuper);			
-		}
+			if (InheritanceSuper)
+			{
+				Ar.Preload(InheritanceSuper);
+			}
 
-		for (UField* Field = Children; Field; Field = Field->Next)
-		{
-			if (!GEventDrivenLoaderEnabled || !Cast<UFunction>(Field))
+			for (UField* Field = Children; Field; Field = Field->Next)
 			{
-				Ar.Preload(Field);
+				if (!GEventDrivenLoaderEnabled || !Cast<UFunction>(Field))
+				{
+					Ar.Preload(Field);
+				}
 			}
 		}
 
diff --git a/Engine/Source/Runtime/CoreUObject/Private/UObject/Enum.cpp b/Engine/Source/Runtime/CoreUObject/Private/UObject/Enum.cpp
--- a/Engine/Source/Runtime/CoreUObject/Private/UObject/Enum.cpp
+++ b/Engine/Source/Runtime/CoreUObject/Private/UObject/Enum.cpp
@@ -12,6 +12,7 @@
 #include "UObject/DevObjectVersion.h"
 #include "UObject/CoreObjectVersion.h"
 #include "UObject/CoreRedirects.h"
+#include "UObject/LinkerLoad.h"
 
 DEFINE_LOG_CATEGORY_STATIC(LogEnum, Log, All);
 
@@ -485,8 +486,15 @@ int32 UEnum::GetIndexByNameString(const FString& InSearchString, EGetByNameFlags
 	else if (!!(Flags & EGetByNameFlags::ErrorIfNotFound) && !InSearchString.IsEmpty() && !InSearchString.Equals(FName().ToString(), StringComparisonMethod))
 	{
 		// None is passed in by blueprints at various points, isn't an error. Any other failed resolve should be fixed
-		FUObjectThreadContext& ThreadContext = FUObjectThreadContext::Get();
-		UE_LOG(LogEnum, Warning, TEXT("In asset '%s', there is an enum property of type '%s' with an invalid value of '%s'"), *GetPathNameSafe(ThreadContext.SerializedObject), *GetName(), *InSearchString);
+		UObject* SerializedObject = nullptr;
+		if (FLinkerLoad* Linker = GetLinker())
+		{
+			if (FUObjectSerializeContext* LoadContext = Linker->GetSerializeContext())
+			{
+				SerializedObject = LoadContext->SerializedObject;
+			}
+		}
+		UE_LOG(LogEnum, Warning, TEXT("In asset '%s', there is an enum property of type '%s' with an invalid value of '%s'"), *GetPathNameSafe(SerializedObject), *GetName(), *InSearchString);
 	}
 
 	return INDEX_NONE;
diff --git a/Engine/Source/Runtime/CoreUObject/Private/UObject/EnumProperty.cpp b/Engine/Source/Runtime/CoreUObject/Private/UObject/EnumProperty.cpp
--- a/Engine/Source/Runtime/CoreUObject/Private/UObject/EnumProperty.cpp
+++ b/Engine/Source/Runtime/CoreUObject/Private/UObject/EnumProperty.cpp
@@ -7,6 +7,7 @@
 #include "Templates/ChooseClass.h"
 #include "Templates/IsSigned.h"
 #include "Algo/Find.h"
+#include "UObject/LinkerLoad.h"
 
 namespace UE4EnumProperty_Private
 {
@@ -279,8 +280,15 @@ const TCHAR* UEnumProperty::ImportText_Internal(const TCHAR* InBuffer, void* Dat
 			// Enum could not be created from value. This indicates a bad value so
 			// return null so that the caller of ImportText can generate a more meaningful
 			// warning/error
-			FUObjectThreadContext& ThreadContext = FUObjectThreadContext::Get();
-			UE_LOG(LogClass, Warning, TEXT("In asset '%s', there is an enum property of type '%s' with an invalid value of '%s'"), *GetPathNameSafe(ThreadContext.SerializedObject), *Enum->GetName(), *Temp);
+			UObject* SerializedObject = nullptr;
+			if (FLinkerLoad* Linker = GetLinker())
+			{
+				if (FUObjectSerializeContext* LoadContext = Linker->GetSerializeContext())
+				{
+					SerializedObject = LoadContext->SerializedObject;
+				}
+			}
+			UE_LOG(LogClass, Warning, TEXT("In asset '%s', there is an enum property of type '%s' with an invalid value of '%s'"), *GetPathNameSafe(SerializedObject), *Enum->GetName(), *Temp);
 			return nullptr;
 		}
 	}
diff --git a/Engine/Source/Runtime/CoreUObject/Private/UObject/Linker.cpp b/Engine/Source/Runtime/CoreUObject/Private/UObject/Linker.cpp
--- a/Engine/Source/Runtime/CoreUObject/Private/UObject/Linker.cpp
+++ b/Engine/Source/Runtime/CoreUObject/Private/UObject/Linker.cpp
@@ -384,24 +384,28 @@ void DeleteLoader(FLinkerLoad* Loader)
 	FLinkerManager::Get().RemoveLinker(Loader);
 }
 
-static void LogGetPackageLinkerError(FArchive* LinkerArchive, const TCHAR* InFilename, const FText& InFullErrorMessage, const FText& InSummaryErrorMessage, UObject* InOuter, uint32 LoadFlags)
+static void LogGetPackageLinkerError(FArchive* LinkerArchive, FUObjectSerializeContext* LoadContext, const TCHAR* InFilename, const FText& InFullErrorMessage, const FText& InSummaryErrorMessage, UObject* InOuter, uint32 LoadFlags)
 {
 	static FName NAME_LoadErrors("LoadErrors");
 	struct Local
 	{
 		/** Helper function to output more detailed error info if available */
-		static void OutputErrorDetail(FArchive* InLinkerArchive, const FName& LogName)
+		static void OutputErrorDetail(FArchive* InLinkerArchive, FUObjectSerializeContext* InLoadContext, const FName& LogName)
 		{
-			FUObjectThreadContext& ThreadContext = FUObjectThreadContext::Get();
-			if (ThreadContext.SerializedObject && ThreadContext.SerializedImportLinker)
+			FUObjectSerializeContext* LoadContext = InLoadContext;
+			if (!LoadContext && InLinkerArchive)
+			{
+				LoadContext = InLinkerArchive->GetSerializeContext();
+			}
+			if (LoadContext && LoadContext->SerializedObject && LoadContext->SerializedImportLinker)
 			{
 				FMessageLog LoadErrors(LogName);
 
 				TSharedRef<FTokenizedMessage> Message = LoadErrors.Info();
 				Message->AddToken(FTextToken::Create(LOCTEXT("FailedLoad_Message", "Failed to load")));
-				Message->AddToken(FAssetNameToken::Create(ThreadContext.SerializedImportLinker->GetImportPathName(ThreadContext.SerializedImportIndex)));
+				Message->AddToken(FAssetNameToken::Create(LoadContext->SerializedImportLinker->GetImportPathName(LoadContext->SerializedImportIndex)));
 				Message->AddToken(FTextToken::Create(LOCTEXT("FailedLoad_Referenced", "Referenced by")));
-				Message->AddToken(FUObjectToken::Create(ThreadContext.SerializedObject));
+				Message->AddToken(FUObjectToken::Create(LoadContext->SerializedObject));
 				UProperty* SerializedProperty = InLinkerArchive ? InLinkerArchive->GetSerializedProperty() : nullptr;
 				if (SerializedProperty != nullptr)
 				{
@@ -455,14 +459,14 @@ static void LogGetPackageLinkerError(FArchive* LinkerArchive, const TCHAR* InFil
 				Message->AddToken(FAssetNameToken::Create(OuterPackageName));
 			}
 
-			Local::OutputErrorDetail(LinkerArchive, NAME_LoadErrors);
+			Local::OutputErrorDetail(LinkerArchive, LoadContext, NAME_LoadErrors);
 		}
 	}
 	else
 	{
 		if (!(LoadFlags & (LOAD_NoWarn | LOAD_Quiet)))
 		{
-			Local::OutputErrorDetail(LinkerArchive, NAME_LoadErrors);
+			Local::OutputErrorDetail(LinkerArchive, LoadContext, NAME_LoadErrors);
 		}
 
 		FFormatNamedArguments Arguments;
@@ -537,7 +541,8 @@ FLinkerLoad* GetPackageLinker
 	uint32			LoadFlags,
 	UPackageMap*	Sandbox,
 	FGuid*			CompatibleGuid,
-	FArchive*		InReaderOverride
+	FArchive*		InReaderOverride,
+	FUObjectSerializeContext* InExistingContext
 )
 {
 	// See if there is already a linker for this package.
@@ -560,7 +565,7 @@ FLinkerLoad* GetPackageLinker
 		{
 			// try to recover from this instead of throwing, it seems recoverable just by doing this
 			FText ErrorText(LOCTEXT("PackageResolveFailed", "Can't resolve asset name"));
-			LogGetPackageLinkerError(Result, InLongPackageName, ErrorText, ErrorText, InOuter, LoadFlags);
+			LogGetPackageLinkerError(Result, InExistingContext, InLongPackageName, ErrorText, ErrorText, InOuter, LoadFlags);
 			return nullptr;
 		}
 	
@@ -588,11 +593,11 @@ FLinkerLoad* GetPackageLinker
 			// In memory-only packages have no linker and this is ok.
 			if (!(LoadFlags & LOAD_AllowDll) && !InOuter->HasAnyPackageFlags(PKG_InMemoryOnly) && !FLinkerLoad::IsKnownMissingPackage(InOuter->GetFName()))
 			{
-				FUObjectThreadContext& ThreadContext = FUObjectThreadContext::Get();
 				FFormatNamedArguments Arguments;
+				FLinkerLoad* SerializedPackageLinker = InExistingContext ? InExistingContext->SerializedPackageLinker : nullptr;
 				Arguments.Add(TEXT("AssetName"), FText::FromString(PackageNameToLoad));
-				Arguments.Add(TEXT("PackageName"), FText::FromString(ThreadContext.SerializedPackageLinker ? *(ThreadContext.SerializedPackageLinker->Filename) : TEXT("NULL")));
-				LogGetPackageLinkerError(Result, ThreadContext.SerializedPackageLinker ? *ThreadContext.SerializedPackageLinker->Filename : nullptr,
+				Arguments.Add(TEXT("PackageName"), FText::FromString(SerializedPackageLinker ? *SerializedPackageLinker->Filename : TEXT("NULL")));
+				LogGetPackageLinkerError(Result, InExistingContext, SerializedPackageLinker ? *SerializedPackageLinker->Filename : nullptr,
 											FText::Format(LOCTEXT("PackageNotFound", "Can't find file for asset '{AssetName}' while loading {PackageName}."), Arguments),
 											LOCTEXT("PackageNotFoundShort", "Can't find file for asset."),
 											InOuter,
@@ -609,7 +614,7 @@ FLinkerLoad* GetPackageLinker
 		{
 			// try to recover from this instead of throwing, it seems recoverable just by doing this
 			FText ErrorText(LOCTEXT("PackageResolveFailed", "Can't resolve asset name"));
-			LogGetPackageLinkerError(Result, InLongPackageName, ErrorText, ErrorText, InOuter, LoadFlags);
+			LogGetPackageLinkerError(Result, InExistingContext, InLongPackageName, ErrorText, ErrorText, InOuter, LoadFlags);
 			return nullptr;
 		}
 
@@ -648,7 +653,7 @@ FLinkerLoad* GetPackageLinker
 				Arguments.Add(TEXT("Filename"), FText::FromString(InLongPackageName));
 
 				// try to recover from this instead of throwing, it seems recoverable just by doing this
-				LogGetPackageLinkerError(Result, InLongPackageName, FText::Format(LOCTEXT("FileNotFound", "Can't find file '{Filename}'"), Arguments), LOCTEXT("FileNotFoundShort", "Can't find file"), InOuter, LoadFlags);
+				LogGetPackageLinkerError(Result, InExistingContext, InLongPackageName, FText::Format(LOCTEXT("FileNotFound", "Can't find file '{Filename}'"), Arguments), LOCTEXT("FileNotFoundShort", "Can't find file"), InOuter, LoadFlags);
 			}
 			return nullptr;
 		}
@@ -673,7 +678,7 @@ FLinkerLoad* GetPackageLinker
 			{
 				FFormatNamedArguments Arguments;
 				Arguments.Add(TEXT("Filename"), FText::FromString(InLongPackageName));
-				LogGetPackageLinkerError(Result, InLongPackageName, FText::Format(LOCTEXT("FilenameToPackage", "Can't convert filename '{Filename}' to asset name"), Arguments), LOCTEXT("FilenameToPackageShort", "Can't convert filename to asset name"), InOuter, LoadFlags);
+				LogGetPackageLinkerError(Result, InExistingContext, InLongPackageName, FText::Format(LOCTEXT("FilenameToPackage", "Can't convert filename '{Filename}' to asset name"), Arguments), LOCTEXT("FilenameToPackageShort", "Can't convert filename to asset name"), InOuter, LoadFlags);
 				return nullptr;
 			}
 			InOuter = FilenamePkg;
@@ -694,7 +699,7 @@ FLinkerLoad* GetPackageLinker
 		FFormatNamedArguments Arguments;
 		Arguments.Add(TEXT("AssetName"), FText::FromString(InOuter->GetName()));
 
-		LogGetPackageLinkerError(Result, InLongPackageName, FText::Format(LOCTEXT("Sandbox", "Asset '{AssetName}' is not accessible in this sandbox"), Arguments), LOCTEXT("SandboxShort", "Asset is not accessible in this sandbox"), InOuter, LoadFlags);
+		LogGetPackageLinkerError(Result, InExistingContext, InLongPackageName, FText::Format(LOCTEXT("Sandbox", "Asset '{AssetName}' is not accessible in this sandbox"), Arguments), LOCTEXT("SandboxShort", "Asset is not accessible in this sandbox"), InOuter, LoadFlags);
 		return nullptr;
 	}
 #endif
@@ -702,12 +707,17 @@ FLinkerLoad* GetPackageLinker
 	// Create new linker.
 	if( !Result )
 	{
-		check(IsLoading());
+		// @todo: RM so never create the context below?
+		check(InExistingContext && InExistingContext->HasStartedLoading());
 
 		// we will already have found the filename above
 		check(NewFilename.Len() > 0);
-
-		Result = FLinkerLoad::CreateLinker( InOuter, *NewFilename, LoadFlags, InReaderOverride);
+		TRefCountPtr<FUObjectSerializeContext> LoadContext(InExistingContext ? InExistingContext : new FUObjectSerializeContext());
+		Result = FLinkerLoad::CreateLinker(LoadContext, InOuter, *NewFilename, LoadFlags, InReaderOverride);
+	}
+	else if (InExistingContext)
+	{
+		Result->SetSerializeContext(InExistingContext);
 	}
 
 	if ( !Result && CreatedPackage )
@@ -728,7 +738,7 @@ FLinkerLoad* GetPackageLinker
 		Arguments.Add(TEXT("AssetName"), FText::FromString(InOuter->GetName()));
 
 		// This should never fire, because FindPackageFile should never return an incompatible file
-		LogGetPackageLinkerError(Result, InLongPackageName, FText::Format(LOCTEXT("PackageVersion", "Asset '{AssetName}' version mismatch"), Arguments), LOCTEXT("PackageVersionShort", "Asset version mismatch"), InOuter, LoadFlags);
+		LogGetPackageLinkerError(Result, InExistingContext, InLongPackageName, FText::Format(LOCTEXT("PackageVersion", "Asset '{AssetName}' version mismatch"), Arguments), LOCTEXT("PackageVersionShort", "Asset version mismatch"), InOuter, LoadFlags);
 		return nullptr;
 	}
 
diff --git a/Engine/Source/Runtime/CoreUObject/Private/UObject/LinkerLoad.cpp b/Engine/Source/Runtime/CoreUObject/Private/UObject/LinkerLoad.cpp
--- a/Engine/Source/Runtime/CoreUObject/Private/UObject/LinkerLoad.cpp
+++ b/Engine/Source/Runtime/CoreUObject/Private/UObject/LinkerLoad.cpp
@@ -301,21 +301,22 @@ void FLinkerLoad::CreateActiveRedirectsMap(const FString& GEngineIniName)
 
 FScopedCreateImportCounter::FScopedCreateImportCounter(FLinkerLoad* Linker, int32 Index)
 {
-	FUObjectThreadContext& ThreadContext = FUObjectThreadContext::Get();
+	LoadContext = Linker->GetSerializeContext();
+	check(LoadContext);
+
 	// Remember the old linker and index
-	PreviousLinker = ThreadContext.SerializedImportLinker;
-	PreviousIndex = ThreadContext.SerializedImportIndex;
+	PreviousLinker = LoadContext->SerializedImportLinker;
+	PreviousIndex = LoadContext->SerializedImportIndex;
 	// Remember the current linker and index.
-	ThreadContext.SerializedImportLinker = Linker;
-	ThreadContext.SerializedImportIndex = Index;
+	LoadContext->SerializedImportLinker = Linker;
+	LoadContext->SerializedImportIndex = Index;
 }
 
 FScopedCreateImportCounter::~FScopedCreateImportCounter()
 {
-	FUObjectThreadContext& ThreadContext = FUObjectThreadContext::Get();
 	// Restore old values
-	ThreadContext.SerializedImportLinker = PreviousLinker;
-	ThreadContext.SerializedImportIndex = PreviousIndex;
+	LoadContext->SerializedImportLinker = PreviousLinker;
+	LoadContext->SerializedImportIndex = PreviousIndex;
 }
 
 
@@ -329,48 +330,33 @@ struct FScopedCreateExportCounter
 	 */
 	FScopedCreateExportCounter(FLinkerLoad* Linker, int32 Index)
 	{
-		FUObjectThreadContext& ThreadContext = FUObjectThreadContext::Get();
+		LoadContext = Linker->GetSerializeContext();
+		check(LoadContext);
+
 		// Remember the old linker and index
-		PreviousLinker = ThreadContext.SerializedExportLinker;
-		PreviousIndex = ThreadContext.SerializedExportIndex;
+		PreviousLinker = LoadContext->SerializedExportLinker;
+		PreviousIndex = LoadContext->SerializedExportIndex;
 		// Remember the current linker and index.
-		ThreadContext.SerializedExportLinker = Linker;
-		ThreadContext.SerializedExportIndex = Index;
+		LoadContext->SerializedExportLinker = Linker;
+		LoadContext->SerializedExportIndex = Index;
 	}
 
 	/** Destructor. Called upon CreateImport() exit. */
 	~FScopedCreateExportCounter()
 	{
-		FUObjectThreadContext& ThreadContext = FUObjectThreadContext::Get();
 		// Restore old values
-		ThreadContext.SerializedExportLinker = PreviousLinker;
-		ThreadContext.SerializedExportIndex = PreviousIndex;
+		LoadContext->SerializedExportLinker = PreviousLinker;
+		LoadContext->SerializedExportIndex = PreviousIndex;
 	}
 
+	/** Current load context object */
+	FUObjectSerializeContext* LoadContext;
 	/** Poreviously stored linker */
 	FLinkerLoad* PreviousLinker;
 	/** Previously stored index */
 	int32 PreviousIndex;
 };
 
-/**
- * Exception save guard to ensure SerializedPackageLinker is reset after this
- * class goes out of scope.
- */
-class FSerializedPackageLinkerGuard
-{
-	/** Pointer to restore to after going out of scope. */
-	FLinkerLoad* PrevSerializedPackageLinker;
-public:
-	FSerializedPackageLinkerGuard() 
-		: PrevSerializedPackageLinker(FUObjectThreadContext::Get().SerializedPackageLinker)
-	{}
-	~FSerializedPackageLinkerGuard() 
-	{ 
-		FUObjectThreadContext::Get().SerializedPackageLinker = PrevSerializedPackageLinker;
-	}
-};
-
 namespace FLinkerDefs
 {
 	/** Number of progress steps for reporting status to a GUI while loading packages */
@@ -421,8 +407,10 @@ void FLinkerLoad::StaticInit(UClass* InUTexture2DStaticClass)
  *
  * @return	new FLinkerLoad object for Parent/ Filename
  */
-FLinkerLoad* FLinkerLoad::CreateLinker(UPackage* Parent, const TCHAR* Filename, uint32 LoadFlags, FArchive* InLoader /*= nullptr*/)
+FLinkerLoad* FLinkerLoad::CreateLinker(FUObjectSerializeContext* LoadContext, UPackage* Parent, const TCHAR* Filename, uint32 LoadFlags, FArchive* InLoader /*= nullptr*/)
 {
+	check(LoadContext);
+
 #if USE_CIRCULAR_DEPENDENCY_LOAD_DEFERRING
 	// we don't want the linker permanently created with the 
 	// DeferDependencyLoads flag (we also want to be able to determine if the 
@@ -437,7 +425,7 @@ FLinkerLoad* FLinkerLoad::CreateLinker(UPackage* Parent, const TCHAR* Filename,
 	LoadFlags &= ~LOAD_DeferDependencyLoads;
 #endif // USE_CIRCULAR_DEPENDENCY_LOAD_DEFERRING
 
-	FLinkerLoad* Linker = CreateLinkerAsync(Parent, Filename, LoadFlags
+	FLinkerLoad* Linker = CreateLinkerAsync(LoadContext, Parent, Filename, LoadFlags
 		, TFunction<void()>([](){})
 		);
 	{
@@ -460,8 +448,7 @@ FLinkerLoad* FLinkerLoad::CreateLinker(UPackage* Parent, const TCHAR* Filename,
 			Linker->ResetStatusInfo();
 		}
 
-		FSerializedPackageLinkerGuard Guard;
-		FUObjectThreadContext::Get().SerializedPackageLinker = Linker;
+		TGuardValue<FLinkerLoad*> SerializedPackageLinkerGuard(LoadContext->SerializedPackageLinker, Linker);
 		if (Linker->Tick(0.f, false, false) == LINKER_Failed)
 		{
 			return nullptr;
@@ -564,13 +551,9 @@ void FLinkerLoad::PRIVATE_PatchNewObjectIntoExport(UObject* OldObject, UObject*
 		NewObject->SetLinker(OldObjectLinker, CachedLinkerIndex);
 		ObjExport.Object = NewObject;
 
-		TArray<UObject*>& ObjLoaded = FUObjectThreadContext::Get().ObjLoaded;
 		// If the object was in the ObjLoaded queue (exported, but not yet serialized), swap out for our new object
-		const int32 ObjLoadedIdx = ObjLoaded.Find(OldObject);
-		if( ObjLoadedIdx != INDEX_NONE )
-		{
-			ObjLoaded[ObjLoadedIdx] = NewObject;
-		}
+		check(OldObjectLinker->GetSerializeContext());
+		OldObjectLinker->GetSerializeContext()->PRIVATE_PatchNewObjectIntoExport(OldObject, NewObject);
 	}
 }
 
@@ -612,7 +595,7 @@ FName FLinkerLoad::FindSubobjectRedirectName(const FName& Name, UClass* Class)
  *
  * @return	new FLinkerLoad object for Parent/ Filename
  */
-FLinkerLoad* FLinkerLoad::CreateLinkerAsync( UPackage* Parent, const TCHAR* Filename, uint32 LoadFlags 
+FLinkerLoad* FLinkerLoad::CreateLinkerAsync(FUObjectSerializeContext* LoadContext, UPackage* Parent, const TCHAR* Filename, uint32 LoadFlags
 	, TFunction<void()>&& InSummaryReadyCallback
 	)
 {
@@ -640,6 +623,7 @@ FLinkerLoad* FLinkerLoad::CreateLinkerAsync( UPackage* Parent, const TCHAR* File
 			LoadFlags |= LOAD_Async;
 		}
 		Linker = new FLinkerLoad(Parent, Filename, LoadFlags );
+		Linker->SetSerializeContext(LoadContext);
 		Parent->LinkerLoad = Linker;
 		if (GEventDrivenLoaderEnabled && Linker)
 		{
@@ -652,6 +636,15 @@ FLinkerLoad* FLinkerLoad::CreateLinkerAsync( UPackage* Parent, const TCHAR* File
 	return Linker;
 }
 
+void FLinkerLoad::SetSerializeContext(FUObjectSerializeContext* InLoadContext)
+{
+	CurrentLoadContext = InLoadContext;
+}
+FUObjectSerializeContext* FLinkerLoad::GetSerializeContext()
+{
+	return CurrentLoadContext;
+}
+
 /**
  * Ticks an in-flight linker and spends InTimeLimit seconds on creation. This is a soft time limit used
  * if bInUseTimeLimit is true.
@@ -2283,7 +2276,10 @@ void FLinkerLoad::GatherImportDependencies(int32 ImportIndex, TSet<FDependencyRe
 		return;
 	}
 
-	BeginLoad(TEXT("GatherImportDependencies"));
+	FUObjectSerializeContext* LoadContext = GetSerializeContext();
+	check(LoadContext);
+
+	BeginLoad(LoadContext, TEXT("GatherImportDependencies"));
 
 	// load the linker and find export in sourcelinker
 	if (Import.SourceLinker == NULL || Import.SourceIndex == INDEX_NONE)
@@ -2320,7 +2316,7 @@ void FLinkerLoad::GatherImportDependencies(int32 ImportIndex, TSet<FDependencyRe
 					Import.SourceLinker, Import.SourceIndex, 
 					*GetImportFullName(ImportIndex), *this->Filename );
 			}
-			EndLoad();
+			EndLoad(LoadContext);
 			return;
 		}
 
@@ -2347,7 +2343,7 @@ void FLinkerLoad::GatherImportDependencies(int32 ImportIndex, TSet<FDependencyRe
 		NewRef.ExportIndex = Import.SourceIndex;
 	}
 
-	EndLoad();
+	EndLoad(LoadContext);
 
 	// Add to set and recurse if not already present.
 	bool bIsAlreadyInSet = false;
@@ -2456,7 +2452,8 @@ FLinkerLoad::EVerifyResult FLinkerLoad::VerifyImport(int32 ImportIndex)
 
 						// now, fake our Import to be what the redirector pointed to
 						Import.XObject = Redir->DestinationObject;
-						FUObjectThreadContext::Get().ImportCount++;
+						check(CurrentLoadContext);
+						CurrentLoadContext->IncrementImportCount();
 						FLinkerManager::Get().AddLoaderWithNewImports(this);
 					}
 				}
@@ -2546,7 +2543,7 @@ FLinkerLoad::EVerifyResult FLinkerLoad::VerifyImport(int32 ImportIndex)
 }
 
 // Internal Load package call so that we can pass the linker that requested this package as an import dependency
-UPackage* LoadPackageInternal(UPackage* InOuter, const TCHAR* InLongPackageName, uint32 LoadFlags, FLinkerLoad* ImportLinker, FArchive* InReaderOverride);
+UPackage* LoadPackageInternal(UPackage* InOuter, const TCHAR* InLongPackageName, uint32 LoadFlags, FLinkerLoad* ImportLinker, FArchive* InReaderOverride, FUObjectSerializeContext* InLoadContext);
 
 /**
  * Safely verify that an import in the ImportMap points to a good object. This decides whether or not
@@ -2627,7 +2624,7 @@ bool FLinkerLoad::VerifyImportInner(const int32 ImportIndex, FString& WarningSuf
 
 			// we now fully load the package that we need a single export from - however, we still use CreatePackage below as it handles all cases when the package
 			// didn't exist (native only), etc		
-			TmpPkg = LoadPackageInternal(NULL, *Import.ObjectName.ToString(), InternalLoadFlags | LOAD_IsVerifying, this, nullptr);
+			TmpPkg = LoadPackageInternal(NULL, *Import.ObjectName.ToString(), InternalLoadFlags | LOAD_IsVerifying, this, nullptr, GetSerializeContext());
 		}
 
 #if WITH_EDITOR
@@ -2873,7 +2870,8 @@ bool FLinkerLoad::VerifyImportInner(const int32 ImportIndex, FString& WarningSuf
 			{
 				// except if we are looking for _the_ package...in which case we are looking for TmpPkg, so we are done
 				Import.XObject = TmpPkg;
-				FUObjectThreadContext::Get().ImportCount++;
+				check(CurrentLoadContext);
+				CurrentLoadContext->IncrementImportCount();
 				FLinkerManager::Get().AddLoaderWithNewImports(this);
 				return false;
 			}
@@ -2926,7 +2924,8 @@ bool FLinkerLoad::VerifyImportInner(const int32 ImportIndex, FString& WarningSuf
 				if (FindObject != NULL && ((LoadFlags & LOAD_FindIfFail) || bIsInMemoryOnlyOrNativeTransient))
 				{
 					Import.XObject = FindObject;
-					FUObjectThreadContext::Get().ImportCount++;
+					check(CurrentLoadContext);
+					CurrentLoadContext->IncrementImportCount();
 					FLinkerManager::Get().AddLoaderWithNewImports(this);
 				}
 				else
@@ -3255,7 +3254,7 @@ UObject* FLinkerLoad::Create( UClass* ObjectClass, FName ObjectName, UObject* Ou
 		Index = FindExportIndex(UObjectRedirector::StaticClass()->GetFName(), NAME_CoreUObject, ObjectName, OuterPackageIndex);
 		if (Index == INDEX_NONE)
 		{
-			Index = FindExportIndex(UObjectRedirector::StaticClass()->GetFName(), GLongCoreUObjectPackageName, ObjectName, OuterPackageIndex);			
+			Index = FindExportIndex(UObjectRedirector::StaticClass()->GetFName(), GLongCoreUObjectPackageName, ObjectName, OuterPackageIndex);
 		}
 
 		// if we found a redirector, create it, and move on down the line
@@ -3461,10 +3460,10 @@ void FLinkerLoad::Preload( UObject* Object )
 						}
 #endif // USE_CIRCULAR_DEPENDENCY_LOAD_DEFERRING
 
-						FUObjectThreadContext& ThreadContext = FUObjectThreadContext::Get();
+						check(CurrentLoadContext);
 						// Maintain the current SerializedObjects.
-						UObject* PrevSerializedObject = ThreadContext.SerializedObject;
-						ThreadContext.SerializedObject = Object;
+						UObject* PrevSerializedObject = CurrentLoadContext->SerializedObject;
+						CurrentLoadContext->SerializedObject = Object;
 
 #if WITH_EDITOR && WITH_TEXT_ARCHIVE_SUPPORT
 						if (IsTextFormat())
@@ -3490,7 +3489,7 @@ void FLinkerLoad::Preload( UObject* Object )
 						}
 
 						Object->SetFlags(RF_LoadCompleted);
-						ThreadContext.SerializedObject = PrevSerializedObject;
+						CurrentLoadContext->SerializedObject = PrevSerializedObject;
 					}
 					else
 					{
@@ -3501,11 +3500,11 @@ void FLinkerLoad::Preload( UObject* Object )
 						FArchive::FScopeAddDebugData N(*this, Object->GetFName());
 						FArchive::FScopeAddDebugData C(*this, Object->GetClass()->GetFName());
 #endif
+						check(CurrentLoadContext);
 
-						FUObjectThreadContext& ThreadContext = FUObjectThreadContext::Get();
 						// Maintain the current SerializedObjects.
-						UObject* PrevSerializedObject = ThreadContext.SerializedObject;
-						ThreadContext.SerializedObject = Object;
+						UObject* PrevSerializedObject = CurrentLoadContext->SerializedObject;
+						CurrentLoadContext->SerializedObject = Object;
 
 #if WITH_EDITOR && WITH_TEXT_ARCHIVE_SUPPORT
 						if (IsTextFormat())
@@ -3540,7 +3539,7 @@ void FLinkerLoad::Preload( UObject* Object )
 						}
 
 						Object->SetFlags(RF_LoadCompleted);
-						ThreadContext.SerializedObject = PrevSerializedObject;
+						CurrentLoadContext->SerializedObject = PrevSerializedObject;
 					}
 				}
 
@@ -3674,7 +3673,9 @@ void FLinkerLoad::Preload( UObject* Object )
 				if (Object->HasAnyFlags(RF_ClassDefaultObject) && Object->GetClass()->HasAnyClassFlags(CLASS_CompiledFromBlueprint))
 				{
 					Object->SetFlags(RF_NeedPostLoad|RF_WasLoaded);
-					FUObjectThreadContext::Get().ObjLoaded.Add(Object);
+					check(LinkerRoot && LinkerRoot == Object->GetOutermost());
+					check(CurrentLoadContext);
+					CurrentLoadContext->AddLoadedObject(Object);
 				}
 			}
 		}
@@ -3772,7 +3773,8 @@ UObject* FLinkerLoad::CreateExport( int32 Index )
 					// The CDO is created (in a custom code) at the end of loading (when it's safe to solve cyclic dependencies).
 					if (!DynamicClass->GetDefaultObject(false))
 					{
-						FUObjectThreadContext::Get().ObjLoaded.Add(Export.Object);
+						check(CurrentLoadContext);
+						CurrentLoadContext->AddLoadedObject(Export.Object);
 					}
 				}
 			}
@@ -3908,7 +3910,9 @@ UObject* FLinkerLoad::CreateExport( int32 Index )
 			// will find and return an existing object if one exists and only create a new one if there doesn't.
 			Export.Object = CreatePackage( NULL, *Export.ObjectName.ToString() );
 			check(Export.Object);
-			FUObjectThreadContext::Get().ForcedExportCount++;
+			check(CurrentLoadContext);
+			CurrentLoadContext->IncrementForcedExportCount();
+			FLinkerManager::Get().AddLoaderWithForcedExports(this);
 		}
 		else
 		{
@@ -4056,7 +4060,9 @@ UObject* FLinkerLoad::CreateExport( int32 Index )
 			// Mark that we need to dissociate forced exports later on if we are a forced export.
 			if (Export.bForcedExport)
 			{
-				FUObjectThreadContext::Get().ForcedExportCount++;
+				check(CurrentLoadContext);
+				CurrentLoadContext->IncrementForcedExportCount();
+				FLinkerManager::Get().AddLoaderWithForcedExports(this);
 			}
 			// Associate linker with object to avoid detachment mismatches.
 			else
@@ -4065,7 +4071,8 @@ UObject* FLinkerLoad::CreateExport( int32 Index )
 
 				// If this object was allocated but never loaded (components created by a constructor) make sure it gets loaded
 				// Don't do this for any packages that have previously fully loaded as they may have in memory changes
-				FUObjectThreadContext::Get().ObjLoaded.AddUnique(Export.Object);
+				check(CurrentLoadContext);
+				CurrentLoadContext->AddLoadedObject(Export.Object);
 				if (!Export.Object->HasAnyFlags(RF_LoadCompleted) && !LinkerRoot->IsFullyLoaded())
 				{
 					check(!GEventDrivenLoaderEnabled || !EVENT_DRIVEN_ASYNC_LOAD_ACTIVE_AT_RUNTIME);
@@ -4095,9 +4102,10 @@ UObject* FLinkerLoad::CreateExport( int32 Index )
 				// A redirector has been found, replace this export with it.
 				LoadClass = UObjectRedirector::StaticClass();
 				// Create new import for UObjectRedirector class
-				FObjectImport* RedirectorImport = new( ImportMap )FObjectImport( UObjectRedirector::StaticClass() );
-				FLinkerManager::Get().AddLoaderWithNewImports(this);
-				FUObjectThreadContext::Get().ImportCount++;
+				FObjectImport* RedirectorImport = new(ImportMap)FObjectImport(UObjectRedirector::StaticClass());
+				check(CurrentLoadContext);
+				CurrentLoadContext->IncrementImportCount();
+				FLinkerManager::Get().AddLoaderWithNewImports(this);				
 				Export.ClassIndex = FPackageIndex::FromImport(ImportMap.Num() - 1);
 				Export.Object = Redirector;
 				Export.Object->SetLinker( this, Index );
@@ -4231,7 +4239,8 @@ UObject* FLinkerLoad::CreateExport( int32 Index )
 				// so add the object to the list of objects that need to be loaded, which will be processed
 				// in EndLoad()
 				Export.Object->SetLinker( this, Index );
-				FUObjectThreadContext::Get().ObjLoaded.Add(Export.Object);
+				check(CurrentLoadContext);
+				CurrentLoadContext->AddLoadedObject(Export.Object);
 			}
 		}
 		else
@@ -4288,7 +4297,9 @@ UObject* FLinkerLoad::CreateExport( int32 Index )
 			// Mark that we need to dissociate forced exports later on.
 			if( Export.bForcedExport )
 			{
-				FUObjectThreadContext::Get().ForcedExportCount++;
+				check(CurrentLoadContext);
+				CurrentLoadContext->IncrementForcedExportCount();
+				FLinkerManager::Get().AddLoaderWithForcedExports(this);
 			}
 		}
 	}
@@ -4421,7 +4432,8 @@ UObject* FLinkerLoad::CreateImport( int32 Index )
 					{		
 						// Associate import and indicate that we associated an import for later cleanup.
 						Import.XObject = FindObject;
-						FUObjectThreadContext::Get().ImportCount++;
+						check(CurrentLoadContext);
+						CurrentLoadContext->IncrementImportCount();
 						FLinkerManager::Get().AddLoaderWithNewImports(this);
 					}
 				}
@@ -4458,7 +4470,8 @@ UObject* FLinkerLoad::CreateImport( int32 Index )
 					}
 				}
 #endif
-				FUObjectThreadContext::Get().ImportCount++;
+				check(CurrentLoadContext);
+				CurrentLoadContext->IncrementImportCount();
 				FLinkerManager::Get().AddLoaderWithNewImports(this);
 			}
 		}
@@ -4623,7 +4636,8 @@ void FLinkerLoad::Detach()
 	FLinkerManager::Get().RemoveLoaderFromObjectLoadersAndLoadersWithNewImports(this);
 	if (!FPlatformProperties::HasEditorOnlyData())
 	{
-		FUObjectThreadContext::Get().DelayedLinkerClosePackages.Remove(this);
+		check(CurrentLoadContext);
+		CurrentLoadContext->RemoveDelayedLinkerClosePackage(this);
 	}
 
 	delete StructuredArchive;
@@ -4720,6 +4734,7 @@ FArchive& FLinkerLoad::operator<<( UObject*& Object )
 		return *this;
 	}
 
+	check(CurrentLoadContext != nullptr);
 	UObject* Temporary = NULL;
 	Temporary = IndexToObject( Index );
 
@@ -4831,7 +4846,8 @@ UObject* FLinkerLoad::GetArchetypeFromLoader(const UObject* Obj)
 {
 	if (GEventDrivenLoaderEnabled)
 	{
-		check(!TemplateForGetArchetypeFromLoader || FUObjectThreadContext::Get().SerializedObject == Obj);
+		check(CurrentLoadContext);
+		check(!TemplateForGetArchetypeFromLoader || CurrentLoadContext->SerializedObject == Obj);
 		return TemplateForGetArchetypeFromLoader;
 	}
 	else
diff --git a/Engine/Source/Runtime/CoreUObject/Private/UObject/LinkerManager.cpp b/Engine/Source/Runtime/CoreUObject/Private/UObject/LinkerManager.cpp
--- a/Engine/Source/Runtime/CoreUObject/Private/UObject/LinkerManager.cpp
+++ b/Engine/Source/Runtime/CoreUObject/Private/UObject/LinkerManager.cpp
@@ -144,43 +144,38 @@ void FLinkerManager::ResetLoaders(UObject* InPkg)
 
 void FLinkerManager::DissociateImportsAndForcedExports()
 {
-	int32& ImportCount = FUObjectThreadContext::Get().ImportCount;
-	if (ImportCount != 0)
 	{
 		// In cooked builds linkers don't stick around long enough to make this worthwhile
 		TSet<FLinkerLoad*> LocalLoadersWithNewImports;
 		GetLoadersWithNewImportsAndEmpty(LocalLoadersWithNewImports);
-		if (LocalLoadersWithNewImports.Num())
+
+		for (FLinkerLoad* Linker : LocalLoadersWithNewImports)
 		{
-			for (FLinkerLoad* Linker : LocalLoadersWithNewImports)
+			for (int32 ImportIndex = 0; ImportIndex < Linker->ImportMap.Num(); ImportIndex++)
 			{
-				for (int32 ImportIndex = 0; ImportIndex < Linker->ImportMap.Num(); ImportIndex++)
+				FObjectImport& Import = Linker->ImportMap[ImportIndex];
+				if (Import.XObject && !Import.XObject->IsNative())
 				{
-					FObjectImport& Import = Linker->ImportMap[ImportIndex];
-					if (Import.XObject && !Import.XObject->IsNative())
-					{
-						Import.XObject = nullptr;
-					}
-					Import.SourceLinker = nullptr;
-					// when the SourceLinker is reset, the SourceIndex must also be reset, or recreating
-					// an import that points to a redirector will fail to find the redirector
-					Import.SourceIndex = INDEX_NONE;
+					Import.XObject = nullptr;
 				}
+				Import.SourceLinker = nullptr;
+				// when the SourceLinker is reset, the SourceIndex must also be reset, or recreating
+				// an import that points to a redirector will fail to find the redirector
+				Import.SourceIndex = INDEX_NONE;
+			}
+			if (Linker->GetSerializeContext())
+			{
+				Linker->GetSerializeContext()->ResetImportCount();
 			}
 		}
-		ImportCount = 0;
 	}
 
-	int32& ForcedExportCount = FUObjectThreadContext::Get().ForcedExportCount;
-	if (ForcedExportCount)
 	{		
-		TSet<FLinkerLoad*> LocalLoaders;
-		GetLoaders(LocalLoaders);
-		for (FLinkerLoad* Linker : LocalLoaders)
+		TSet<FLinkerLoad*> LocalLoadersWithForcedExports;
+		GetLoadersWithForcedExportsAndEmpty(LocalLoadersWithForcedExports);
+		for (FLinkerLoad* Linker : LocalLoadersWithForcedExports)
 		{
-			//@todo optimization: only dissociate exports for loaders that had forced exports created
-			//@todo optimization: since the last time this function was called.
-			for (auto& Export : Linker->ExportMap)
+			for (FObjectExport& Export : Linker->ExportMap)
 			{
 				if (Export.Object && Export.bForcedExport)
 				{
@@ -188,9 +183,12 @@ void FLinkerManager::DissociateImportsAndForcedExports()
 					Export.ResetObject();
 				}
 			}
-		}	
-		ForcedExportCount = 0;
-	}	
+			if (Linker->GetSerializeContext())
+			{
+				Linker->GetSerializeContext()->ResetForcedExports();
+			}
+		}
+	}
 }
 
 void FLinkerManager::DeleteLinkers()
diff --git a/Engine/Source/Runtime/CoreUObject/Private/UObject/LinkerManager.h b/Engine/Source/Runtime/CoreUObject/Private/UObject/LinkerManager.h
--- a/Engine/Source/Runtime/CoreUObject/Private/UObject/LinkerManager.h
+++ b/Engine/Source/Runtime/CoreUObject/Private/UObject/LinkerManager.h
@@ -67,6 +67,22 @@ public:
 		LoadersWithNewImports.Add(LinkerLoad);
 	}
 
+	FORCEINLINE void GetLoadersWithForcedExportsAndEmpty(TSet<FLinkerLoad*>& OutLoaders)
+	{
+#if THREADSAFE_UOBJECTS
+		FScopeLock ObjectLoadersLock(&LoadersWithForcedExportsCritical);
+#endif
+		OutLoaders = MoveTemp(LoadersWithForcedExports);
+	}
+
+	FORCEINLINE void AddLoaderWithForcedExports(FLinkerLoad* LinkerLoad)
+	{
+#if THREADSAFE_UOBJECTS
+		FScopeLock ObjectLoadersLock(&LoadersWithForcedExportsCritical);
+#endif
+		LoadersWithForcedExports.Add(LinkerLoad);
+	}
+
 #if !UE_BUILD_SHIPPING && !UE_BUILD_TEST
 	FORCEINLINE TArray<FLinkerLoad*>& GetLiveLinkers()
 	{
@@ -110,6 +126,12 @@ private:
 	TArray<FLinkerLoad*> LiveLinkers;
 #endif
 	
+	/** List of loaders that have new imports **/
+	TSet<FLinkerLoad*> LoadersWithForcedExports;
+#if THREADSAFE_UOBJECTS
+	FCriticalSection LoadersWithForcedExportsCritical;
+#endif
+
 	/** List of linkers to delete **/
 	TSet<FLinkerLoad*>	PendingCleanupList;
 #if THREADSAFE_UOBJECTS
diff --git a/Engine/Source/Runtime/CoreUObject/Private/UObject/LinkerSave.cpp b/Engine/Source/Runtime/CoreUObject/Private/UObject/LinkerSave.cpp
--- a/Engine/Source/Runtime/CoreUObject/Private/UObject/LinkerSave.cpp
+++ b/Engine/Source/Runtime/CoreUObject/Private/UObject/LinkerSave.cpp
@@ -276,10 +276,11 @@ FArchive& FLinkerSave::operator<<( FName& InName )
 {
 	int32 Save = MapName(InName);
 
+	check(GetSerializeContext());
 	ensureMsgf(Save != INDEX_NONE, TEXT("Name \"%s\" is not mapped when saving %s (object: %s, property: %s)"), 
 		*InName.ToString(),
 		*GetArchiveName(),
-		*FUObjectThreadContext::Get().SerializedObject->GetFullName(),
+		*GetSerializeContext()->SerializedObject->GetFullName(),
 		*GetSerializedProperty()->GetFullName());
 
 	int32 Number = InName.GetNumber();
@@ -302,4 +303,18 @@ FArchive& FLinkerSave::operator<<(FLazyObjectPtr& LazyObjectPtr)
 	FUniqueObjectGuid ID;
 	ID = LazyObjectPtr.GetUniqueID();
 	return *this << ID;
+}
+
+void FLinkerSave::SetSerializeContext(FUObjectSerializeContext* InLoadContext)
+{
+	SaveContext = InLoadContext;
+	if (Saver)
+	{
+		Saver->SetSerializeContext(InLoadContext);
+	}
+}
+
+FUObjectSerializeContext* FLinkerSave::GetSerializeContext()
+{
+	return SaveContext;
 }
\ No newline at end of file
diff --git a/Engine/Source/Runtime/CoreUObject/Private/UObject/Obj.cpp b/Engine/Source/Runtime/CoreUObject/Private/UObject/Obj.cpp
--- a/Engine/Source/Runtime/CoreUObject/Private/UObject/Obj.cpp
+++ b/Engine/Source/Runtime/CoreUObject/Private/UObject/Obj.cpp
@@ -4246,7 +4246,6 @@ void CleanupCachedArchetypes();
 //
 void StaticExit()
 {
-	check(FUObjectThreadContext::Get().ObjLoaded.Num() == 0);
 	if (UObjectInitialized() == false)
 	{
 		return;
@@ -4336,7 +4335,6 @@ void StaticExit()
 
 	UObjectBaseShutdown();
 	// Empty arrays to prevent falsely-reported memory leaks.
-	FUObjectThreadContext::Get().ObjLoaded.Empty();
 	FDeferredMessageLog::Cleanup();
 	CleanupGCArrayPools();
 	CleanupLinkerAnnotations();
diff --git a/Engine/Source/Runtime/CoreUObject/Private/UObject/Package.cpp b/Engine/Source/Runtime/CoreUObject/Private/UObject/Package.cpp
--- a/Engine/Source/Runtime/CoreUObject/Private/UObject/Package.cpp
+++ b/Engine/Source/Runtime/CoreUObject/Private/UObject/Package.cpp
@@ -7,6 +7,7 @@
 #include "Misc/PackageName.h"
 #include "UObject/LinkerLoad.h"
 #include "UObject/LinkerManager.h"
+#include "UObject/UObjectThreadContext.h"
 
 /*-----------------------------------------------------------------------------
 	UPackage.
@@ -149,7 +150,9 @@ UMetaData* UPackage::GetMetaData()
 
 	if (MetaData->HasAnyFlags(RF_NeedLoad))
 	{
-		MetaData->GetLinker()->Preload(MetaData);
+		FLinkerLoad* MetaDataLinker = MetaData->GetLinker();
+		check(MetaDataLinker);
+		MetaDataLinker->Preload(MetaData);
 	}
 
 	return MetaData;
diff --git a/Engine/Source/Runtime/CoreUObject/Private/UObject/PropertyByte.cpp b/Engine/Source/Runtime/CoreUObject/Private/UObject/PropertyByte.cpp
--- a/Engine/Source/Runtime/CoreUObject/Private/UObject/PropertyByte.cpp
+++ b/Engine/Source/Runtime/CoreUObject/Private/UObject/PropertyByte.cpp
@@ -8,6 +8,7 @@
 #include "UObject/UObjectThreadContext.h"
 #include "Serialization/ArchiveUObjectFromStructuredArchive.h"
 #include "Algo/Find.h"
+#include "UObject/LinkerLoad.h"
 
 /*-----------------------------------------------------------------------------
 	UByteProperty.
@@ -375,8 +376,15 @@ const TCHAR* UByteProperty::ImportText_Internal( const TCHAR* InBuffer, void* Da
 			// Enum could not be created from value. This indicates a bad value so
 			// return null so that the caller of ImportText can generate a more meaningful
 			// warning/error
-			FUObjectThreadContext& ThreadContext = FUObjectThreadContext::Get();
-			UE_LOG(LogClass, Warning, TEXT("In asset '%s', there is an enum property of type '%s' with an invalid value of '%s'"), *GetPathNameSafe(ThreadContext.SerializedObject), *Enum->GetName(), *Temp);
+			UObject* SerializedObject = nullptr;
+			if (FLinkerLoad* Linker = GetLinker())
+			{
+				if (FUObjectSerializeContext* LoadContext = Linker->GetSerializeContext())
+				{
+					SerializedObject = LoadContext->SerializedObject;
+				}
+			}
+			UE_LOG(LogClass, Warning, TEXT("In asset '%s', there is an enum property of type '%s' with an invalid value of '%s'"), *GetPathNameSafe(SerializedObject), *Enum->GetName(), *Temp);
 			return nullptr;
 		}
 	}
diff --git a/Engine/Source/Runtime/CoreUObject/Private/UObject/PropertyMap.cpp b/Engine/Source/Runtime/CoreUObject/Private/UObject/PropertyMap.cpp
--- a/Engine/Source/Runtime/CoreUObject/Private/UObject/PropertyMap.cpp
+++ b/Engine/Source/Runtime/CoreUObject/Private/UObject/PropertyMap.cpp
@@ -9,6 +9,7 @@
 #include "UObject/PropertyHelper.h"
 #include "Misc/ScopeExit.h"
 #include "Serialization/ArchiveUObjectFromStructuredArchive.h"
+#include "UObject/UObjectThreadContext.h"
 
 namespace UE4MapProperty_Private
 {
diff --git a/Engine/Source/Runtime/CoreUObject/Private/UObject/SavePackage.cpp b/Engine/Source/Runtime/CoreUObject/Private/UObject/SavePackage.cpp
--- a/Engine/Source/Runtime/CoreUObject/Private/UObject/SavePackage.cpp
+++ b/Engine/Source/Runtime/CoreUObject/Private/UObject/SavePackage.cpp
@@ -382,10 +382,14 @@ static void CheckObjectPriorToSave(FArchiveUObject& Ar, UObject* InObj, UPackage
 	{
 		return;
 	}
-	FUObjectThreadContext& ThreadContext = FUObjectThreadContext::Get();
+	UObject* SerializedObject = nullptr;
+	FUObjectSerializeContext* SaveContext = Ar.GetSerializeContext();
+	check(SaveContext);
+	SerializedObject = SaveContext->SerializedObject;
+
 	if (!InObj->IsValidLowLevelFast() || !InObj->IsValidLowLevel())
 	{
-		UE_LOG(LogLinker, Fatal, TEXT("Attempt to save bogus object %p ThreadContext.SerializedObject=%s  SerializedProperty=%s"), (void*)InObj, *GetFullNameSafe(ThreadContext.SerializedObject), *GetFullNameSafe(Ar.GetSerializedProperty()));
+		UE_LOG(LogLinker, Fatal, TEXT("Attempt to save bogus object %p SaveContext.SerializedObject=%s  SerializedProperty=%s"), (void*)InObj, *GetFullNameSafe(SerializedObject), *GetFullNameSafe(Ar.GetSerializedProperty()));
 		return;
 	}
 	// if the object class is abstract or has been marked as deprecated, mark this
@@ -400,7 +404,7 @@ static void CheckObjectPriorToSave(FArchiveUObject& Ar, UObject* InObj, UPackage
 		{
 			TArray<UObject*> ComponentReferences;
 			FReferenceFinder ComponentCollector(ComponentReferences, InObj, false, true, true);
-			ComponentCollector.FindReferences(InObj, ThreadContext.SerializedObject, Ar.GetSerializedProperty());
+			ComponentCollector.FindReferences(InObj, SerializedObject, Ar.GetSerializedProperty());
 
 			for ( int32 Index = 0; Index < ComponentReferences.Num(); Index++ )
 			{
@@ -1006,7 +1010,14 @@ public:
 	void ProcessBaseObject(UObject* BaseObject);
 	virtual FArchive& operator<<(UObject*& Obj) override;
 	virtual FArchive& operator<<(FWeakObjectPtr& Value) override;
-
+	virtual void SetSerializeContext(FUObjectSerializeContext* InLoadContext) override
+	{
+		LoadContext = InLoadContext;
+	}
+	virtual FUObjectSerializeContext* GetSerializeContext() override
+	{
+		return LoadContext;
+	}
 	/**
 	 * Package we're currently saving.  Only objects contained
 	 * within this package will be tagged for serialization.
@@ -1018,6 +1029,7 @@ public:
 private:
 
 	TArray<UObject*> TaggedObjects;
+	TRefCountPtr<FUObjectSerializeContext> LoadContext;
 
 	void ProcessTaggedObjects();
 };
@@ -1231,6 +1243,18 @@ public:
 	
 	virtual void MarkSearchableName(const UObject* TypeObject, const FName& ValueName) const override;
 	virtual FString GetArchiveName() const override;
+	virtual void SetSerializeContext(FUObjectSerializeContext* InLoadContext) override
+	{
+		LoadContext = InLoadContext;
+	}
+	virtual FUObjectSerializeContext* GetSerializeContext() override
+	{
+		return LoadContext;
+	}
+
+private:
+
+	TRefCountPtr<FUObjectSerializeContext> LoadContext;
 };
 
 FString FArchiveSaveTagImports::GetArchiveName() const
@@ -2968,19 +2992,21 @@ static bool ValidateConformCompatibility(UPackage* NewPackage, FLinkerLoad* OldL
 	// unfortunately, to get at the classes and their properties we will also need to load the default objects
 	// but the remapped package won't be bound to its native instance, so we need to manually copy the constructors
 	// so that classes with their own Serialize() implementations are loaded correctly
-	BeginLoad();
-	for (int32 i = 0; i < OldLinker->ExportMap.Num(); i++)
 	{
-		UClass* NewClass = (UClass*)StaticFindObjectFast(UClass::StaticClass(), NewPackage, OldLinker->ExportMap[i].ObjectName, true, false);
-		UClass* OldClass = static_cast<UClass*>(OldLinker->Create(UClass::StaticClass(), OldLinker->ExportMap[i].ObjectName, OldLinker->LinkerRoot, LOAD_None, false));
-		if (OldClass != nullptr && NewClass != nullptr && OldClass->IsNative() && NewClass->IsNative())
+		BeginLoad(OldLinker->GetSerializeContext());
+		for (int32 i = 0; i < OldLinker->ExportMap.Num(); i++)
 		{
-			OldClass->ClassConstructor = NewClass->ClassConstructor;
-			OldClass->ClassVTableHelperCtorCaller = NewClass->ClassVTableHelperCtorCaller;
-			OldClass->ClassAddReferencedObjects = NewClass->ClassAddReferencedObjects;
+			UClass* NewClass = (UClass*)StaticFindObjectFast(UClass::StaticClass(), NewPackage, OldLinker->ExportMap[i].ObjectName, true, false);
+			UClass* OldClass = static_cast<UClass*>(OldLinker->Create(UClass::StaticClass(), OldLinker->ExportMap[i].ObjectName, OldLinker->LinkerRoot, LOAD_None, false));
+			if (OldClass != nullptr && NewClass != nullptr && OldClass->IsNative() && NewClass->IsNative())
+			{
+				OldClass->ClassConstructor = NewClass->ClassConstructor;
+				OldClass->ClassVTableHelperCtorCaller = NewClass->ClassVTableHelperCtorCaller;
+				OldClass->ClassAddReferencedObjects = NewClass->ClassAddReferencedObjects;
+			}
 		}
+		EndLoad(OldLinker->GetSerializeContext());
 	}
-	EndLoad();
 
 	bool bHadCompatibilityErrors = false;
 	// check for illegal change of networking flags on class fields
@@ -2989,9 +3015,9 @@ static bool ValidateConformCompatibility(UPackage* NewPackage, FLinkerLoad* OldL
 		if (OldLinker->GetExportClassName(i) == NAME_Class)
 		{
 			// load the object so we can analyze it
-			BeginLoad();
-			UClass* OldClass = static_cast<UClass*>(OldLinker->Create(UClass::StaticClass(), OldLinker->ExportMap[i].ObjectName, OldLinker->LinkerRoot, LOAD_None, false));
-			EndLoad();
+			BeginLoad(OldLinker->GetSerializeContext());
+			UClass* OldClass = static_cast<UClass*>(OldLinker->Create(UClass::StaticClass(), OldLinker->ExportMap[i].ObjectName, OldLinker->LinkerRoot, LOAD_None, false));			
+			EndLoad(OldLinker->GetSerializeContext());
 			if (OldClass != nullptr)
 			{
 				UClass* NewClass = FindObjectFast<UClass>(NewPackage, OldClass->GetFName(), true, false);
@@ -3519,6 +3545,8 @@ FSavePackageResultStruct UPackage::Save(UPackage* InOuter, UObject* Base, EObjec
 
 	if (FPlatformProperties::HasEditorOnlyData())
 	{
+		TRefCountPtr<FUObjectSerializeContext> SaveContext(new FUObjectSerializeContext());
+
 #if WITH_EDITOR
 		struct FDiffSettings
 		{
@@ -3761,7 +3789,8 @@ FSavePackageResultStruct UPackage::Save(UPackage* InOuter, UObject* Base, EObjec
 			FArchiveSaveTagExports ExportTaggerArchive( InOuter );
 			ExportTaggerArchive.SetPortFlags( ComparisonFlags );
 			ExportTaggerArchive.SetCookingTarget(TargetPlatform);
-			
+			ExportTaggerArchive.SetSerializeContext(SaveContext);
+
 			check( ExportTaggerArchive.IsCooking() == !!TargetPlatform );
 			check( ExportTaggerArchive.IsCooking() == bIsCooking );
 
@@ -3940,6 +3969,7 @@ FSavePackageResultStruct UPackage::Save(UPackage* InOuter, UObject* Base, EObjec
 
 				FStructuredArchive* StructuredArchive = new FStructuredArchive(*Formatter);
 				FStructuredArchive::FRecord StructuredArchiveRoot = StructuredArchive->Open().EnterRecord();
+				StructuredArchiveRoot.GetUnderlyingArchive().SetSerializeContext(SaveContext);
 #if WITH_EDITOR
 				if (!!TargetPlatform)
 				{
@@ -4044,6 +4074,7 @@ FSavePackageResultStruct UPackage::Save(UPackage* InOuter, UObject* Base, EObjec
 						FArchiveSaveTagImports ImportTagger(Linker.Get(), SavePackageState);
 						ImportTagger.SetPortFlags(ComparisonFlags);
 						ImportTagger.SetFilterEditorOnly(FilterEditorOnly);
+						ImportTagger.SetSerializeContext(SaveContext);
 
 						UClass* Class = Obj->GetClass();
 
@@ -5454,7 +5485,7 @@ FSavePackageResultStruct UPackage::Save(UPackage* InOuter, UObject* Base, EObjec
 								TimingInfo.Value++;
 								FScopedDurationTimer SerializeTimer(TimingInfo.Key);
 #endif
-								TGuardValue<UObject*> GuardSerializedObject(ThreadContext.SerializedObject, Export.Object);
+								TGuardValue<UObject*> GuardSerializedObject(SaveContext->SerializedObject, Export.Object);
 
 								if (bSupportsText)
 								{
diff --git a/Engine/Source/Runtime/CoreUObject/Private/UObject/SoftObjectPath.cpp b/Engine/Source/Runtime/CoreUObject/Private/UObject/SoftObjectPath.cpp
--- a/Engine/Source/Runtime/CoreUObject/Private/UObject/SoftObjectPath.cpp
+++ b/Engine/Source/Runtime/CoreUObject/Private/UObject/SoftObjectPath.cpp
@@ -564,13 +564,11 @@ bool FSoftObjectPathThreadContext::GetSerializationOptions(FName& OutPackageName
 		bFoundAnything = true;
 	}
 	
-	
-	// Check UObject thread context as a backup, this only works for loads
-	FUObjectThreadContext& ThreadContext = FUObjectThreadContext::Get();
-	if (ThreadContext.SerializedObject)
+	// Check UObject serialize context as a backup
+	FUObjectSerializeContext* LoadContext = Archive ? Archive->GetSerializeContext() : nullptr;
+	if (LoadContext && LoadContext->SerializedObject)
 	{
-		FLinkerLoad* Linker = ThreadContext.SerializedObject->GetLinker();
-
+		FLinkerLoad* Linker = LoadContext->SerializedObject->GetLinker();
 		if (Linker)
 		{
 			if (CurrentPackageName == NAME_None)
diff --git a/Engine/Source/Runtime/CoreUObject/Private/UObject/UObjectGlobals.cpp b/Engine/Source/Runtime/CoreUObject/Private/UObject/UObjectGlobals.cpp
--- a/Engine/Source/Runtime/CoreUObject/Private/UObject/UObjectGlobals.cpp
+++ b/Engine/Source/Runtime/CoreUObject/Private/UObject/UObjectGlobals.cpp
@@ -1149,7 +1149,7 @@ public:
 * @param	ImportLinker	Linker that requests this package through one of its imports
 * @return	Loaded package if successful, NULL otherwise
 */
-UPackage* LoadPackageInternal(UPackage* InOuter, const TCHAR* InLongPackageNameOrFilename, uint32 LoadFlags, FLinkerLoad* ImportLinker, FArchive* InReaderOverride)
+UPackage* LoadPackageInternal(UPackage* InOuter, const TCHAR* InLongPackageNameOrFilename, uint32 LoadFlags, FLinkerLoad* ImportLinker, FArchive* InReaderOverride, FUObjectSerializeContext* InLoadContext)
 {
 	DECLARE_SCOPE_CYCLE_COUNTER(TEXT("LoadPackageInternal"), STAT_LoadPackageInternal, STATGROUP_ObjectVerbose);
 
@@ -1253,7 +1253,8 @@ UPackage* LoadPackageInternal(UPackage* InOuter, const TCHAR* InLongPackageNameO
 	}
 
 	// Try to load.
-	BeginLoad(InLongPackageNameOrFilename);
+	TRefCountPtr<FUObjectSerializeContext> LoadContext(InLoadContext ? InLoadContext : new FUObjectSerializeContext());
+	BeginLoad(LoadContext, InLongPackageNameOrFilename);
 
 	bool bFullyLoadSkipped = false;
 
@@ -1279,11 +1280,11 @@ UPackage* LoadPackageInternal(UPackage* InOuter, const TCHAR* InLongPackageNameO
 		}
 #endif
 
-		Linker = GetPackageLinker(InOuter, *FileToLoad, LoadFlags, nullptr, nullptr, InReaderOverride);
+		Linker = GetPackageLinker(InOuter, *FileToLoad, LoadFlags, nullptr, nullptr, InReaderOverride, LoadContext);
 		
 		if (!Linker)
 		{
-			EndLoad();
+			EndLoad(LoadContext);
 			return nullptr;
 		}
 
@@ -1292,7 +1293,7 @@ UPackage* LoadPackageInternal(UPackage* InOuter, const TCHAR* InLongPackageNameO
 
 		auto EndLoadAndCopyLocalizationGatherFlag = [&]
 		{
-			EndLoad();
+			EndLoad(Linker->GetSerializeContext());
 			// Set package-requires-localization flags from archive after loading. This reinforces flagging of packages that haven't yet been resaved.
 			Result->ThisRequiresLocalizationGather(Linker->RequiresLocalizationGather());
 		};
@@ -1399,7 +1400,7 @@ UPackage* LoadPackageInternal(UPackage* InOuter, const TCHAR* InLongPackageNameO
 		}
 
 		// Only set time it took to load package if the above EndLoad is the "outermost" EndLoad.
-		if( Result && !IsLoading() && !(LoadFlags & LOAD_Verify) )
+		if( Result && !LoadContext->HasLoadedObjects() && !(LoadFlags & LOAD_Verify) )
 		{
 			Result->SetLoadTime( FPlatformTime::Seconds() - StartTime );
 		}
@@ -1421,7 +1422,8 @@ UPackage* LoadPackageInternal(UPackage* InOuter, const TCHAR* InLongPackageNameO
 		{
 			if (!IsInAsyncLoadingThread())
 			{
-				if (FUObjectThreadContext::Get().ObjBeginLoadCount == 0)
+				check(Linker->GetSerializeContext());
+				if (!Linker->GetSerializeContext()->HasStartedLoading())
 				{
 					// Sanity check to make sure that Linker is the linker that loaded our Result package or the linker has already been detached
 					check(!Result || Result->LinkerLoad == Linker || Result->LinkerLoad == nullptr);
@@ -1440,12 +1442,14 @@ UPackage* LoadPackageInternal(UPackage* InOuter, const TCHAR* InLongPackageNameO
 				// Async loading removes delayed linkers on the game thread after streaming has finished
 				else
 				{
-					FUObjectThreadContext::Get().DelayedLinkerClosePackages.AddUnique(Linker);
+					check(Linker->GetSerializeContext());
+					Linker->GetSerializeContext()->AddDelayedLinkerClosePackage(Linker);
 				}
 			}
 			else
 			{
-				FUObjectThreadContext::Get().DelayedLinkerClosePackages.AddUnique(Linker);
+				check(Linker->GetSerializeContext());
+				Linker->GetSerializeContext()->AddDelayedLinkerClosePackage(Linker);
 			}
 		}
 	}
@@ -1474,9 +1478,11 @@ UPackage* LoadPackage(UPackage* InOuter, const TCHAR* InLongPackageName, uint32
 	// since we are faking the object name, this is basically a duplicate of LLM_SCOPED_TAG_WITH_OBJECT_IN_SET
 	FString FakePackageName = FString(TEXT("Package ")) + InLongPackageName;
 	LLM_SCOPED_TAG_WITH_STAT_NAME_IN_SET(FLowLevelMemTracker::Get().IsTagSetActive(ELLMTagSet::Assets) ? FDynamicStats::CreateMemoryStatId<FStatGroup_STATGROUP_LLMAssets>(FName(*FakePackageName)).GetName() : NAME_None, ELLMTagSet::Assets, ELLMTracker::Default);
-	return LoadPackageInternal(InOuter, InLongPackageName, LoadFlags, /*ImportLinker =*/ nullptr, InReaderOverride);
+	return LoadPackageInternal(InOuter, InLongPackageName, LoadFlags, /*ImportLinker =*/ nullptr, InReaderOverride, nullptr);
 }
 
+static int32 GGameThreadLoadCounter = 0;
+
 /**
  * Returns whether we are currently loading a package (sync or async)
  *
@@ -1484,18 +1490,17 @@ UPackage* LoadPackage(UPackage* InOuter, const TCHAR* InLongPackageName, uint32
  */
 bool IsLoading()
 {
-	check(FUObjectThreadContext::Get().ObjBeginLoadCount >= 0);
-	return FUObjectThreadContext::Get().ObjBeginLoadCount > 0;
+	return GGameThreadLoadCounter > 0;
 }
 
 //
 // Begin loading packages.
 //warning: Objects may not be destroyed between BeginLoad/EndLoad calls.
 //
-void BeginLoad(const TCHAR* DebugContext)
+void BeginLoad(FUObjectSerializeContext* LoadContext, const TCHAR* DebugContext)
 {
-	auto& ThreadContext = FUObjectThreadContext::Get();
-	if (ThreadContext.ObjBeginLoadCount == 0 && !IsInAsyncLoadingThread())
+	check(LoadContext);
+	if (!LoadContext->HasStartedLoading() && !IsInAsyncLoadingThread())
 	{
 		if (IsAsyncLoading() && (DebugContext != nullptr))
 		{
@@ -1504,12 +1509,13 @@ void BeginLoad(const TCHAR* DebugContext)
 
 		// Make sure we're finishing up all pending async loads, and trigger texture streaming next tick if necessary.
 		FlushAsyncLoading();
-
-		// Validate clean load state.
-		check(ThreadContext.ObjLoaded.Num() == 0);
+	}
+	if (IsInGameThread() && !IsInAsyncLoadingThread())
+	{
+		GGameThreadLoadCounter++;
 	}
 
-	++ThreadContext.ObjBeginLoadCount;
+	LoadContext->IncrementBeginLoadCount();
 }
 
 // Sort objects by linker name and file offset
@@ -1552,16 +1558,20 @@ struct FCompareUObjectByLinkerAndOffset
 //
 // End loading packages.
 //
-void EndLoad()
+void EndLoad(FUObjectSerializeContext* LoadContext)
 {
-	auto& ThreadContext = FUObjectThreadContext::Get();
+	check(LoadContext);
 
-	check(ThreadContext.ObjBeginLoadCount > 0);
 	if (IsInAsyncLoadingThread())
 	{
-		ThreadContext.ObjBeginLoadCount--;
+		LoadContext->DecrementBeginLoadCount();
 		return;
 	}
+	else if (IsInGameThread())
+	{
+		GGameThreadLoadCounter--;
+		check(GGameThreadLoadCounter >= 0);
+	}
 
 #if WITH_EDITOR
 	FScopedSlowTask SlowTask(0, NSLOCTEXT("Core", "PerformingPostLoad", "Performing post-load..."), ShouldReportProgress());
@@ -1570,23 +1580,22 @@ void EndLoad()
 	TSet<UObject*> AssetsLoaded;
 #endif
 
-	while (--ThreadContext.ObjBeginLoadCount == 0 && (ThreadContext.ObjLoaded.Num() || ThreadContext.ImportCount || ThreadContext.ForcedExportCount))
+	while (LoadContext->DecrementBeginLoadCount() == 0 && (LoadContext->HasLoadedObjects() || LoadContext->HasPendingImportsOrForcedExports()))
 	{
 		// The time tracker keeps track of time spent in EndLoad.
 		FExclusiveLoadPackageTimeTracker::FScopedEndLoadTracker Tracker;
 
 		// Make sure we're not recursively calling EndLoad as e.g. loading a config file could cause
 		// BeginLoad/EndLoad to be called.
-		ThreadContext.ObjBeginLoadCount++;
+		LoadContext->IncrementBeginLoadCount();
 
 		// Temporary list of loaded objects as GObjLoaded might expand during iteration.
 		TArray<UObject*> ObjLoaded;
 		TSet<FLinkerLoad*> LoadedLinkers;
-		while (ThreadContext.ObjLoaded.Num())
+		while (LoadContext->HasLoadedObjects())
 		{
 			// Accumulate till GObjLoaded no longer increases.
-			ObjLoaded += ThreadContext.ObjLoaded;
-			ThreadContext.ObjLoaded.Empty();
+			LoadContext->AppendLoadedObjectsAndEmpty(ObjLoaded);
 
 			// Sort by Filename and Offset.
 			ObjLoaded.Sort(FCompareUObjectByLinkerAndOffset());
@@ -1605,7 +1614,7 @@ void EndLoad()
 
 			// Start over again as new objects have been loaded that need to have "Preload" called on them before
 			// we can safely PostLoad them.
-			if (ThreadContext.ObjLoaded.Num())
+			if (LoadContext->HasLoadedObjects())
 			{
 				continue;
 			}
@@ -1706,7 +1715,7 @@ void EndLoad()
 #endif	// WITH_EDITOR
 
 			// Empty array before next iteration as we finished postloading all objects.
-			ObjLoaded.Empty(ThreadContext.ObjLoaded.Num());
+			ObjLoaded.Reset();
 		}
 
 		if ( GIsEditor && LoadedLinkers.Num() > 0 )
@@ -1743,7 +1752,8 @@ void EndLoad()
 		FLinkerManager::Get().DissociateImportsAndForcedExports();
 
 		// close any linkers' loaders that were requested to be closed once GObjBeginLoadCount goes to 0
-		TArray<FLinkerLoad*> PackagesToClose = MoveTemp(ThreadContext.DelayedLinkerClosePackages);
+		TArray<FLinkerLoad*> PackagesToClose;
+		LoadContext->MoveDelayedLinkerClosePackages(PackagesToClose);
 		for (FLinkerLoad* Linker : PackagesToClose)
 		{
 			if (Linker)
@@ -1965,7 +1975,7 @@ FObjectDuplicationParameters::FObjectDuplicationParameters( UObject* InSourceObj
 
 UObject* StaticDuplicateObject(UObject const* SourceObject, UObject* DestOuter, const FName DestName, EObjectFlags FlagMask, UClass* DestClass, EDuplicateMode::Type DuplicateMode, EInternalObjectFlags InternalFlagsMask)
 {
-	if (!IsAsyncLoading() && !IsLoading() && SourceObject->HasAnyFlags(RF_ClassDefaultObject))
+	if (!IsAsyncLoading() && /*!IsLoading() &&*/ SourceObject->HasAnyFlags(RF_ClassDefaultObject))
 	{
 		// Detach linker for the outer if it already exists, to avoid problems with PostLoad checking the Linker version
 		ResetLoaders(DestOuter);
@@ -2100,8 +2110,9 @@ UObject* StaticDuplicateObjectEx( FObjectDuplicationParameters& Parameters )
 		SerializedObjects.Add(Object);
 	};
 
-
-	FDuplicateDataReader	Reader(DuplicatedObjectAnnotation, ObjectData, Parameters.PortFlags, Parameters.DestOuter);
+	TRefCountPtr<FUObjectSerializeContext> LoadContext(new FUObjectSerializeContext());
+	FDuplicateDataReader Reader(DuplicatedObjectAnnotation, ObjectData, Parameters.PortFlags, Parameters.DestOuter);
+	Reader.SetSerializeContext(LoadContext);
 	for(int32 ObjectIndex = 0;ObjectIndex < SerializedObjects.Num();ObjectIndex++)
 	{
 		UObject* SerializedObject = SerializedObjects[ObjectIndex];
@@ -2109,6 +2120,7 @@ UObject* StaticDuplicateObjectEx( FObjectDuplicationParameters& Parameters )
 		FDuplicatedObject ObjectInfo = DuplicatedObjectAnnotation.GetAnnotation( SerializedObject );
 		checkSlow( !ObjectInfo.IsDefault() );
 
+		TGuardValue<UObject*> SerializedObjectGuard(LoadContext->SerializedObject, ObjectInfo.DuplicatedObject);
 		if ( !SerializedObject->HasAnyFlags(RF_ClassDefaultObject) )
 		{
 			ObjectInfo.DuplicatedObject->Serialize(Reader);
diff --git a/Engine/Source/Runtime/CoreUObject/Private/UObject/UObjectThreadContext.cpp b/Engine/Source/Runtime/CoreUObject/Private/UObject/UObjectThreadContext.cpp
--- a/Engine/Source/Runtime/CoreUObject/Private/UObject/UObjectThreadContext.cpp
+++ b/Engine/Source/Runtime/CoreUObject/Private/UObject/UObjectThreadContext.cpp
@@ -9,19 +9,39 @@
 DEFINE_LOG_CATEGORY(LogUObjectThreadContext);
 
 FUObjectThreadContext::FUObjectThreadContext()
-: ImportCount(0)
-, ForcedExportCount(0)
-, ObjBeginLoadCount(0)
-, IsRoutingPostLoad(false)
+: IsRoutingPostLoad(false)
 , CurrentlyPostLoadedObjectByALT(nullptr)
 , IsDeletingLinkers(false)
 , IsInConstructor(0)
 , ConstructedObject(nullptr)
-, SerializedObject(nullptr)
-, SerializedPackageLinker(nullptr)
-, SerializedImportIndex(INDEX_NONE)
-, SerializedImportLinker(nullptr)
-, SerializedExportIndex(INDEX_NONE)
-, SerializedExportLinker(nullptr)
 , AsyncPackage(nullptr)
 {}
+
+
+FUObjectSerializeContext::FUObjectSerializeContext()
+	: RefCount(0)
+	, ImportCount(0)
+	, ForcedExportCount(0)
+	, ObjBeginLoadCount(0)
+	, SerializedObject(nullptr)
+	, SerializedPackageLinker(nullptr)
+	, SerializedImportIndex(0)
+	, SerializedImportLinker(nullptr)
+	, SerializedExportIndex(0)
+	, SerializedExportLinker(nullptr)
+{}
+
+FUObjectSerializeContext::~FUObjectSerializeContext()
+{
+	checkf(!HasLoadedObjects(), TEXT("FUObjectSerializeContext is being destroyed but it still has pending loaded objects in its ObjectsLoaded list."));
+}
+
+int32 FUObjectSerializeContext::IncrementBeginLoadCount()
+{
+	return ++ObjBeginLoadCount;
+}
+int32 FUObjectSerializeContext::DecrementBeginLoadCount()
+{
+	check(HasStartedLoading());
+	return --ObjBeginLoadCount;
+}
\ No newline at end of file
diff --git a/Engine/Source/Runtime/CoreUObject/Public/Serialization/ArchiveStackTrace.h b/Engine/Source/Runtime/CoreUObject/Public/Serialization/ArchiveStackTrace.h
--- a/Engine/Source/Runtime/CoreUObject/Public/Serialization/ArchiveStackTrace.h
+++ b/Engine/Source/Runtime/CoreUObject/Public/Serialization/ArchiveStackTrace.h
@@ -5,6 +5,7 @@
 #include "CoreMinimal.h"
 #include "Serialization/LargeMemoryWriter.h"
 #include "Serialization/LargeMemoryReader.h"
+#include "Templates/RefCounting.h"
 
 class FUObjectThreadContext;
 class FLinkerLoad;
@@ -147,6 +148,8 @@ class COREUOBJECT_API FArchiveStackTrace : public FLargeMemoryWriter
 	ANSICHAR* LastSerializeCallstack;
 	/** Cached thread context */
 	FUObjectThreadContext& ThreadContext;
+	/** Current save/load context */
+	TRefCountPtr<FUObjectSerializeContext> LoadContext;
 
 #if WITH_EDITOR
 	/** Current debug name stack */
@@ -205,7 +208,7 @@ private:
 	void CompareWithInternal(const FPackageData& SourcePackage, const FPackageData& DestPackage, const TCHAR* AssetFilename, const TCHAR* CallstackCutoffText, const int64 MaxDiffsToLog, int32& InOutDiffsLogged, TMap<FName, FArchiveDiffStats>& OutStats);
 	bool GenerateDiffMapInternal(const FPackageData& SourcePackage, const FPackageData& DestPackage, int32 MaxDiffsToFind, FArchiveDiffMap& OutDiffMap);
 	void DumpPackageHeaderDiffs(const FPackageData& SourcePackage, const FPackageData& DestPackage, const FString& AssetFilename, const int32 MaxDiffsToLog);
-	static FLinkerLoad* CreateLinkerForPackage(const FString& InPackageName, const FString& InFilename, const FPackageData& PackageData);
+	static FLinkerLoad* CreateLinkerForPackage(FUObjectSerializeContext* LoadContext, const FString& InPackageName, const FString& InFilename, const FPackageData& PackageData);
 
 public:
 
@@ -214,6 +217,8 @@ public:
 
 	//~ Begin FArchive Interface
 	virtual void Serialize(void* Data, int64 Num) override;
+	virtual void SetSerializeContext(FUObjectSerializeContext* InLoadContext) override;
+	virtual FUObjectSerializeContext* GetSerializeContext() override;
 #if WITH_EDITOR
 	virtual void PushDebugDataString(const FName& DebugData) override
 	{
diff --git a/Engine/Source/Runtime/CoreUObject/Public/Serialization/AsyncLoading.h b/Engine/Source/Runtime/CoreUObject/Public/Serialization/AsyncLoading.h
--- a/Engine/Source/Runtime/CoreUObject/Public/Serialization/AsyncLoading.h
+++ b/Engine/Source/Runtime/CoreUObject/Public/Serialization/AsyncLoading.h
@@ -14,6 +14,8 @@
 #include "UObject/Package.h"
 #include "Templates/Casts.h"
 #include "UObject/ObjectRedirector.h"
+#include "UObject/UObjectThreadContext.h"
+#include "Templates/RefCounting.h"
 
 class IAsyncReadRequest;
 struct FAsyncPackage;
@@ -924,6 +926,11 @@ private:
 	/** Add this time taken for object of class Class to have PostLoad called, to the stats we track. */
 	void TrackPostLoadTimeForClass(const UClass* Class, double Time);
 #endif // PERF_TRACK_DETAILED_ASYNC_STATS
+
+public:
+
+	/** Serialization context for this package */
+	TRefCountPtr<FUObjectSerializeContext> LoadContext;
 };
 
 struct FScopedAsyncPackageEvent
diff --git a/Engine/Source/Runtime/CoreUObject/Public/UObject/Class.h b/Engine/Source/Runtime/CoreUObject/Public/UObject/Class.h
--- a/Engine/Source/Runtime/CoreUObject/Public/UObject/Class.h
+++ b/Engine/Source/Runtime/CoreUObject/Public/UObject/Class.h
@@ -2166,7 +2166,7 @@ public:
 	 * Conditionally recompiles the class after loading, in case any dependencies were also newly loaded
 	 * @param ObjLoaded	If set this is the list of objects that are currently loading, usualy GObjLoaded
 	 */
-	virtual void ConditionalRecompileClass(TArray<UObject*>* ObjLoaded) {}
+	virtual void ConditionalRecompileClass() {}
 	virtual void FlushCompilationQueueForLevel() {}
 #endif //WITH_EDITOR
 
diff --git a/Engine/Source/Runtime/CoreUObject/Public/UObject/Linker.h b/Engine/Source/Runtime/CoreUObject/Public/UObject/Linker.h
--- a/Engine/Source/Runtime/CoreUObject/Public/UObject/Linker.h
+++ b/Engine/Source/Runtime/CoreUObject/Public/UObject/Linker.h
@@ -564,7 +564,7 @@ COREUOBJECT_API void DeleteLoaders();
 /** Queues linker for deletion */
 COREUOBJECT_API void DeleteLoader(FLinkerLoad* Loader);
 
-COREUOBJECT_API FLinkerLoad* GetPackageLinker(UPackage* InOuter, const TCHAR* InLongPackageName, uint32 LoadFlags, UPackageMap* Sandbox, FGuid* CompatibleGuid, FArchive* InReaderOverride = nullptr);
+COREUOBJECT_API FLinkerLoad* GetPackageLinker(UPackage* InOuter, const TCHAR* InLongPackageName, uint32 LoadFlags, UPackageMap* Sandbox, FGuid* CompatibleGuid, FArchive* InReaderOverride = nullptr, FUObjectSerializeContext* InExistingContext = nullptr);
 COREUOBJECT_API FString GetPrestreamPackageLinkerName(const TCHAR* InLongPackageName, bool bExistSkip = true);
 
 
diff --git a/Engine/Source/Runtime/CoreUObject/Public/UObject/LinkerLoad.h b/Engine/Source/Runtime/CoreUObject/Public/UObject/LinkerLoad.h
--- a/Engine/Source/Runtime/CoreUObject/Public/UObject/LinkerLoad.h
+++ b/Engine/Source/Runtime/CoreUObject/Public/UObject/LinkerLoad.h
@@ -8,6 +8,7 @@
 #include "UObject/SoftObjectPtr.h"
 #include "UObject/ObjectResource.h"
 #include "UObject/Linker.h"
+#include "Templates/RefCounting.h"
 
 class FLinkerPlaceholderBase;
 class ULinkerPlaceholderExportObject;
@@ -59,6 +60,8 @@ struct FScopedCreateImportCounter
 	/** Destructor. Called upon CreateImport() exit. */
 	~FScopedCreateImportCounter();
 
+	/** Current load context object */
+	FUObjectSerializeContext* LoadContext;
 	/** Previously stored linker */
 	FLinkerLoad* PreviousLinker;
 	/** Previously stored index */
@@ -445,7 +448,7 @@ public:
 	 *
 	 * @return	new FLinkerLoad object for Parent/ Filename
 	 */
-	COREUOBJECT_API static FLinkerLoad* CreateLinker(UPackage* Parent, const TCHAR* Filename, uint32 LoadFlags, FArchive* InLoader = nullptr);
+	COREUOBJECT_API static FLinkerLoad* CreateLinker(FUObjectSerializeContext* LoadContext, UPackage* Parent, const TCHAR* Filename, uint32 LoadFlags, FArchive* InLoader = nullptr);
 
 	void Verify();
 
@@ -857,8 +860,8 @@ private:
 	 *
 	 * @return	new FLinkerLoad object for Parent/ Filename
 	 */
-	COREUOBJECT_API static FLinkerLoad* CreateLinkerAsync(UPackage* Parent, const TCHAR* Filename, uint32 LoadFlags
-		, TFunction<void()>&& InSummaryReadyCallback	
+	COREUOBJECT_API static FLinkerLoad* CreateLinkerAsync(FUObjectSerializeContext* LoadContext, UPackage* Parent, const TCHAR* Filename, uint32 LoadFlags
+		, TFunction<void()>&& InSummaryReadyCallback
 	);
 
 protected: // Daniel L: Made this protected so I can override the constructor and create a custom loader to load the header of the linker in the DiffFilesCommandlet
@@ -1180,6 +1183,17 @@ private:
 	//
 	// FLinkerLoad creation helpers END
 	//
+
+private:
+
+	/** Current UObject serialization context */
+	TRefCountPtr<FUObjectSerializeContext> CurrentLoadContext;
+
+public:
+
+	//~ FArchive interface
+	virtual void SetSerializeContext(FUObjectSerializeContext* InLoadContext) override;
+	virtual FUObjectSerializeContext* GetSerializeContext() override;
 };
 
 // used by the EDL at boot time to coordinate loading with what is going on with the deferred registration stuff
diff --git a/Engine/Source/Runtime/CoreUObject/Public/UObject/LinkerSave.h b/Engine/Source/Runtime/CoreUObject/Public/UObject/LinkerSave.h
--- a/Engine/Source/Runtime/CoreUObject/Public/UObject/LinkerSave.h
+++ b/Engine/Source/Runtime/CoreUObject/Public/UObject/LinkerSave.h
@@ -6,6 +6,8 @@
 #include "Serialization/ArchiveUObject.h"
 #include "UObject/ObjectResource.h"
 #include "UObject/Linker.h"
+#include "UObject/UObjectThreadContext.h"
+#include "Templates/RefCounting.h"
 
 struct FUntypedBulkData;
 
@@ -43,6 +45,9 @@ public:
 	/** Index array - location of the name in the NameMap array for each FName is stored in the NameIndices array using the FName's Index */
 	TMap<FName, int32, FDefaultSetAllocator, TLinkerNameMapKeyFuncs<int32>> NameIndices;
 
+	/** Save context associated with this linker */
+	TRefCountPtr<FUObjectSerializeContext> SaveContext;
+
 	/** List of bulkdata that needs to be stored at the end of the file */
 	struct FBulkDataStorageInfo
 	{
@@ -80,6 +85,9 @@ public:
 	FArchive& operator<<( FName& InName );
 	FArchive& operator<<( UObject*& Obj );
 	FArchive& operator<<( FLazyObjectPtr& LazyObjectPtr );
+	virtual void SetSerializeContext(FUObjectSerializeContext* InLoadContext) override;
+	FUObjectSerializeContext* GetSerializeContext() override;
+
 
 #if WITH_EDITOR
 	// proxy for debugdata
diff --git a/Engine/Source/Runtime/CoreUObject/Public/UObject/Object.h b/Engine/Source/Runtime/CoreUObject/Public/UObject/Object.h
--- a/Engine/Source/Runtime/CoreUObject/Public/UObject/Object.h
+++ b/Engine/Source/Runtime/CoreUObject/Public/UObject/Object.h
@@ -1109,7 +1109,7 @@ public:
 	 * Called on the target when a class is loaded with ClassGeneratedBy is loaded.  Should regenerate the class if needed, and return the updated class
 	 * @return	Updated instance of the class, if needed, or NULL if no regeneration was necessary
 	 */
-	virtual UClass* RegenerateClass(UClass* ClassToRegenerate, UObject* PreviousCDO, TArray<UObject*>& ObjLoaded) { return NULL; };
+	virtual UClass* RegenerateClass(UClass* ClassToRegenerate, UObject* PreviousCDO) { return NULL; };
 
 	/** 
 	 * Returns whether this object is contained in or part of a blueprint object
diff --git a/Engine/Source/Runtime/CoreUObject/Public/UObject/UObjectGlobals.h b/Engine/Source/Runtime/CoreUObject/Public/UObject/UObjectGlobals.h
--- a/Engine/Source/Runtime/CoreUObject/Public/UObject/UObjectGlobals.h
+++ b/Engine/Source/Runtime/CoreUObject/Public/UObject/UObjectGlobals.h
@@ -507,8 +507,8 @@ COREUOBJECT_API EAsyncPackageState::Type ProcessAsyncLoading( bool bUseTimeLimit
 COREUOBJECT_API EAsyncPackageState::Type ProcessAsyncLoadingUntilComplete(TFunctionRef<bool()> CompletionPredicate, float TimeLimit);
 
 /** UObjects are being loaded between these calls */
-COREUOBJECT_API void BeginLoad(const TCHAR* DebugContext = nullptr);
-COREUOBJECT_API void EndLoad();
+COREUOBJECT_API void BeginLoad(FUObjectSerializeContext* LoadContext, const TCHAR* DebugContext = nullptr);
+COREUOBJECT_API void EndLoad(FUObjectSerializeContext* LoadContext);
 
 /**
  * Find an existing package by name
diff --git a/Engine/Source/Runtime/CoreUObject/Public/UObject/UObjectThreadContext.h b/Engine/Source/Runtime/CoreUObject/Public/UObject/UObjectThreadContext.h
--- a/Engine/Source/Runtime/CoreUObject/Public/UObject/UObjectThreadContext.h
+++ b/Engine/Source/Runtime/CoreUObject/Public/UObject/UObjectThreadContext.h
@@ -8,8 +8,10 @@
 
 #include "CoreMinimal.h"
 #include "HAL/ThreadSingleton.h"
+#include "HAL/ThreadSafeCounter.h"
 
 class FObjectInitializer;
+struct FUObjectSerializeContext;
 
 DECLARE_LOG_CATEGORY_EXTERN(LogUObjectThreadContext, Log, All);
 
@@ -61,17 +63,6 @@
 		return *ObjectInitializerPtr;
 	}
 
-
-	/** Imports for EndLoad optimization.	*/
-	int32 ImportCount;
-	/** Forced exports for EndLoad optimization. */
-	int32 ForcedExportCount;
-	/** Count for BeginLoad multiple loads.	*/
-	int32 ObjBeginLoadCount;
-	/** Objects that might need preloading. */
-	TArray<UObject*> ObjLoaded;
-	/** List of linkers that we want to close the loaders for (to free file handles) - needs to be delayed until EndLoad is called with GObjBeginLoadCount of 0 */
-	TArray<FLinkerLoad*> DelayedLinkerClosePackages;
 	/** true when we are routing ConditionalPostLoad/PostLoad to objects										*/
 	bool IsRoutingPostLoad;
 	/** The object we are routing PostLoad from the Async Loading code for */
@@ -82,18 +73,6 @@
 	int32 IsInConstructor;
 	/* Object that is currently being constructed with ObjectInitializer */
 	UObject* ConstructedObject;
-	/** Points to the main UObject currently being serialized */
-	UObject* SerializedObject;
-	/** Points to the main PackageLinker currently being serialized (Defined in Linker.cpp) */
-	FLinkerLoad* SerializedPackageLinker;
-	/** The main Import Index currently being used for serialization by CreateImports() (Defined in Linker.cpp) */
-	int32 SerializedImportIndex;
-	/** Points to the main Linker currently being used for serialization by CreateImports() (Defined in Linker.cpp) */
-	FLinkerLoad* SerializedImportLinker;
-	/** The most recently used export Index for serialization by CreateExport() */
-	int32 SerializedExportIndex;
-	/** Points to the most recently used Linker for serialization by CreateExport() */
-	FLinkerLoad* SerializedExportLinker;
 	/** Async Package currently processing objects */
 	struct FAsyncPackage* AsyncPackage;
 
@@ -108,3 +87,143 @@
 	TMap<FName, TSet<FName>> PackagesMarkedEditorOnlyByOtherPackage;
 #endif
 };
+
+/** Structure that holds the current serialization state of UObjects */
+struct COREUOBJECT_API FUObjectSerializeContext
+{
+private:
+
+	/** Reference count of this context */
+	int32 RefCount;
+
+	/** Imports for EndLoad optimization.	*/
+	int32 ImportCount;
+	/** Forced exports for EndLoad optimization. */
+	int32 ForcedExportCount;
+	/** Count for BeginLoad multiple loads.	*/
+	int32 ObjBeginLoadCount;
+	/** Objects that might need preloading. */
+	TArray<UObject*> ObjectsLoaded;
+	/** List of linkers that we want to close the loaders for (to free file handles) - needs to be delayed until EndLoad is called with GObjBeginLoadCount of 0 */
+	TArray<class FLinkerLoad*> DelayedLinkerClosePackages;
+
+public:
+
+	/** Points to the main UObject currently being serialized */
+	UObject* SerializedObject;
+	/** Points to the main PackageLinker currently being serialized (Defined in Linker.cpp) */
+	FLinkerLoad* SerializedPackageLinker;
+	/** The main Import Index currently being used for serialization by CreateImports() (Defined in Linker.cpp) */
+	int32 SerializedImportIndex;
+	/** Points to the main Linker currently being used for serialization by CreateImports() (Defined in Linker.cpp) */
+	FLinkerLoad* SerializedImportLinker;
+	/** The most recently used export Index for serialization by CreateExport() */
+	int32 SerializedExportIndex;
+	/** Points to the most recently used Linker for serialization by CreateExport() */
+	FLinkerLoad* SerializedExportLinker;
+
+	/** Constructor */
+	FUObjectSerializeContext();
+
+	/** Destructor */
+	~FUObjectSerializeContext();
+
+	/** Adds a new loaded object */
+	void AddLoadedObject(UObject* InObject)
+	{
+		ObjectsLoaded.Add(InObject);
+	}
+
+	/** Checks if object loading has started */
+	bool HasStartedLoading() const
+	{
+		return ObjBeginLoadCount > 0;
+	}
+
+	int32 IncrementBeginLoadCount();
+	int32 DecrementBeginLoadCount();
+
+	int32 IncrementImportCount()
+	{
+		return ++ImportCount;
+	}
+	void ResetImportCount()
+	{
+		ImportCount = 0;
+	}
+
+	int32 IncrementForcedExportCount()
+	{
+		return ++ForcedExportCount;
+	}
+	void ResetForcedExports()
+	{
+		ForcedExportCount = 0;
+	}
+
+	bool HasPendingImportsOrForcedExports() const
+	{
+		return ImportCount || ForcedExportCount;
+	}
+
+	bool HasLoadedObjects() const
+	{
+		return !!ObjectsLoaded.Num();
+	}
+
+	void PRIVATE_PatchNewObjectIntoExport(UObject* OldObject, UObject* NewObject)
+	{
+		const int32 ObjLoadedIdx = ObjectsLoaded.Find(OldObject);
+		if (ObjLoadedIdx != INDEX_NONE)
+		{
+			ObjectsLoaded[ObjLoadedIdx] = NewObject;
+		}
+	}
+
+	void AppendLoadedObjectsAndEmpty(TArray<UObject*>& InLoadedObject)
+	{
+		InLoadedObject.Append(ObjectsLoaded);
+		ObjectsLoaded.Reset();
+	}
+
+	TArray<UObject*>& GetObjectsLoaded()
+	{
+		return ObjectsLoaded;
+	}
+
+	void AddDelayedLinkerClosePackage(class FLinkerLoad* InLinker)
+	{
+		DelayedLinkerClosePackages.AddUnique(InLinker);
+	}
+
+	void RemoveDelayedLinkerClosePackage(class FLinkerLoad* InLinker)
+	{
+		DelayedLinkerClosePackages.Remove(InLinker);
+	}
+
+	void MoveDelayedLinkerClosePackages(TArray<class FLinkerLoad*>& OutDelayedLinkerClosePackages)
+	{
+		OutDelayedLinkerClosePackages = MoveTemp(DelayedLinkerClosePackages);
+	}
+
+	//~ TRefCountPtr interface
+	int32 AddRef()
+	{
+		return ++RefCount;
+	}
+	int32 Release()
+	{
+		int32 CurrentRefCount = --RefCount;
+		check(CurrentRefCount >= 0);
+		if (CurrentRefCount == 0)
+		{
+			delete this;
+		}
+		return CurrentRefCount;
+	}
+	int32 GetRefCount() const
+	{
+		return RefCount;
+	}
+	//~ TRefCountPtr interface
+};
diff --git a/Engine/Source/Runtime/Engine/Classes/Engine/Blueprint.h b/Engine/Source/Runtime/Engine/Classes/Engine/Blueprint.h
--- a/Engine/Source/Runtime/Engine/Classes/Engine/Blueprint.h
+++ b/Engine/Source/Runtime/Engine/Classes/Engine/Blueprint.h
@@ -789,7 +789,7 @@ public:
 	//~ Begin UObject Interface (WITH_EDITOR)
 	virtual void PostDuplicate(bool bDuplicateForPIE) override;
 	virtual bool Rename(const TCHAR* NewName = nullptr, UObject* NewOuter = nullptr, ERenameFlags Flags = REN_None) override;
-	virtual UClass* RegenerateClass(UClass* ClassToRegenerate, UObject* PreviousCDO, TArray<UObject*>& ObjLoaded) override;
+	virtual UClass* RegenerateClass(UClass* ClassToRegenerate, UObject* PreviousCDO) override;
 	virtual void PostLoad() override;
 	virtual void PostLoadSubobjects( FObjectInstancingGraph* OuterInstanceGraph ) override;
 	virtual bool Modify(bool bAlwaysMarkDirty = true) override;
diff --git a/Engine/Source/Runtime/Engine/Classes/Engine/BlueprintGeneratedClass.h b/Engine/Source/Runtime/Engine/Classes/Engine/BlueprintGeneratedClass.h
--- a/Engine/Source/Runtime/Engine/Classes/Engine/BlueprintGeneratedClass.h
+++ b/Engine/Source/Runtime/Engine/Classes/Engine/BlueprintGeneratedClass.h
@@ -694,7 +694,7 @@ public:
 	// UClass interface
 #if WITH_EDITOR
 	virtual UClass* GetAuthoritativeClass() override;
-	virtual void ConditionalRecompileClass(TArray<UObject*>* ObjLoaded) override;
+	virtual void ConditionalRecompileClass() override;
 	virtual void FlushCompilationQueueForLevel() override;
 	virtual UObject* GetArchetypeForCDO() const override;
 #endif //WITH_EDITOR
diff --git a/Engine/Source/Runtime/Engine/Private/Blueprint.cpp b/Engine/Source/Runtime/Engine/Private/Blueprint.cpp
--- a/Engine/Source/Runtime/Engine/Private/Blueprint.cpp
+++ b/Engine/Source/Runtime/Engine/Private/Blueprint.cpp
@@ -596,7 +596,7 @@ void UBlueprint::PostDuplicate(bool bDuplicateForPIE)
 
 extern COREUOBJECT_API bool GBlueprintUseCompilationManager;
 
-UClass* UBlueprint::RegenerateClass(UClass* ClassToRegenerate, UObject* PreviousCDO, TArray<UObject*>& ObjLoaded)
+UClass* UBlueprint::RegenerateClass(UClass* ClassToRegenerate, UObject* PreviousCDO)
 {
 	LoadModulesRequiredForCompilation();
 
@@ -655,7 +655,7 @@ UClass* UBlueprint::RegenerateClass(UClass* ClassToRegenerate, UObject* Previous
 	}
 	else
 	{
-		return FBlueprintEditorUtils::RegenerateBlueprintClass(this, ClassToRegenerate, PreviousCDO, ObjLoaded);
+		return FBlueprintEditorUtils::RegenerateBlueprintClass(this, ClassToRegenerate, PreviousCDO);
 	}
 }
 
diff --git a/Engine/Source/Runtime/Engine/Private/BlueprintGeneratedClass.cpp b/Engine/Source/Runtime/Engine/Private/BlueprintGeneratedClass.cpp
--- a/Engine/Source/Runtime/Engine/Private/BlueprintGeneratedClass.cpp
+++ b/Engine/Source/Runtime/Engine/Private/BlueprintGeneratedClass.cpp
@@ -216,11 +216,11 @@ struct FConditionalRecompileClassHepler
 extern UNREALED_API FSecondsCounterData BlueprintCompileAndLoadTimerData;
 extern COREUOBJECT_API bool GBlueprintUseCompilationManager;
 
-void UBlueprintGeneratedClass::ConditionalRecompileClass(TArray<UObject*>* ObjLoaded)
+void UBlueprintGeneratedClass::ConditionalRecompileClass()
 {
 	if(GBlueprintUseCompilationManager)
 	{
-		FBlueprintCompilationManager::FlushCompilationQueue(ObjLoaded);
+		FBlueprintCompilationManager::FlushCompilationQueue();
 		return;
 	}
 	
@@ -255,7 +255,7 @@ void UBlueprintGeneratedClass::ConditionalRecompileClass(TArray<UObject*>* ObjLo
 			}
 			if ((GeneratingBP->Status != BS_Error) && (GeneratingBP->BlueprintType != EBlueprintType::BPTYPE_MacroLibrary))
 			{
-				FKismetEditorUtilities::RecompileBlueprintBytecode(GeneratingBP, ObjLoaded);
+				FKismetEditorUtilities::RecompileBlueprintBytecode(GeneratingBP);
 			}
 
 			GeneratingBP->bIsRegeneratingOnLoad = bWasRegenerating;
diff --git a/Engine/Source/Runtime/Engine/Private/Materials/Material.cpp b/Engine/Source/Runtime/Engine/Private/Materials/Material.cpp
--- a/Engine/Source/Runtime/Engine/Private/Materials/Material.cpp
+++ b/Engine/Source/Runtime/Engine/Private/Materials/Material.cpp
@@ -535,7 +535,6 @@ void UMaterialInterface::PostLoadDefaultMaterials()
 			UMaterial* Material = GDefaultMaterials[Domain];
 #if USE_EVENT_DRIVEN_ASYNC_LOAD_AT_BOOT_TIME
 			check(Material || (GIsInitialLoad && GEventDrivenLoaderEnabled));
-			check((GIsInitialLoad && GEventDrivenLoaderEnabled) || !Material->HasAnyFlags(RF_NeedLoad)); //-V595
 			if (Material && !Material->HasAnyFlags(RF_NeedLoad))
 #else
 			check(Material);
diff --git a/Engine/Source/Runtime/Engine/Private/UnrealEngine.cpp b/Engine/Source/Runtime/Engine/Private/UnrealEngine.cpp
--- a/Engine/Source/Runtime/Engine/Private/UnrealEngine.cpp
+++ b/Engine/Source/Runtime/Engine/Private/UnrealEngine.cpp
@@ -221,6 +221,7 @@ UnrealEngine.cpp: Implements the UEngine class and helpers.
 #endif
 
 #include "HAL/FileManagerGeneric.h"
+#include "UObject/UObjectThreadContext.h"
 
 DEFINE_LOG_CATEGORY(LogEngine);
 IMPLEMENT_MODULE( FEngineModule, Engine );
@@ -8575,20 +8576,22 @@ FGuid UEngine::GetPackageGuid(FName PackageName, bool bForPIE)
 {
 	FGuid Result(0,0,0,0);
 
-	BeginLoad(*PackageName.ToString());
+	TRefCountPtr<FUObjectSerializeContext> LoadContext(new FUObjectSerializeContext());
+	BeginLoad(LoadContext, *PackageName.ToString());
 	uint32 LoadFlags = LOAD_NoWarn | LOAD_NoVerify;
 	if (bForPIE)
 	{
 		LoadFlags |= LOAD_PackageForPIE;
 	}
 	UPackage* PackageToReset = nullptr;
-	FLinkerLoad* Linker = GetPackageLinker(NULL, *PackageName.ToString(), LoadFlags, NULL, NULL);
-	if (Linker != NULL && Linker->LinkerRoot != NULL)
+	FLinkerLoad* Linker = GetPackageLinker(nullptr, *PackageName.ToString(), LoadFlags, nullptr, nullptr, nullptr, LoadContext);
+	check(Linker);
+	if (Linker != nullptr && Linker->LinkerRoot != nullptr)
 	{
 		Result = Linker->LinkerRoot->GetGuid();
 		PackageToReset = Linker->LinkerRoot;
 	}
-	EndLoad();
+	EndLoad(Linker->GetSerializeContext());
 
 	ResetLoaders(PackageToReset);
 	Linker = nullptr;
diff --git a/Engine/Source/Runtime/GameplayTags/Private/GameplayTagsManager.cpp b/Engine/Source/Runtime/GameplayTags/Private/GameplayTagsManager.cpp
--- a/Engine/Source/Runtime/GameplayTags/Private/GameplayTagsManager.cpp
+++ b/Engine/Source/Runtime/GameplayTags/Private/GameplayTagsManager.cpp
@@ -59,7 +59,7 @@ void UGameplayTagsManager::LoadGameplayTagTables()
 		// Handle case where the module is dynamically-loaded within a LoadPackage stack, which would otherwise
 		// result in the tag table not having its RowStruct serialized in time. Without the RowStruct, the tags manager
 		// will not be initialized correctly.
-		if (TagTable && IsLoading())
+		if (TagTable)
 		{
 			FLinkerLoad* TagLinker = TagTable->GetLinker();
 			if (TagLinker)
