Fixing HTML5 CIS warning

diff --git a/Engine/Source/Runtime/CoreUObject/Private/UObject/Linker.cpp b/Engine/Source/Runtime/CoreUObject/Private/UObject/Linker.cpp
--- a/Engine/Source/Runtime/CoreUObject/Private/UObject/Linker.cpp
+++ b/Engine/Source/Runtime/CoreUObject/Private/UObject/Linker.cpp
@@ -392,20 +392,20 @@ static void LogGetPackageLinkerError(FArchive* LinkerArchive, FUObjectSerializeC
 		/** Helper function to output more detailed error info if available */
 		static void OutputErrorDetail(FArchive* InLinkerArchive, FUObjectSerializeContext* InLoadContext, const FName& LogName)
 		{
-			FUObjectSerializeContext* LoadContext = InLoadContext;
-			if (!LoadContext && InLinkerArchive)
+			FUObjectSerializeContext* LoadContextToReport = InLoadContext;
+			if (!LoadContextToReport && InLinkerArchive)
 			{
-				LoadContext = InLinkerArchive->GetSerializeContext();
+				LoadContextToReport = InLinkerArchive->GetSerializeContext();
 			}
-			if (LoadContext && LoadContext->SerializedObject && LoadContext->SerializedImportLinker)
+			if (LoadContextToReport && LoadContextToReport->SerializedObject && LoadContextToReport->SerializedImportLinker)
 			{
 				FMessageLog LoadErrors(LogName);
 
 				TSharedRef<FTokenizedMessage> Message = LoadErrors.Info();
 				Message->AddToken(FTextToken::Create(LOCTEXT("FailedLoad_Message", "Failed to load")));
-				Message->AddToken(FAssetNameToken::Create(LoadContext->SerializedImportLinker->GetImportPathName(LoadContext->SerializedImportIndex)));
+				Message->AddToken(FAssetNameToken::Create(LoadContextToReport->SerializedImportLinker->GetImportPathName(LoadContextToReport->SerializedImportIndex)));
 				Message->AddToken(FTextToken::Create(LOCTEXT("FailedLoad_Referenced", "Referenced by")));
-				Message->AddToken(FUObjectToken::Create(LoadContext->SerializedObject));
+				Message->AddToken(FUObjectToken::Create(LoadContextToReport->SerializedObject));
 				UProperty* SerializedProperty = InLinkerArchive ? InLinkerArchive->GetSerializedProperty() : nullptr;
 				if (SerializedProperty != nullptr)
 				{
