Fixing cook crash related to de-globalization of UObject serialization state variables

diff --git a/Engine/Source/Runtime/CoreUObject/Private/Serialization/AsyncLoading.cpp b/Engine/Source/Runtime/CoreUObject/Private/Serialization/AsyncLoading.cpp
--- a/Engine/Source/Runtime/CoreUObject/Private/Serialization/AsyncLoading.cpp
+++ b/Engine/Source/Runtime/CoreUObject/Private/Serialization/AsyncLoading.cpp
@@ -5831,14 +5831,23 @@ EAsyncPackageState::Type FAsyncPackage::CreateLinker()
 		// if the linker already exists, we don't need to lookup the file (it may have been pre-created with
 		// a different filename)
 		Linker = FLinkerLoad::FindExistingLinkerForPackage(Package);
-		if (Linker && GEventDrivenLoaderEnabled)
-		{
-			// this almost works, but the EDL does not tolerate redoing steps it already did
-			UE_LOG(LogStreaming, Fatal, TEXT("Package %s was reloaded before it even closed the linker from a previous load. Seems like a waste of time eh?"), *Desc.Name.ToString());
-			check(Package);
-			FWeakAsyncPackagePtr WeakPtr(this);
-			GPrecacheCallbackHandler.RegisterNewSummaryRequest(this);
-			GPrecacheCallbackHandler.SummaryComplete(WeakPtr);
+		if (Linker)
+		{
+			// Swap the load context to the currently associated with the existing linker
+			check(Linker->GetSerializeContext());
+			LoadContext->DecrementBeginLoadCount();
+			LoadContext = Linker->GetSerializeContext();
+			LoadContext->IncrementBeginLoadCount();
+
+			if (GEventDrivenLoaderEnabled)
+			{
+				// this almost works, but the EDL does not tolerate redoing steps it already did
+				UE_LOG(LogStreaming, Fatal, TEXT("Package %s was reloaded before it even closed the linker from a previous load. Seems like a waste of time eh?"), *Desc.Name.ToString());
+				check(Package);
+				FWeakAsyncPackagePtr WeakPtr(this);
+				GPrecacheCallbackHandler.RegisterNewSummaryRequest(this);
+				GPrecacheCallbackHandler.SummaryComplete(WeakPtr);
+			}
 		}
 
 		if (!Linker)
