CIS fix

diff --git a/Engine/Source/Runtime/CoreUObject/Public/UObject/LinkerLoad.h b/Engine/Source/Runtime/CoreUObject/Public/UObject/LinkerLoad.h
--- a/Engine/Source/Runtime/CoreUObject/Public/UObject/LinkerLoad.h
+++ b/Engine/Source/Runtime/CoreUObject/Public/UObject/LinkerLoad.h
@@ -1192,8 +1192,8 @@ private:
 public:
 
 	//~ FArchive interface
-	virtual void SetSerializeContext(FUObjectSerializeContext* InLoadContext) override;
-	virtual FUObjectSerializeContext* GetSerializeContext() override;
+	COREUOBJECT_API virtual void SetSerializeContext(FUObjectSerializeContext* InLoadContext) override;
+	COREUOBJECT_API virtual FUObjectSerializeContext* GetSerializeContext() override;
 };
 
 // used by the EDL at boot time to coordinate loading with what is going on with the deferred registration stuff
