Fixing a crash when cooking caused by PostLoading text imports too early in the nested Load callstack

diff --git a/Engine/Source/Runtime/CoreUObject/Private/UObject/PropertyBaseObject.cpp b/Engine/Source/Runtime/CoreUObject/Private/UObject/PropertyBaseObject.cpp
--- a/Engine/Source/Runtime/CoreUObject/Private/UObject/PropertyBaseObject.cpp
+++ b/Engine/Source/Runtime/CoreUObject/Private/UObject/PropertyBaseObject.cpp
@@ -232,7 +232,7 @@ void UObjectPropertyBase::ExportTextItem( FString& ValueStr, const void* Propert
  *
  * @return	true if the text is successfully resolved into a valid object reference of the correct type, false otherwise.
  */
-bool UObjectPropertyBase::ParseObjectPropertyValue( const UProperty* Property, UObject* OwnerObject, UClass* RequiredMetaClass, uint32 PortFlags, const TCHAR*& Buffer, UObject*& out_ResolvedValue )
+bool UObjectPropertyBase::ParseObjectPropertyValue(const UProperty* Property, UObject* OwnerObject, UClass* RequiredMetaClass, uint32 PortFlags, const TCHAR*& Buffer, UObject*& out_ResolvedValue, FUObjectSerializeContext* InSerializeContext /*= nullptr*/)
 {
 	check(Property);
 	if (!RequiredMetaClass)
@@ -279,12 +279,12 @@ bool UObjectPropertyBase::ParseObjectPropertyValue( const UProperty* Property, U
 
 			// ignore the object class, it isn't fully qualified, and searching ANY_PACKAGE might get the wrong one!
 			// Try the find the object.
-			out_ResolvedValue = UObjectPropertyBase::FindImportedObject(Property, OwnerObject, ObjectClass, RequiredMetaClass, *ObjectText, PortFlags);
+			out_ResolvedValue = UObjectPropertyBase::FindImportedObject(Property, OwnerObject, ObjectClass, RequiredMetaClass, *ObjectText, PortFlags, InSerializeContext);
 		}
 		else
 		{
 			// Try the find the object.
-			out_ResolvedValue = UObjectPropertyBase::FindImportedObject(Property, OwnerObject, ObjectClass, RequiredMetaClass, *Temp, PortFlags);
+			out_ResolvedValue = UObjectPropertyBase::FindImportedObject(Property, OwnerObject, ObjectClass, RequiredMetaClass, *Temp, PortFlags, InSerializeContext);
 		}
 
 		if ( out_ResolvedValue != nullptr && !out_ResolvedValue->GetClass()->IsChildOf(RequiredMetaClass) )
@@ -316,14 +316,15 @@ const TCHAR* UObjectPropertyBase::ImportText_Internal( const TCHAR* InBuffer, vo
 {
 	const TCHAR* Buffer = InBuffer;
 	UObject* Result = nullptr;
+	FLinkerLoad* Linker = GetLinker();
 
-	bool bOk = ParseObjectPropertyValue(this, Parent, PropertyClass, PortFlags, Buffer, Result);
+	bool bOk = ParseObjectPropertyValue(this, Parent, PropertyClass, PortFlags, Buffer, Result, Linker ? Linker->GetSerializeContext() : nullptr);
 
 	SetObjectPropertyValue(Data, Result);
 	return Buffer;
 }
 
-UObject* UObjectPropertyBase::FindImportedObject( const UProperty* Property, UObject* OwnerObject, UClass* ObjectClass, UClass* RequiredMetaClass, const TCHAR* Text, uint32 PortFlags/*=0*/ )
+UObject* UObjectPropertyBase::FindImportedObject( const UProperty* Property, UObject* OwnerObject, UClass* ObjectClass, UClass* RequiredMetaClass, const TCHAR* Text, uint32 PortFlags/*=0*/, FUObjectSerializeContext* InSerializeContext /*= nullptr*/)
 {
 	UObject*	Result = nullptr;
 	check( ObjectClass->IsChildOf(RequiredMetaClass) );
@@ -408,7 +409,7 @@ UObject* UObjectPropertyBase::FindImportedObject( const UProperty* Property, UOb
 		if (Dot && AttemptNonQualifiedSearch)
 		{
 			// search with just the object name
-			Result = FindImportedObject(Property, OwnerObject, ObjectClass, RequiredMetaClass, Dot + 1);
+			Result = FindImportedObject(Property, OwnerObject, ObjectClass, RequiredMetaClass, Dot + 1, 0, InSerializeContext);
 		}
 		FString NewText(Text);
 		// if it didn't have a dot, then maybe they just gave a uasset package name
@@ -445,7 +446,7 @@ UObject* UObjectPropertyBase::FindImportedObject( const UProperty* Property, UOb
 				const uint32 LoadFlags = LOAD_NoWarn | LOAD_FindIfFail;		
 
 				UE_LOG(LogProperty, Verbose, TEXT("FindImportedObject is attempting to import [%s] (class = %s) with StaticLoadObject"), Text, *GetFullNameSafe(ObjectClass));
-				Result = StaticLoadObject(ObjectClass, nullptr, Text, nullptr, LoadFlags, nullptr);
+				Result = StaticLoadObject(ObjectClass, nullptr, Text, nullptr, LoadFlags, nullptr, true, InSerializeContext);
 
 #if USE_DEFERRED_DEPENDENCY_CHECK_VERIFICATION_TESTS
 				check(!bDeferAssetImports || !Result || !FBlueprintSupport::IsInBlueprintPackage(Result));
diff --git a/Engine/Source/Runtime/CoreUObject/Private/UObject/UObjectGlobals.cpp b/Engine/Source/Runtime/CoreUObject/Private/UObject/UObjectGlobals.cpp
--- a/Engine/Source/Runtime/CoreUObject/Private/UObject/UObjectGlobals.cpp
+++ b/Engine/Source/Runtime/CoreUObject/Private/UObject/UObjectGlobals.cpp
@@ -752,7 +752,7 @@ const FString* GetIniFilenameFromObjectsReference(const FString& Name)
 //
 // Resolve a package and name.
 //
-bool ResolveName(UObject*& InPackage, FString& InOutName, bool Create, bool Throw, uint32 LoadFlags /*= LOAD_None*/)
+bool ResolveName(UObject*& InPackage, FString& InOutName, bool Create, bool Throw, uint32 LoadFlags /*= LOAD_None*/, FUObjectSerializeContext* InLoadContext /*= nullptr*/)
 {
 	// Strip off the object class.
 	ConstructorHelpers::StripObjectClass( InOutName );
@@ -823,7 +823,7 @@ bool ResolveName(UObject*& InPackage, FString& InOutName, bool Create, bool Thro
 			InPackage = StaticFindObjectFast(UPackage::StaticClass(), InPackage, *PartialName);
 			if (!ScriptPackageName && !InPackage)
 			{
-				InPackage = LoadPackage(dynamic_cast<UPackage*>(InPackage), *PartialName, LoadFlags);
+				InPackage = LoadPackage(dynamic_cast<UPackage*>(InPackage), *PartialName, LoadFlags, nullptr, InLoadContext);
 			}
 			if (!InPackage)
 			{
@@ -898,7 +898,7 @@ bool ParseObject( const TCHAR* Stream, const TCHAR* Match, UClass* Class, UObjec
  *
  * @return The object that was loaded or found. NULL for a failure.
  */
-UObject* StaticLoadObjectInternal(UClass* ObjectClass, UObject* InOuter, const TCHAR* InName, const TCHAR* Filename, uint32 LoadFlags, UPackageMap* Sandbox, bool bAllowObjectReconciliation)
+UObject* StaticLoadObjectInternal(UClass* ObjectClass, UObject* InOuter, const TCHAR* InName, const TCHAR* Filename, uint32 LoadFlags, UPackageMap* Sandbox, bool bAllowObjectReconciliation, FUObjectSerializeContext* InSerializeContext)
 {
 	SCOPE_CYCLE_COUNTER(STAT_LoadObject);
 	check(ObjectClass);
@@ -910,7 +910,7 @@ UObject* StaticLoadObjectInternal(UClass* ObjectClass, UObject* InOuter, const T
 	const bool bContainsObjectName = !!FCString::Strstr(InName, TEXT("."));
 
 	// break up the name into packages, returning the innermost name and its outer
-	ResolveName(InOuter, StrName, true, true, LoadFlags & (LOAD_EditorOnly | LOAD_Quiet | LOAD_NoWarn | LOAD_DeferDependencyLoads));
+	ResolveName(InOuter, StrName, true, true, LoadFlags & (LOAD_EditorOnly | LOAD_Quiet | LOAD_NoWarn | LOAD_DeferDependencyLoads), InSerializeContext);
 	if (InOuter)
 	{
 		// If we have a full UObject name then attempt to find the object in memory first,
@@ -933,7 +933,7 @@ UObject* StaticLoadObjectInternal(UClass* ObjectClass, UObject* InOuter, const T
 			if (!InOuter->GetOutermost()->HasAnyPackageFlags(PKG_CompiledIn))
 			{
 				// now that we have one asset per package, we load the entire package whenever a single object is requested
-				LoadPackage(NULL, *InOuter->GetOutermost()->GetName(), LoadFlags & ~LOAD_Verify);
+				LoadPackage(NULL, *InOuter->GetOutermost()->GetName(), LoadFlags & ~LOAD_Verify, nullptr, InSerializeContext);
 			}
 
 			// now, find the object in the package
@@ -962,7 +962,7 @@ UObject* StaticLoadObjectInternal(UClass* ObjectClass, UObject* InOuter, const T
 		StrName = InName;
 		StrName += TEXT(".");
 		StrName += FPackageName::GetShortName(InName);
-		Result = StaticLoadObjectInternal(ObjectClass, InOuter, *StrName, Filename, LoadFlags, Sandbox, bAllowObjectReconciliation);
+		Result = StaticLoadObjectInternal(ObjectClass, InOuter, *StrName, Filename, LoadFlags, Sandbox, bAllowObjectReconciliation, InSerializeContext);
 	}
 #if WITH_EDITORONLY_DATA
 	else if (Result && !(LoadFlags & LOAD_EditorOnly))
@@ -974,7 +974,7 @@ UObject* StaticLoadObjectInternal(UClass* ObjectClass, UObject* InOuter, const T
 	return Result;
 }
 
-UObject* StaticLoadObject(UClass* ObjectClass, UObject* InOuter, const TCHAR* InName, const TCHAR* Filename, uint32 LoadFlags, UPackageMap* Sandbox, bool bAllowObjectReconciliation )
+UObject* StaticLoadObject(UClass* ObjectClass, UObject* InOuter, const TCHAR* InName, const TCHAR* Filename, uint32 LoadFlags, UPackageMap* Sandbox, bool bAllowObjectReconciliation, FUObjectSerializeContext* InSerializeContext)
 {
 	FUObjectThreadContext& ThreadContext = FUObjectThreadContext::Get();
 	if (ThreadContext.IsRoutingPostLoad && IsInAsyncLoadingThread())
@@ -986,7 +986,7 @@ UObject* StaticLoadObject(UClass* ObjectClass, UObject* InOuter, const TCHAR* In
 			*GetFullNameSafe(ThreadContext.CurrentlyPostLoadedObjectByALT));
 	}
 
-	UObject* Result = StaticLoadObjectInternal(ObjectClass, InOuter, InName, Filename, LoadFlags, Sandbox, bAllowObjectReconciliation);
+	UObject* Result = StaticLoadObjectInternal(ObjectClass, InOuter, InName, Filename, LoadFlags, Sandbox, bAllowObjectReconciliation, InSerializeContext);
 	if (!Result)
 	{
 		FString ObjectName = InName;
@@ -1463,7 +1463,7 @@ UPackage* LoadPackageInternal(UPackage* InOuter, const TCHAR* InLongPackageNameO
 	return Result;
 }
 
-UPackage* LoadPackage(UPackage* InOuter, const TCHAR* InLongPackageName, uint32 LoadFlags, FArchive* InReaderOverride)
+UPackage* LoadPackage(UPackage* InOuter, const TCHAR* InLongPackageName, uint32 LoadFlags, FArchive* InReaderOverride, FUObjectSerializeContext* InLoadContext)
 {
 	COOK_STAT(LoadPackageStats::NumPackagesLoaded++);
 	COOK_STAT(FScopedDurationTimer LoadTimer(LoadPackageStats::LoadPackageTimeSec));
@@ -1478,7 +1478,7 @@ UPackage* LoadPackage(UPackage* InOuter, const TCHAR* InLongPackageName, uint32
 	// since we are faking the object name, this is basically a duplicate of LLM_SCOPED_TAG_WITH_OBJECT_IN_SET
 	FString FakePackageName = FString(TEXT("Package ")) + InLongPackageName;
 	LLM_SCOPED_TAG_WITH_STAT_NAME_IN_SET(FLowLevelMemTracker::Get().IsTagSetActive(ELLMTagSet::Assets) ? FDynamicStats::CreateMemoryStatId<FStatGroup_STATGROUP_LLMAssets>(FName(*FakePackageName)).GetName() : NAME_None, ELLMTagSet::Assets, ELLMTracker::Default);
-	return LoadPackageInternal(InOuter, InLongPackageName, LoadFlags, /*ImportLinker =*/ nullptr, InReaderOverride, nullptr);
+	return LoadPackageInternal(InOuter, InLongPackageName, LoadFlags, /*ImportLinker =*/ nullptr, InReaderOverride, InLoadContext);
 }
 
 static int32 GGameThreadLoadCounter = 0;
diff --git a/Engine/Source/Runtime/CoreUObject/Public/UObject/UObjectGlobals.h b/Engine/Source/Runtime/CoreUObject/Public/UObject/UObjectGlobals.h
--- a/Engine/Source/Runtime/CoreUObject/Public/UObject/UObjectGlobals.h
+++ b/Engine/Source/Runtime/CoreUObject/Public/UObject/UObjectGlobals.h
@@ -174,7 +174,7 @@ COREUOBJECT_API const FString* GetIniFilenameFromObjectsReference(const FString&
  */
 COREUOBJECT_API FString ResolveIniObjectsReference(const FString& ObjectReference, const FString* IniFilename = nullptr, bool bThrow = false);
 
-COREUOBJECT_API bool ResolveName(UObject*& Outer, FString& ObjectsReferenceString, bool Create, bool Throw, uint32 LoadFlags = LOAD_None);
+COREUOBJECT_API bool ResolveName(UObject*& Outer, FString& ObjectsReferenceString, bool Create, bool Throw, uint32 LoadFlags = LOAD_None, FUObjectSerializeContext* InLoadContext = nullptr);
 COREUOBJECT_API void SafeLoadError( UObject* Outer, uint32 LoadFlags, const TCHAR* ErrorMessage);
 
 COREUOBJECT_API int32 UpdateSuffixForNextNewObject(UObject* Parent, UClass* Class, TFunctionRef<void(int32&)> IndexMutator);
@@ -222,10 +222,11 @@ COREUOBJECT_API bool ParseObject( const TCHAR* Stream, const TCHAR* Match, UClas
  * @param LoadFlags		Flags controlling how to handle loading from disk
  * @param Sandbox		A list of packages to restrict the search for the object
  * @param bAllowObjectReconciliation	Whether to allow the object to be found via FindObject in the case of seek free loading
+ * @param InSerializeContext	Additional context when called during serialization
  *
  * @return The object that was loaded or found. NULL for a failure.
  */
-COREUOBJECT_API UObject* StaticLoadObject( UClass* Class, UObject* InOuter, const TCHAR* Name, const TCHAR* Filename = NULL, uint32 LoadFlags = LOAD_None, UPackageMap* Sandbox = NULL, bool bAllowObjectReconciliation = true );
+COREUOBJECT_API UObject* StaticLoadObject( UClass* Class, UObject* InOuter, const TCHAR* Name, const TCHAR* Filename = NULL, uint32 LoadFlags = LOAD_None, UPackageMap* Sandbox = NULL, bool bAllowObjectReconciliation = true, FUObjectSerializeContext* InSerializeContext = nullptr );
 COREUOBJECT_API UClass* StaticLoadClass(UClass* BaseClass, UObject* InOuter, const TCHAR* Name, const TCHAR* Filename = NULL, uint32 LoadFlags = LOAD_None, UPackageMap* Sandbox = NULL);
 
 /**
@@ -281,9 +282,10 @@ COREUOBJECT_API void StaticTick( float DeltaTime, bool bUseFullTimeLimit = true,
  * @param	InOuter		Package to load new package into (usually NULL or ULevel->GetOuter())
  * @param	Filename	Long package name to load
  * @param	LoadFlags	Flags controlling loading behavior
+ * @param InLoadContext	Additional context when called during serialization
  * @return	Loaded package if successful, NULL otherwise
  */
-COREUOBJECT_API UPackage* LoadPackage( UPackage* InOuter, const TCHAR* InLongPackageName, uint32 LoadFlags, FArchive* InReaderOverride = nullptr );
+COREUOBJECT_API UPackage* LoadPackage( UPackage* InOuter, const TCHAR* InLongPackageName, uint32 LoadFlags, FArchive* InReaderOverride = nullptr, FUObjectSerializeContext* InLoadContext = nullptr );
 
 /* Async package loading result */
 namespace EAsyncLoadingResult
diff --git a/Engine/Source/Runtime/CoreUObject/Public/UObject/UnrealType.h b/Engine/Source/Runtime/CoreUObject/Public/UObject/UnrealType.h
--- a/Engine/Source/Runtime/CoreUObject/Public/UObject/UnrealType.h
+++ b/Engine/Source/Runtime/CoreUObject/Public/UObject/UnrealType.h
@@ -1944,11 +1944,11 @@ public:
 	 *								fully qualified or not), or can be formatted as a const object reference (i.e. SomeClass'SomePackage.TheObject')
 	 *								When the function returns, Buffer will be pointing to the first character after the object value text in the input stream.
 	 * @param	ResolvedValue		receives the object that is resolved from the input text.
-	 *
+	 * @param InSerializeContext	Additional context when called during serialization
 	 * @return	true if the text is successfully resolved into a valid object reference of the correct type, false otherwise.
 	 */
-	static bool ParseObjectPropertyValue( const UProperty* Property, UObject* OwnerObject, UClass* RequiredMetaClass, uint32 PortFlags, const TCHAR*& Buffer, UObject*& out_ResolvedValue );
-	static UObject* FindImportedObject( const UProperty* Property, UObject* OwnerObject, UClass* ObjectClass, UClass* RequiredMetaClass, const TCHAR* Text, uint32 PortFlags = 0);
+	static bool ParseObjectPropertyValue( const UProperty* Property, UObject* OwnerObject, UClass* RequiredMetaClass, uint32 PortFlags, const TCHAR*& Buffer, UObject*& out_ResolvedValue, FUObjectSerializeContext* InSerializeContext = nullptr );
+	static UObject* FindImportedObject( const UProperty* Property, UObject* OwnerObject, UClass* ObjectClass, UClass* RequiredMetaClass, const TCHAR* Text, uint32 PortFlags = 0, FUObjectSerializeContext* InSerializeContext = nullptr);
 	
 	// Returns the qualified export path for a given object, parent, and export root scope
 	static FString GetExportPath(const UObject* Object, const UObject* Parent, const UObject* ExportRootScope, const uint32 PortFlags);
