IsValidLowLevelFast can fail during initial load so skip it when getting an archetype of an object early in engine init. Fixes a crash when enabling default subobject checks in cooked builds.

diff --git a/Engine/Source/Runtime/CoreUObject/Private/UObject/UObjectArchetype.cpp b/Engine/Source/Runtime/CoreUObject/Private/UObject/UObjectArchetype.cpp
--- a/Engine/Source/Runtime/CoreUObject/Private/UObject/UObjectArchetype.cpp
+++ b/Engine/Source/Runtime/CoreUObject/Private/UObject/UObjectArchetype.cpp
@@ -168,7 +168,8 @@ UObject* UObject::GetArchetype() const
 		}
 	}
 #endif // UE_VERIFY_CACHED_ARCHETYPE
-	check(Archetype == nullptr || Archetype->IsValidLowLevelFast());
+	// Note that IsValidLowLevelFast check may fail during initial load as not all classes are initialized at this point so skip it
+	check(Archetype == nullptr || GIsInitialLoad || Archetype->IsValidLowLevelFast());
 
 	return Archetype;
 #else
