Added gc.DumpRefsToCluster Name=ClusterRootObjectName command to dump references to all objects within a cluster (useful for determining if clusters are indeed 'stale')

diff --git a/Engine/Source/Runtime/CoreUObject/Private/UObject/UObjectClusters.cpp b/Engine/Source/Runtime/CoreUObject/Private/UObject/UObjectClusters.cpp
--- a/Engine/Source/Runtime/CoreUObject/Private/UObject/UObjectClusters.cpp
+++ b/Engine/Source/Runtime/CoreUObject/Private/UObject/UObjectClusters.cpp
@@ -468,7 +468,7 @@ void FindStaleClusters(const TArray<FString>& Args)
 			TotalNumClusters++;
 
 			UObject* ClusterRootObject = static_cast<UObject*>(ObjectItem->Object);
-			FReferenceChainSearch SearchRefs(ClusterRootObject, EReferenceChainSearchMode::ExternalOnly|EReferenceChainSearchMode::Direct);
+			FReferenceChainSearch SearchRefs(ClusterRootObject, EReferenceChainSearchMode::ExternalOnly);
 			
 			bool bReferenced = false;
 			if (SearchRefs.GetReferenceChains().Num() > 0)
@@ -491,6 +491,59 @@ void FindStaleClusters(const TArray<FString>& Args)
 	UE_LOG(LogObj, Display, TEXT("Found %d clusters, including %d stale."), TotalNumClusters, NumStaleClusters);
 }
 
+
+void DumpRefsToCluster(FUObjectCluster* Cluster)
+{
+	FUObjectItem* RootItem = GUObjectArray.IndexToObject(Cluster->RootIndex);
+
+	UE_LOG(LogObj, Display, TEXT("Dumping references to objects in cluster %s"), *static_cast<UObject*>(RootItem->Object)->GetFullName());
+
+	bool bIsReferenced = false;
+	for (int32 ObjectIndex : Cluster->Objects)
+	{
+		FUObjectItem* ObjectItem = GUObjectArray.IndexToObject(ObjectIndex);
+		FReferenceChainSearch SearchRefs(static_cast<UObject*>(ObjectItem->Object), EReferenceChainSearchMode::ExternalOnly | EReferenceChainSearchMode::Shortest);
+		if (SearchRefs.GetReferenceChains().Num())
+		{
+			bIsReferenced = true;
+			SearchRefs.PrintResults(true);
+		}
+	}
+	if (!bIsReferenced)
+	{
+		UE_LOG(LogObj, Display, TEXT("Cluster %s is not currently referenced by anything."), *static_cast<UObject*>(RootItem->Object)->GetFullName());
+	}
+}
+
+void DumpRefsToCluster(const TArray<FString>& Args)
+{
+	// This is seriously slow.
+	UE_LOG(LogObj, Display, TEXT("Searching for references to clusteres. This may take a while..."));
+
+	TArray<int32> RootObjects;
+	for (const FString& Arg : Args)
+	{
+		if (Arg.StartsWith("Root="))
+		{
+			FString ObjectsList = Arg.Mid(5);
+			TArray<FString> ObjectNames;
+			ObjectsList.ParseIntoArray(ObjectNames, TEXT(","));
+			ParseObjectNameArrayForClusters(RootObjects, ObjectNames);
+		}
+	}
+
+	for (int32 RootIndex : RootObjects)
+	{
+		FUObjectItem* RootItem = GUObjectArray.IndexToObject(RootIndex);
+		if (RootItem->HasAnyFlags(EInternalObjectFlags::ClusterRoot))
+		{
+			FUObjectCluster* Cluster = GUObjectClusters.GetObjectCluster(static_cast<UObject*>(RootItem->Object));
+			DumpRefsToCluster(Cluster);
+		}
+	}
+}
+
+
 static FAutoConsoleCommand ListClustersCommand(
 	TEXT("gc.ListClusters"),
 	TEXT("Dumps all clusters do output log. When 'Hiearchy' argument is specified lists all objects inside clusters."),
@@ -503,6 +556,12 @@ static FAutoConsoleCommand FindStaleClustersCommand(
 	FConsoleCommandWithArgsDelegate::CreateStatic(FindStaleClusters)
 	);
 
+static FAutoConsoleCommand DumpRefsToClusterCommand(
+	TEXT("gc.DumpRefsToCluster"),
+	TEXT("Dumps references to all objects within a cluster. Specify the cluster name with Root=Name."),
+	FConsoleCommandWithArgsDelegate::CreateStatic(DumpRefsToCluster)
+);
+
 #endif // !UE_BUILD_SHIPPING
 
 /**
