More efficient implementation of FExclusiveLoadPackageTimeTracker::Get() that inlines better and avoids the thread-safe machinery that function-level statics entail.

diff --git a/Engine/Source/Runtime/CoreUObject/Private/Misc/ExclusiveLoadPackageTimeTracker.cpp b/Engine/Source/Runtime/CoreUObject/Private/Misc/ExclusiveLoadPackageTimeTracker.cpp
--- a/Engine/Source/Runtime/CoreUObject/Private/Misc/ExclusiveLoadPackageTimeTracker.cpp
+++ b/Engine/Source/Runtime/CoreUObject/Private/Misc/ExclusiveLoadPackageTimeTracker.cpp
@@ -13,6 +13,8 @@
 
 #define LOCTEXT_NAMESPACE "ExclusiveLoadPackageTimeTracker"
 
+FExclusiveLoadPackageTimeTracker* FExclusiveLoadPackageTimeTracker::TrackerInstance;
+
 FExclusiveLoadPackageTimeTracker::FExclusiveLoadPackageTimeTracker()
 	: TrackerOverhead(0.f)
 	, EndLoadName(FName(TEXT("EndLoad")))
@@ -28,9 +30,10 @@ FExclusiveLoadPackageTimeTracker::FExclusiveLoadPackageTimeTracker()
 {
 }
 
-FExclusiveLoadPackageTimeTracker& FExclusiveLoadPackageTimeTracker::Get()
+FExclusiveLoadPackageTimeTracker& FExclusiveLoadPackageTimeTracker::Construct()
 {
 	static FExclusiveLoadPackageTimeTracker Tracker;
+	TrackerInstance = &Tracker;
 	return Tracker;
 }
 
diff --git a/Engine/Source/Runtime/CoreUObject/Public/Misc/ExclusiveLoadPackageTimeTracker.h b/Engine/Source/Runtime/CoreUObject/Public/Misc/ExclusiveLoadPackageTimeTracker.h
--- a/Engine/Source/Runtime/CoreUObject/Public/Misc/ExclusiveLoadPackageTimeTracker.h
+++ b/Engine/Source/Runtime/CoreUObject/Public/Misc/ExclusiveLoadPackageTimeTracker.h
@@ -110,7 +110,8 @@ public:
 	FORCEINLINE static void PushEndLoad()
 	{
 #if WITH_LOADPACKAGE_TIME_TRACKER
-		Get().InternalPushLoadPackage(Get().EndLoadName);
+		FExclusiveLoadPackageTimeTracker& Instance = Get();
+		Instance.InternalPushLoadPackage(Instance.EndLoadName);
 #endif
 	}
 
@@ -186,8 +187,16 @@ private:
 		{}
 	};
 
+	inline static FExclusiveLoadPackageTimeTracker& Get()
+	{
+		if (TrackerInstance)
+			return *TrackerInstance;
+		else
+			return Construct();
+	}
+
 	FExclusiveLoadPackageTimeTracker();
-	COREUOBJECT_API static FExclusiveLoadPackageTimeTracker& Get();
+	COREUOBJECT_API static FExclusiveLoadPackageTimeTracker& Construct();
 
 	/** Starts a time for the specified package name. */
 	COREUOBJECT_API void InternalPushLoadPackage(FName PackageName);
@@ -238,5 +247,7 @@ private:
 	/** Auto-registered console command to call ResetReport() */
 	const FAutoConsoleCommand ResetReportCommand;
 
+	static FExclusiveLoadPackageTimeTracker* TrackerInstance;
+
 #endif // WITH_LOADPACKAGE_TIME_TRACKER
 };
