Fixed automation test errors caused by de-globalization of UObject load state

Added LoadPackageLinker function that replaces { BeginLoad(); GetPackageLinker(); EndLoad(); } calls in a convenient way.

diff --git a/Engine/Source/Editor/MovieSceneTools/Private/K2Node_GetSequenceBinding.cpp b/Engine/Source/Editor/MovieSceneTools/Private/K2Node_GetSequenceBinding.cpp
--- a/Engine/Source/Editor/MovieSceneTools/Private/K2Node_GetSequenceBinding.cpp
+++ b/Engine/Source/Editor/MovieSceneTools/Private/K2Node_GetSequenceBinding.cpp
@@ -95,7 +95,8 @@ public:
 
 UMovieSceneSequence* UK2Node_GetSequenceBinding::GetSequence() const
 {
-	return Cast<UMovieSceneSequence>(SourceSequence.TryLoad());
+	FLinkerLoad* Linker = GetLinker();
+	return Cast<UMovieSceneSequence>(SourceSequence.TryLoad(Linker ? Linker->GetSerializeContext() : nullptr));
 }
 
 void UK2Node_GetSequenceBinding::ValidateNodeDuringCompilation(FCompilerResultsLog& MessageLog) const
diff --git a/Engine/Source/Runtime/CoreUObject/Private/UObject/Linker.cpp b/Engine/Source/Runtime/CoreUObject/Private/UObject/Linker.cpp
--- a/Engine/Source/Runtime/CoreUObject/Private/UObject/Linker.cpp
+++ b/Engine/Source/Runtime/CoreUObject/Private/UObject/Linker.cpp
@@ -542,9 +542,10 @@ FLinkerLoad* GetPackageLinker
 	UPackageMap*	Sandbox,
 	FGuid*			CompatibleGuid,
 	FArchive*		InReaderOverride,
-	FUObjectSerializeContext* InExistingContext
+	FUObjectSerializeContext** InOutLoadContext
 )
 {
+	FUObjectSerializeContext* InExistingContext = InOutLoadContext ? *InOutLoadContext : nullptr;
 	// See if there is already a linker for this package.
 	FLinkerLoad* Result = FLinkerLoad::FindExistingLinkerForPackage(InOuter);
 
@@ -552,6 +553,13 @@ FLinkerLoad* GetPackageLinker
 	// See if the linker is already loaded.
 	if (Result)
 	{
+		if (InExistingContext && Result->GetSerializeContext() && Result->GetSerializeContext() != InExistingContext)
+		{
+			if (!Result->GetSerializeContext()->HasStartedLoading())
+			{
+				Result->SetSerializeContext(InExistingContext);
+			}
+		}
 		return Result;
 	}
 
@@ -717,7 +725,16 @@ FLinkerLoad* GetPackageLinker
 	}
 	else if (InExistingContext)
 	{
-		Result->SetSerializeContext(InExistingContext);
+		if (Result->GetSerializeContext() && Result->GetSerializeContext()->HasStartedLoading() && InExistingContext->GetBeginLoadCount() == 1)
+		{
+			// Use the context associated with the linker because it has already started loading objects
+			*InOutLoadContext = Result->GetSerializeContext();
+		}
+		else
+		{
+			// Replace the linker context with the one passed into this function
+			Result->SetSerializeContext(InExistingContext);
+		}
 	}
 
 	if ( !Result && CreatedPackage )
@@ -745,6 +762,26 @@ FLinkerLoad* GetPackageLinker
 	return Result;
 }
 
+FLinkerLoad* LoadPackageLinker(UPackage* InOuter, const TCHAR* InLongPackageName, uint32 LoadFlags, UPackageMap* Sandbox, FGuid* CompatibleGuid, FArchive* InReaderOverride)
+{
+	FLinkerLoad* Linker = nullptr;
+	TRefCountPtr<FUObjectSerializeContext> LoadContext(new FUObjectSerializeContext());
+	BeginLoad(LoadContext);
+	{
+		FUObjectSerializeContext* InOutLoadContext = LoadContext;
+		Linker = GetPackageLinker(InOuter, InLongPackageName, LoadFlags, Sandbox, CompatibleGuid, InReaderOverride, &InOutLoadContext);
+		if (InOutLoadContext != LoadContext)
+		{
+			// The linker already existed and was associated with another context
+			LoadContext->DecrementBeginLoadCount();
+			LoadContext = InOutLoadContext;
+			LoadContext->IncrementBeginLoadCount();
+		}
+	}
+	EndLoad(Linker ? Linker->GetSerializeContext() : LoadContext.GetReference());
+	return Linker;
+}
+
 /**
  * 
  * Ensure thumbnails are loaded and then reset the loader in preparation for a package save
diff --git a/Engine/Source/Runtime/CoreUObject/Private/UObject/SoftObjectPath.cpp b/Engine/Source/Runtime/CoreUObject/Private/UObject/SoftObjectPath.cpp
--- a/Engine/Source/Runtime/CoreUObject/Private/UObject/SoftObjectPath.cpp
+++ b/Engine/Source/Runtime/CoreUObject/Private/UObject/SoftObjectPath.cpp
@@ -344,13 +344,13 @@ bool FSoftObjectPath::SerializeFromMismatchedTag(struct FPropertyTag const& Tag,
 	return bReturn;
 }
 
-UObject* FSoftObjectPath::TryLoad() const
+UObject* FSoftObjectPath::TryLoad(FUObjectSerializeContext* InLoadContext) const
 {
 	UObject* LoadedObject = nullptr;
 
 	if (!IsNull())
 	{
-		LoadedObject = LoadObject<UObject>(nullptr, *ToString());
+		LoadedObject = StaticLoadObject(UObject::StaticClass(), nullptr, *ToString(), nullptr, LOAD_None, nullptr, true, InLoadContext);
 
 #if WITH_EDITOR
 		// Look at core redirects if we didn't find the object
diff --git a/Engine/Source/Runtime/CoreUObject/Private/UObject/UObjectGlobals.cpp b/Engine/Source/Runtime/CoreUObject/Private/UObject/UObjectGlobals.cpp
--- a/Engine/Source/Runtime/CoreUObject/Private/UObject/UObjectGlobals.cpp
+++ b/Engine/Source/Runtime/CoreUObject/Private/UObject/UObjectGlobals.cpp
@@ -1280,8 +1280,18 @@ UPackage* LoadPackageInternal(UPackage* InOuter, const TCHAR* InLongPackageNameO
 		}
 #endif
 
-		Linker = GetPackageLinker(InOuter, *FileToLoad, LoadFlags, nullptr, nullptr, InReaderOverride, LoadContext);
-		
+		{
+			FUObjectSerializeContext* InOutLoadContext = LoadContext;
+			Linker = GetPackageLinker(InOuter, *FileToLoad, LoadFlags, nullptr, nullptr, InReaderOverride, &InOutLoadContext);
+			if (InOutLoadContext != LoadContext)
+			{
+				// The linker already existed and was associated with another context
+				LoadContext->DecrementBeginLoadCount();
+				LoadContext = InOutLoadContext;
+				LoadContext->IncrementBeginLoadCount();
+			}			
+		}
+
 		if (!Linker)
 		{
 			EndLoad(LoadContext);
diff --git a/Engine/Source/Runtime/CoreUObject/Public/UObject/Linker.h b/Engine/Source/Runtime/CoreUObject/Public/UObject/Linker.h
--- a/Engine/Source/Runtime/CoreUObject/Public/UObject/Linker.h
+++ b/Engine/Source/Runtime/CoreUObject/Public/UObject/Linker.h
@@ -6,6 +6,7 @@
 #include "UObject/ObjectResource.h"
 #include "Internationalization/GatherableTextData.h"
 #include "UObject/PackageFileSummary.h"
+#include "Templates/RefCounting.h"
 
 class FReferenceCollector;
 
@@ -564,7 +565,31 @@ COREUOBJECT_API void DeleteLoaders();
 /** Queues linker for deletion */
 COREUOBJECT_API void DeleteLoader(FLinkerLoad* Loader);
 
-COREUOBJECT_API FLinkerLoad* GetPackageLinker(UPackage* InOuter, const TCHAR* InLongPackageName, uint32 LoadFlags, UPackageMap* Sandbox, FGuid* CompatibleGuid, FArchive* InReaderOverride = nullptr, FUObjectSerializeContext* InExistingContext = nullptr);
+/** 
+  * Loads a linker for a package and returns it without loading any objects.
+  * @param InOuter Package if known, can be null
+	* @param InLongPackageName Name of the package to load
+	* @param LoadFlags Flags to pass to the new linker
+	* @param Sandbox Additional sandbox for loading
+	* @param CompatibleGuid Net GUID
+	* @param InReaderOverride Optional archive to use for reading package data
+	* @return Pointer to the loaded linker or null if the file didn't exist
+	*/
+COREUOBJECT_API FLinkerLoad* LoadPackageLinker(UPackage* InOuter, const TCHAR* InLongPackageName, uint32 LoadFlags, UPackageMap* Sandbox, FGuid* CompatibleGuid, FArchive* InReaderOverride = nullptr);
+
+/** 
+  * Gets a linker for a package and returns it without loading any objects. This call must be preceeded by BeginLoad and followed by EndLoad calls
+  * @param InOuter Package if known, can be null
+  * @param InLongPackageName Name of the package to load
+  * @param LoadFlags Flags to pass to the new linker
+  * @param Sandbox Additional sandbox for loading
+  * @param CompatibleGuid Net GUID
+  * @param InReaderOverride Optional archive to use for reading package data
+	* @param InOutLoadContext Optional load context. If the package linker is already associated with a context that's currently loading objects that context will be returned in this param
+  * @return Pointer to the loaded linker or null if the file didn't exist
+  */
+COREUOBJECT_API FLinkerLoad* GetPackageLinker(UPackage* InOuter, const TCHAR* InLongPackageName, uint32 LoadFlags, UPackageMap* Sandbox, FGuid* CompatibleGuid, FArchive* InReaderOverride = nullptr, FUObjectSerializeContext** InOutLoadContext = nullptr);
+
 COREUOBJECT_API FString GetPrestreamPackageLinkerName(const TCHAR* InLongPackageName, bool bExistSkip = true);
 
 
diff --git a/Engine/Source/Runtime/CoreUObject/Public/UObject/LinkerLoad.h b/Engine/Source/Runtime/CoreUObject/Public/UObject/LinkerLoad.h
--- a/Engine/Source/Runtime/CoreUObject/Public/UObject/LinkerLoad.h
+++ b/Engine/Source/Runtime/CoreUObject/Public/UObject/LinkerLoad.h
@@ -8,7 +8,6 @@
 #include "UObject/SoftObjectPtr.h"
 #include "UObject/ObjectResource.h"
 #include "UObject/Linker.h"
-#include "Templates/RefCounting.h"
 
 class FLinkerPlaceholderBase;
 class ULinkerPlaceholderExportObject;
diff --git a/Engine/Source/Runtime/CoreUObject/Public/UObject/SoftObjectPath.h b/Engine/Source/Runtime/CoreUObject/Public/UObject/SoftObjectPath.h
--- a/Engine/Source/Runtime/CoreUObject/Public/UObject/SoftObjectPath.h
+++ b/Engine/Source/Runtime/CoreUObject/Public/UObject/SoftObjectPath.h
@@ -95,10 +95,10 @@ struct COREUOBJECT_API FSoftObjectPath
 
 	/**
 	 * Attempts to load the asset, this will call LoadObject which can be very slow
-	 *
+	 * @param InLoadContext Optional load context when called from nested load callstack
 	 * @return Loaded UObject, or nullptr if the reference is null or the asset fails to load
 	 */
-	UObject* TryLoad() const;
+	UObject* TryLoad(FUObjectSerializeContext* InLoadContext = nullptr) const;
 
 	/**
 	 * Attempts to find a currently loaded object that matches this path
diff --git a/Engine/Source/Runtime/CoreUObject/Public/UObject/UObjectThreadContext.h b/Engine/Source/Runtime/CoreUObject/Public/UObject/UObjectThreadContext.h
--- a/Engine/Source/Runtime/CoreUObject/Public/UObject/UObjectThreadContext.h
+++ b/Engine/Source/Runtime/CoreUObject/Public/UObject/UObjectThreadContext.h
@@ -139,6 +139,10 @@ public:
 	{
 		return ObjBeginLoadCount > 0;
 	}
+	int32 GetBeginLoadCount() const
+	{
+		return ObjBeginLoadCount;
+	}
 
 	int32 IncrementBeginLoadCount();
 	int32 DecrementBeginLoadCount();
diff --git a/Engine/Source/Runtime/Engine/Private/UnrealEngine.cpp b/Engine/Source/Runtime/Engine/Private/UnrealEngine.cpp
--- a/Engine/Source/Runtime/Engine/Private/UnrealEngine.cpp
+++ b/Engine/Source/Runtime/Engine/Private/UnrealEngine.cpp
@@ -8566,7 +8566,18 @@ FGuid UEngine::GetPackageGuid(FName PackageName, bool bForPIE)
 		LoadFlags |= LOAD_PackageForPIE;
 	}
 	UPackage* PackageToReset = nullptr;
-	FLinkerLoad* Linker = GetPackageLinker(nullptr, *PackageName.ToString(), LoadFlags, nullptr, nullptr, nullptr, LoadContext);
+	FLinkerLoad* Linker = nullptr;
+	{
+		FUObjectSerializeContext* InOutLoadContext = LoadContext;
+		Linker = GetPackageLinker(nullptr, *PackageName.ToString(), LoadFlags, nullptr, nullptr, nullptr, &InOutLoadContext);
+		if (InOutLoadContext != LoadContext)
+		{
+			// The linker already existed and was associated with another context
+			LoadContext->DecrementBeginLoadCount();
+			LoadContext = InOutLoadContext;
+			LoadContext->IncrementBeginLoadCount();
+		}
+	}
 	check(Linker);
 	if (Linker != nullptr && Linker->LinkerRoot != nullptr)
 	{
