Added new override for LoadPackageLinker that takes a lambda that can be used to process the linker before EndLoad is being called

diff --git a/Engine/Source/Editor/UnrealEd/Private/Commandlets/ContentCommandlets.cpp b/Engine/Source/Editor/UnrealEd/Private/Commandlets/ContentCommandlets.cpp
--- a/Engine/Source/Editor/UnrealEd/Private/Commandlets/ContentCommandlets.cpp
+++ b/Engine/Source/Editor/UnrealEd/Private/Commandlets/ContentCommandlets.cpp
@@ -572,7 +572,7 @@ void UResavePackagesCommandlet::LoadAndSaveOnePackage(const FString& Filename)
 		// Get the package linker.
 		VerboseMessage(TEXT("Pre GetPackageLinker"));
 
-		FLinkerLoad* Linker = LoadPackageLinker(NULL, *Filename, LOAD_NoVerify, nullptr, nullptr, nullptr);
+		FLinkerLoad* Linker = LoadPackageLinker(nullptr, *Filename, LOAD_NoVerify);
 	
 		// Bail early if we don't have a valid linker (package was out of date, etc)
 		if( !Linker )
@@ -2170,7 +2170,7 @@ int32 UWrangleContentCommandlet::Main( const FString& Params )
 					// load the package fully
 					UPackage* Package = LoadPackage(nullptr, *PackageFilename, LOAD_None);
 
-					FLinkerLoad* Linker = LoadPackageLinker(nullptr, *PackageFilename, LOAD_Quiet | LOAD_NoWarn | LOAD_NoVerify, nullptr, nullptr, nullptr);
+					FLinkerLoad* Linker = LoadPackageLinker(nullptr, *PackageFilename, LOAD_Quiet | LOAD_NoWarn | LOAD_NoVerify);
 
 					// look for special package types
 					bool bIsMap = Linker->ContainsMap();
@@ -2322,7 +2322,7 @@ int32 UWrangleContentCommandlet::Main( const FString& Params )
 				}
 			}
 
-			FLinkerLoad* Linker = LoadPackageLinker(nullptr, *PackageFilename, LOAD_Quiet | LOAD_NoWarn | LOAD_NoVerify, nullptr, nullptr, nullptr);
+			FLinkerLoad* Linker = LoadPackageLinker(nullptr, *PackageFilename, LOAD_Quiet | LOAD_NoWarn | LOAD_NoVerify);
 
 			// go through the exports in the package, looking for public objects
 			for (int32 ExportIndex = 0; ExportIndex < Linker->ExportMap.Num(); ExportIndex++)
diff --git a/Engine/Source/Editor/UnrealEd/Private/Commandlets/PackageUtilities.cpp b/Engine/Source/Editor/UnrealEd/Private/Commandlets/PackageUtilities.cpp
--- a/Engine/Source/Editor/UnrealEd/Private/Commandlets/PackageUtilities.cpp
+++ b/Engine/Source/Editor/UnrealEd/Private/Commandlets/PackageUtilities.cpp
@@ -781,7 +781,7 @@ int32 ULoadPackageCommandlet::Main( const FString& Params )
 
 		if (bCheckForLegacyPackages)
 		{
-			FLinkerLoad* Linker = LoadPackageLinker(nullptr, *Filename, LOAD_NoVerify, nullptr, nullptr, nullptr);
+			FLinkerLoad* Linker = LoadPackageLinker(nullptr, *Filename, LOAD_NoVerify);
 			MinVersion = FMath::Min<int32>(MinVersion, Linker->Summary.GetFileVersionUE4());
 		}
 		else
diff --git a/Engine/Source/Runtime/CoreUObject/Private/UObject/Linker.cpp b/Engine/Source/Runtime/CoreUObject/Private/UObject/Linker.cpp
--- a/Engine/Source/Runtime/CoreUObject/Private/UObject/Linker.cpp
+++ b/Engine/Source/Runtime/CoreUObject/Private/UObject/Linker.cpp
@@ -762,7 +762,7 @@ FLinkerLoad* GetPackageLinker
 	return Result;
 }
 
-FLinkerLoad* LoadPackageLinker(UPackage* InOuter, const TCHAR* InLongPackageName, uint32 LoadFlags, UPackageMap* Sandbox, FGuid* CompatibleGuid, FArchive* InReaderOverride)
+FLinkerLoad* LoadPackageLinker(UPackage* InOuter, const TCHAR* InLongPackageName, uint32 LoadFlags, UPackageMap* Sandbox, FGuid* CompatibleGuid, FArchive* InReaderOverride, TFunctionRef<void(FLinkerLoad* LoadedLinker)> LinkerLoadedCallback)
 {
 	FLinkerLoad* Linker = nullptr;
 	TRefCountPtr<FUObjectSerializeContext> LoadContext(new FUObjectSerializeContext());
@@ -778,10 +778,17 @@ FLinkerLoad* LoadPackageLinker(UPackage* InOuter, const TCHAR* InLongPackageName
 			LoadContext->IncrementBeginLoadCount();
 		}
 	}
+	// Allow external code to work with the linker before EndLoad()
+	LinkerLoadedCallback(Linker);
 	EndLoad(Linker ? Linker->GetSerializeContext() : LoadContext.GetReference());
 	return Linker;
 }
 
+FLinkerLoad* LoadPackageLinker(UPackage* InOuter, const TCHAR* InLongPackageName, uint32 LoadFlags, UPackageMap* Sandbox, FGuid* CompatibleGuid, FArchive* InReaderOverride)
+{
+	return LoadPackageLinker(InOuter, InLongPackageName, LoadFlags, Sandbox, CompatibleGuid, InReaderOverride, [](FLinkerLoad* InLinker) {});
+}
+
 /**
  * 
  * Ensure thumbnails are loaded and then reset the loader in preparation for a package save
diff --git a/Engine/Source/Runtime/CoreUObject/Public/UObject/Linker.h b/Engine/Source/Runtime/CoreUObject/Public/UObject/Linker.h
--- a/Engine/Source/Runtime/CoreUObject/Public/UObject/Linker.h
+++ b/Engine/Source/Runtime/CoreUObject/Public/UObject/Linker.h
@@ -573,9 +573,11 @@ COREUOBJECT_API void DeleteLoader(FLinkerLoad* Loader);
 	* @param Sandbox Additional sandbox for loading
 	* @param CompatibleGuid Net GUID
 	* @param InReaderOverride Optional archive to use for reading package data
+	* @param LinkerLoadedCallback Callback when the linker is loaded (or not found)
 	* @return Pointer to the loaded linker or null if the file didn't exist
 	*/
-COREUOBJECT_API FLinkerLoad* LoadPackageLinker(UPackage* InOuter, const TCHAR* InLongPackageName, uint32 LoadFlags, UPackageMap* Sandbox, FGuid* CompatibleGuid, FArchive* InReaderOverride = nullptr);
+COREUOBJECT_API FLinkerLoad* LoadPackageLinker(UPackage* InOuter, const TCHAR* InLongPackageName, uint32 LoadFlags, UPackageMap* Sandbox, FGuid* CompatibleGuid, FArchive* InReaderOverride, TFunctionRef<void(FLinkerLoad* LoadedLinker)> LinkerLoadedCallback);
+COREUOBJECT_API FLinkerLoad* LoadPackageLinker(UPackage* InOuter, const TCHAR* InLongPackageName, uint32 LoadFlags = LOAD_None, UPackageMap* Sandbox = nullptr, FGuid* CompatibleGuid = nullptr, FArchive* InReaderOverride = nullptr);
 
 /** 
   * Gets a linker for a package and returns it without loading any objects. This call must be preceeded by BeginLoad and followed by EndLoad calls
@@ -590,6 +592,8 @@ COREUOBJECT_API FLinkerLoad* LoadPackageLinker(UPackage* InOuter, const TCHAR* I
   */
 COREUOBJECT_API FLinkerLoad* GetPackageLinker(UPackage* InOuter, const TCHAR* InLongPackageName, uint32 LoadFlags, UPackageMap* Sandbox, FGuid* CompatibleGuid, FArchive* InReaderOverride = nullptr, FUObjectSerializeContext** InOutLoadContext = nullptr);
 
+
+
 COREUOBJECT_API FString GetPrestreamPackageLinkerName(const TCHAR* InLongPackageName, bool bExistSkip = true);
 
 
diff --git a/Engine/Source/Runtime/Engine/Private/UnrealEngine.cpp b/Engine/Source/Runtime/Engine/Private/UnrealEngine.cpp
--- a/Engine/Source/Runtime/Engine/Private/UnrealEngine.cpp
+++ b/Engine/Source/Runtime/Engine/Private/UnrealEngine.cpp
@@ -8566,25 +8566,15 @@ FGuid UEngine::GetPackageGuid(FName PackageName, bool bForPIE)
 		LoadFlags |= LOAD_PackageForPIE;
 	}
 	UPackage* PackageToReset = nullptr;
-	FLinkerLoad* Linker = nullptr;
+	FLinkerLoad* Linker = LoadPackageLinker(nullptr, *PackageName.ToString(), LoadFlags, nullptr, nullptr, nullptr, [&PackageToReset, &Result](FLinkerLoad* InLinker)
 	{
-		FUObjectSerializeContext* InOutLoadContext = LoadContext;
-		Linker = GetPackageLinker(nullptr, *PackageName.ToString(), LoadFlags, nullptr, nullptr, nullptr, &InOutLoadContext);
-		if (InOutLoadContext != LoadContext)
+		check(InLinker);
+		if (InLinker != nullptr && InLinker->LinkerRoot != nullptr)
 		{
-			// The linker already existed and was associated with another context
-			LoadContext->DecrementBeginLoadCount();
-			LoadContext = InOutLoadContext;
-			LoadContext->IncrementBeginLoadCount();
+			Result = InLinker->LinkerRoot->GetGuid();
+			PackageToReset = InLinker->LinkerRoot;
 		}
-	}
-	check(Linker);
-	if (Linker != nullptr && Linker->LinkerRoot != nullptr)
-	{
-		Result = Linker->LinkerRoot->GetGuid();
-		PackageToReset = Linker->LinkerRoot;
-	}
-	EndLoad(Linker->GetSerializeContext());
+	});
 
 	ResetLoaders(PackageToReset);
 	Linker = nullptr;
