Potential thread-safety fix in UClass::ClearFunctionMapsCaches()

I observed the function cache being stomped twice while debugging but I lost the repro. It could have been caused by a VS 15.8 compiler bug I suffered from atm or some other stomp that has been fixed, but seems prudent to acquire the lock before clearing.

diff --git a/Engine/Source/Runtime/CoreUObject/Private/UObject/Class.cpp b/Engine/Source/Runtime/CoreUObject/Private/UObject/Class.cpp
--- a/Engine/Source/Runtime/CoreUObject/Private/UObject/Class.cpp
+++ b/Engine/Source/Runtime/CoreUObject/Private/UObject/Class.cpp
@@ -4464,6 +4464,12 @@ void UClass::CreateLinkAndAddChildFunctionsToMap(const FClassFunctionLinkInfo* F
 	}
 }
 
+void UClass::ClearFunctionMapsCaches()
+{
+	FRWScopeLock ScopeLock(SuperFuncMapLock, FRWScopeLockType::SLT_Write);
+	SuperFuncMap.Empty();
+}
+
 UFunction* UClass::FindFunctionByName(FName InName, EIncludeSuperFlag::Type IncludeSuper) const
 {
 	UFunction* Result = FuncMap.FindRef(InName);
diff --git a/Engine/Source/Runtime/CoreUObject/Public/UObject/Class.h b/Engine/Source/Runtime/CoreUObject/Public/UObject/Class.h
--- a/Engine/Source/Runtime/CoreUObject/Public/UObject/Class.h
+++ b/Engine/Source/Runtime/CoreUObject/Public/UObject/Class.h
@@ -2427,10 +2427,7 @@ public:
 	}
 
 	/** Clears the function name caches, in case things have changed */
-	void ClearFunctionMapsCaches()
-	{
-		SuperFuncMap.Empty();
-	}
+	void ClearFunctionMapsCaches();
 
 	/** Looks for a given function name */
 	UFunction* FindFunctionByName(FName InName, EIncludeSuperFlag::Type IncludeSuper = EIncludeSuperFlag::IncludeSuper) const;
