Fixing soft object paths not resolving the package name properly in some cases when loading

diff --git a/Engine/Source/Runtime/CoreUObject/Private/UObject/PropertySoftObjectPtr.cpp b/Engine/Source/Runtime/CoreUObject/Private/UObject/PropertySoftObjectPtr.cpp
--- a/Engine/Source/Runtime/CoreUObject/Private/UObject/PropertySoftObjectPtr.cpp
+++ b/Engine/Source/Runtime/CoreUObject/Private/UObject/PropertySoftObjectPtr.cpp
@@ -5,6 +5,7 @@
 #include "UObject/SoftObjectPtr.h"
 #include "UObject/PropertyPortFlags.h"
 #include "UObject/UnrealType.h"
+#include "UObject/LinkerLoad.h"
 
 /*-----------------------------------------------------------------------------
 	USoftObjectProperty.
@@ -101,7 +102,7 @@ const TCHAR* USoftObjectProperty::ImportText_Internal( const TCHAR* InBuffer, vo
 
 	FSoftObjectPath SoftObjectPath;
 
-	if (SoftObjectPath.ImportTextItem(InBuffer, PortFlags, Parent, ErrorText))
+	if (SoftObjectPath.ImportTextItem(InBuffer, PortFlags, Parent, ErrorText, GetLinker()))
 	{
 		SoftObjectPtr = SoftObjectPath;
 		return InBuffer;
diff --git a/Engine/Source/Runtime/CoreUObject/Private/UObject/SoftObjectPath.cpp b/Engine/Source/Runtime/CoreUObject/Private/UObject/SoftObjectPath.cpp
--- a/Engine/Source/Runtime/CoreUObject/Private/UObject/SoftObjectPath.cpp
+++ b/Engine/Source/Runtime/CoreUObject/Private/UObject/SoftObjectPath.cpp
@@ -232,7 +232,7 @@ bool FSoftObjectPath::ExportTextItem(FString& ValueStr, FSoftObjectPath const& D
 	return true;
 }
 
-bool FSoftObjectPath::ImportTextItem(const TCHAR*& Buffer, int32 PortFlags, UObject* Parent, FOutputDevice* ErrorText)
+bool FSoftObjectPath::ImportTextItem(const TCHAR*& Buffer, int32 PortFlags, UObject* Parent, FOutputDevice* ErrorText, FArchive* InSerializingArchive)
 {
 	FString ImportedPath;
 	const TCHAR* NewBuffer = UPropertyHelpers::ReadToken(Buffer, ImportedPath, 1);
@@ -274,13 +274,13 @@ bool FSoftObjectPath::ImportTextItem(const TCHAR*& Buffer, int32 PortFlags, UObj
 		// We're probably reading config for an editor only object, we need to mark this reference as editor only
 		FSoftObjectPathSerializationScope SerializationScope(NAME_None, NAME_None, ESoftObjectPathCollectType::EditorOnlyCollect, ESoftObjectPathSerializeType::AlwaysSerialize);
 
-		PostLoadPath(nullptr);
+		PostLoadPath(InSerializingArchive);
 	}
 	else
 #endif
 	{
 		// Consider this a load, so Config string references get cooked
-		PostLoadPath(nullptr);
+		PostLoadPath(InSerializingArchive);
 	}
 
 	return true;
diff --git a/Engine/Source/Runtime/CoreUObject/Public/UObject/SoftObjectPath.h b/Engine/Source/Runtime/CoreUObject/Public/UObject/SoftObjectPath.h
--- a/Engine/Source/Runtime/CoreUObject/Public/UObject/SoftObjectPath.h
+++ b/Engine/Source/Runtime/CoreUObject/Public/UObject/SoftObjectPath.h
@@ -148,7 +148,7 @@ struct COREUOBJECT_API FSoftObjectPath
 	}
 	FSoftObjectPath& operator=(FSoftObjectPath Other);
 	bool ExportTextItem(FString& ValueStr, FSoftObjectPath const& DefaultValue, UObject* Parent, int32 PortFlags, UObject* ExportRootScope) const;
-	bool ImportTextItem( const TCHAR*& Buffer, int32 PortFlags, UObject* Parent, FOutputDevice* ErrorText );
+	bool ImportTextItem( const TCHAR*& Buffer, int32 PortFlags, UObject* Parent, FOutputDevice* ErrorText, FArchive* InSerializingArchive = nullptr );
 	bool SerializeFromMismatchedTag(struct FPropertyTag const& Tag, FStructuredArchive::FSlot Slot);
 
 	/** Serializes the internal path and also handles save/PIE fixups. Call this from the archiver overrides */
