Adding code that dumps the callstacks that added a new custom version after the package's that's being saved summary has already been serialized

diff --git a/Engine/Source/Runtime/Core/Public/Serialization/Archive.h b/Engine/Source/Runtime/Core/Public/Serialization/Archive.h
--- a/Engine/Source/Runtime/Core/Public/Serialization/Archive.h
+++ b/Engine/Source/Runtime/Core/Public/Serialization/Archive.h
@@ -879,7 +879,7 @@ public:
 	 *
 	 * @param Guid The guid of the custom version.  This must have previously been registered with FCustomVersionRegistration.
 	 */
-	void UsingCustomVersion(const struct FGuid& Guid);
+	virtual void UsingCustomVersion(const struct FGuid& Guid);
 
 	/**
 	 * Queries a custom version from the archive.  If the archive is being used to write, the custom version must have already been registered.
diff --git a/Engine/Source/Runtime/CoreUObject/Private/UObject/LinkerSave.cpp b/Engine/Source/Runtime/CoreUObject/Private/UObject/LinkerSave.cpp
--- a/Engine/Source/Runtime/CoreUObject/Private/UObject/LinkerSave.cpp
+++ b/Engine/Source/Runtime/CoreUObject/Private/UObject/LinkerSave.cpp
@@ -11,6 +11,7 @@
 #include "Internationalization/TextPackageNamespaceUtil.h"
 #include "UObject/UObjectThreadContext.h"
 #include "UObject/UnrealType.h"
+#include "HAL/PlatformStackWalk.h"
 
 /*----------------------------------------------------------------------------
 	FLinkerSave.
@@ -304,7 +305,6 @@ FArchive& FLinkerSave::operator<<(FLazyObjectPtr& LazyObjectPtr)
 	ID = LazyObjectPtr.GetUniqueID();
 	return *this << ID;
 }
-
 void FLinkerSave::SetSerializeContext(FUObjectSerializeContext* InLoadContext)
 {
 	SaveContext = InLoadContext;
@@ -317,4 +317,43 @@ void FLinkerSave::SetSerializeContext(FUObjectSerializeContext* InLoadContext)
 FUObjectSerializeContext* FLinkerSave::GetSerializeContext()
 {
 	return SaveContext;
-}
\ No newline at end of file
+}
+
+void FLinkerSave::UsingCustomVersion(const struct FGuid& Guid)
+{
+	FArchiveUObject::UsingCustomVersion(Guid);
+
+	// Here we're going to try and dump the callstack that added a new custom version after package summary has been serialized
+	if (Summary.GetCustomVersionContainer().GetVersion(Guid) == nullptr)
+	{
+		const FCustomVersion* RegisteredVersion = FCustomVersionContainer::GetRegistered().GetVersion(Guid);
+		check(RegisteredVersion);
+
+		FString CustomVersionWarning = FString::Printf(TEXT("Unexpected custom version \"%s\" used after package %s summary has been serialized. Callstack:\n"),
+			*RegisteredVersion->GetFriendlyName().ToString(), *LinkerRoot->GetName());
+
+		const int32 MaxStackFrames = 100;
+		uint64 StackFrames[MaxStackFrames];
+		int32 NumStackFrames = FPlatformStackWalk::CaptureStackBackTrace(StackFrames, MaxStackFrames);
+
+		// Convert the stack trace to text, ignore the first functions (ProgramCounterToHumanReadableString)
+		const int32 IgnoreStackLinesCount = 1;		
+		ANSICHAR Buffer[1024];
+		const ANSICHAR* CutoffFunction = "UPackage::Save";
+		for (int32 Idx = IgnoreStackLinesCount; Idx < NumStackFrames; Idx++)
+		{			
+			Buffer[0] = '\0';
+			FPlatformStackWalk::ProgramCounterToHumanReadableString(Idx, StackFrames[Idx], Buffer, sizeof(Buffer));
+			CustomVersionWarning += TEXT("\t");
+			CustomVersionWarning += Buffer;
+			CustomVersionWarning += "\n";
+			if (FCStringAnsi::Strstr(Buffer, CutoffFunction))
+			{
+				// Anything below UPackage::Save is not interesting from the point of view of what we're trying to find
+				break;
+			}
+		}
+
+		UE_LOG(LogLinker, Warning, TEXT("%s"), *CustomVersionWarning);
+	}
+}
diff --git a/Engine/Source/Runtime/CoreUObject/Public/UObject/LinkerSave.h b/Engine/Source/Runtime/CoreUObject/Public/UObject/LinkerSave.h
--- a/Engine/Source/Runtime/CoreUObject/Public/UObject/LinkerSave.h
+++ b/Engine/Source/Runtime/CoreUObject/Public/UObject/LinkerSave.h
@@ -87,7 +87,7 @@ public:
 	FArchive& operator<<( FLazyObjectPtr& LazyObjectPtr );
 	virtual void SetSerializeContext(FUObjectSerializeContext* InLoadContext) override;
 	FUObjectSerializeContext* GetSerializeContext() override;
-
+	virtual void UsingCustomVersion(const struct FGuid& Guid) override;
 
 #if WITH_EDITOR
 	// proxy for debugdata
