Fixing potential crash when verifying imports when async loading and resolving deferred exports

diff --git a/Engine/Source/Runtime/CoreUObject/Private/UObject/Linker.cpp b/Engine/Source/Runtime/CoreUObject/Private/UObject/Linker.cpp
--- a/Engine/Source/Runtime/CoreUObject/Private/UObject/Linker.cpp
+++ b/Engine/Source/Runtime/CoreUObject/Private/UObject/Linker.cpp
@@ -725,9 +725,10 @@ FLinkerLoad* GetPackageLinker
 	}
 	else if (InExistingContext)
 	{
-		if (Result->GetSerializeContext() && Result->GetSerializeContext()->HasStartedLoading() && InExistingContext->GetBeginLoadCount() == 1)
+		if ((Result->GetSerializeContext() && Result->GetSerializeContext()->HasStartedLoading() && InExistingContext->GetBeginLoadCount() == 1) ||
+			  IsInAsyncLoadingThread())
 		{
-			// Use the context associated with the linker because it has already started loading objects
+			// Use the context associated with the linker because it has already started loading objects (or we're in ALT where each package needs its own context)
 			*InOutLoadContext = Result->GetSerializeContext();
 		}
 		else
