Fix a potential crash when loading thumbnail data

diff --git a/Engine/Source/Runtime/CoreUObject/Private/UObject/LinkerLoad.cpp b/Engine/Source/Runtime/CoreUObject/Private/UObject/LinkerLoad.cpp
--- a/Engine/Source/Runtime/CoreUObject/Private/UObject/LinkerLoad.cpp
+++ b/Engine/Source/Runtime/CoreUObject/Private/UObject/LinkerLoad.cpp
@@ -1875,7 +1875,7 @@ FLinkerLoad::ELinkerStatus FLinkerLoad::SerializeThumbnails( bool bForceEnableIn
 			}
 
 
-			FStructuredArchive::FArray DataArray = Record.EnterArray(FIELD_NAME_TEXT("Thumbnails"), ThumbnailCount);
+			FStructuredArchive::FStream DataStream = Record.EnterStream(FIELD_NAME_TEXT("Thumbnails"));
 
 			// Now go and load and cache all of the thumbnails
 			for (int32 CurObjectIndex = 0; CurObjectIndex < ThumbnailInfoArray.Num(); ++CurObjectIndex)
@@ -1890,7 +1890,7 @@ FLinkerLoad::ELinkerStatus FLinkerLoad::SerializeThumbnails( bool bForceEnableIn
 
 				// Load the image data
 				FObjectThumbnail LoadedThumbnail;
-				LoadedThumbnail.Serialize(DataArray.EnterElement());
+				LoadedThumbnail.Serialize(DataStream.EnterElement());
 
 				// Store the data!
 				LinkerRoot->ThumbnailMap->Add(CurThumbnailInfo.ObjectFullName, LoadedThumbnail);
