Fix a crash when saving a package with duplicate asset registry tag names

diff --git a/Engine/Source/Runtime/CoreUObject/Private/UObject/SavePackage.cpp b/Engine/Source/Runtime/CoreUObject/Private/UObject/SavePackage.cpp
--- a/Engine/Source/Runtime/CoreUObject/Private/UObject/SavePackage.cpp
+++ b/Engine/Source/Runtime/CoreUObject/Private/UObject/SavePackage.cpp
@@ -6302,10 +6302,10 @@ void UPackage::SaveAssetRegistryData(UPackage* InOuter, FLinkerSave* Linker, FSt
 		TArray<FAssetRegistryTag> Tags;
 		for (FAssetRegistryTag& SourceTag : SourceTags)
 		{
-			const FAssetRegistryTag* Existing = Tags.FindByPredicate([SourceTag](const FAssetRegistryTag& InTag) { return InTag.Name == SourceTag.Name; });
+			FAssetRegistryTag* Existing = Tags.FindByPredicate([SourceTag](const FAssetRegistryTag& InTag) { return InTag.Name == SourceTag.Name; });
 			if (Existing)
 			{
-				check(Existing->Value == SourceTag.Value);
+				Existing->Value = SourceTag.Value;
 			}
 			else
 			{
