Allow interface properties to be imported as object properties.

diff --git a/Engine/Source/Runtime/CoreUObject/Private/UObject/PropertyObject.cpp b/Engine/Source/Runtime/CoreUObject/Private/UObject/PropertyObject.cpp
--- a/Engine/Source/Runtime/CoreUObject/Private/UObject/PropertyObject.cpp
+++ b/Engine/Source/Runtime/CoreUObject/Private/UObject/PropertyObject.cpp
@@ -65,6 +65,21 @@ EConvertFromTypeResult UObjectProperty::ConvertFromType(const FPropertyTag& Tag,
 
 		return EConvertFromTypeResult::Converted;
 	}
+	else if (Tag.Type == NAME_InterfaceProperty)
+	{
+		UObject* ObjectValue;
+		Slot << ObjectValue;
+
+		if (ObjectValue && !ObjectValue->IsA(PropertyClass))
+		{
+			UE_LOG(LogClass, Warning, TEXT("Failed to convert interface property %s of %s from Interface to %s"), *this->GetName(), *Slot.GetUnderlyingArchive().GetArchiveName(), *PropertyClass->GetName());
+			return EConvertFromTypeResult::CannotConvert;
+		}
+
+		SetPropertyValue_InContainer(Data, ObjectValue, Tag.ArrayIndex);
+		CheckValidObject(GetPropertyValuePtr_InContainer(Data, Tag.ArrayIndex));
+		return EConvertFromTypeResult::Converted;
+	}
 
 	return EConvertFromTypeResult::UseSerializeItem;
 }
