Fixing cases when duting heavy recursive loading of assets some of the dependencies that started loading early on haven't finished loading by the time they are required

diff --git a/Engine/Source/Runtime/CoreUObject/Private/Serialization/AsyncLoading.cpp b/Engine/Source/Runtime/CoreUObject/Private/Serialization/AsyncLoading.cpp
--- a/Engine/Source/Runtime/CoreUObject/Private/Serialization/AsyncLoading.cpp
+++ b/Engine/Source/Runtime/CoreUObject/Private/Serialization/AsyncLoading.cpp
@@ -6467,7 +6467,7 @@ EAsyncPackageState::Type FAsyncPackage::PreLoadObjects()
 	// GC can't run in here
 	FGCScopeGuard GCGuard;
 
-	TArray<UObject*>& ThreadObjLoaded = LoadContext->PRIVATE_GetObjectsLoadedForFAsyncPackage();
+	TArray<UObject*>& ThreadObjLoaded = LoadContext->PRIVATE_GetObjectsLoadedInternalUseOnly();
 	PackageObjLoaded.Append(ThreadObjLoaded);
 	ThreadObjLoaded.Reset();
 
@@ -6566,7 +6566,7 @@ EAsyncPackageState::Type FAsyncPackage::PostLoadObjects()
 	FUObjectThreadContext& ThreadContext = FUObjectThreadContext::Get();
 	TGuardValue<bool> GuardIsRoutingPostLoad(ThreadContext.IsRoutingPostLoad, true);
 
-	TArray<UObject*>& ThreadObjLoaded = LoadContext->PRIVATE_GetObjectsLoadedForFAsyncPackage();
+	TArray<UObject*>& ThreadObjLoaded = LoadContext->PRIVATE_GetObjectsLoadedInternalUseOnly();
 	if (ThreadObjLoaded.Num())
 	{
 		// New objects have been loaded. They need to go through PreLoad first so exit now and come back after they've been preloaded.
@@ -6646,7 +6646,7 @@ EAsyncPackageState::Type FAsyncPackage::PostLoadDeferredObjects(double InTickSta
 	TGuardValue<bool> GuardIsRoutingPostLoad(PackageScope.ThreadContext.IsRoutingPostLoad, true);
 	FAsyncLoadingTickScope InAsyncLoadingTick;
 
-	TArray<UObject*>& ObjLoadedInPostLoad = LoadContext->PRIVATE_GetObjectsLoadedForFAsyncPackage();
+	TArray<UObject*>& ObjLoadedInPostLoad = LoadContext->PRIVATE_GetObjectsLoadedInternalUseOnly();
 	TArray<UObject*> ObjLoadedInPostLoadLocal;
 
 	STAT(double PostLoadStartTime = FPlatformTime::Seconds());
@@ -6836,7 +6836,7 @@ EAsyncPackageState::Type FAsyncPackage::FinishObjects()
 	LastTypeOfWorkPerformed			= TEXT("finishing all objects");
 
 	check(!Linker || LoadContext == Linker->GetSerializeContext());		
-	TArray<UObject*>& ThreadObjLoaded = LoadContext->PRIVATE_GetObjectsLoadedForFAsyncPackage();
+	TArray<UObject*>& ThreadObjLoaded = LoadContext->PRIVATE_GetObjectsLoadedInternalUseOnly();
 
 	EAsyncLoadingResult::Type LoadingResult;
 	if (!bLoadHasFailed)
diff --git a/Engine/Source/Runtime/CoreUObject/Private/UObject/Linker.cpp b/Engine/Source/Runtime/CoreUObject/Private/UObject/Linker.cpp
--- a/Engine/Source/Runtime/CoreUObject/Private/UObject/Linker.cpp
+++ b/Engine/Source/Runtime/CoreUObject/Private/UObject/Linker.cpp
@@ -733,6 +733,13 @@ FLinkerLoad* GetPackageLinker
 		}
 		else
 		{
+			if (Result->GetSerializeContext() && Result->GetSerializeContext() != InExistingContext)
+			{
+				// Make sure the objects already loaded with the context associated with the existing linker
+				// are copied to the context provided for this function call to make sure they all get loaded ASAP
+				check(!IsInAsyncLoadingThread());
+				InExistingContext->AddUniqueLoadedObjects(Result->GetSerializeContext()->PRIVATE_GetObjectsLoadedInternalUseOnly());
+			}
 			// Replace the linker context with the one passed into this function
 			Result->SetSerializeContext(InExistingContext);
 		}
diff --git a/Engine/Source/Runtime/CoreUObject/Private/UObject/UObjectThreadContext.cpp b/Engine/Source/Runtime/CoreUObject/Private/UObject/UObjectThreadContext.cpp
--- a/Engine/Source/Runtime/CoreUObject/Private/UObject/UObjectThreadContext.cpp
+++ b/Engine/Source/Runtime/CoreUObject/Private/UObject/UObjectThreadContext.cpp
@@ -44,4 +44,13 @@ int32 FUObjectSerializeContext::DecrementBeginLoadCount()
 {
 	check(HasStartedLoading());
 	return --ObjBeginLoadCount;
+}
+
+void FUObjectSerializeContext::AddUniqueLoadedObjects(const TArray<UObject*>& InObjects)
+{
+	for (UObject* NewLoadedObject : InObjects)
+	{
+		ObjectsLoaded.AddUnique(NewLoadedObject);
+	}
+	
 }
\ No newline at end of file
diff --git a/Engine/Source/Runtime/CoreUObject/Public/UObject/UObjectThreadContext.h b/Engine/Source/Runtime/CoreUObject/Public/UObject/UObjectThreadContext.h
--- a/Engine/Source/Runtime/CoreUObject/Public/UObject/UObjectThreadContext.h
+++ b/Engine/Source/Runtime/CoreUObject/Public/UObject/UObjectThreadContext.h
@@ -133,6 +133,7 @@ public:
 	{
 		ObjectsLoaded.Add(InObject);
 	}
+	void AddUniqueLoadedObjects(const TArray<UObject*>& InObjects);
 
 	/** Checks if object loading has started */
 	bool HasStartedLoading() const
@@ -190,7 +191,7 @@ public:
 	}
 
 	/** This is only meant to be used by FAsyncPackage for performance reasons. The ObjectsLoaded array should not be manipulated directly! */
-	TArray<UObject*>& PRIVATE_GetObjectsLoadedForFAsyncPackage()
+	TArray<UObject*>& PRIVATE_GetObjectsLoadedInternalUseOnly()
 	{
 		return ObjectsLoaded;
 	}
