Fixing an issue with objects that are loaded from constructors during async loading not being serialized properly

diff --git a/Engine/Source/Runtime/CoreUObject/Private/UObject/UObjectGlobals.cpp b/Engine/Source/Runtime/CoreUObject/Private/UObject/UObjectGlobals.cpp
--- a/Engine/Source/Runtime/CoreUObject/Private/UObject/UObjectGlobals.cpp
+++ b/Engine/Source/Runtime/CoreUObject/Private/UObject/UObjectGlobals.cpp
@@ -1401,6 +1401,14 @@ UPackage* LoadPackageInternal(UPackage* InOuter, const TCHAR* InLongPackageNameO
 			// that triggered this LoadPackage call
 			FSerializedPropertyScope SerializedProperty(*Linker, ImportLinker ? ImportLinker->GetSerializedProperty() : Linker->GetSerializedProperty());
 			Linker->LoadAllObjects(GEventDrivenLoaderEnabled);
+
+			// @todo: remove me when loading can be self-contained (and EndLoad doesn't check for IsInAsyncLoadingThread) or there's just one loading path
+			// If we start a non-async loading during async loading and the serialization context is not associated with any other package and
+			// doesn't come from an async package, queue this package to be async loaded, otherwise we'll end up not loading its exports
+			if (!Linker->AsyncRoot && Linker->GetSerializeContext()->GetBeginLoadCount() == 1 && IsInAsyncLoadingThread())
+			{
+				LoadPackageAsync(Linker->LinkerRoot->GetName());
+			}
 		}
 		else
 		{
