Fix to an assert which fails for classes with no properties but still differs in size from its superclass (e.g. because it has non-properties).

diff --git a/Engine/Source/Runtime/CoreUObject/Private/UObject/Obj.cpp b/Engine/Source/Runtime/CoreUObject/Private/UObject/Obj.cpp
--- a/Engine/Source/Runtime/CoreUObject/Private/UObject/Obj.cpp
+++ b/Engine/Source/Runtime/CoreUObject/Private/UObject/Obj.cpp
@@ -1941,12 +1941,42 @@ void UObject::LoadConfig( UClass* ConfigClass/*=NULL*/, const TCHAR* InFilename/
 	}
 
 #if !IS_PROGRAM
+	auto HaveSameProperties = [](const UStruct* Struct1, const UStruct* Struct2) -> bool
+	{
+		TFieldIterator<UProperty> It1(Struct1);
+		TFieldIterator<UProperty> It2(Struct2);
+
+		for (;;)
+		{
+			bool bAtEnd1 = !It1;
+			bool bAtEnd2 = !It2;
+
+			// If one iterator is at the end and one isn't, the property lists are different
+			if (bAtEnd1 != bAtEnd2)
+			{
+				return false;
+			}
+
+			// If both iterators have reached the end, the property lists are the same
+			if (bAtEnd1)
+			{
+				return true;
+			}
+
+			// If the properties are different, the property lists are different
+			if (*It1 != *It2)
+			{
+				return false;
+			}
+		}
+	};
+
 	// Do we have properties that don't exist yet?
 	// If this happens then we're trying to load the config for an object that doesn't
 	// know what its layout is. Usually a call to GetDefaultObject that occurs too early
 	// because ProcessNewlyLoadedUObjects hasn't happened yet
 	checkf(ConfigClass->PropertyLink != nullptr
-		|| (ConfigClass->GetSuperStruct() && ConfigClass->PropertiesSize == ConfigClass->GetSuperStruct()->PropertiesSize)
+		|| (ConfigClass->GetSuperStruct() && HaveSameProperties(ConfigClass, ConfigClass->GetSuperStruct()))
 		|| ConfigClass->PropertiesSize == 0
 		|| GIsRequestingExit, // Ignore this check when exiting as we may have requested exit during init when not everything is initialized
 		TEXT("class %s has uninitialized properties. Accessed too early?"), *ConfigClass->GetName());
