Fixing a potential crash when saving packages and fixing up Soft Object Paths

diff --git a/Engine/Source/Runtime/CoreUObject/Private/UObject/SoftObjectPath.cpp b/Engine/Source/Runtime/CoreUObject/Private/UObject/SoftObjectPath.cpp
--- a/Engine/Source/Runtime/CoreUObject/Private/UObject/SoftObjectPath.cpp
+++ b/Engine/Source/Runtime/CoreUObject/Private/UObject/SoftObjectPath.cpp
@@ -483,7 +483,7 @@ bool FSoftObjectPath::FixupCoreRedirects()
 	if (OldName != NewName)
 	{
 		// Only do the fixup if the old object isn't in memory, this avoids false positives
-		UObject* FoundOldObject = FindObject<UObject>(nullptr, *OldString);
+		UObject* FoundOldObject = FindObjectSafe<UObject>(nullptr, *OldString);
 
 		if (!FoundOldObject)
 		{
