Fixed UClass::HasProperty() returning false positives for properties that belong to a child class of the one being tested

diff --git a/Engine/Source/Runtime/CoreUObject/Private/UObject/Class.cpp b/Engine/Source/Runtime/CoreUObject/Private/UObject/Class.cpp
--- a/Engine/Source/Runtime/CoreUObject/Private/UObject/Class.cpp
+++ b/Engine/Source/Runtime/CoreUObject/Private/UObject/Class.cpp
@@ -3938,7 +3938,7 @@ bool UClass::HasProperty(UProperty* InProperty) const
 {
 	if ( UClass* PropertiesClass = dynamic_cast<UClass*>(InProperty->GetOuter()) )
 	{
-		return PropertiesClass->FindNearestCommonBaseClass(this) != nullptr;
+		return IsChildOf(PropertiesClass);
 	}
 
 	return false;
