Adding a missing mutex lock to RemoveLoaderFromObjectLoadersAndLoadersWithNewImports

diff --git a/Engine/Source/Runtime/CoreUObject/Private/UObject/LinkerManager.h b/Engine/Source/Runtime/CoreUObject/Private/UObject/LinkerManager.h
--- a/Engine/Source/Runtime/CoreUObject/Private/UObject/LinkerManager.h
+++ b/Engine/Source/Runtime/CoreUObject/Private/UObject/LinkerManager.h
@@ -44,11 +44,18 @@ public:
 
 	FORCEINLINE void RemoveLoaderFromObjectLoadersAndLoadersWithNewImports(FLinkerLoad* LinkerLoad)
 	{
+		{
 #if THREADSAFE_UOBJECTS
-		FScopeLock ObjectLoadersLock(&ObjectLoadersCritical);
+			FScopeLock ObjectLoadersLock(&ObjectLoadersCritical);
+#endif
+			ObjectLoaders.Remove(LinkerLoad);
+		}
+		{
+#if THREADSAFE_UOBJECTS
+			FScopeLock ObjectLoadersLock(&LoadersWithNewImportsCritical);
 #endif
-		ObjectLoaders.Remove(LinkerLoad);
-		LoadersWithNewImports.Remove(LinkerLoad);
+			LoadersWithNewImports.Remove(LinkerLoad);
+		}
 	}
 
 	FORCEINLINE void GetLoadersWithNewImportsAndEmpty(TSet<FLinkerLoad*>& OutLoaders)
