Removing the functionality to merge GC clusters as it was never used and generally could lead to keeping all assets in memory forever

diff --git a/Engine/Config/BaseEngine.ini b/Engine/Config/BaseEngine.ini
--- a/Engine/Config/BaseEngine.ini
+++ b/Engine/Config/BaseEngine.ini
@@ -774,7 +774,6 @@ gc.TimeBetweenPurgingPendingKillObjects=61.1
 gc.MaxObjectsInEditor=16777216
 gc.IncrementalBeginDestroyEnabled=True
 gc.CreateGCClusters=True
-gc.MergeGCClusters=False
 gc.MinGCClusterSize=5
 gc.ActorClusteringEnabled=False
 gc.BlueprintClusteringEnabled=False
diff --git a/Engine/Source/Runtime/CoreUObject/Private/UObject/UObjectClusters.cpp b/Engine/Source/Runtime/CoreUObject/Private/UObject/UObjectClusters.cpp
--- a/Engine/Source/Runtime/CoreUObject/Private/UObject/UObjectClusters.cpp
+++ b/Engine/Source/Runtime/CoreUObject/Private/UObject/UObjectClusters.cpp
@@ -29,14 +29,6 @@ static FAutoConsoleVariableRef CCreateGCClusters(
 	ECVF_Default
 	);
 
-int32 GMergeGCClusters = 0;
-static FAutoConsoleVariableRef CMergeGCClusters(
-	TEXT("gc.MergeGCClusters"),
-	GMergeGCClusters,
-	TEXT("If true, when creating clusters, the clusters referenced from another cluster will get merged into one big cluster."),
-	ECVF_Default
-	);
-
 int32 GMinGCClusterSize = 5;
 static FAutoConsoleVariableRef CMinGCClusterSize(
 	TEXT("gc.MinGCClusterSize"),
@@ -653,48 +645,6 @@ public:
 		}
 	}
 
-	/**
-	* Merges an existing cluster with the currently constructed one
-	*
-	* @param ObjectItem UObject's entry in GUObjectArray. This is either the other cluster root or one if the cluster objects.
-	* @param Obj The object to add to cluster
-	* @param ObjectsToSerialize An array of remaining objects to serialize (Obj must be added to it if Obj can be added to cluster)
-	*/
-	void MergeCluster(FUObjectItem* ObjectItem, UObject* Object, TArray<UObject*>& ObjectsToSerialize)
-	{
-		FUObjectItem* OtherClusterRootItem = nullptr;
-		int32 OtherClusterRootIndex = INDEX_NONE;
-		// First find the other cluster root item and its index
-		if (ObjectItem->HasAnyFlags(EInternalObjectFlags::ClusterRoot))
-		{
-			OtherClusterRootIndex = GUObjectArray.ObjectToIndex(Object);
-			OtherClusterRootItem = ObjectItem;
-		}
-		else
-		{
-			OtherClusterRootIndex = ObjectItem->GetOwnerIndex();
-			check(OtherClusterRootIndex > 0);
-			OtherClusterRootItem = GUObjectArray.IndexToObjectUnsafeForGC(OtherClusterRootIndex);
-		}
-		// This is another cluster, merge it with this one
-		FUObjectCluster& ClusterToMerge = GUObjectClusters[OtherClusterRootItem->GetClusterIndex()];
-		for (int32 OtherClusterObjectIndex : ClusterToMerge.Objects)
-		{
-			FUObjectItem* OtherClusterObjectItem = GUObjectArray.IndexToObjectUnsafeForGC(OtherClusterObjectIndex);
-			OtherClusterObjectItem->SetOwnerIndex(0);
-			AddObjectToCluster(OtherClusterObjectIndex, OtherClusterObjectItem, static_cast<UObject*>(OtherClusterObjectItem->Object), ObjectsToSerialize, true);
-		}		
-		GUObjectClusters.FreeCluster(OtherClusterRootItem->GetClusterIndex());
-
-		// Make sure the root object is also added to the current cluster
-		OtherClusterRootItem->ClearFlags(EInternalObjectFlags::ClusterRoot);
-		OtherClusterRootItem->SetOwnerIndex(0);
-		AddObjectToCluster(OtherClusterRootIndex, OtherClusterRootItem, static_cast<UObject*>(OtherClusterRootItem->Object), ObjectsToSerialize, true);
-
-		// Sanity check so that we make sure the object was actually in the lister it said it belonged to
-		check(ObjectItem->GetOwnerIndex() == ClusterRootIndex);
-	}
-
 	/**
 	* Handles UObject reference from the token stream. Performance is critical here so we're FORCEINLINING this function.
 	*
@@ -720,33 +670,25 @@ public:
 			{
 				if (ObjectItem->HasAnyFlags(EInternalObjectFlags::ClusterRoot) || ObjectItem->GetOwnerIndex() != 0)
 				{					
-					if (GMergeGCClusters)
+					// Simply reference this cluster and all clusters it's referencing
+					const int32 OtherClusterRootIndex = ObjectItem->HasAnyFlags(EInternalObjectFlags::ClusterRoot) ? GUObjectArray.ObjectToIndex(Object) : ObjectItem->GetOwnerIndex();
+					FUObjectItem* OtherClusterRootItem = GUObjectArray.IndexToObject(OtherClusterRootIndex);
+					const int32 OtherClusterIndex = OtherClusterRootItem->GetClusterIndex();
+					FUObjectCluster& OtherCluster = GUObjectClusters[OtherClusterIndex];
+					Cluster.ReferencedClusters.AddUnique(OtherClusterRootIndex);
+					OtherCluster.ReferencedByClusters.AddUnique(ClusterRootIndex);
+
+					for (int32 OtherClusterReferencedCluster : OtherCluster.ReferencedClusters)
 					{
-						// This is an existing cluster, merge it with the current one.
-						MergeCluster(ObjectItem, Object, ObjectsToSerialize);
-					}
-					else
-					{
-						// Simply reference this cluster and all clusters it's referencing
-						const int32 OtherClusterRootIndex = ObjectItem->HasAnyFlags(EInternalObjectFlags::ClusterRoot) ? GUObjectArray.ObjectToIndex(Object) : ObjectItem->GetOwnerIndex();
-						FUObjectItem* OtherClusterRootItem = GUObjectArray.IndexToObject(OtherClusterRootIndex);
-						const int32 OtherClusterIndex = OtherClusterRootItem->GetClusterIndex();
-						FUObjectCluster& OtherCluster = GUObjectClusters[OtherClusterIndex];
-						Cluster.ReferencedClusters.AddUnique(OtherClusterRootIndex);
-						OtherCluster.ReferencedByClusters.AddUnique(ClusterRootIndex);
-
-						for (int32 OtherClusterReferencedCluster : OtherCluster.ReferencedClusters)
-						{
-							if (OtherClusterReferencedCluster != ClusterRootIndex)
-							{
-								Cluster.ReferencedClusters.AddUnique(OtherClusterReferencedCluster);
-							}
-						}
-						for (int32 OtherClusterReferencedMutableObjectIndex : OtherCluster.MutableObjects)
+						if (OtherClusterReferencedCluster != ClusterRootIndex)
 						{
-							Cluster.MutableObjects.AddUnique(OtherClusterReferencedMutableObjectIndex);
+							Cluster.ReferencedClusters.AddUnique(OtherClusterReferencedCluster);
 						}
 					}
+					for (int32 OtherClusterReferencedMutableObjectIndex : OtherCluster.MutableObjects)
+					{
+						Cluster.MutableObjects.AddUnique(OtherClusterReferencedMutableObjectIndex);
+					}
 				}
 				else if (!GUObjectArray.IsDisregardForGC(Object)) // We know that disregard for GC objects will never be GC'd so no reference is necessary
 				{
diff --git a/Engine/Source/Runtime/Engine/Classes/Engine/CoreSettings.h b/Engine/Source/Runtime/Engine/Classes/Engine/CoreSettings.h
--- a/Engine/Source/Runtime/Engine/Classes/Engine/CoreSettings.h
+++ b/Engine/Source/Runtime/Engine/Classes/Engine/CoreSettings.h
@@ -180,11 +180,6 @@ protected:
 		ToolTip = "If true, the engine will attempt to create clusters of objects for better garbage collection performance."))
 	uint32 CreateGCClusters : 1;
 
-	UPROPERTY(EditAnywhere, config, Category = Optimization, meta = (
-		ConsoleVariable = "gc.MergeGCClusters", DisplayName = "Merge GC Clusters",
-		ToolTip = "If true, when creating clusters, the clusters referenced from another cluster will get merged into one big cluster."))
-	uint32 MergeGCClusters : 1;
-
 	UPROPERTY(EditAnywhere, config, Category = Optimization, meta = (
 		ConsoleVariable = "gc.ActorClusteringEnabled", DisplayName = "Actor Clustering Enabled",
 		ToolTip = "Whether to allow levels to create actor clusters for GC."))
diff --git a/Engine/Source/Runtime/Engine/Private/CoreSettings.cpp b/Engine/Source/Runtime/Engine/Private/CoreSettings.cpp
--- a/Engine/Source/Runtime/Engine/Private/CoreSettings.cpp
+++ b/Engine/Source/Runtime/Engine/Private/CoreSettings.cpp
@@ -175,7 +175,6 @@ UGarbageCollectionSettings::UGarbageCollectionSettings()
 	MaxObjectsInEditor = 12 * 1024 * 1024;
 	MaxObjectsInGame = 2 * 1024 * 1024;	
 	CreateGCClusters = true;	
-	MergeGCClusters = false;
 	MinGCClusterSize = 5;
 	ActorClusteringEnabled = true;
 	BlueprintClusteringEnabled = false;
