Fixed a crash when loading cooked asset in the editor caused by seeking past the header before uexp file was loaded

diff --git a/Engine/Source/Runtime/CoreUObject/Private/Serialization/AsyncLoading.cpp b/Engine/Source/Runtime/CoreUObject/Private/Serialization/AsyncLoading.cpp
--- a/Engine/Source/Runtime/CoreUObject/Private/Serialization/AsyncLoading.cpp
+++ b/Engine/Source/Runtime/CoreUObject/Private/Serialization/AsyncLoading.cpp
@@ -7711,7 +7711,7 @@ FORCEINLINE void FArchiveAsync2::SetPosAndUpdatePrecacheBuffer(int64 Pos)
 
 void FArchiveAsync2::Seek(int64 InPos)
 {
-	if (GEventDrivenLoaderEnabled && LoadPhase < ELoadPhase::ProcessingExports)
+	if ((GEventDrivenLoaderEnabled || bCookedForEDLInEditor) && LoadPhase < ELoadPhase::ProcessingExports)
 	{
 		check(!HeaderSizeWhenReadingExportsFromSplitFile && HeaderSize && TotalSize() == HeaderSize);
 		if (InPos >= HeaderSize)
