It's now possible to enable support for cooked data in the editor in project settings. The ini setting moved from [Core.System] AllowCookedDataInEditorBuilds=True to [/Script/UnrealEd.CookerSettings] cook.AllowCookedDataInEditorBuilds=True

diff --git a/Engine/Source/Editor/UnrealEd/Private/CookerSettings.cpp b/Engine/Source/Editor/UnrealEd/Private/CookerSettings.cpp
--- a/Engine/Source/Editor/UnrealEd/Private/CookerSettings.cpp
+++ b/Engine/Source/Editor/UnrealEd/Private/CookerSettings.cpp
@@ -15,6 +15,7 @@ UCookerSettings::UCookerSettings(const FObjectInitializer& ObjectInitializer)
 	, bCompileBlueprintsInDevelopmentMode(true)
 	, BlueprintComponentDataCookingMethod(EBlueprintComponentDataCookingMethod::EnabledBlueprintsOnly)
 	, bCookBlueprintComponentTemplateData(false)
+	, bAllowCookedDataInEditorBuilds(false)
 {
 	SectionName = TEXT("Cooker");
 	DefaultPVRTCQuality = 1;
diff --git a/Engine/Source/Editor/UnrealEd/Public/CookerSettings.h b/Engine/Source/Editor/UnrealEd/Public/CookerSettings.h
--- a/Engine/Source/Editor/UnrealEd/Public/CookerSettings.h
+++ b/Engine/Source/Editor/UnrealEd/Public/CookerSettings.h
@@ -121,6 +121,12 @@ public:
 	UPROPERTY(GlobalConfig, EditAnywhere, Category = Textures, meta = (DisplayName = "ASTC Compression Quality vs Size (0-4, 0 is smallest)"))
 	int32 DefaultASTCQualityBySize;
 
+	/** Allows opening cooked assets in the editor */
+	UPROPERTY(EditAnywhere, config, Category = Editor, meta = (
+		ConsoleVariable = "cook.AllowCookedDataInEditorBuilds", DisplayName = "Allow Cooked Content In The Editor",
+		ToolTip = "If true, the editor will be able to open cooked assets (limited to a subset of supported asset types)."))
+	uint32 bAllowCookedDataInEditorBuilds : 1;
+
 private:
 	/** Deprecated. Use BlueprintComponentDataCookingMethod instead. */
 	UPROPERTY(GlobalConfig)
diff --git a/Engine/Source/Runtime/CoreUObject/Private/UObject/LinkerLoad.cpp b/Engine/Source/Runtime/CoreUObject/Private/UObject/LinkerLoad.cpp
--- a/Engine/Source/Runtime/CoreUObject/Private/UObject/LinkerLoad.cpp
+++ b/Engine/Source/Runtime/CoreUObject/Private/UObject/LinkerLoad.cpp
@@ -81,19 +81,14 @@ bool FLinkerLoad::ShouldCreateThrottledSlowTask() const
 }
 #endif
 
-// Helper struct for getting the value of [Core.System] AllowCookedDataInEditorBuilds from ini
-struct FInitCookedDatataInEditorBuildsSupport
-{
-	bool bAllowCookedData;
-	FInitCookedDatataInEditorBuildsSupport()
-	{
-		if (!GConfig->GetBool(TEXT("Core.System"), TEXT("AllowCookedDataInEditorBuilds"), bAllowCookedData, GEngineIni))
-		{
-			bAllowCookedData = true;
-		}
-	}
-	FORCEINLINE operator bool() const { return bAllowCookedData; }
-};
+
+static int32 GAllowCookedDataInEditorBuilds = 0;
+static FAutoConsoleVariableRef CVarAllowCookedDataInEditorBuilds(
+	TEXT("cook.AllowCookedDataInEditorBuilds"),
+		GAllowCookedDataInEditorBuilds,
+	TEXT("If true, allows cooked assets to be loaded in the editor."),
+	ECVF_Default
+);
 
 /**
 * Test whether the given package index is a valid import or export in this package
@@ -1010,8 +1005,7 @@ FLinkerLoad::ELinkerStatus FLinkerLoad::CreateLoader(
 			{
 				// If want to be able to load cooked data in the editor we need to use FArchiveAsync2 which supports EDL cooked packages,
 				// otherwise the generic file reader is faster in the editor so use that
-				static FInitCookedDatataInEditorBuildsSupport AllowCookedDataInEditorBuilds;
-				bool bCanUseFArchiveAsync2 = FPlatformProperties::RequiresCookedData() || AllowCookedDataInEditorBuilds;
+				bool bCanUseFArchiveAsync2 = FPlatformProperties::RequiresCookedData() || GAllowCookedDataInEditorBuilds;
 				if (bCanUseFArchiveAsync2)
 				{
 					Loader = new FArchiveAsync2(*Filename
@@ -1207,9 +1201,8 @@ FLinkerLoad::ELinkerStatus FLinkerLoad::SerializePackageFileSummary()
 		// don't load packages that contain editor only data in builds that don't support that and vise versa
 		if (FPlatformProperties::HasEditorOnlyData() && !!(Summary.PackageFlags & PKG_FilterEditorOnly))
 		{
-			// This warning can be disabled in ini with [Core.System] AllowCookedDataInEditorBuilds=False
-			static FInitCookedDatataInEditorBuildsSupport AllowCookedDataInEditorBuilds;
-			if (!AllowCookedDataInEditorBuilds)
+			// This warning can be disabled in ini or project settings
+			if (!GAllowCookedDataInEditorBuilds)
 			{
 				UE_LOG(LogLinker, Warning, 
 					TEXT("Unable to load package (%s). Package contains cooked data which is not supported by the current build. Set [Core.System] AllowCookedDataInEditorBuilds to true in Engine.ini to allow it."), 
