Fixing cook crash on startup / bug with  highly recursive synchronous loads

diff --git a/Engine/Source/Runtime/CoreUObject/Private/UObject/UObjectGlobals.cpp b/Engine/Source/Runtime/CoreUObject/Private/UObject/UObjectGlobals.cpp
--- a/Engine/Source/Runtime/CoreUObject/Private/UObject/UObjectGlobals.cpp
+++ b/Engine/Source/Runtime/CoreUObject/Private/UObject/UObjectGlobals.cpp
@@ -1267,9 +1267,34 @@ UPackage* LoadPackageInternal(UPackage* InOuter, const TCHAR* InLongPackageNameO
 	{
 		FCoreDelegates::OnSyncLoadPackage.Broadcast(FileToLoad);
 	}
+	
+	// Set up a load context
+	TRefCountPtr<FUObjectSerializeContext> LoadContext;
+	if (InLoadContext)
+	{
+		// Use the privided context
+		LoadContext = InLoadContext;
+	}
+	else
+	{
+		// Try to get the context from the callstack
+		FUObjectThreadContext& ThreadContext = FUObjectThreadContext::Get();
+		if (ThreadContext.SerializeContext)
+		{
+			LoadContext = ThreadContext.SerializeContext;
+		}
+		else
+		{
+			// Create a new context
+			LoadContext = new FUObjectSerializeContext();
+			if (!IsInAsyncLoadingThread())
+			{
+				ThreadContext.SerializeContext = LoadContext;
+			}
+		}
+	}
 
 	// Try to load.
-	TRefCountPtr<FUObjectSerializeContext> LoadContext(InLoadContext ? InLoadContext : new FUObjectSerializeContext());
 	BeginLoad(LoadContext, InLongPackageNameOrFilename);
 
 	bool bFullyLoadSkipped = false;
@@ -1841,6 +1866,20 @@ void EndLoad(FUObjectSerializeContext* LoadContext)
 		FCoreUObjectDelegates::OnAssetLoaded.Broadcast(LoadedAsset);
 	}
 #endif	// WITH_EDITOR
+
+
+	if (LoadContext->GetBeginLoadCount() == 0)
+	{
+		FUObjectThreadContext& ThreadContext = FUObjectThreadContext::Get();
+		if (ThreadContext.SerializeContext == LoadContext)
+		{
+			ThreadContext.SerializeContext = nullptr;
+		}
+		else
+		{
+			check(!ThreadContext.SerializeContext);
+		}
+	}
 }
 
 /*-----------------------------------------------------------------------------
diff --git a/Engine/Source/Runtime/CoreUObject/Private/UObject/UObjectThreadContext.cpp b/Engine/Source/Runtime/CoreUObject/Private/UObject/UObjectThreadContext.cpp
--- a/Engine/Source/Runtime/CoreUObject/Private/UObject/UObjectThreadContext.cpp
+++ b/Engine/Source/Runtime/CoreUObject/Private/UObject/UObjectThreadContext.cpp
@@ -15,6 +15,7 @@ FUObjectThreadContext::FUObjectThreadContext()
 , IsInConstructor(0)
 , ConstructedObject(nullptr)
 , AsyncPackage(nullptr)
+, SerializeContext(nullptr)
 {}
 
 
diff --git a/Engine/Source/Runtime/CoreUObject/Public/UObject/UObjectThreadContext.h b/Engine/Source/Runtime/CoreUObject/Public/UObject/UObjectThreadContext.h
--- a/Engine/Source/Runtime/CoreUObject/Public/UObject/UObjectThreadContext.h
+++ b/Engine/Source/Runtime/CoreUObject/Public/UObject/UObjectThreadContext.h
@@ -75,6 +75,8 @@ public:
 	UObject* ConstructedObject;
 	/** Async Package currently processing objects */
 	struct FAsyncPackage* AsyncPackage;
+	/** Current serialization context (only valid for non-async loads) */
+	FUObjectSerializeContext* SerializeContext;
 
 #if !(UE_BUILD_SHIPPING || UE_BUILD_TEST)
 	/** Stack to ensure that PostInitProperties is routed through Super:: calls. **/
