Fixing more CIS issues

diff --git a/Engine/Source/Runtime/CoreUObject/Private/UObject/EnumProperty.cpp b/Engine/Source/Runtime/CoreUObject/Private/UObject/EnumProperty.cpp
--- a/Engine/Source/Runtime/CoreUObject/Private/UObject/EnumProperty.cpp
+++ b/Engine/Source/Runtime/CoreUObject/Private/UObject/EnumProperty.cpp
@@ -7,6 +7,7 @@
 #include "Templates/ChooseClass.h"
 #include "Templates/IsSigned.h"
 #include "Algo/Find.h"
+#include "UObject/LinkerLoad.h"
 
 namespace UE4EnumProperty_Private
 {
diff --git a/Engine/Source/Runtime/CoreUObject/Private/UObject/PropertyByte.cpp b/Engine/Source/Runtime/CoreUObject/Private/UObject/PropertyByte.cpp
--- a/Engine/Source/Runtime/CoreUObject/Private/UObject/PropertyByte.cpp
+++ b/Engine/Source/Runtime/CoreUObject/Private/UObject/PropertyByte.cpp
@@ -8,6 +8,7 @@
 #include "UObject/UObjectThreadContext.h"
 #include "Serialization/ArchiveUObjectFromStructuredArchive.h"
 #include "Algo/Find.h"
+#include "UObject/LinkerLoad.h"
 
 /*-----------------------------------------------------------------------------
 	UByteProperty.
