Fixed live server crash with invalid package names passed to LoadPackage, in UPackageMap::InternalLoadObject.

diff --git a/Engine/Source/Runtime/CoreUObject/Private/Serialization/AsyncLoading.cpp b/Engine/Source/Runtime/CoreUObject/Private/Serialization/AsyncLoading.cpp
--- a/Engine/Source/Runtime/CoreUObject/Private/Serialization/AsyncLoading.cpp
+++ b/Engine/Source/Runtime/CoreUObject/Private/Serialization/AsyncLoading.cpp
@@ -6850,6 +6850,8 @@ void FAsyncPackage::UpdateLoadPercentage()
 
 int32 LoadPackageAsync(const FString& InName, const FGuid* InGuid /*= nullptr*/, const TCHAR* InPackageToLoadFrom /*= nullptr*/, FLoadPackageAsyncDelegate InCompletionDelegate /*= FLoadPackageAsyncDelegate()*/, EPackageFlags InPackageFlags /*= PKG_None*/, int32 InPIEInstanceID /*= INDEX_NONE*/, int32 InPackagePriority /*= 0*/)
 {
+	int32 RequestID = INDEX_NONE;
+
 	static bool bOnce = false;
 	if (!bOnce && GEventDrivenLoaderEnabled)
 	{
@@ -6860,6 +6862,8 @@ int32 LoadPackageAsync(const FString& InName, const FGuid* InGuid /*= nullptr*/,
 	// The comments clearly state that it should be a package name but we also handle it being a filename as this function is not perf critical
 	// and LoadPackage handles having a filename being passed in as well.
 	FString PackageName;
+	bool bValidPackageName = true;
+
 	if (FPackageName::IsValidLongPackageName(InName, /*bIncludeReadOnlyRoots*/true))
 	{
 		PackageName = InName;
@@ -6867,47 +6871,65 @@ int32 LoadPackageAsync(const FString& InName, const FGuid* InGuid /*= nullptr*/,
 	// PackageName got populated by the conditional function
 	else if (!(FPackageName::IsPackageFilename(InName) && FPackageName::TryConvertFilenameToLongPackageName(InName, PackageName)))
 	{
-		// PackageName will get populated by the conditional function
+		// PackageName may get populated by the conditional function
 		FString ClassName;
+
 		if (!FPackageName::ParseExportTextPath(PackageName, &ClassName, &PackageName))
 		{
-			UE_LOG(LogStreaming, Fatal, TEXT("LoadPackageAsync failed to begin to load a package because the supplied package name was neither a valid long package name nor a filename of a map within a content folder: '%s'"), *PackageName);
+			UE_LOG(LogStreaming, Warning, TEXT("LoadPackageAsync failed to begin to load a package because the supplied package name ")
+					TEXT("was neither a valid long package name nor a filename of a map within a content folder: '%s' (%s)"),
+					*PackageName, *InName);
+
+			bValidPackageName = false;
 		}
 	}
 
 	FString PackageNameToLoad(InPackageToLoadFrom);
-	if (PackageNameToLoad.IsEmpty())
-	{
-		PackageNameToLoad = PackageName;
-	}
-	// Make sure long package name is passed to FAsyncPackage so that it doesn't attempt to 
-	// create a package with short name.
-	if (FPackageName::IsShortPackageName(PackageNameToLoad))
+
+	if (bValidPackageName)
 	{
-		UE_LOG(LogStreaming, Fatal, TEXT("Async loading code requires long package names (%s)."), *PackageNameToLoad);
+		if (PackageNameToLoad.IsEmpty())
+		{
+			PackageNameToLoad = PackageName;
+		}
+		// Make sure long package name is passed to FAsyncPackage so that it doesn't attempt to 
+		// create a package with short name.
+		if (FPackageName::IsShortPackageName(PackageNameToLoad))
+		{
+			UE_LOG(LogStreaming, Warning, TEXT("Async loading code requires long package names (%s)."), *PackageNameToLoad);
+
+			bValidPackageName = false;
+		}
 	}
 
-	if ( FCoreDelegates::OnAsyncLoadPackage.IsBound() )
+	if (bValidPackageName)
 	{
-		FCoreDelegates::OnAsyncLoadPackage.Broadcast(InName);
-	}
+		if ( FCoreDelegates::OnAsyncLoadPackage.IsBound() )
+		{
+			FCoreDelegates::OnAsyncLoadPackage.Broadcast(InName);
+		}
+
+		// Generate new request ID and add it immediately to the global request list (it needs to be there before we exit
+		// this function, otherwise it would be added when the packages are being processed on the async thread).
+		RequestID = GPackageRequestID.Increment();
+		FAsyncLoadingThread::Get().AddPendingRequest(RequestID);
 
-	// Generate new request ID and add it immediately to the global request list (it needs to be there before we exit
-	// this function, otherwise it would be added when the packages are being processed on the async thread).
-	const int32 RequestID = GPackageRequestID.Increment();
-	FAsyncLoadingThread::Get().AddPendingRequest(RequestID);
+		// Allocate delegate on Game Thread, it is not safe to copy delegates by value on other threads
+		TUniquePtr<FLoadPackageAsyncDelegate> CompletionDelegatePtr;
+		if (InCompletionDelegate.IsBound())
+		{
+			CompletionDelegatePtr.Reset(new FLoadPackageAsyncDelegate(InCompletionDelegate));
+		}
 
-	// Allocate delegate on Game Thread, it is not safe to copy delegates by value on other threads
-	TUniquePtr<FLoadPackageAsyncDelegate> CompletionDelegatePtr;
-	if (InCompletionDelegate.IsBound())
+		// Add new package request
+		FAsyncPackageDesc PackageDesc(RequestID, *PackageName, *PackageNameToLoad, InGuid ? *InGuid : FGuid(), MoveTemp(CompletionDelegatePtr), InPackageFlags, InPIEInstanceID, InPackagePriority);
+		FAsyncLoadingThread::Get().QueuePackage(PackageDesc);
+	}
+	else
 	{
-		CompletionDelegatePtr.Reset(new FLoadPackageAsyncDelegate(InCompletionDelegate));
+		InCompletionDelegate.ExecuteIfBound(FName(*InName), nullptr, EAsyncLoadingResult::Failed);
 	}
 
-	// Add new package request
-	FAsyncPackageDesc PackageDesc(RequestID, *PackageName, *PackageNameToLoad, InGuid ? *InGuid : FGuid(), MoveTemp(CompletionDelegatePtr), InPackageFlags, InPIEInstanceID, InPackagePriority);
-	FAsyncLoadingThread::Get().QueuePackage(PackageDesc);
-
 	return RequestID;
 }
 
diff --git a/Engine/Source/Runtime/CoreUObject/Private/UObject/UObjectGlobals.cpp b/Engine/Source/Runtime/CoreUObject/Private/UObject/UObjectGlobals.cpp
--- a/Engine/Source/Runtime/CoreUObject/Private/UObject/UObjectGlobals.cpp
+++ b/Engine/Source/Runtime/CoreUObject/Private/UObject/UObjectGlobals.cpp
@@ -1189,7 +1189,11 @@ UPackage* LoadPackageInternal(UPackage* InOuter, const TCHAR* InLongPackageNameO
 			}
 
 			int32 RequestID = LoadPackageAsync(InName, nullptr, *InPackageName);
-			FlushAsyncLoading(RequestID);
+
+			if (RequestID != INDEX_NONE)
+			{
+				FlushAsyncLoading(RequestID);
+			}
 		}
 
 		Result = FindObjectFast<UPackage>(nullptr, PackageFName);
