FFallbackStruct needs CPP struct ops

diff --git a/Engine/Source/Runtime/Core/Public/Misc/FallbackStruct.h b/Engine/Source/Runtime/Core/Public/Misc/FallbackStruct.h
new file mode 100644
--- /dev/null
+++ b/Engine/Source/Runtime/Core/Public/Misc/FallbackStruct.h
@@ -0,0 +1,8 @@
+// Copyright 1998-2019 Epic Games, Inc. All Rights Reserved.
+
+#pragma once
+
+// A struct used as stub for deleted ones.
+struct FFallbackStruct
+{
+};
\ No newline at end of file
diff --git a/Engine/Source/Runtime/CoreUObject/Private/UObject/Property.cpp b/Engine/Source/Runtime/CoreUObject/Private/UObject/Property.cpp
--- a/Engine/Source/Runtime/CoreUObject/Private/UObject/Property.cpp
+++ b/Engine/Source/Runtime/CoreUObject/Private/UObject/Property.cpp
@@ -352,6 +352,16 @@ struct TStructOpsTypeTraits<FPrimaryAssetId> : public TStructOpsTypeTraitsBase2<
 	};
 };
 IMPLEMENT_STRUCT(PrimaryAssetId);
+
+template<>
+struct TStructOpsTypeTraits<FFallbackStruct> : public TStructOpsTypeTraitsBase2<FFallbackStruct>
+{
+	enum
+	{
+	};
+};
+IMPLEMENT_STRUCT(FallbackStruct);
+
 /*-----------------------------------------------------------------------------
 	Helpers.
 -----------------------------------------------------------------------------*/
diff --git a/Engine/Source/Runtime/CoreUObject/Public/UObject/Class.h b/Engine/Source/Runtime/CoreUObject/Public/UObject/Class.h
--- a/Engine/Source/Runtime/CoreUObject/Public/UObject/Class.h
+++ b/Engine/Source/Runtime/CoreUObject/Public/UObject/Class.h
@@ -11,6 +11,7 @@
 #include "UObject/ObjectMacros.h"
 #include "UObject/UObjectGlobals.h"
 #include "UObject/Object.h"
+#include "Misc/FallbackStruct.h"
 #include "Misc/Guid.h"
 #include "Math/RandomStream.h"
 #include "UObject/GarbageCollection.h"
@@ -3242,7 +3243,6 @@ template<> struct TBaseStructure<FBox2D>
 	COREUOBJECT_API static UScriptStruct* Get();
 };	
 
-struct FFallbackStruct;
 template<> struct TBaseStructure<FFallbackStruct>
 {
 	COREUOBJECT_API static UScriptStruct* Get();
diff --git a/Engine/Source/Runtime/CoreUObject/Public/UObject/NoExportTypes.h b/Engine/Source/Runtime/CoreUObject/Public/UObject/NoExportTypes.h
--- a/Engine/Source/Runtime/CoreUObject/Public/UObject/NoExportTypes.h
+++ b/Engine/Source/Runtime/CoreUObject/Public/UObject/NoExportTypes.h
@@ -17,6 +17,7 @@
 
 #include "PixelFormat.h"
 
+#include "Misc/FallbackStruct.h"
 #include "Misc/Guid.h"
 #include "Misc/DateTime.h"
 #include "Misc/Timespan.h"
