Try to fix Linux's unnecessary complaints
#rnx

diff --git a/Engine/Source/Runtime/CoreUObject/Private/UObject/Property.cpp b/Engine/Source/Runtime/CoreUObject/Private/UObject/Property.cpp
--- a/Engine/Source/Runtime/CoreUObject/Private/UObject/Property.cpp
+++ b/Engine/Source/Runtime/CoreUObject/Private/UObject/Property.cpp
@@ -356,9 +356,6 @@ IMPLEMENT_STRUCT(PrimaryAssetId);
 template<>
 struct TStructOpsTypeTraits<FFallbackStruct> : public TStructOpsTypeTraitsBase2<FFallbackStruct>
 {
-	enum
-	{
-	};
 };
 IMPLEMENT_STRUCT(FallbackStruct);
 
