DR - Add missing enum entries
#rnx

diff --git a/Engine/Source/Runtime/CoreUObject/Public/UObject/NoExportTypes.h b/Engine/Source/Runtime/CoreUObject/Public/UObject/NoExportTypes.h
--- a/Engine/Source/Runtime/CoreUObject/Public/UObject/NoExportTypes.h
+++ b/Engine/Source/Runtime/CoreUObject/Public/UObject/NoExportTypes.h
@@ -214,6 +214,10 @@ enum EPixelFormat
 	PF_R8G8B8A8_SNORM,
 	PF_R16G16B16A16_UNORM,
 	PF_R16G16B16A16_SNORM,
+	PF_PLATFORM_HDR_0,
+	PF_PLATFORM_HDR_1,
+	PF_PLATFORM_HDR_2,
+	PF_NV12,
 	PF_MAX,
 };
 
