Avoid a copy when adding to annotation if possible
Fix some spelling errors
#rnx

diff --git a/Engine/Source/Runtime/CoreUObject/Public/UObject/UObjectAnnotation.h b/Engine/Source/Runtime/CoreUObject/Public/UObject/UObjectAnnotation.h
--- a/Engine/Source/Runtime/CoreUObject/Public/UObject/UObjectAnnotation.h
+++ b/Engine/Source/Runtime/CoreUObject/Public/UObject/UObjectAnnotation.h
@@ -40,7 +40,7 @@ public:
 #if !(UE_BUILD_SHIPPING || UE_BUILD_TEST)
 		if (!bAutoRemove)
 		{
-			FScopeLock AnntationMapLock(&AnnotationMapCritical);
+			FScopeLock AnnotationMapLock(&AnnotationMapCritical);
 			// in this case we are only verifying that the external assurances of removal are met
 			check(!AnnotationMap.Find(Object));
 		}
@@ -78,10 +78,10 @@ public:
 	void AddAnnotation(const UObjectBase *Object,TAnnotation Annotation)
 	{
 		check(Object);
-		FScopeLock AnntationMapLock(&AnnotationMapCritical);
+		FScopeLock AnnotationMapLock(&AnnotationMapCritical);
 		AnnotationCacheKey = Object;
-		AnnotationCacheValue = Annotation;
-		if (Annotation.IsDefault())
+		AnnotationCacheValue = MoveTemp(Annotation);
+		if (AnnotationCacheValue.IsDefault())
 		{
 			RemoveAnnotation(Object); // adding the default annotation is the same as removing an annotation
 		}
@@ -109,7 +109,7 @@ public:
 	TAnnotation GetAndRemoveAnnotation(const UObjectBase *Object)
 	{		
 		check(Object);
-		FScopeLock AnntationMapLock(&AnnotationMapCritical);
+		FScopeLock AnnotationMapLock(&AnnotationMapCritical);
 		AnnotationCacheKey = Object;
 		AnnotationCacheValue = TAnnotation();
 		const bool bHadElements = (AnnotationMap.Num() > 0);
@@ -135,7 +135,7 @@ public:
 	void RemoveAnnotation(const UObjectBase *Object)
 	{
 		check(Object);
-		FScopeLock AnntationMapLock(&AnnotationMapCritical);
+		FScopeLock AnnotationMapLock(&AnnotationMapCritical);
 		AnnotationCacheKey = Object;
 		AnnotationCacheValue = TAnnotation();
 		const bool bHadElements = (AnnotationMap.Num() > 0);
@@ -157,7 +157,7 @@ public:
 	 */
 	void RemoveAllAnnotations()
 	{
-		FScopeLock AnntationMapLock(&AnnotationMapCritical);
+		FScopeLock AnnotationMapLock(&AnnotationMapCritical);
 		AnnotationCacheKey = NULL;
 		AnnotationCacheValue = TAnnotation();
 		const bool bHadElements = (AnnotationMap.Num() > 0);
@@ -182,7 +182,7 @@ public:
 	FORCEINLINE TAnnotation GetAnnotation(const UObjectBase *Object)
 	{
 		check(Object);
-		FScopeLock AnntationMapLock(&AnnotationMapCritical);
+		FScopeLock AnnotationMapLock(&AnnotationMapCritical);
 		if (Object != AnnotationCacheKey)
 		{			
 			AnnotationCacheKey = Object;
@@ -213,7 +213,7 @@ public:
 	 */
 	void Reserve(int32 ExpectedNumElements)
 	{
-		FScopeLock AnntationMapLock(&AnnotationMapCritical);
+		FScopeLock AnnotationMapLock(&AnnotationMapCritical);
 		AnnotationMap.Empty(ExpectedNumElements);
 	}
 
@@ -280,7 +280,7 @@ public:
 	 */
 	UObject *Find(TAnnotation Annotation)
 	{
-		FScopeLock InverseAnntationMapLock(&InverseAnnotationMapCritical);
+		FScopeLock InverseAnnotationMapLock(&InverseAnnotationMapCritical);
 		checkSlow(!Annotation.IsDefault()); // it is not legal to search for the default annotation
 		return (UObject *)InverseAnnotationMap.FindRef(Annotation);
 	}
@@ -293,7 +293,7 @@ public:
 	 */
 	void AddAnnotation(const UObjectBase *Object,TAnnotation Annotation)
 	{
-		FScopeLock InverseAnntationMapLock(&InverseAnnotationMapCritical);
+		FScopeLock InverseAnnotationMapLock(&InverseAnnotationMapCritical);
 		if (Annotation.IsDefault())
 		{
 			RemoveAnnotation(Object); // adding the default annotation is the same as removing an annotation
@@ -318,7 +318,7 @@ public:
 	 */
 	void RemoveAnnotation(const UObjectBase *Object)
 	{
-		FScopeLock InverseAnntationMapLock(&InverseAnnotationMapCritical);
+		FScopeLock InverseAnnotationMapLock(&InverseAnnotationMapCritical);
 		TAnnotation Annotation = this->GetAndRemoveAnnotation(Object);
 		if (Annotation.IsDefault())
 		{
@@ -337,7 +337,7 @@ public:
 	 */
 	void RemoveAllAnnotations()
 	{
-		FScopeLock InverseAnntationMapLock(&InverseAnnotationMapCritical);
+		FScopeLock InverseAnnotationMapLock(&InverseAnnotationMapCritical);
 		Super::RemoveAllAnnotations();
 		InverseAnnotationMap.Empty();
 	}
