Ensure that NextOutParm is always null on the last element in the linked list, so that the list can be reliably enumerated

diff --git a/Engine/Source/Runtime/CoreUObject/Private/UObject/ScriptCore.cpp b/Engine/Source/Runtime/CoreUObject/Private/UObject/ScriptCore.cpp
--- a/Engine/Source/Runtime/CoreUObject/Private/UObject/ScriptCore.cpp
+++ b/Engine/Source/Runtime/CoreUObject/Private/UObject/ScriptCore.cpp
@@ -828,13 +828,11 @@ void ProcessScriptFunction(UObject* Context, UFunction* Function, FFrame& Stack,
 		}
 	}
 	Stack.Code++;
-#if UE_BUILD_DEBUG
-	// set the next pointer of the last item to NULL so we'll properly assert if something goes wrong
+	// set the next pointer of the last item to NULL to mark the end of the list
 	if (*LastOut)
 	{
 		(*LastOut)->NextOutParm = NULL;
 	}
-#endif
 
 	if (!bUsePersistentFrame)
 	{
@@ -1483,13 +1481,11 @@ void UObject::ProcessEvent( UFunction* Function, void* Parms )
 				}
 			}
 
-#if UE_BUILD_DEBUG
-			// set the next pointer of the last item to NULL so we'll properly assert if something goes wrong
+			// set the next pointer of the last item to NULL to mark the end of the list
 			if (*LastOut)
 			{
 				(*LastOut)->NextOutParm = NULL;
 			}
-#endif
 		}
 
 		if (!bUsePersistentFrame)
