[NO NON-PII COMMIT MESSAGE]

diff --git a/Engine/Source/Runtime/CoreUObject/Private/UObject/GarbageCollection.cpp b/Engine/Source/Runtime/CoreUObject/Private/UObject/GarbageCollection.cpp
--- a/Engine/Source/Runtime/CoreUObject/Private/UObject/GarbageCollection.cpp
+++ b/Engine/Source/Runtime/CoreUObject/Private/UObject/GarbageCollection.cpp
@@ -152,14 +152,6 @@ static FAutoConsoleVariableRef CVarMinDesiredObjectsPerSubTask(
 	ECVF_Default
 	);
 
-static int32 GCheckForIllegalMarkPendingKill = !(UE_BUILD_TEST || UE_BUILD_SHIPPING);
-static FAutoConsoleVariableRef CVarCheckForIllegalMarkPendingKill(
-	TEXT("gc.CheckForIllegalMarkPendingKill"),
-	GCheckForIllegalMarkPendingKill,
-	TEXT("If > 0, garbage collection will check for certainly rendering uobjects being illegally marked pending kill. This eventually causes mysterious and hard to find crashes in the renderer. There is a large performance penalty, so by default this is not enabled in shipping and test configurations."),
-	ECVF_Default
-);
-
 static int32 GIncrementalBeginDestroyEnabled = 1;
 static FAutoConsoleVariableRef CIncrementalBeginDestroyEnabled(
 	TEXT("gc.IncrementalBeginDestroyEnabled"),
@@ -481,20 +473,6 @@ public:
 
 			// Null out reference.
 			Object = NULL;
-
-			// Silently nulling out references can be fatal for some objects.  Usually rendering objects which would need to recreate renderthread proxies to avoid using deleted data and crashing.  e.g.
-			// If MarkPendingKill destroyed a UTexture that was still referenced by a Material then that can cause a crash as the RT data of the material will still try to render with the bad texture.
-			// Unfortunately this is often a race condition between threads, so we want to log errors early and deterministically.
-			if (GCheckForIllegalMarkPendingKill && ReferencingObject && !ReferencingObject->IsPendingKill())
-			{
-				const int32 ObjectIndexReferencer = GUObjectArray.ObjectToIndex(ReferencingObject);
-				FUObjectItem* ObjectItemReferencer = GUObjectArray.IndexToObjectUnsafeForGC(ObjectIndexReferencer);
-
-				//set HadReferenceKilled so we can later call NotifyObjectReferenceEliminated() on objects that have had references silently null'd out.  We don't do it immediately here to avoid false positives in the case where
-				//the Referencer is unreachable.  i.e. If the referencing object is dead anyway we don't need to notify it.
-				ObjectItemReferencer->SetFlags(EInternalObjectFlags::HadReferenceKilled);
-				UE_LOG(LogGarbage, Verbose, TEXT("NotifyObjectReferenceEliminated %s %s %s"), *ReferencingObject->GetPathName(), *ObjectItem->Object->GetFName().ToString(), *ObjectItem->Object->GetOuter()->GetName());				
-			}
 		}
 		// Add encountered object reference to list of to be serialized objects if it hasn't already been added.
 		else if (ObjectItem->IsUnreachable())
@@ -1529,7 +1507,7 @@ void CollectGarbageInternal(EObjectFlags KeepFlags, bool bPerformFullPurge)
 		// This has to be unlocked before we call post GC callbacks
 		FGCScopeLock GCLock;
 
-		UE_LOG(LogGarbage, Log, TEXT("Collecting garbage%s   (GCheckForIllegalMarkPendingKill = %d)"), IsAsyncLoading() ? TEXT(" while async loading") : TEXT(""), GCheckForIllegalMarkPendingKill);
+		UE_LOG(LogGarbage, Log, TEXT("Collecting garbage%s"), IsAsyncLoading() ? TEXT(" while async loading") : TEXT(""));
 
 		// Make sure previous incremental purge has finished or we do a full purge pass in case we haven't kicked one
 		// off yet since the last call to garbage collection.
@@ -1654,7 +1632,7 @@ bool UnhashUnreachableObjects(bool bUseTimeLimit, float TimeLimit)
 			// Begin the object's asynchronous destruction.
 			Object->ConditionalBeginDestroy();
 		}
-		ObjectItem->ClearFlags(EInternalObjectFlags::HadReferenceKilled);
+
 		Items++;
 
 		const bool bPollTimeLimit = ((TimePollCounter++) % TimeLimitEnforcementGranularityForBeginDestroy == 0);
diff --git a/Engine/Source/Runtime/CoreUObject/Public/UObject/Object.h b/Engine/Source/Runtime/CoreUObject/Public/UObject/Object.h
--- a/Engine/Source/Runtime/CoreUObject/Public/UObject/Object.h
+++ b/Engine/Source/Runtime/CoreUObject/Public/UObject/Object.h
@@ -696,8 +696,6 @@ public:
 	/** Returns true if this object is considered an asset. */
 	virtual bool IsAsset() const;
 
-	virtual void NotifyObjectReferenceEliminated() const {}
-
 	/**
 	 * Returns an Type:Name pair representing the PrimaryAssetId for this object.
 	 * Assets that need to be globally referenced at runtime should return a valid Identifier.
diff --git a/Engine/Source/Runtime/CoreUObject/Public/UObject/ObjectMacros.h b/Engine/Source/Runtime/CoreUObject/Public/UObject/ObjectMacros.h
--- a/Engine/Source/Runtime/CoreUObject/Public/UObject/ObjectMacros.h
+++ b/Engine/Source/Runtime/CoreUObject/Public/UObject/ObjectMacros.h
@@ -524,11 +524,11 @@ enum class EInternalObjectFlags : int32
 	Unreachable = 1 << 28, ///< Object is not reachable on the object graph.
 	PendingKill = 1 << 29, ///< Objects that are pending destruction (invalid for gameplay but valid objects)
 	RootSet = 1 << 30, ///< Object will not be garbage collected, even if unreferenced.
-	HadReferenceKilled = 1 << 31, ///< Object had a reference null'd out by markpendingkill
+	//~ UnusedFlag = 1 << 31,
 
 	GarbageCollectionKeepFlags = Native | Async | AsyncLoading,
 	//~ Make sure this is up to date!
-	AllFlags = ReachableInCluster | ClusterRoot | Native | Async | AsyncLoading | Unreachable | PendingKill | RootSet | HadReferenceKilled
+	AllFlags = ReachableInCluster | ClusterRoot | Native | Async | AsyncLoading | Unreachable | PendingKill | RootSet
 };
 ENUM_CLASS_FLAGS(EInternalObjectFlags);
 
diff --git a/Engine/Source/Runtime/Engine/Classes/Components/SkeletalMeshComponent.h b/Engine/Source/Runtime/Engine/Classes/Components/SkeletalMeshComponent.h
--- a/Engine/Source/Runtime/Engine/Classes/Components/SkeletalMeshComponent.h
+++ b/Engine/Source/Runtime/Engine/Classes/Components/SkeletalMeshComponent.h
@@ -1297,7 +1297,6 @@ public:
 public:
 	//~ Begin UObject Interface.
 	virtual void Serialize(FArchive& Ar) override;
-	virtual void NotifyObjectReferenceEliminated() const override;
 	virtual void PostLoad() override;
 #if WITH_EDITOR
 	DECLARE_MULTICAST_DELEGATE(FOnSkeletalMeshPropertyChangedMulticaster)
diff --git a/Engine/Source/Runtime/Engine/Classes/Components/StaticMeshComponent.h b/Engine/Source/Runtime/Engine/Classes/Components/StaticMeshComponent.h
--- a/Engine/Source/Runtime/Engine/Classes/Components/StaticMeshComponent.h
+++ b/Engine/Source/Runtime/Engine/Classes/Components/StaticMeshComponent.h
@@ -399,7 +399,6 @@ public:
 	virtual void ImportCustomProperties(const TCHAR* SourceText, FFeedbackContext* Warn) override;	
 	virtual void Serialize(FArchive& Ar) override;
 	virtual void PostInitProperties() override;
-	virtual void NotifyObjectReferenceEliminated() const override;
 #if WITH_EDITOR
 	virtual void PostEditUndo() override;
 	virtual void PreEditUndo() override;
diff --git a/Engine/Source/Runtime/Engine/Classes/Materials/Material.h b/Engine/Source/Runtime/Engine/Classes/Materials/Material.h
--- a/Engine/Source/Runtime/Engine/Classes/Materials/Material.h
+++ b/Engine/Source/Runtime/Engine/Classes/Materials/Material.h
@@ -1023,7 +1023,6 @@ public:
 	ENGINE_API virtual void BeginDestroy() override;
 	ENGINE_API virtual bool IsReadyForFinishDestroy() override;
 	ENGINE_API virtual void FinishDestroy() override;
-	ENGINE_API virtual void NotifyObjectReferenceEliminated() const override;
 	ENGINE_API virtual void GetResourceSizeEx(FResourceSizeEx& CumulativeResourceSize) override;
 	ENGINE_API static void AddReferencedObjects(UObject* InThis, FReferenceCollector& Collector);
 	ENGINE_API virtual bool CanBeClusterRoot() const override;
diff --git a/Engine/Source/Runtime/Engine/Classes/Materials/MaterialInstance.h b/Engine/Source/Runtime/Engine/Classes/Materials/MaterialInstance.h
--- a/Engine/Source/Runtime/Engine/Classes/Materials/MaterialInstance.h
+++ b/Engine/Source/Runtime/Engine/Classes/Materials/MaterialInstance.h
@@ -489,7 +489,6 @@ public:
 	virtual ENGINE_API bool IsReadyForFinishDestroy() override;
 	virtual ENGINE_API void FinishDestroy() override;
 	ENGINE_API static void AddReferencedObjects(UObject* InThis, FReferenceCollector& Collector);
-	ENGINE_API virtual void NotifyObjectReferenceEliminated() const override;
 #if WITH_EDITOR
 	virtual ENGINE_API void PostEditChangeProperty(FPropertyChangedEvent& PropertyChangedEvent) override;
 
diff --git a/Engine/Source/Runtime/Engine/Classes/Particles/ParticleSystemComponent.h b/Engine/Source/Runtime/Engine/Classes/Particles/ParticleSystemComponent.h
--- a/Engine/Source/Runtime/Engine/Classes/Particles/ParticleSystemComponent.h
+++ b/Engine/Source/Runtime/Engine/Classes/Particles/ParticleSystemComponent.h
@@ -1224,7 +1224,6 @@ public:
 	virtual void PostLoad() override;
 	virtual void BeginDestroy() override;
 	virtual void FinishDestroy() override;
-	virtual void NotifyObjectReferenceEliminated() const override;
 #if WITH_EDITOR
 	virtual void PreEditChange(UProperty* PropertyThatWillChange) override;
 	virtual void PostEditChangeChainProperty(FPropertyChangedChainEvent& PropertyChangedEvent) override;
diff --git a/Engine/Source/Runtime/Engine/Private/Components/SkeletalMeshComponent.cpp b/Engine/Source/Runtime/Engine/Private/Components/SkeletalMeshComponent.cpp
--- a/Engine/Source/Runtime/Engine/Private/Components/SkeletalMeshComponent.cpp
+++ b/Engine/Source/Runtime/Engine/Private/Components/SkeletalMeshComponent.cpp
@@ -335,12 +335,6 @@ void USkeletalMeshComponent::Serialize(FArchive& Ar)
 	PRAGMA_ENABLE_DEPRECATION_WARNINGS
 }
 
-void USkeletalMeshComponent::NotifyObjectReferenceEliminated() const
-{
-	UE_LOG(LogSkeletalMesh, Error, TEXT("Garbage collector eliminated reference from skeletalmeshcomponent!  SkeletalMesh objects should not be cleaned up via MarkPendingKill().\n           SkeletalMesh=%s"),
-		*GetPathName());
-}
-
 void USkeletalMeshComponent::PostLoad()
 {
 	Super::PostLoad();
diff --git a/Engine/Source/Runtime/Engine/Private/Components/StaticMeshComponent.cpp b/Engine/Source/Runtime/Engine/Private/Components/StaticMeshComponent.cpp
--- a/Engine/Source/Runtime/Engine/Private/Components/StaticMeshComponent.cpp
+++ b/Engine/Source/Runtime/Engine/Private/Components/StaticMeshComponent.cpp
@@ -347,12 +347,6 @@ void UStaticMeshComponent::PostInitProperties()
 	}
 }
 
-void UStaticMeshComponent::NotifyObjectReferenceEliminated() const
-{
-	UE_LOG(LogStaticMesh, Error, TEXT("Garbage collector eliminated reference from staticmeshcomponent!  Static Mesh objects should not be cleaned up via MarkPendingKill().\n           StaticMesh=%s"),
-		*GetPathName());
-}
-
 bool UStaticMeshComponent::AreNativePropertiesIdenticalTo( UObject* Other ) const
 {
 	bool bNativePropertiesAreIdentical = Super::AreNativePropertiesIdenticalTo( Other );
diff --git a/Engine/Source/Runtime/Engine/Private/Materials/Material.cpp b/Engine/Source/Runtime/Engine/Private/Materials/Material.cpp
--- a/Engine/Source/Runtime/Engine/Private/Materials/Material.cpp
+++ b/Engine/Source/Runtime/Engine/Private/Materials/Material.cpp
@@ -4713,12 +4713,6 @@ void UMaterial::FinishDestroy()
 	Super::FinishDestroy();
 }
 
-void UMaterial::NotifyObjectReferenceEliminated() const
-{
-	UE_LOG(LogMaterial, Error, TEXT("Garbage collector eliminated reference from material!  Material referenced objects should not be cleaned up via MarkPendingKill().\n           Material=%s"), 
-		*GetPathName());
-}
-
 void UMaterial::GetResourceSizeEx(FResourceSizeEx& CumulativeResourceSize)
 {
 	Super::GetResourceSizeEx(CumulativeResourceSize);
diff --git a/Engine/Source/Runtime/Engine/Private/Materials/MaterialInstance.cpp b/Engine/Source/Runtime/Engine/Private/Materials/MaterialInstance.cpp
--- a/Engine/Source/Runtime/Engine/Private/Materials/MaterialInstance.cpp
+++ b/Engine/Source/Runtime/Engine/Private/Materials/MaterialInstance.cpp
@@ -3255,12 +3255,6 @@ void UMaterialInstance::FinishDestroy()
 	Super::FinishDestroy();
 }
 
-void UMaterialInstance::NotifyObjectReferenceEliminated() const
-{
-	UE_LOG(LogMaterial, Error, TEXT("Garbage collector eliminated reference from material instance!  Material instance referenced objects should not be cleaned up via MarkPendingKill().\n           MI=%s\n"), 
-		*GetPathName());
-}
-
 void UMaterialInstance::AddReferencedObjects(UObject* InThis, FReferenceCollector& Collector)
 {
 	UMaterialInstance* This = CastChecked<UMaterialInstance>(InThis);
diff --git a/Engine/Source/Runtime/Engine/Private/Particles/ParticleComponents.cpp b/Engine/Source/Runtime/Engine/Private/Particles/ParticleComponents.cpp
--- a/Engine/Source/Runtime/Engine/Private/Particles/ParticleComponents.cpp
+++ b/Engine/Source/Runtime/Engine/Private/Particles/ParticleComponents.cpp
@@ -3706,13 +3706,6 @@ void UParticleSystemComponent::FinishDestroy()
 	Super::FinishDestroy();
 }
 
-void UParticleSystemComponent::NotifyObjectReferenceEliminated() const
-{
-	UE_LOG(LogParticles, Error, TEXT("Garbage collector eliminated reference from particle system!  Particle system objects should not be cleaned up via MarkPendingKill().\n           ParticleSystem=%s"),
-		*GetPathName());
-}
-
-
 void UParticleSystemComponent::GetResourceSizeEx(FResourceSizeEx& CumulativeResourceSize)
 {
 	ForceAsyncWorkCompletion(ENSURE_AND_STALL);
