There should be no functionality changes

diff --git a/Engine/Source/Runtime/CoreUObject/Public/UObject/NoExportTypes.h b/Engine/Source/Runtime/CoreUObject/Public/UObject/NoExportTypes.h
--- a/Engine/Source/Runtime/CoreUObject/Public/UObject/NoExportTypes.h
+++ b/Engine/Source/Runtime/CoreUObject/Public/UObject/NoExportTypes.h
@@ -1,13 +1,8 @@
 // Copyright 1998-2019 Epic Games, Inc. All Rights Reserved.
 
-
-
-//=============================================================================
-// Unreal base structures.
-
-// Temporary mirrors of C++ structs, used mainly as forward declarations for the core object module
-// to avoid including the full engine source. In most cases the full class definition is in another
-// file and is noted as such. More complete documentation will generally be found in those files.
+// Reflection mirrors of C++ structs defined in Core or CoreUObject, those modules are not parsed by the Unreal Header Tool.
+// The documentation comments here are only for use in the editor tooltips, and is ignored for the API docs.
+// More complete documentation will be found in the files that have the full class definition, listed below.
 
 #pragma once
 
@@ -54,7 +49,12 @@
 
 #if !CPP      //noexport class
 
-/** String search case used in UnrealString.h */
+/// @cond DOXYGEN_IGNORE
+
+/**
+ * Determines case sensitivity options for string comparisons. 
+ * @note Mirrored from Engine\Source\Runtime\Core\Public\Containers\UnrealString.h
+ */
 UENUM()
 namespace ESearchCase
 {
@@ -65,7 +65,10 @@ namespace ESearchCase
 	};
 }
 
-/** String search dir used in UnrealString.h */
+/**
+ * Determines search direction for string operations.
+ * @note Mirrored from Engine\Source\Runtime\Core\Public\Containers\UnrealString.h
+ */
 UENUM()
 namespace ESearchDir
 {
@@ -77,7 +80,8 @@ namespace ESearchDir
 }
 
 /**
- * Enum that defines how the log times are to be displayed (mirrored from OutputDevice.h).
+ * Enum that defines how the log times are to be displayed.
+ * @note Mirrored from Engine\Source\Runtime\Core\Public\Misc\OutputDevice.h
  */
 UENUM()
 namespace ELogTimes
@@ -111,22 +115,37 @@ namespace EAxis
 	};
 }
 
-// Interpolation data types.
+/** Describes shape of an interpolation curve (mirrored from InterpCurvePoint.h). */
 UENUM()
 enum EInterpCurveMode
 {
+	/** A straight line between two keypoint values. */
 	CIM_Linear UMETA(DisplayName="Linear"),
+	
+	/** A cubic-hermite curve between two keypoints, using Arrive/Leave tangents. These tangents will be automatically
+		updated when points are moved, etc.  Tangents are unclamped and will plateau at curve start and end points. */
 	CIM_CurveAuto UMETA(DisplayName="Curve Auto"),
+	
+	/** The out value is held constant until the next key, then will jump to that value. */
 	CIM_Constant UMETA(DisplayName="Constant"),
+	
+	/** A smooth curve just like CIM_Curve, but tangents are not automatically updated so you can have manual control over them (eg. in Curve Editor). */
 	CIM_CurveUser UMETA(DisplayName="Curve User"),
+	
+	/** A curve like CIM_Curve, but the arrive and leave tangents are not forced to be the same, so you can create a 'corner' at this key. */
 	CIM_CurveBreak UMETA(DisplayName="Curve Break"),
+	
+	/** A cubic-hermite curve between two keypoints, using Arrive/Leave tangents. These tangents will be automatically
+	    updated when points are moved, etc.  Tangents are clamped and will plateau at curve start and end points. */
 	CIM_CurveAutoClamped UMETA(DisplayName="Curve Auto Clamped"),
 };
 
-// @warning:	When you update this, you must add an entry to GPixelFormats(see RenderUtils.cpp)
-// @warning:	When you update this, you must add an entries to PixelFormat.h, usually just copy the generated section on the header into EPixelFormat
-// @warning:	The *Tools DLLs will also need to be recompiled if the ordering is changed, but should not need code changes.
-
+/**
+ * Describes the format of a each pixel in a graphics buffer.
+ * @warning: When you update this, you must add an entry to GPixelFormats(see RenderUtils.cpp)
+ * @warning: When you update this, you must add an entries to PixelFormat.h, usually just copy the generated section on the header into EPixelFormat
+ * @warning: The *Tools DLLs will also need to be recompiled if the ordering is changed, but should not need code changes.
+ */
 UENUM()
 enum EPixelFormat
 {
@@ -222,6 +241,7 @@ enum EPixelFormat
 	PF_MAX,
 };
 
+/** Mouse cursor types (mirrored from ICursor.h) */
 UENUM()
 namespace EMouseCursor
 {
@@ -310,15 +330,13 @@ enum class EUnit : uint8
 
 
 	/** Percentage */
-	
 	Percentage,
 
 	/** Symbolic entry, not specifiable on meta data. */
-	
 	Unspecified
 };
 
-/** A globally unique identifier. */
+/** A globally unique identifier (mirrored from Guid.h) */
 USTRUCT(immutable, noexport, BlueprintType)
 struct FGuid
 {
@@ -337,7 +355,7 @@ struct FGuid
 
 /**
  * A point or direction FVector in 3d space.
- * The full C++ class is located here: Engine\Source\Runtime\Core\Public\Math\Vector.h
+ * @note The full C++ class is located here: Engine\Source\Runtime\Core\Public\Math\Vector.h
  */
 USTRUCT(immutable, noexport, BlueprintType, meta=(HasNativeMake="Engine.KismetMathLibrary.MakeVector", HasNativeBreak="Engine.KismetMathLibrary.BreakVector"))
 struct FVector
@@ -354,7 +372,7 @@ struct FVector
 
 /**
 * A 4-D homogeneous vector.
-* The full C++ class is located here: Engine\Source\Runtime\Core\Public\Math\Vector4.h
+* @note The full C++ class is located here: Engine\Source\Runtime\Core\Public\Math\Vector4.h
 */
 USTRUCT(immutable, noexport, BlueprintType)
 struct FVector4
@@ -374,9 +392,9 @@ struct FVector4
 };
 
 /**
-* A vector in 2-D space composed of components (X, Y) with floating point precision.
-* The full C++ class is located here: Engine\Source\Runtime\Core\Public\Math\Vector2D.h
-*/
+ * A vector in 2-D space composed of components (X, Y) with floating point precision.
+ * @note The full C++ class is located here: Engine\Source\Runtime\Core\Public\Math\Vector2D.h
+ */
 USTRUCT(immutable, noexport, BlueprintType, meta=(HasNativeMake="Engine.KismetMathLibrary.MakeVector2D", HasNativeBreak="Engine.KismetMathLibrary.BreakVector2D"))
 struct FVector2D
 {
@@ -388,8 +406,7 @@ struct FVector2D
 
 };
 
-
-
+/** A pair of 3D vectors (mirrored from TwoVectors.h). */
 USTRUCT(immutable, BlueprintType, noexport)
 struct FTwoVectors
 {
@@ -398,26 +415,22 @@ struct FTwoVectors
 
 	UPROPERTY(EditAnywhere, BlueprintReadWrite, Category=TwoVectors, SaveGame)
 	FVector v2;
-
 };
 
-
 /**
  * A plane definition in 3D space.
- * The full C++ class is located here: Engine\Source\Runtime\Core\Public\Math\Plane.h
+ * @note The full C++ class is located here: Engine\Source\Runtime\Core\Public\Math\Plane.h
  */
 USTRUCT(immutable, noexport, BlueprintType)
 struct FPlane : public FVector
 {
 	UPROPERTY(EditAnywhere, BlueprintReadWrite, Category=Plane, SaveGame)
 	float W;
-
 };
 
-
 /**
  * An orthogonal rotation in 3d space.
- * The full C++ class is located here: Engine\Source\Runtime\Core\Public\Math\Rotator.h
+ * @note The full C++ class is located here: Engine\Source\Runtime\Core\Public\Math\Rotator.h
  */
 USTRUCT(immutable, noexport, BlueprintType, meta=(HasNativeMake="Engine.KismetMathLibrary.MakeRotator", HasNativeBreak="Engine.KismetMathLibrary.BreakRotator"))
 struct FRotator
@@ -436,10 +449,9 @@ struct FRotator
 
 };
 
-
 /**
  * Quaternion.
- * The full C++ class is located here: Engine\Source\Runtime\Core\Public\Math\Quat.h
+ * @note The full C++ class is located here: Engine\Source\Runtime\Core\Public\Math\Quat.h
  */
 USTRUCT(immutable, noexport, BlueprintType)
 struct FQuat
@@ -458,10 +470,9 @@ struct FQuat
 
 };
 
-
 /**
  * A packed normal.
- * The full C++ class is located here: Engine\Source\Runtime\RenderCore\Public\PackedNormal.h
+ * @note The full C++ class is located here: Engine\Source\Runtime\RenderCore\Public\PackedNormal.h
  */
 USTRUCT(immutable, noexport)
 struct FPackedNormal
@@ -481,9 +492,9 @@ struct FPackedNormal
 };
 
 /**
-* A packed basis vector.
-* The full C++ class is located here: Engine\Source\Runtime\RenderCore\Public\PackedNormal.h
-*/
+ * A packed basis vector.
+ * @note The full C++ class is located here: Engine\Source\Runtime\RenderCore\Public\PackedNormal.h
+ */
 USTRUCT(immutable, noexport)
 struct FPackedRGB10A2N
 {
@@ -492,9 +503,9 @@ struct FPackedRGB10A2N
 };
 
 /**
-* A packed vector.
-* The full C++ class is located here: Engine\Source\Runtime\RenderCore\Public\PackedNormal.h
-*/
+ * A packed vector.
+ * @note The full C++ class is located here: Engine\Source\Runtime\RenderCore\Public\PackedNormal.h
+ */
 USTRUCT(immutable, noexport)
 struct FPackedRGBA16N
 {
@@ -507,7 +518,7 @@ struct FPackedRGBA16N
 
 /**
  * Screen coordinates.
- * The full C++ class is located here: Engine\Source\Runtime\Core\Public\Math\IntPoint.h
+ * @note The full C++ class is located here: Engine\Source\Runtime\Core\Public\Math\IntPoint.h
  */
 USTRUCT(immutable, noexport, BlueprintType)
 struct FIntPoint
@@ -522,7 +533,7 @@ struct FIntPoint
 
 /**
  * An integer vector in 3D space.
- * The full C++ class is located here: Engine\Source\Runtime\Core\Public\Math\IntVector.h
+ * @note The full C++ class is located here: Engine\Source\Runtime\Core\Public\Math\IntVector.h
  */
 USTRUCT(immutable, noexport, BlueprintType)
 struct FIntVector
@@ -537,10 +548,9 @@ struct FIntVector
 	int32 Z;
 };
 
-
 /**
  * Stores a color with 8 bits of precision per channel. (BGRA).
- * The full C++ class is located here: Engine\Source\Runtime\Core\Public\Math\Color.h
+ * @note The full C++ class is located here: Engine\Source\Runtime\Core\Public\Math\Color.h
  */
 USTRUCT(immutable, noexport, BlueprintType)
 struct FColor
@@ -559,10 +569,9 @@ struct FColor
 
 };
 
-
 /**
  * A linear, 32-bit/component floating point RGBA color.
- * The full C++ class is located here: Engine\Source\Runtime\Core\Public\Math\Color.h
+ * @note The full C++ class is located here: Engine\Source\Runtime\Core\Public\Math\Color.h
  */
 USTRUCT(immutable, noexport, BlueprintType)
 struct FLinearColor
@@ -581,10 +590,9 @@ struct FLinearColor
 
 };
 
-
 /**
  * A bounding box.
- * The full C++ class is located here: Engine\Source\Runtime\Core\Public\Math\Box.h
+ * @note The full C++ class is located here: Engine\Source\Runtime\Core\Public\Math\Box.h
  */
 USTRUCT(immutable, noexport, BlueprintType, meta=(HasNativeMake="Engine.KismetMathLibrary.MakeBox"))
 struct FBox
@@ -602,7 +610,7 @@ struct FBox
 
 /**
  * A rectangular 2D Box.
- * The full C++ class is located here: Engine\Source\Runtime\Core\Public\Math\Box2D.h
+ * @note The full C++ class is located here: Engine\Source\Runtime\Core\Public\Math\Box2D.h
  */
 USTRUCT(immutable, noexport, BlueprintType, meta=(HasNativeMake="Engine.KismetMathLibrary.MakeBox2D"))
 struct FBox2D
@@ -618,20 +626,22 @@ struct FBox2D
 
 };
 
-
 /**
  * A bounding box and bounding sphere with the same origin.
- * The full C++ class is located here : Engine\Source\Runtime\Core\Public\Math\BoxSphereBounds.h
+ * @note The full C++ class is located here : Engine\Source\Runtime\Core\Public\Math\BoxSphereBounds.h
  */
 USTRUCT(noexport, BlueprintType)
 struct FBoxSphereBounds
 {
+	/** Holds the origin of the bounding box and sphere. */
 	UPROPERTY(EditAnywhere, BlueprintReadWrite, Category=BoxSphereBounds, SaveGame)
 	FVector Origin;
 
+	/** Holds the extent of the bounding box, which is half the size of the box in 3D space */
 	UPROPERTY(EditAnywhere, BlueprintReadWrite, Category=BoxSphereBounds, SaveGame)
 	FVector BoxExtent;
 
+	/** Holds the radius of the bounding sphere. */
 	UPROPERTY(EditAnywhere, BlueprintReadWrite, Category=BoxSphereBounds, SaveGame)
 	float SphereRadius;
 
@@ -639,36 +649,43 @@ struct FBoxSphereBounds
 
 /**
  * Structure for arbitrarily oriented boxes (i.e. not necessarily axis-aligned).
- * The full C++ class is located here: Engine\Source\Runtime\Core\Public\Math\OrientedBox.h
+ * @note The full C++ class is located here: Engine\Source\Runtime\Core\Public\Math\OrientedBox.h
  */
 USTRUCT(immutable, noexport)
 struct FOrientedBox
 {
+	/** Holds the center of the box. */
 	UPROPERTY(EditAnywhere, Category=OrientedBox, SaveGame)
 	FVector Center;
 
+	/** Holds the x-axis vector of the box. Must be a unit vector. */
 	UPROPERTY(EditAnywhere, Category=OrientedBox, SaveGame)
 	FVector AxisX;
 	
+	/** Holds the y-axis vector of the box. Must be a unit vector. */
 	UPROPERTY(EditAnywhere, Category=OrientedBox, SaveGame)
 	FVector AxisY;
 	
+	/** Holds the z-axis vector of the box. Must be a unit vector. */
 	UPROPERTY(EditAnywhere, Category=OrientedBox, SaveGame)
 	FVector AxisZ;
 
+	/** Holds the extent of the box along its x-axis. */
 	UPROPERTY(EditAnywhere, Category=OrientedBox, SaveGame)
 	float ExtentX;
 	
+	/** Holds the extent of the box along its y-axis. */
 	UPROPERTY(EditAnywhere, Category=OrientedBox, SaveGame)
 	float ExtentY;
 
+	/** Holds the extent of the box along its z-axis. */
 	UPROPERTY(EditAnywhere, Category=OrientedBox, SaveGame)
 	float ExtentZ;
 };
 
-/*
+/**
  * A 4x4 matrix.
- * The full C++ class is located here: Engine\Source\Runtime\Core\Public\Math\Matrix.h
+ * @note The full C++ class is located here: Engine\Source\Runtime\Core\Public\Math\Matrix.h
  */
 USTRUCT(immutable, noexport, BlueprintType)
 struct FMatrix
@@ -687,273 +704,346 @@ struct FMatrix
 
 };
 
-
-
+/** 
+ * Describes one specific point on an interpolation curve.
+ * @note This is a mirror of TInterpCurvePoint<float>, defined in InterpCurvePoint.h
+ */
 USTRUCT(noexport, BlueprintType)
 struct FInterpCurvePointFloat
 {
+	/** Float input value that corresponds to this key (eg. time). */
 	UPROPERTY(EditAnywhere, BlueprintReadWrite, Category=InterpCurvePointFloat)
 	float InVal;
 
+	/** Float output value type when input is equal to InVal. */
 	UPROPERTY(EditAnywhere, BlueprintReadWrite, Category=InterpCurvePointFloat)
 	float OutVal;
 
+	/** Tangent of curve arriving at this point. */
 	UPROPERTY(EditAnywhere, BlueprintReadWrite, Category=InterpCurvePointFloat)
 	float ArriveTangent;
 
+	/** Tangent of curve leaving this point. */
 	UPROPERTY(EditAnywhere, BlueprintReadWrite, Category=InterpCurvePointFloat)
 	float LeaveTangent;
 
+	/** Interpolation mode between this point and the next one. */
 	UPROPERTY(EditAnywhere, BlueprintReadWrite, Category=InterpCurvePointFloat)
 	TEnumAsByte<enum EInterpCurveMode> InterpMode;
 
 };
 
-
-
+/** 
+ * Describes an entire curve that is used to compute a float output value from a float input.
+ * @note This is a mirror of TInterpCurve<float>, defined in InterpCurve.h
+ */
 USTRUCT(noexport, BlueprintType)
 struct FInterpCurveFloat
 {
+	/** Holds the collection of interpolation points. */
 	UPROPERTY(EditAnywhere, BlueprintReadWrite, Category=InterpCurveFloat)
 	TArray<FInterpCurvePointFloat> Points;
 
+	/** Specify whether the curve is looped or not */
 	UPROPERTY(EditAnywhere, BlueprintReadWrite, Category=InterpCurveFloat)
 	bool bIsLooped;
 
+	/** Specify the offset from the last point's input key corresponding to the loop point */
 	UPROPERTY(EditAnywhere, BlueprintReadWrite, Category=InterpCurveFloat)
 	float LoopKeyOffset;
 };
 
-
-
+/** 
+ * Describes one specific point on an interpolation curve.
+ * @note This is a mirror of TInterpCurvePoint<FVector2D>, defined in InterpCurvePoint.h
+ */
 USTRUCT(noexport, BlueprintType)
 struct FInterpCurvePointVector2D
 {
+	/** Float input value that corresponds to this key (eg. time). */
 	UPROPERTY(EditAnywhere, BlueprintReadWrite, Category=InterpCurvePointVector2D)
 	float InVal;
 
+	/** 2D vector output value of when input is equal to InVal. */
 	UPROPERTY(EditAnywhere, BlueprintReadWrite, Category=InterpCurvePointVector2D)
 	FVector2D OutVal;
 
+	/** Tangent of curve arriving at this point. */
 	UPROPERTY(EditAnywhere, BlueprintReadWrite, Category=InterpCurvePointVector2D)
 	FVector2D ArriveTangent;
 
+	/** Tangent of curve leaving this point. */
 	UPROPERTY(EditAnywhere, BlueprintReadWrite, Category=InterpCurvePointVector2D)
 	FVector2D LeaveTangent;
 
+	/** Interpolation mode between this point and the next one. */
 	UPROPERTY(EditAnywhere, BlueprintReadWrite, Category=InterpCurvePointVector2D)
 	TEnumAsByte<enum EInterpCurveMode> InterpMode;
-
 };
 
-
-
+/** 
+ * Describes an entire curve that is used to compute a 2D vector output value from a float input.
+ * @note This is a mirror of TInterpCurve<FVector2D>, defined in InterpCurve.h
+ */
 USTRUCT(noexport, BlueprintType)
 struct FInterpCurveVector2D
 {
+	/** Holds the collection of interpolation points. */
 	UPROPERTY(EditAnywhere, BlueprintReadWrite, Category=InterpCurveVector2D)
 	TArray<FInterpCurvePointVector2D> Points;
 
+	/** Specify whether the curve is looped or not */
 	UPROPERTY(EditAnywhere, BlueprintReadWrite, Category=InterpCurveVector2D)
 	bool bIsLooped;
 
+	/** Specify the offset from the last point's input key corresponding to the loop point */
 	UPROPERTY(EditAnywhere, BlueprintReadWrite, Category=InterpCurveVector2D)
 	float LoopKeyOffset;
 };
 
-
-
+/** 
+ * Describes one specific point on an interpolation curve.
+ * @note This is a mirror of TInterpCurvePoint<FVector>, defined in InterpCurvePoint.h
+ */
 USTRUCT(noexport, BlueprintType)
 struct FInterpCurvePointVector
 {
+	/** Float input value that corresponds to this key (eg. time). */
 	UPROPERTY(EditAnywhere, BlueprintReadWrite, Category=InterpCurvePointVector)
 	float InVal;
 
+	/** 3D vector output value of when input is equal to InVal. */
 	UPROPERTY(EditAnywhere, BlueprintReadWrite, Category=InterpCurvePointVector)
 	FVector OutVal;
 
+	/** Tangent of curve arriving at this point. */
 	UPROPERTY(EditAnywhere, BlueprintReadWrite, Category=InterpCurvePointVector)
 	FVector ArriveTangent;
 
+	/** Tangent of curve leaving this point. */
 	UPROPERTY(EditAnywhere, BlueprintReadWrite, Category=InterpCurvePointVector)
 	FVector LeaveTangent;
 
+	/** Interpolation mode between this point and the next one. */
 	UPROPERTY(EditAnywhere, BlueprintReadWrite, Category=InterpCurvePointVector)
 	TEnumAsByte<enum EInterpCurveMode> InterpMode;
-
 };
 
-
-
+/** 
+ * Describes an entire curve that is used to compute a 3D vector output value from a float input.
+ * @note This is a mirror of TInterpCurve<FVector>, defined in InterpCurve.h
+ */
 USTRUCT(noexport, BlueprintType)
 struct FInterpCurveVector
 {
+	/** Holds the collection of interpolation points. */
 	UPROPERTY(EditAnywhere, BlueprintReadWrite, Category=InterpCurveVector)
 	TArray<FInterpCurvePointVector> Points;
 
+	/** Specify whether the curve is looped or not */
 	UPROPERTY(EditAnywhere, BlueprintReadWrite, Category=InterpCurveVector)
 	bool bIsLooped;
 
+	/** Specify the offset from the last point's input key corresponding to the loop point */
 	UPROPERTY(EditAnywhere, BlueprintReadWrite, Category=InterpCurveVector)
 	float LoopKeyOffset;
 };
 
-
-
+/**
+ * Describes one specific point on an interpolation curve.
+ * @note This is a mirror of TInterpCurvePoint<FQuat>, defined in InterpCurvePoint.h
+ */
 USTRUCT(noexport, BlueprintType)
 struct FInterpCurvePointQuat
 {
+	/** Float input value that corresponds to this key (eg. time). */
 	UPROPERTY(EditAnywhere, BlueprintReadWrite, Category=InterpCurvePointQuat)
 	float InVal;
 
+	/** Quaternion output value of when input is equal to InVal. */
 	UPROPERTY(EditAnywhere, BlueprintReadWrite, Category=InterpCurvePointQuat)
 	FQuat OutVal;
 
+	/** Tangent of curve arriving at this point. */
 	UPROPERTY(EditAnywhere, BlueprintReadWrite, Category=InterpCurvePointQuat)
 	FQuat ArriveTangent;
 
+	/** Tangent of curve leaving this point. */
 	UPROPERTY(EditAnywhere, BlueprintReadWrite, Category=InterpCurvePointQuat)
 	FQuat LeaveTangent;
 
+	/** Interpolation mode between this point and the next one. */
 	UPROPERTY(EditAnywhere, BlueprintReadWrite, Category=InterpCurvePointQuat)
 	TEnumAsByte<enum EInterpCurveMode> InterpMode;
-
 };
 
-
-
+/**
+ * Describes an entire curve that is used to compute a quaternion output value from a float input.
+ * @note This is a mirror of TInterpCurve<FQuat>, defined in InterpCurve.h
+ */
 USTRUCT(noexport, BlueprintType)
 struct FInterpCurveQuat
 {
+	/** Holds the collection of interpolation points. */
 	UPROPERTY(EditAnywhere, BlueprintReadWrite, Category=InterpCurveQuat)
 	TArray<FInterpCurvePointQuat> Points;
 
+	/** Specify whether the curve is looped or not */
 	UPROPERTY(EditAnywhere, BlueprintReadWrite, Category=InterpCurveQuat)
 	bool bIsLooped;
 
+	/** Specify the offset from the last point's input key corresponding to the loop point */
 	UPROPERTY(EditAnywhere, BlueprintReadWrite, Category=InterpCurveQuat)
 	float LoopKeyOffset;
 };
 
-
-
+/**
+ * Describes one specific point on an interpolation curve.
+ * @note This is a mirror of TInterpCurvePoint<FTwoVectors>, defined in InterpCurvePoint.h
+ */
 USTRUCT(noexport, BlueprintType)
 struct FInterpCurvePointTwoVectors
 {
+	/** Float input value that corresponds to this key (eg. time). */
 	UPROPERTY(EditAnywhere, BlueprintReadWrite, Category=InterpCurvePointTwoVectors)
 	float InVal;
 
+	/** Two 3D vectors output value of when input is equal to InVal. */
 	UPROPERTY(EditAnywhere, BlueprintReadWrite, Category=InterpCurvePointTwoVectors)
 	FTwoVectors OutVal;
 
+	/** Tangent of curve arriving at this point. */
 	UPROPERTY(EditAnywhere, BlueprintReadWrite, Category=InterpCurvePointTwoVectors)
 	FTwoVectors ArriveTangent;
 
+	/** Tangent of curve leaving this point. */
 	UPROPERTY(EditAnywhere, BlueprintReadWrite, Category=InterpCurvePointTwoVectors)
 	FTwoVectors LeaveTangent;
 
+	/** Interpolation mode between this point and the next one. */
 	UPROPERTY(EditAnywhere, BlueprintReadWrite, Category=InterpCurvePointTwoVectors)
 	TEnumAsByte<enum EInterpCurveMode> InterpMode;
-
 };
 
-
-
+/**
+ * Describes an entire curve that is used to compute two 3D vector values from a float input.
+ * @note This is a mirror of TInterpCurve<FTwoVectors>, defined in InterpCurve.h
+ */
 USTRUCT(noexport, BlueprintType)
 struct FInterpCurveTwoVectors
 {
+	/** Holds the collection of interpolation points. */
 	UPROPERTY(EditAnywhere, BlueprintReadWrite, Category=InterpCurveTwoVectors)
 	TArray<FInterpCurvePointTwoVectors> Points;
 
+	/** Specify whether the curve is looped or not */
 	UPROPERTY(EditAnywhere, BlueprintReadWrite, Category=InterpCurveTwoVectors)
 	bool bIsLooped;
 
+	/** Specify the offset from the last point's input key corresponding to the loop point */
 	UPROPERTY(EditAnywhere, BlueprintReadWrite, Category=InterpCurveTwoVectors)
 	float LoopKeyOffset;
 };
 
-
-
+/**
+ * Describes one specific point on an interpolation curve.
+ * @note This is a mirror of TInterpCurvePoint<FLinearColor>, defined in InterpCurvePoint.h
+ */
 USTRUCT(noexport, BlueprintType)
 struct FInterpCurvePointLinearColor
 {
+	/** Float input value that corresponds to this key (eg. time). */
 	UPROPERTY(EditAnywhere, BlueprintReadWrite, Category=InterpCurvePointLinearColor)
 	float InVal;
 
+	/** Color output value of when input is equal to InVal. */
 	UPROPERTY(EditAnywhere, BlueprintReadWrite, Category=InterpCurvePointLinearColor)
 	FLinearColor OutVal;
 
+	/** Tangent of curve arriving at this point. */
 	UPROPERTY(EditAnywhere, BlueprintReadWrite, Category=InterpCurvePointLinearColor)
 	FLinearColor ArriveTangent;
 
+	/** Tangent of curve leaving this point. */
 	UPROPERTY(EditAnywhere, BlueprintReadWrite, Category=InterpCurvePointLinearColor)
 	FLinearColor LeaveTangent;
 
+	/** Interpolation mode between this point and the next one. */
 	UPROPERTY(EditAnywhere, BlueprintReadWrite, Category=InterpCurvePointLinearColor)
 	TEnumAsByte<enum EInterpCurveMode> InterpMode;
-
 };
 
-
-
+/**
+ * Describes an entire curve that is used to compute a color output value from a float input.
+ * @note This is a mirror of TInterpCurve<FLinearColor>, defined in InterpCurve.h
+ */
 USTRUCT(noexport, BlueprintType)
 struct FInterpCurveLinearColor
 {
+	/** Holds the collection of interpolation points. */
 	UPROPERTY(EditAnywhere, BlueprintReadWrite, Category=InterpCurveLinearColor)
 	TArray<FInterpCurvePointLinearColor> Points;
 
+	/** Specify whether the curve is looped or not */
 	UPROPERTY(EditAnywhere, BlueprintReadWrite, Category=InterpCurveLinearColor)
 	bool bIsLooped;
 
+	/** Specify the offset from the last point's input key corresponding to the loop point */
 	UPROPERTY(EditAnywhere, BlueprintReadWrite, Category=InterpCurveLinearColor)
 	float LoopKeyOffset;
 };
 
-
 /**
  * Transform composed of Quat/Translation/Scale.
- * The full C++ class is located here: Engine\Source\Runtime\Core\Public\Math\Transform.h
+ * @note This is implemented in either TransformVectorized.h or TransformNonVectorized.h depending on the platform.
  */
 USTRUCT(noexport, BlueprintType, meta=(HasNativeMake="Engine.KismetMathLibrary.MakeTransform", HasNativeBreak="Engine.KismetMathLibrary.BreakTransform"))
 struct FTransform
 {
+	/** Rotation of this transformation, as a quaternion. */
 	UPROPERTY(EditAnywhere, BlueprintReadWrite, Category=Transform, SaveGame)
 	FQuat Rotation;
 
+	/** Translation of this transformation, as a vector. */
 	UPROPERTY(EditAnywhere, BlueprintReadWrite, Category=Transform, SaveGame)
 	FVector Translation;
 
+	/** 3D scale (always applied in local space) as a vector. */
 	UPROPERTY(EditAnywhere, BlueprintReadWrite, Category=Transform, SaveGame, meta=(MakeStructureDefaultValue = "1,1,1"))
 	FVector Scale3D;
-
 };
 
-
 /**
- * Thread-safe RNG.
- * The full C++ class is located here: Engine\Source\Runtime\Core\Public\Math\RandomStream.h
+ * Thread-safe random number generator that can be manually seeded.
+ * @note The full C++ class is located here: Engine\Source\Runtime\Core\Public\Math\RandomStream.h
  */
 USTRUCT(noexport, BlueprintType, meta = (HasNativeMake = "Engine.KismetMathLibrary.MakeRandomStream", HasNativeBreak = "Engine.KismetMathLibrary.BreakRandomStream"))
 struct FRandomStream
 {
 public:
+	/** Holds the initial seed. */
 	UPROPERTY(EditAnywhere, BlueprintReadWrite, Category=RandomStream, SaveGame)
 	int32 InitialSeed;
 	
+	/** Holds the current seed. */
 	UPROPERTY()
 	int32 Seed;
 };
 
-/** A date/time value. */
-
+/** 
+ * A value representing a specific point date and time over a wide range of years.
+ * @note The full C++ class is located here: Engine\Source\Runtime\Core\Public\Misc\DateTime.h
+ */
 USTRUCT(immutable, noexport, BlueprintType, meta=(HasNativeMake="Engine.KismetMathLibrary.MakeDateTime", HasNativeBreak="Engine.KismetMathLibrary.BreakDateTime"))
 struct FDateTime
 {
 	int64 Ticks;
 };
 
-/** A frame number value. */
-
+/** 
+ * A frame number value, representing discrete frames since the start of timing.
+ * @note The full C++ class is located here: Engine\Source\Runtime\Core\Public\Misc\FrameNumber.h
+ */
 USTRUCT(noexport, BlueprintType)
 struct FFrameNumber
 {
@@ -961,8 +1051,10 @@ struct FFrameNumber
 	int32 Value;
 };
 
-/** A frame rate represented as a fraction comprising 2 integers: a numerator (number of frames), and a denominator (per second). */
-
+/** 
+ * A frame rate represented as a fraction comprising 2 integers: a numerator (number of frames), and a denominator (per second).
+ * @note The full C++ class is located here: Engine\Source\Runtime\Core\Public\Misc\FrameRate.h
+ */
 USTRUCT(noexport, BlueprintType, meta=(HasNativeMake="Engine.KismetMathLibrary.MakeFrameRate", HasNativeBreak="Engine.KismetMathLibrary.BreakFrameRate"))
 struct FFrameRate
 {
@@ -975,32 +1067,42 @@ struct FFrameRate
 	int32 Denominator;
 };
 
-/** Represents a time by a context-free frame number, plus a sub frame value in the range [0:1). */
-
+/** 
+ * Represents a time by a context-free frame number, plus a sub frame value in the range [0:1). 
+ * @note The full C++ class is located here: Engine\Source\Runtime\Core\Public\Misc\FrameTime.h
+ */
 USTRUCT(noexport, BlueprintType)
 struct FFrameTime
 {
+	/** Count of frames from start of timing */
 	UPROPERTY(BlueprintReadWrite, Category=FrameTime)
 	FFrameNumber FrameNumber;
 
+	/** Time within a frame, always between >= 0 and < 1 */
 	UPROPERTY(BlueprintReadWrite, Category=FrameTime)
 	float SubFrame;
 };
 
-/** A frame time qualified by a frame rate context. */
-
+/** 
+ * A frame time qualified by a frame rate context.
+ * @note The full C++ class is located here: Engine\Source\Runtime\Core\Public\Misc\QualifiedFrameTime.h
+ */
 USTRUCT(noexport, BlueprintType, meta=(ScriptName="QualifiedTime", HasNativeMake="Engine.KismetMathLibrary.MakeQualifiedFrameTime", HasNativeBreak="Engine.KismetMathLibrary.BreakQualifiedFrameTime"))
 struct FQualifiedFrameTime
 {
+	/** The frame time */
 	UPROPERTY(BlueprintReadWrite, Category=QualifiedFrameTime)
 	FFrameTime Time;
 
+	/** The rate that this frame time is in */
 	UPROPERTY(BlueprintReadWrite, Category=QualifiedFrameTime)
 	FFrameRate Rate;
 };
 
-/** A timecode that stores time in HH:MM:SS format with the remainder of time represented by an integer frame count. */
-
+/** 
+ * A timecode that stores time in HH:MM:SS format with the remainder of time represented by an integer frame count. 
+ * @note The full C++ class is located here: Engine\Source\Runtime\Core\Public\Misc\TimeCode.h
+ */
 USTRUCT(noexport, BlueprintType)
 struct FTimecode
 {
@@ -1021,15 +1123,20 @@ struct FTimecode
 	bool bDropFrameFormat;
 };
 
-/** A time span value. */
-
+/** 
+ * A time span value, which is the difference between two dates and times.
+ * @note The full C++ class is located here: Engine\Source\Runtime\Core\Public\Misc\Timespan.h
+ */
 USTRUCT(immutable, noexport, BlueprintType, meta=(HasNativeMake="Engine.KismetMathLibrary.MakeTimespan", HasNativeBreak="Engine.KismetMathLibrary.BreakTimespan"))
 struct FTimespan
 {
 	int64 Ticks;
 };
 
-/** An object path, this is saved as a name/string pair */
+/** 
+ * A struct that contains a string reference to an object, either a top level asset or a subobject.
+ * @note The full C++ class is located here: Engine\Source\Runtime\CoreUObject\Public\UObject\SoftObjectPath.h
+ */
 USTRUCT(noexport, BlueprintType, meta=(HasNativeMake="Engine.KismetSystemLibrary.MakeSoftObjectPath", HasNativeBreak="Engine.KismetSystemLibrary.BreakSoftObjectPath"))
 struct FSoftObjectPath
 {
@@ -1042,47 +1149,53 @@ struct FSoftObjectPath
 	FString SubPathString;
 };
 
-/** Subclass of FSoftObjectPath that is only valid to use with UClass* */
+/** 
+ * A struct that contains a string reference to a class, can be used to make soft references to classes.
+ * @note The full C++ class is located here: Engine\Source\Runtime\CoreUObject\Public\UObject\SoftObjectPath.h
+ */
 USTRUCT(noexport, BlueprintType, meta=(HasNativeMake="Engine.KismetSystemLibrary.MakeSoftClassPath", HasNativeBreak="Engine.KismetSystemLibrary.BreakSoftClassPath"))
 struct FSoftClassPath : public FSoftObjectPath
 {
 };
 
-/** A type of primary asset, this is a wrapper around FName and can be cast back and forth */
+/** 
+ * A type of primary asset, used by the Asset Manager system.
+ * @note The full C++ class is located here: Engine\Source\Runtime\CoreUObject\Public\UObject\PrimaryAssetId.h
+ */
 USTRUCT(noexport, BlueprintType)
 struct FPrimaryAssetType
 {
-	/** The Type of this object, by default it's base class's name */
+	/** The Type of this object, by default its base class's name */
 	UPROPERTY(EditAnywhere, SaveGame, BlueprintReadWrite, Category = PrimaryAssetType)
 	FName Name;
 };
 
-/** This identifies an object as a "primary" asset that can be searched for by the AssetManager and used in various tools */
+/** 
+ * This identifies an object as a "primary" asset that can be searched for by the AssetManager and used in various tools
+ * @note The full C++ class is located here: Engine\Source\Runtime\CoreUObject\Public\UObject\PrimaryAssetId.h
+ */
 USTRUCT(noexport, BlueprintType)
 struct FPrimaryAssetId
 {
-	/** The Type of this object, by default it's base class's name */
+	/** The Type of this object, by default its base class's name */
 	UPROPERTY(EditAnywhere, SaveGame, BlueprintReadWrite, Category = PrimaryAssetId)
 	FPrimaryAssetType PrimaryAssetType;
 
-	/** The Name of this object, by default it's short name */
+	/** The Name of this object, by default its short name */
 	UPROPERTY(EditAnywhere, SaveGame, BlueprintReadWrite, Category = PrimaryAssetId)
 	FName PrimaryAssetName;
 };
 
-// A struct used as stub for deleted ones.
-
+/** A struct used as stub for deleted ones. */
 USTRUCT(noexport)
 struct FFallbackStruct  
 {
 };
 
+/** Enumerates the valid types of range bounds (mirrored from RangeBound.h) */
 UENUM(BlueprintType)
 namespace ERangeBoundTypes
 {
-	/**
-	* Enumerates the valid types of range bounds.
-	*/
 	enum Type
 	{
 		/**
@@ -1102,80 +1215,103 @@ namespace ERangeBoundTypes
 	};
 }
 
-// A float range bound
-
+/**
+ * Defines a single bound for a range of values.
+ * @note This is a mirror of TRangeBound<float>, defined in RangeBound.h
+ */
 USTRUCT(noexport, BlueprintType)
 struct FFloatRangeBound
 {
+	/** Holds the type of the bound. */
 	UPROPERTY(EditAnywhere, BlueprintReadWrite, Category=Range)
 	TEnumAsByte<ERangeBoundTypes::Type> Type;
 
+	/** Holds the bound's value. */
 	UPROPERTY(EditAnywhere, BlueprintReadWrite, Category=Range)
 	float Value;
 };
 
-// A float range
-
+/**
+ * A contiguous set of floats described by lower and upper bound values.
+ * @note This is a mirror of TRange<float>, defined in Range.h
+ */
 USTRUCT(noexport, BlueprintType)
 struct FFloatRange
 {
+	/** Holds the range's lower bound. */
 	UPROPERTY(EditAnywhere, BlueprintReadWrite, Category=Range)
 	FFloatRangeBound LowerBound;
 
+	/** Holds the range's upper bound. */
 	UPROPERTY(EditAnywhere, BlueprintReadWrite, Category=Range)
 	FFloatRangeBound UpperBound;
 };
 
-// An int32 range bound
-
+/**
+ * Defines a single bound for a range of values.
+ * @note This is a mirror of TRangeBound<int32>, defined in RangeBound.h
+ */
 USTRUCT(noexport, BlueprintType)
 struct FInt32RangeBound
 {
+	/** Holds the type of the bound. */
 	UPROPERTY(EditAnywhere, BlueprintReadWrite, Category = Range)
 	TEnumAsByte<ERangeBoundTypes::Type> Type;
 
+	/** Holds the bound's value. */
 	UPROPERTY(EditAnywhere, BlueprintReadWrite, Category = Range)
 	int32 Value;
 };
 
-// An int32 range
-
+/**
+ * A contiguous set of floats described by lower and upper bound values.
+ * @note This is a mirror of TRange<int32>, defined in Range.h
+ */
 USTRUCT(noexport, BlueprintType)
 struct FInt32Range
 {
+	/** Holds the range's lower bound. */
 	UPROPERTY(EditAnywhere, BlueprintReadWrite, Category = Range)
 	FInt32RangeBound LowerBound;
 
+	/** Holds the range's upper bound. */
 	UPROPERTY(EditAnywhere, BlueprintReadWrite, Category = Range)
 	FInt32RangeBound UpperBound;
 };
 
-// A float interval
-
+/**
+ * An interval of floats, defined by inclusive min and max values
+ * @note This is a mirror of TInterval<float>, defined in Interval.h
+ */
 USTRUCT(noexport)
 struct FFloatInterval
 {
+	/** Values must be >= Min */
 	UPROPERTY(EditAnywhere, Category=Interval)
 	float Min;
 
+	/** Values must be <= Max */
 	UPROPERTY(EditAnywhere, Category=Interval)
 	float Max;
 };
 
-// An int32 interval
-
+/**
+ * An interval of integers, defined by inclusive min and max values
+ * @note This is a mirror of TInterval<int32>, defined in Interval.h
+ */
 USTRUCT(noexport)
 struct FInt32Interval
 {
+	/** Values must be >= Min */
 	UPROPERTY(EditAnywhere, Category=Interval)
 	int32 Min;
 
+	/** Values must be <= Max */
 	UPROPERTY(EditAnywhere, Category=Interval)
 	int32 Max;
 };
 
-// Localization
-
+/** Categories of localized text (mirrored in LocalizedTextSourceTypes.h */
 UENUM(BlueprintType)
 enum class ELocalizedTextSourceCategory : uint8
 {
@@ -1184,6 +1320,10 @@ enum class ELocalizedTextSourceCategory : uint8
 	Editor,
 };
 
+/**
+ * Polyglot data that may be registered to the text localization manager at runtime.
+ * @note This struct is mirrored in PolyglotTextData.h
+ */
 USTRUCT(noexport, BlueprintType)
 struct FPolyglotTextData
 {
@@ -1243,8 +1383,7 @@ struct FPolyglotTextData
 	FText CachedText;
 };
 
-// Automation
-
+/** Report level of automation events (mirrored in AutomationEvent.h). */
 UENUM()
 enum class EAutomationEventType : uint8
 {
@@ -1253,6 +1392,7 @@ enum class EAutomationEventType : uint8
 	Error
 };
 
+/** Event emitted by automation system (mirrored in AutomationEvent.h). */
 USTRUCT(noexport)
 struct FAutomationEvent
 {
@@ -1269,6 +1409,7 @@ struct FAutomationEvent
 	FGuid Artifact;
 };
 
+/** Information about the execution of an automation task (mirrored in AutomationEvent.h). */
 USTRUCT(noexport)
 struct FAutomationExecutionEntry
 {
@@ -1285,29 +1426,18 @@ struct FAutomationExecutionEntry
 	FDateTime Timestamp;
 };
 
-//=============================================================================
 /**
- * Object: The base class all objects.
- * This is a built-in Unreal class and it shouldn't be modified by mod authors.
- * The full C++ class is located here: Engine\Source\Runtime\CoreUObject\Public\UObject\UObject.h
+ * Direct base class for all UE4 objects
+ * @note The full C++ class is located here: Engine\Source\Runtime\CoreUObject\Public\UObject\Object.h
  */
-//=============================================================================
-
 UCLASS(abstract, noexport)
 class UObject
 {
 	GENERATED_BODY()
 public:
 
-	/**
-	 * Default UObject constructor.
-	 */
 	UObject(const FObjectInitializer& ObjectInitializer = FObjectInitializer::Get());
-	/** DO NOT USE. This constructor is for internal usage only for hot-reload purposes. */
 	UObject(FVTableHelper& Helper);
-
-	//=============================================================================
-	// K2 support functions.
 	
 	/**
 	 * Executes some portion of the ubergraph.
@@ -1318,5 +1448,7 @@ public:
 	void ExecuteUbergraph(int32 EntryPoint);
 };
 
+/// @endcond
+
 #endif
 
diff --git a/Engine/Source/Runtime/CoreUObject/Public/UObject/Object.h b/Engine/Source/Runtime/CoreUObject/Public/UObject/Object.h
--- a/Engine/Source/Runtime/CoreUObject/Public/UObject/Object.h
+++ b/Engine/Source/Runtime/CoreUObject/Public/UObject/Object.h
@@ -2,6 +2,10 @@
 
 #pragma once
 
+/*=============================================================================
+	Object.h: Direct base class for all UE4 objects
+=============================================================================*/
+
 #include "CoreMinimal.h"
 #include "UObject/Script.h"
 #include "UObject/ObjectMacros.h"
@@ -20,51 +24,63 @@ struct FPropertyChangedChainEvent;
 
 DECLARE_LOG_CATEGORY_EXTERN(LogObj, Log, All);
 
-
+/** Parameter enum for CastChecked() function, defines when it will check/assert */
 namespace ECastCheckedType
 {
 	enum Type
 	{
+		/** Null is okay, only assert on incorrect type */
 		NullAllowed,
+		/** Null is not allowed, assert on incorrect type or null */
 		NullChecked
 	};
 };
 
-//
-// The base class of all objects.
-//
+/** 
+ * The base class of all UE4 objects. The type of an object is defined by its UClass.
+ * This provides support functions for creating and using objects, and virtual functions that should be overridden in child classes.
+ * 
+ * @see https://docs.unrealengine.com/en-us/Programming/UnrealArchitecture/Objects
+ */
 class COREUOBJECT_API UObject : public UObjectBaseUtility
 {
-	// Declarations.
+	// Declarations, normally created by UnrealHeaderTool boilerplate code
 	DECLARE_CLASS(UObject,UObject,CLASS_Abstract|CLASS_NoExport|CLASS_Intrinsic|CLASS_MatchedSerializers,CASTCLASS_None,TEXT("/Script/CoreUObject"),NO_API)
 	DEFINE_DEFAULT_OBJECT_INITIALIZER_CONSTRUCTOR_CALL(UObject)
+	typedef UObject WithinClass;
 	static UObject* __VTableCtorCaller(FVTableHelper& Helper)
 	{
 		return new (EC_InternalUseOnlyConstructor, (UObject*)GetTransientPackage(), NAME_None, RF_NeedLoad | RF_ClassDefaultObject | RF_TagGarbageTemp) UObject(Helper);
 	}
+	static const TCHAR* StaticConfigName() 
+	{
+		return TEXT("Engine");
+	}
+	static void StaticRegisterNativesUObject() 
+	{
+	}
 
-	typedef UObject WithinClass;
-	static const TCHAR* StaticConfigName() {return TEXT("Engine");}
-
-	// Constructors and destructors.
+	/** Default constructor */
 	UObject();
+
+	/** Deprecated constructor, ObjectInitializer is no longer needed but is supported for older classes. */
 	UObject(const FObjectInitializer& ObjectInitializer);
-	UObject( EStaticConstructor, EObjectFlags InFlags );
+
+	/** DO NOT USE. This constructor is for internal usage only for statically-created objects. */
+	UObject(EStaticConstructor, EObjectFlags InFlags);
+
 	/** DO NOT USE. This constructor is for internal usage only for hot-reload purposes. */
 	UObject(FVTableHelper& Helper);
 
-	static void StaticRegisterNativesUObject() 
-	{
-	}
-
+	/** Utility function for templates below */
 	UObject* CreateDefaultSubobject(FName SubobjectFName, UClass* ReturnType, UClass* ClassToCreateByDefault, bool bIsRequired, bool bAbstract, bool bIsTransient);
 
 	/**
-	* Create a component or subobject only to be used with the editor.
-	* @param	TReturnType					class of return type, all overrides must be of this type
-	* @param	SubobjectName				name of the new component
-	* @param	bTransient					true if the component is being assigned to a transient property. This does not make the component itself transient, but does stop it from inheriting parent defaults
-	*/
+	 * Create a component or subobject only to be used with the editor. They will be stripped out in packaged builds.
+	 * @param	TReturnType					Class of return type, all overrides must be of this type
+	 * @param	SubobjectName				Name of the new component
+	 * @param	bTransient					True if the component is being assigned to a transient property. This does not make the component itself transient, but does stop it from inheriting parent defaults
+	 */
 	template<class TReturnType>
 	TReturnType* CreateEditorOnlyDefaultSubobject(FName SubobjectName, bool bTransient = false)
 	{
@@ -73,11 +89,11 @@ class COREUOBJECT_API UObject : public UObjectBaseUtility
 	}
 
 	/**
-	* Create a component or subobject
-	* @param	TReturnType					class of return type, all overrides must be of this type
-	* @param	SubobjectName				name of the new component
-	* @param	bTransient					true if the component is being assigned to a transient property. This does not make the component itself transient, but does stop it from inheriting parent defaults
-	*/
+	 * Create a component or subobject.
+	 * @param	TReturnType					Class of return type, all overrides must be of this type
+	 * @param	SubobjectName				Name of the new component
+	 * @param	bTransient					True if the component is being assigned to a transient property. This does not make the component itself transient, but does stop it from inheriting parent defaults
+	 */
 	template<class TReturnType>
 	TReturnType* CreateDefaultSubobject(FName SubobjectName, bool bTransient = false)
 	{
@@ -86,12 +102,12 @@ class COREUOBJECT_API UObject : public UObjectBaseUtility
 	}
 	
 	/**
-	* Create a component or subobject
-	* @param TReturnType class of return type, all overrides must be of this type
-	* @param TClassToConstructByDefault class to construct by default
-	* @param SubobjectName name of the new component
-	* @param bTransient		true if the component is being assigned to a transient property. This does not make the component itself transient, but does stop it from inheriting parent defaults
-	*/
+	 * Create a component or subobject, allows creating a child class and returning the parent class.
+	 * @param	TReturnType					Class of return type, all overrides must be of this type
+	 * @param	TClassToConstructByDefault	Class of object to actually construct
+	 * @param	SubobjectName				Name of the new component
+	 * @param	bTransient					True if the component is being assigned to a transient property. This does not make the component itself transient, but does stop it from inheriting parent defaults
+	 */
 	template<class TReturnType, class TClassToConstructByDefault>
 	TReturnType* CreateDefaultSubobject(FName SubobjectName, bool bTransient = false)
 	{
@@ -99,12 +115,12 @@ class COREUOBJECT_API UObject : public UObjectBaseUtility
 	}
 	
 	/**
-	* Create optional component or subobject. Optional subobjects may not get created
-	* when a derived class specified DoNotCreateDefaultSubobject with the subobject's name.
-	* @param	TReturnType					class of return type, all overrides must be of this type
-	* @param	SubobjectName				name of the new component
-	* @param	bTransient					true if the component is being assigned to a transient property. This does not make the component itself transient, but does stop it from inheriting parent defaults
-	*/
+	 * Create an optional component or subobject. Optional subobjects may not get created.
+	 * when a derived class specified DoNotCreateDefaultSubobject with the subobject's name.
+	 * @param	TReturnType					Class of return type, all overrides must be of this type
+	 * @param	SubobjectName				Name of the new component
+	 * @param	bTransient					True if the component is being assigned to a transient property. This does not make the component itself transient, but does stop it from inheriting parent defaults
+	 */
 	template<class TReturnType>
 	TReturnType* CreateOptionalDefaultSubobject(FName SubobjectName, bool bTransient = false)
 	{
@@ -113,12 +129,11 @@ class COREUOBJECT_API UObject : public UObjectBaseUtility
 	}
 	
 	/**
-	* Create optional component or subobject. Optional subobjects may not get created
-	* when a derived class specified DoNotCreateDefaultSubobject with the subobject's name.
-	* @param	TReturnType					class of return type, all overrides must be of this type
-	* @param	SubobjectName				name of the new component
-	* @param	bTransient					true if the component is being assigned to a transient property. This does not make the component itself transient, but does stop it from inheriting parent defaults
-	*/
+	 * Create a subobject that has the Abstract class flag, child classes are expected to override this by calling CreateDEfaultObject with the same name and a non-abstract class.
+	 * @param	TReturnType					Class of return type, all overrides must be of this type
+	 * @param	SubobjectName				Name of the new component
+	 * @param	bTransient					True if the component is being assigned to a transient property. This does not make the component itself transient, but does stop it from inheriting parent defaults
+	 */
 	template<class TReturnType>
 	TReturnType* CreateAbstractDefaultSubobject(FName SubobjectName, bool bTransient = false)
 	{
@@ -127,30 +142,32 @@ class COREUOBJECT_API UObject : public UObjectBaseUtility
 	}
 
 	/**
-	* Gets all default subobjects associated with this object instance.
-	* @param	OutDefaultSubobjects	Array containing all default subobjects of this object.
-	*/
+	 * Gets all default subobjects associated with this object instance.
+	 * @param	OutDefaultSubobjects	Array containing all default subobjects of this object.
+	 */
 	void GetDefaultSubobjects(TArray<UObject*>& OutDefaultSubobjects);
 
 	/**
-	* Finds a subobject associated with this object instance by its name
-	* @param	OutDefaultSubobjects	Array containing all default subobjects of this object.
-	*/
+	 * Finds a subobject associated with this object instance by its name
+	 * @param	Name	Object name to look for
+	 */
 	UObject* GetDefaultSubobjectByName(FName ToFind);
 
-	//==========================================
-	// UObject interface.
-	//==========================================	
+	/*----------------------------------
+			UObject interface
+	----------------------------------*/
+
 protected:
     /** 
-     * This function actually does the work for the GetDetailInfo and is virtual.  
+     * This function actually does the work for the GetDetailedInfo() and is virtual.  
      * It should only be called from GetDetailedInfo as GetDetailedInfo is safe to call on NULL object pointers
-     **/
+     */
 	virtual FString GetDetailedInfoInternal() const { return TEXT("No_Detailed_Info_Specified"); }
+
 public:
 	/**
 	 * Called after the C++ constructor and after the properties have been initialized, including those loaded from config.
-	 * mainly this is to emulate some behavior of when the constructor was called after the properties were initialized.
+	 * This is called before any serialization or other setup has happened.
 	 */
 	virtual void PostInitProperties();
 
@@ -163,20 +180,20 @@ public:
 	}
 
 	/**
-	 * Called from within SavePackage on the passed in base/ root. The return value of this function will be passed to
-	 * PostSaveRoot. This is used to allow objects used as base to perform required actions before saving and cleanup
-	 * afterwards.
+	 * Called from within SavePackage on the passed in base/root object. The return value of this function will be passed to PostSaveRoot. 
+	 * This is used to allow objects used as a base to perform required actions before saving and cleanup afterwards.
 	 * @param Filename: Name of the file being saved to (includes path)
-	 *
+
 	 * @return	Whether PostSaveRoot needs to perform internal cleanup
 	 */
 	virtual bool PreSaveRoot(const TCHAR* Filename)
 	{
 		return false;
 	}
+
 	/**
-	 * Called from within SavePackage on the passed in base/ root. This function is being called after the package
-	 * has been saved and can perform cleanup.
+	 * Called from within SavePackage on the passed in base/root object. 
+	 * This function is called after the package has been saved and can perform cleanup.
 	 *
 	 * @param	bCleanupIsRequired	Whether PreSaveRoot dirtied state that needs to be cleaned up
 	 */
@@ -201,9 +218,7 @@ public:
 	 */
 	virtual bool Modify( bool bAlwaysMarkDirty=true );
 
-	/** 
-	 * Utility to allow overrides of Modify to avoid doing work if the base class is not going modify anyways.
-	 */
+	/** Utility to allow overrides of Modify to avoid doing work if this object cannot be safely modified */
 	bool CanModify() const;
 
 #if WITH_EDITOR
@@ -214,8 +229,8 @@ public:
 #endif
 
 	/** 
-	 * Do any object-specific cleanup required immediately after loading an object, 
-	 * and immediately after any undo/redo.
+	 * Do any object-specific cleanup required immediately after loading an object.
+	 * This is not called for newly-created objects, and by default will always execute on the game thread.
 	 */
 	virtual void PostLoad();
 
@@ -251,15 +266,18 @@ public:
 	/**
 	 * Called to finish destroying the object.  After UObject::FinishDestroy is called, the object's memory should no longer be accessed.
 	 *
-	 * note: because properties are destroyed here, Super::FinishDestroy() should always be called at the end of your child class's
-	 * FinishDestroy() method, rather than at the beginning.
+	 * @warning Because properties are destroyed here, Super::FinishDestroy() should always be called at the end of your child class's FinishDestroy() method, rather than at the beginning.
 	 */
 	virtual void FinishDestroy();
 
-	/** UObject serializer. */
+	/** 
+	 * Handles reading, writing, and reference collecting using FArchive.
+	 * This implementation handles all UProperty serialization, but can be overridden for native variables.
+	 */
 	virtual void Serialize(FArchive& Ar);
 	virtual void Serialize(FStructuredArchive::FRecord Record);
 
+	/** After a critical error, perform any mission-critical cleanup, such as restoring the video mode orreleasing hardware resources. */
 	virtual void ShutdownAfterError() {}
 
 	/** 
@@ -271,7 +289,7 @@ public:
 
 #if WITH_EDITOR
 	/** 
-	 * This is called when property is about to be modified by InterpPropertyTracks
+	 * This is called when a property is about to be modified externally
 	 *
 	 * @param PropertyThatWillChange	Property that will be changed
 	 */
@@ -353,7 +371,7 @@ private:
 public:
 #endif // WITH_EDITOR
 
-	// @todo document
+	/** Called at the end of Rename(), but only if the rename was actually carried out */
 	virtual void PostRename(UObject* OldOuter, const FName OldName) {}
 	
 	/**
@@ -361,7 +379,6 @@ public:
 	 * Note: NOT called on components on actor duplication (alt-drag or copy-paste).  Use PostEditImport as well to cover that case.
 	 */
 	virtual void PostDuplicate(bool bDuplicateForPIE) {}
-
 	virtual void PostDuplicate(EDuplicateMode::Type DuplicateMode) 
 	{
 		PostDuplicate(DuplicateMode == EDuplicateMode::PIE);
@@ -392,10 +409,10 @@ public:
 	virtual bool NeedsLoadForTargetPlatform(const class ITargetPlatform* TargetPlatform) const;
 
 	/**
-	 * Called during saving to determine the load flags to save with the object.
+	 * Called during saving to include this object in client/servers running in editor builds, even if they wouldn't normally be.
 	 * If false, this object will still get loaded if NeedsLoadForServer/Client are true
 	 * 
-	 * @return	true if this object should always be loaded for editor game
+	 * @return	true if this object should always be loaded for games running in editor builds
 	 */
 	virtual bool NeedsLoadForEditorGame() const
 	{
@@ -455,25 +472,25 @@ public:
 	static bool CanCreateInCurrentContext(UObject* Template);
 
 	/**
-	* Exports the property values for the specified object as text to the output device. Required for Copy&Paste
-	* Most objects don't need this as unreal script can handle most cases.
-	*
-	* @param	Out				the output device to send the exported text to
-	* @param	Indent			number of spaces to prepend to each line of output
-	*
-	* see also: ImportCustomProperties()
-	*/
+	 * Exports the property values for the specified object as text to the output device.
+	 * Override this if you need custom support for copy/paste.
+	 *
+	 * @param	Out				The output device to send the exported text to
+	 * @param	Indent			Number of spaces to prepend to each line of output
+	 *
+	 * @see ImportCustomProperties()
+	 */
 	virtual void ExportCustomProperties(FOutputDevice& Out, uint32 Indent)	{}
 
 	/**
-	* Exports the property values for the specified object as text to the output device. Required for Copy&Paste
-	* Most objects don't need this as unreal script can handle most cases.
-	*
-	* @param	SourceText		the input data (zero terminated), must not be 0
-	* @param	Warn			for error reporting, must not be 0
-	*
-	* see also: ExportCustomProperties()
-	*/
+	 * Exports the property values for the specified object as text to the output device. Required for Copy&Paste
+	 * Override this if you need custom support for copy/paste.
+	 *
+	 * @param	SourceText		The input data (zero terminated), will never be null
+	 * @param	Warn			For error reporting, will never be null
+	 *
+	 * @see ExportCustomProperties()
+	 */
 	virtual void ImportCustomProperties(const TCHAR* SourceText, FFeedbackContext* Warn) {}
 
 	/**
@@ -488,16 +505,30 @@ public:
 	 */
 	virtual void PostReloadConfig( class UProperty* PropertyThatWasLoaded ) {}
 
-	/** Rename this object to a unique name.*/
-	virtual bool Rename( const TCHAR* NewName=NULL, UObject* NewOuter=NULL, ERenameFlags Flags=REN_None );
-
+	/** 
+	 * Rename this object to a unique name, or change its outer.
+	 * @warning Unless ForceNoResetLoaders is passed in, this will cause a flush of all level streaming.
+	 * 
+	 * @param	NewName		The new name of the object, if null then NewOuter should be set
+	 * @param	NewOuter	New Outer this object will be placed within, if null it will use the current outer
+	 * @param	Flags		Flags to specify what happens during the rename
+	 */
+	virtual bool Rename(const TCHAR* NewName=nullptr, UObject* NewOuter=nullptr, ERenameFlags Flags=REN_None);
 
-	/** @return a one line description of an object for viewing in the thumbnail view of the generic browser */
+	/** Return a one line description of an object for viewing in the thumbnail view of the generic browser */
 	virtual FString GetDesc() { return TEXT( "" ); }
 
 #if WITH_ENGINE
+	/** 
+	 * Returns what UWorld this object is contained within. 
+	 * By default this will follow its Outer chain, but it should be overridden if that will not work.
+	 */
 	virtual class UWorld* GetWorld() const;
+
+	/** Internal function used by UEngine::GetWorldFromContextObject() */
 	class UWorld* GetWorldChecked(bool& bSupported) const;
+
+	/** Checks to see if GetWorld() is implemented on a specific class */
 	bool ImplementsGetWorld() const;
 #endif
 
@@ -588,12 +619,7 @@ public:
 		return true;
 	}
 
-	/**
-	 * Gathers a list of asset registry searchable tags which are name/value pairs with some type information
-	 * This only needs to be implemented for asset objects
-	 *
-	 * @param	OutTags		A list of key-value pairs associated with this object and their types
-	 */
+	/** Struct used by GetAssetRegistryTags() to return tag info */
 	struct FAssetRegistryTag
 	{
 		/** Enum specifying the type of this tag */
@@ -626,9 +652,16 @@ public:
 			TD_Memory = 1<<3,
 		};
 		
+		/** Logical name of this tag */
 		FName Name;
+
+		/** Value string for this tag, may represent any data type */
 		FString Value;
+
+		/** Broad description of kind of data represented in Value */
 		ETagType Type;
+
+		/** Flags describing more detail for displaying in the UI */
 		uint32 DisplayFlags;
 
 		FAssetRegistryTag(FName InName, const FString& InValue, ETagType InType, uint32 InDisplayFlags = TD_None)
@@ -640,6 +673,13 @@ public:
 		/** Returns true if this FName is a special UStruct that should be exported even if not tagged, with the struct name as the tag name */
 		COREUOBJECT_API static bool IsUniqueAssetRegistryTagStruct(FName StructName, ETagType& TagType);
 	};
+
+	/**
+	 * Gathers a list of asset registry searchable tags which are name/value pairs with some type information
+	 * This only needs to be implemented for asset objects
+	 *
+	 * @param	OutTags		A list of key-value pairs associated with this object and their types
+	 */
 	virtual void GetAssetRegistryTags(TArray<FAssetRegistryTag>& OutTags) const;
 
 	/** Get the common tag name used for all asset source file import paths */
@@ -750,11 +790,9 @@ public:
 	virtual EDataValidationResult IsDataValid(TArray<FText>& ValidationErrors);
 #endif // WITH_EDITOR
 
-	/**
- 	 *******************************************************
-	 * Non virtual functions, not intended to be overridden
-	 *******************************************************
-	 */
+	/*----------------------------------------------------------
+		Non virtual functions, not intended to be overridden
+	----------------------------------------------------------*/
 
 	/**
 	 * Test the selection state of a UObject
@@ -822,6 +860,8 @@ public:
 	 * asynchronous cleanup process.
 	 */
 	bool ConditionalBeginDestroy();
+
+	/** Called when an object is actually destroyed, memory should never be accessed again */
 	bool ConditionalFinishDestroy();
 	
 	/** PostLoad if needed. */
@@ -835,6 +875,7 @@ public:
 	 *								subobjects and components for a subobject root.
 	 */
 	void ConditionalPostLoadSubobjects( struct FObjectInstancingGraph* OuterInstanceGraph=NULL );
+
 #if WITH_EDITOR
 	/**
 	 * Starts caching of platform specific data for the target platform
@@ -885,7 +926,10 @@ public:
 	 */
 	inline bool IsBasedOnArchetype( const UObject* const SomeObject ) const;
 
+	/** Returns a UFunction with the specified name, wrapper for UClass::FindFunctionByName() */
 	UFunction* FindFunction( FName InName ) const;
+	
+	/** Version of FindFunction() that will assert if the function was not found */
 	UFunction* FindFunctionChecked( FName InName ) const;
 
 	/**
@@ -914,9 +958,8 @@ public:
 	bool CheckDefaultSubobjects(bool bForceCheck = false) const;
 
 	/**
-	 * Save configuration.
-	 * @warning: Must be safe on class-default metaobjects.
-	 * !!may benefit from hierarchical propagation, deleting keys that match superclass...not sure what's best yet.
+	 * Save configuration out to ini files
+	 * @warning Must be safe to call on class-default object
 	 */
 	void SaveConfig( uint64 Flags=CPF_Config, const TCHAR* Filename=NULL, FConfigCacheIni* Config=GConfig );
 
@@ -948,7 +991,6 @@ private:
 	void EnsureNotRetrievingVTablePtr() const;
 
 public:
-	
 	/**
 	 * Get the default config filename for the specified UObject
 	 */
@@ -981,7 +1023,7 @@ public:
 
 	/**
 	 * Wrapper method for LoadConfig that is used when reloading the config data for objects at runtime which have already loaded their config data at least once.
-	 * Allows the objects the receive a callback that it's configuration data has been reloaded.
+	 * Allows the objects the receive a callback that its configuration data has been reloaded.
 	 *
 	 * @param	Class				the class to use for determining which section of the ini to retrieve text values from
 	 * @param	Filename			indicates the filename to load values from; if not specified, uses ConfigClass's ClassConfigName
@@ -1001,7 +1043,7 @@ public:
 	 */
 	void OutputReferencers( FOutputDevice& Ar, FReferencerInformationList* Referencers=NULL );
 	
-	// @todo document
+	/** Called by OutputReferencers() to get the internal list of referencers to write */
 	void RetrieveReferencers( TArray<FReferencerInformation>* OutInternalReferencers, TArray<FReferencerInformation>* OutExternalReferencers);
 
 	/**
@@ -1046,14 +1088,12 @@ public:
 	template<class T>
 	bool Implements() const;
 
-	//==========================================
-	// Kismet Virtual Machine
-	//==========================================
-	//@todo UE4 - do these really need to be virtual to support actors?
 
-	//
-	// Script processing functions.
-	//
+	/*-----------------------------
+			Virtual Machine
+	-----------------------------*/
+
+	/** Called by VM to execute a UFunction with a filled in UStruct of parameters */
 	virtual void ProcessEvent( UFunction* Function, void* Parms );
 
 	/**
@@ -1082,16 +1122,16 @@ public:
 		return false;
 	}
 
-	/** Command line. */
+	/** Handle calling a function by name when executed from the console or a command line */
 	bool CallFunctionByNameWithArguments( const TCHAR* Cmd, FOutputDevice& Ar, UObject* Executor, bool bForceCallWithNonExec = false );
 
-	// Call a function
+	/** Internal VM method for executing a function */
 	void CallFunction( FFrame& Stack, RESULT_DECL, UFunction* Function );
 
-	//
-	// Internal function call processing.
-	// @warning: might not write anything to Result if proper type isn't returned.
-	//
+	/**
+	 * Internal function call processing.
+	 * @warning: might not write anything to Result if proper type isn't returned.
+	 */
 	DECLARE_FUNCTION(ProcessInternal);
 
 	/**
@@ -1126,6 +1166,7 @@ public:
 	 */
 	void DestroyNonNativeProperties();
 
+	/** Called during subobject creation to mark this component as editor only, which causes it to get stripped in packaged builds */
 	virtual void MarkAsEditorOnlySubobject() { }
 
 	/**
@@ -1133,13 +1174,11 @@ public:
 	 */
 	void AbortInsideMemberFunction() const;
 
-	// UnrealScript intrinsics.
+	// UnrealScript intrinsics, do not call directly
 
 	// Undefined native handler
 	DECLARE_FUNCTION(execUndefined);
 
-	
-	// @todo document below declared functions
 	// Variables
 	DECLARE_FUNCTION(execLocalVariable);
 	DECLARE_FUNCTION(execInstanceVariable);
@@ -1152,7 +1191,6 @@ public:
 	DECLARE_FUNCTION(execClassDefaultVariable);
 	DECLARE_FUNCTION(execEndFunctionParms);
 
-
 	// Do Nothing 
 	DECLARE_FUNCTION(execNothing);
 	DECLARE_FUNCTION(execNothingOp4a);
@@ -1229,7 +1267,6 @@ public:
 	DECLARE_FUNCTION(execStructCmpNe);
 	DECLARE_FUNCTION(execStructMember);
 
-	// @todo delegate: Delegate comparison is not supported for multi-cast delegates
 	DECLARE_FUNCTION(execEqualEqual_DelegateDelegate);
 	DECLARE_FUNCTION(execNotEqual_DelegateDelegate);
 	DECLARE_FUNCTION(execEqualEqual_DelegateFunction);
@@ -1247,7 +1284,6 @@ public:
 	DECLARE_FUNCTION(execObjectConst);
 	DECLARE_FUNCTION(execSoftObjectConst);
 
-	// @todo delegate: Multi-cast versions needed for script execution!		(Need Add, Remove, Clear/Empty)
 	DECLARE_FUNCTION(execInstanceDelegate);
 	DECLARE_FUNCTION(execNameConst);
 	DECLARE_FUNCTION(execByteConst);
@@ -1314,11 +1350,13 @@ public:
 
 	DECLARE_FUNCTION(execArrayGetByRef);
 
-	// -- K2 support functions
+	/** Wrapper struct to hold the entrypoint in the right memory address */
 	struct Object_eventExecuteUbergraph_Parms
 	{
 		int32 EntryPoint;
 	};
+
+	/** Execute the ubergraph from a specific entry point */
 	void ExecuteUbergraph(int32 EntryPoint)
 	{
 		Object_eventExecuteUbergraph_Parms Parms;
@@ -1327,7 +1365,6 @@ public:
 	}
 
 protected: 
-
 	/** Checks it's ok to perform subobjects check at this time. */
 	bool CanCheckDefaultSubObjects(bool bForceCheck, bool& bResult) const;
 
diff --git a/Engine/Source/Runtime/CoreUObject/Public/UObject/ObjectMacros.h b/Engine/Source/Runtime/CoreUObject/Public/UObject/ObjectMacros.h
--- a/Engine/Source/Runtime/CoreUObject/Public/UObject/ObjectMacros.h
+++ b/Engine/Source/Runtime/CoreUObject/Public/UObject/ObjectMacros.h
@@ -1,7 +1,7 @@
 // Copyright 1998-2019 Epic Games, Inc. All Rights Reserved.
 
 /*=============================================================================
-	ObjectBase.h: Unreal object base class.
+	ObjectMacros.h: Helper macros and defines for UObject system
 =============================================================================*/
 
 #pragma once
@@ -52,84 +52,78 @@ typedef	uint64 ScriptPointerType;
 	Core enumerations.
 -----------------------------------------------------------------------------*/
 
-//
-// Flags for loading objects.
-//
+/** Flags for loading objects, used by LoadObject() and related functions and passed as a uint32 */
 enum ELoadFlags
 {
-	LOAD_None						= 0x00000000,	// No flags.
-	LOAD_Async						= 0x00000001,	// Loads the package using async loading path/ reader
-	LOAD_NoWarn						= 0x00000002,	// Don't display warning if load fails.
-	LOAD_EditorOnly					= 0x00000004, // Load for editor-only purposes and by editor-only code
-	LOAD_ResolvingDeferredExports	= 0x00000008,	// Denotes that we should not defer export loading (as we're resolving them)
-	LOAD_Verify						= 0x00000010,	// Only verify existance; don't actually load.
-	LOAD_AllowDll					= 0x00000020,	// Allow plain DLLs.
+	LOAD_None						= 0x00000000,	///< No flags.
+	LOAD_Async						= 0x00000001,	///< Loads the package using async loading path/ reader
+	LOAD_NoWarn						= 0x00000002,	///< Don't display warning if load fails.
+	LOAD_EditorOnly					= 0x00000004,	///< Load for editor-only purposes and by editor-only code
+	LOAD_ResolvingDeferredExports	= 0x00000008,	///< Denotes that we should not defer export loading (as we're resolving them)
+	LOAD_Verify						= 0x00000010,	///< Only verify existance; don't actually load.
+	LOAD_AllowDll					= 0x00000020,	///< Allow plain DLLs.
 //	LOAD_Unused						= 0x00000040
-	LOAD_NoVerify					= 0x00000080,   // Don't verify imports yet.
-	LOAD_IsVerifying				= 0x00000100,		// Is verifying imports
+	LOAD_NoVerify					= 0x00000080,   ///< Don't verify imports yet.
+	LOAD_IsVerifying				= 0x00000100,	///< Is verifying imports
 //	LOAD_Unused						= 0x00000200,
 //	LOAD_Unused						= 0x00000400,
 //	LOAD_Unused						= 0x00000800,
-	LOAD_DisableDependencyPreloading = 0x00001000,	// Bypass dependency preloading system
-	LOAD_Quiet						= 0x00002000,   // No log warnings.
-	LOAD_FindIfFail					= 0x00004000,	// Tries FindObject if a linker cannot be obtained (e.g. package is currently being compiled)
-	LOAD_MemoryReader				= 0x00008000,	// Loads the file into memory and serializes from there.
-	LOAD_NoRedirects				= 0x00010000,	// Never follow redirects when loading objects; redirected loads will fail
-	LOAD_ForDiff					= 0x00020000,	// Loading for diffing.
-	LOAD_PackageForPIE				= 0x00080000,   // This package is being loaded for PIE, it must be flagged as such immediately
-	LOAD_DeferDependencyLoads       = 0x00100000,   // Do not load external (blueprint) dependencies (instead, track them for deferred loading)
-	LOAD_ForFileDiff				= 0x00200000,	// Load the package (not for diffing in the editor), instead verify at the two packages serialized output are the same, if they are not then debug break so that you can get the callstack and object information
-	LOAD_DisableCompileOnLoad		= 0x00400000,	// Prevent this load call from running compile on load for the loaded blueprint (intentionally not recursive, dependencies will still compile on load)
+	LOAD_DisableDependencyPreloading = 0x00001000,	///< Bypass dependency preloading system
+	LOAD_Quiet						= 0x00002000,   ///< No log warnings.
+	LOAD_FindIfFail					= 0x00004000,	///< Tries FindObject if a linker cannot be obtained (e.g. package is currently being compiled)
+	LOAD_MemoryReader				= 0x00008000,	///< Loads the file into memory and serializes from there.
+	LOAD_NoRedirects				= 0x00010000,	///< Never follow redirects when loading objects; redirected loads will fail
+	LOAD_ForDiff					= 0x00020000,	///< Loading for diffing in the editor
+	LOAD_PackageForPIE				= 0x00080000,   ///< This package is being loaded for PIE, it must be flagged as such immediately
+	LOAD_DeferDependencyLoads       = 0x00100000,   ///< Do not load external (blueprint) dependencies (instead, track them for deferred loading)
+	LOAD_ForFileDiff				= 0x00200000,	///< Load the package (not for diffing in the editor), instead verify at the two packages serialized output are the same, if they are not then debug break so that you can get the callstack and object information
+	LOAD_DisableCompileOnLoad		= 0x00400000,	///< Prevent this load call from running compile on load for the loaded blueprint (intentionally not recursive, dependencies will still compile on load)
 };
 
-//
-// Flags for saving packages
-//
+/** Flags for saving objects/packages, passed into UPackage::SavePackage() as a uint32 */
 enum ESaveFlags
 {
-	SAVE_None			= 0x00000000,	// No flags
-	SAVE_NoError		= 0x00000001,	// Don't generate errors on save
-	SAVE_FromAutosave	= 0x00000002,   // Used to indicate this save was initiated automatically
-	SAVE_KeepDirty		= 0x00000004,	// Do not clear the dirty flag when saving
-	SAVE_KeepGUID		= 0x00000008,	// Keep the same guid, used to save cooked packages
-	SAVE_Async			= 0x00000010,	// Save to a memory writer, then actually write to disk async
-	SAVE_Unversioned	= 0x00000020,	// Save all versions as zero. Upon load this is changed to the current version. This is only reasonable to use with full cooked builds for distribution.
-	SAVE_CutdownPackage	= 0x00000040,	// Saving cutdown packages in a temp location WITHOUT renaming the package.
-	SAVE_KeepEditorOnlyCookedPackages = 0x00000080,  // keep packages which are marked as editor only even though we are cooking
-	SAVE_Concurrent		= 0x00000100,	// We are save packages in multiple threads at once and should not call non-threadsafe functions or rely on globals. GIsSavingPackage should be set and PreSave/Postsave functions should be called before/after the entire concurrent save.
-	SAVE_DiffOnly       = 0x00000200, // Serializes the package to a special memory archive that performs a diff with an existing file on disk
-	SAVE_DiffCallstack  = 0x00000400, // Serializes the package to a special memory archive that compares all differences against a file on disk and dumps relevant callstacks
-	SAVE_ComputeHash    = 0x00000800, // Compute the MD5 hash of the cooked data
+	SAVE_None			= 0x00000000,	///< No flags
+	SAVE_NoError		= 0x00000001,	///< Don't generate errors on save
+	SAVE_FromAutosave	= 0x00000002,   ///< Used to indicate this save was initiated automatically
+	SAVE_KeepDirty		= 0x00000004,	///< Do not clear the dirty flag when saving
+	SAVE_KeepGUID		= 0x00000008,	///< Keep the same guid, used to save cooked packages
+	SAVE_Async			= 0x00000010,	///< Save to a memory writer, then actually write to disk async
+	SAVE_Unversioned	= 0x00000020,	///< Save all versions as zero. Upon load this is changed to the current version. This is only reasonable to use with full cooked builds for distribution.
+	SAVE_CutdownPackage	= 0x00000040,	///< Saving cutdown packages in a temp location WITHOUT renaming the package.
+	SAVE_KeepEditorOnlyCookedPackages = 0x00000080,  ///< Keep packages which are marked as editor only even though we are cooking
+	SAVE_Concurrent		= 0x00000100,	///< We are save packages in multiple threads at once and should not call non-threadsafe functions or rely on globals. GIsSavingPackage should be set and PreSave/Postsave functions should be called before/after the entire concurrent save.
+	SAVE_DiffOnly       = 0x00000200,	///< Serializes the package to a special memory archive that performs a diff with an existing file on disk
+	SAVE_DiffCallstack  = 0x00000400,	///< Serializes the package to a special memory archive that compares all differences against a file on disk and dumps relevant callstacks
+	SAVE_ComputeHash    = 0x00000800,	///< Compute the MD5 hash of the cooked data
 };
 
-//
-// Package flags.
-//
+/** Package flags, passed into UPackage::SetPackageFlags and related functions */
 enum EPackageFlags
 {
-	PKG_None						= 0x00000000,	// No flags
-	PKG_NewlyCreated				= 0x00000001,	// Newly created package, not saved yet. In editor only.
-	PKG_ClientOptional				= 0x00000002,	// Purely optional for clients.
-	PKG_ServerSideOnly				= 0x00000004,   // Only needed on the server side.
-	PKG_CompiledIn					= 0x00000010,   // This package is from "compiled in" classes.
-	PKG_ForDiffing					= 0x00000020,	// This package was loaded just for the purposes of diffing
-	PKG_EditorOnly					= 0x00000040, // This is editor-only package (for example: editor module script package)
-	PKG_Developer					= 0x00000080,	// Developer module
+	PKG_None						= 0x00000000,	///< No flags
+	PKG_NewlyCreated				= 0x00000001,	///< Newly created package, not saved yet. In editor only.
+	PKG_ClientOptional				= 0x00000002,	///< Purely optional for clients.
+	PKG_ServerSideOnly				= 0x00000004,   ///< Only needed on the server side.
+	PKG_CompiledIn					= 0x00000010,   ///< This package is from "compiled in" classes.
+	PKG_ForDiffing					= 0x00000020,	///< This package was loaded just for the purposes of diffing
+	PKG_EditorOnly					= 0x00000040,	///< This is editor-only package (for example: editor module script package)
+	PKG_Developer					= 0x00000080,	///< Developer module
 //	PKG_Unused						= 0x00000100,
 //	PKG_Unused						= 0x00000200,
 //	PKG_Unused						= 0x00000400,
 //	PKG_Unused						= 0x00000800,
 //	PKG_Unused						= 0x00001000,
 //	PKG_Unused						= 0x00002000,
-	PKG_ContainsMapData				= 0x00004000,   // Contains map data (UObjects only referenced by a single ULevel) but is stored in a different package
-	PKG_Need						= 0x00008000,	// Client needs to download this package.
-	PKG_Compiling					= 0x00010000,	// package is currently being compiled
-	PKG_ContainsMap					= 0x00020000,	// Set if the package contains a ULevel/ UWorld object
-	PKG_RequiresLocalizationGather		= 0x00040000,	// Set if the package contains any data to be gathered by localization
-	PKG_DisallowLazyLoading			= 0x00080000,	// Set if the archive serializing this package cannot use lazy loading
-	PKG_PlayInEditor				= 0x00100000,	// Set if the package was created for the purpose of PIE
-	PKG_ContainsScript				= 0x00200000,	// Package is allowed to contain UClass objects
-	PKG_DisallowExport				= 0x00400000,	// Editor should not export asset in this package
+	PKG_ContainsMapData				= 0x00004000,   ///< Contains map data (UObjects only referenced by a single ULevel) but is stored in a different package
+	PKG_Need						= 0x00008000,	///< Client needs to download this package.
+	PKG_Compiling					= 0x00010000,	///< package is currently being compiled
+	PKG_ContainsMap					= 0x00020000,	///< Set if the package contains a ULevel/ UWorld object
+	PKG_RequiresLocalizationGather	= 0x00040000,	///< Set if the package contains any data to be gathered by localization
+	PKG_DisallowLazyLoading			= 0x00080000,	///< Set if the archive serializing this package cannot use lazy loading
+	PKG_PlayInEditor				= 0x00100000,	///< Set if the package was created for the purpose of PIE
+	PKG_ContainsScript				= 0x00200000,	///< Package is allowed to contain UClass objects
+	PKG_DisallowExport				= 0x00400000,	///< Editor should not export asset in this package
 //	PKG_Unused						= 0x00800000,
 //	PKG_Unused						= 0x01000000,	
 //	PKG_Unused						= 0x02000000,	
@@ -137,17 +131,17 @@ enum EPackageFlags
 //	PKG_Unused						= 0x08000000,	
 //	PKG_Unused						= 0x10000000,	
 //	PKG_Unused						= 0x20000000,
-	PKG_ReloadingForCooker			= 0x40000000,   // this package is reloading in the cooker, try to avoid getting data we will never need. We won't save this package.
-	PKG_FilterEditorOnly			= 0x80000000,	// Package has editor-only data filtered
+	PKG_ReloadingForCooker			= 0x40000000,   ///< This package is reloading in the cooker, try to avoid getting data we will never need. We won't save this package.
+	PKG_FilterEditorOnly			= 0x80000000,	///< Package has editor-only data filtered out
 };
 
-#define PKG_InMemoryOnly	(EPackageFlags)(PKG_CompiledIn | PKG_NewlyCreated) // Flag mask that indicates if this package is a package that exists in memory only.
+/** Flag mask that indicates if this package is a package that exists in memory only. */
+#define PKG_InMemoryOnly	(EPackageFlags)(PKG_CompiledIn | PKG_NewlyCreated)
 
 ENUM_CLASS_FLAGS(EPackageFlags);
 
-//
 // Internal enums.
-//
+
 enum EStaticConstructor				{EC_StaticConstructor};
 enum EInternal						{EC_InternalUseOnlyConstructor};
 enum ECppProperty					{EC_CppProperty};
@@ -166,14 +160,12 @@ public:
 /** Empty API definition.  Used as a placeholder parameter when no DLL export/import API is needed for a UObject class */
 #define NO_API
 
-
 /**
  * Flags describing a class.
  */
 enum EClassFlags
 {
-	/** @name Base flags */
-	//@{
+	/** No Flags */
 	CLASS_None				  = 0x00000000u,
 	/** Class is abstract and can't be instantiated directly. */
 	CLASS_Abstract            = 0x00000001u,
@@ -247,27 +239,23 @@ enum EClassFlags
 	CLASS_ConfigDoNotCheckDefaults = 0x40000000u,
 	/** Class has been consigned to oblivion as part of a blueprint recompile, and a newer version currently exists. */
 	CLASS_NewerVersionExists  = 0x80000000u,
-
-	//@}
 };
 
 // Declare bitwise operators to allow EClassFlags to be combined but still retain type safety
 ENUM_CLASS_FLAGS(EClassFlags);
 
-/** @name Flags to inherit from base class */
-//@{
+/** Flags to inherit from base class */
 #define CLASS_Inherit ((EClassFlags)(CLASS_Transient | CLASS_DefaultConfig | CLASS_Config | CLASS_PerObjectConfig | CLASS_ConfigDoNotCheckDefaults | CLASS_NotPlaceable \
 						| CLASS_Const | CLASS_HasInstancedReference | CLASS_Deprecated | CLASS_DefaultToInstanced | CLASS_GlobalUserConfig))
 
-/** these flags will be cleared by the compiler when the class is parsed during script compilation */
+/** These flags will be cleared by the compiler when the class is parsed during script compilation */
 #define CLASS_RecompilerClear ((EClassFlags)(CLASS_Inherit | CLASS_Abstract | CLASS_NoExport | CLASS_Native | CLASS_Intrinsic | CLASS_TokenStreamAssembled))
 
-/** these flags will be cleared by the compiler when the class is parsed during script compilation */
+/** These flags will be cleared by the compiler when the class is parsed during script compilation */
 #define CLASS_ShouldNeverBeLoaded ((EClassFlags)(CLASS_Native | CLASS_Intrinsic | CLASS_TokenStreamAssembled))
 
-/** these flags will be inherited from the base class only for non-intrinsic classes */
+/** These flags will be inherited from the base class only for non-intrinsic classes */
 #define CLASS_ScriptInherit ((EClassFlags)(CLASS_Inherit | CLASS_EditInlineNew | CLASS_CollapseCategories))
-//@}
 
 /** This is used as a mask for the flags put into generated code for "compiled in" classes. */
 #define CLASS_SaveInCompiledInClasses ((EClassFlags)(\
@@ -361,98 +349,99 @@ enum EClassCastFlags : uint64
 ENUM_CLASS_FLAGS(EClassCastFlags)
 
 
-//
-// Flags associated with each property in a class, overriding the
-// property's default behavior.
-// NOTE: When adding one here, please update ParsePropertyFlags
-//
+/**
+ * Flags associated with each property in a class, overriding the
+ * property's default behavior.
+ * @warning When adding one here, please update ParsePropertyFlags()
+ */
 enum EPropertyFlags : uint64
 {
 	CPF_None = 0,
 
-	CPF_Edit							= 0x0000000000000001,	// Property is user-settable in the editor.
-	CPF_ConstParm						= 0x0000000000000002,	// This is a constant function parameter
-	CPF_BlueprintVisible				= 0x0000000000000004,	// This property can be read by blueprint code
-	CPF_ExportObject					= 0x0000000000000008,	// Object can be exported with actor.
-	CPF_BlueprintReadOnly				= 0x0000000000000010,	// This property cannot be modified by blueprint code
-	CPF_Net								= 0x0000000000000020,	// Property is relevant to network replication.
-	CPF_EditFixedSize					= 0x0000000000000040,	// Indicates that elements of an array can be modified, but its size cannot be changed.
-	CPF_Parm							= 0x0000000000000080,	// Function/When call parameter.
-	CPF_OutParm							= 0x0000000000000100,	// Value is copied out after function call.
-	CPF_ZeroConstructor					= 0x0000000000000200,	// memset is fine for construction
-	CPF_ReturnParm						= 0x0000000000000400,	// Return value.
-	CPF_DisableEditOnTemplate			= 0x0000000000000800,	// Disable editing of this property on an archetype/sub-blueprint
-	//CPF_      						= 0x0000000000001000,	// 
-	CPF_Transient   					= 0x0000000000002000,	// Property is transient: shouldn't be saved or loaded, except for Blueprint CDOs.
-	CPF_Config      					= 0x0000000000004000,	// Property should be loaded/saved as permanent profile.
-	//CPF_								= 0x0000000000008000,	// 
-	CPF_DisableEditOnInstance			= 0x0000000000010000,	// Disable editing on an instance of this class
-	CPF_EditConst   					= 0x0000000000020000,	// Property is uneditable in the editor.
-	CPF_GlobalConfig					= 0x0000000000040000,	// Load config from base class, not subclass.
-	CPF_InstancedReference				= 0x0000000000080000,	// Property is a component references.
-	//CPF_								= 0x0000000000100000,
-	CPF_DuplicateTransient				= 0x0000000000200000,	// Property should always be reset to the default value during any type of duplication (copy/paste, binary duplication, etc.)
-	CPF_SubobjectReference				= 0x0000000000400000,	// Property contains subobject references (TSubobjectPtr)
-	//CPF_    							= 0x0000000000800000,	// 
-	CPF_SaveGame						= 0x0000000001000000,	// Property should be serialized for save games
-	CPF_NoClear							= 0x0000000002000000,	// Hide clear (and browse) button.
-	//CPF_  							= 0x0000000004000000,	//
-	CPF_ReferenceParm					= 0x0000000008000000,	// Value is passed by reference; CPF_OutParam and CPF_Param should also be set.
-	CPF_BlueprintAssignable				= 0x0000000010000000,	// MC Delegates only.  Property should be exposed for assigning in blueprint code
-	CPF_Deprecated  					= 0x0000000020000000,	// Property is deprecated.  Read it from an archive, but don't save it.
-	CPF_IsPlainOldData					= 0x0000000040000000,	// If this is set, then the property can be memcopied instead of CopyCompleteValue / CopySingleValue
-	CPF_RepSkip							= 0x0000000080000000,	// Not replicated. For non replicated properties in replicated structs 
-	CPF_RepNotify						= 0x0000000100000000,	// Notify actors when a property is replicated
-	CPF_Interp							= 0x0000000200000000,	// interpolatable property for use with matinee
-	CPF_NonTransactional				= 0x0000000400000000,	// Property isn't transacted
-	CPF_EditorOnly						= 0x0000000800000000,	// Property should only be loaded in the editor
-	CPF_NoDestructor					= 0x0000001000000000,	// No destructor
-	//CPF_								= 0x0000002000000000,	//
-	CPF_AutoWeak						= 0x0000004000000000,	// Only used for weak pointers, means the export type is autoweak
-	CPF_ContainsInstancedReference		= 0x0000008000000000,	// Property contains component references.
-	CPF_AssetRegistrySearchable			= 0x0000010000000000,	// asset instances will add properties with this flag to the asset registry automatically
-	CPF_SimpleDisplay					= 0x0000020000000000,	// The property is visible by default in the editor details view
-	CPF_AdvancedDisplay					= 0x0000040000000000,	// The property is advanced and not visible by default in the editor details view
-	CPF_Protected						= 0x0000080000000000,	// property is protected from the perspective of script
-	CPF_BlueprintCallable				= 0x0000100000000000,	// MC Delegates only.  Property should be exposed for calling in blueprint code
-	CPF_BlueprintAuthorityOnly			= 0x0000200000000000,	// MC Delegates only.  This delegate accepts (only in blueprint) only events with BlueprintAuthorityOnly.
-	CPF_TextExportTransient				= 0x0000400000000000,	// Property shouldn't be exported to text format (e.g. copy/paste)
-	CPF_NonPIEDuplicateTransient		= 0x0000800000000000,	// Property should only be copied in PIE
-	CPF_ExposeOnSpawn					= 0x0001000000000000,	// Property is exposed on spawn
-	CPF_PersistentInstance				= 0x0002000000000000,	// A object referenced by the property is duplicated like a component. (Each actor should have an own instance.)
-	CPF_UObjectWrapper					= 0x0004000000000000,	// Property was parsed as a wrapper class like TSubclassOf<T>, FScriptInterface etc., rather than a USomething*
-	CPF_HasGetValueTypeHash				= 0x0008000000000000,	// This property can generate a meaningful hash value.
-	CPF_NativeAccessSpecifierPublic		= 0x0010000000000000,	// Public native access specifier
-	CPF_NativeAccessSpecifierProtected	= 0x0020000000000000,	// Protected native access specifier
-	CPF_NativeAccessSpecifierPrivate	= 0x0040000000000000,	// Private native access specifier
-	CPF_SkipSerialization				= 0x0080000000000000,	// Property shouldn't be serialized, can still be exported to text
+	CPF_Edit							= 0x0000000000000001,	///< Property is user-settable in the editor.
+	CPF_ConstParm						= 0x0000000000000002,	///< This is a constant function parameter
+	CPF_BlueprintVisible				= 0x0000000000000004,	///< This property can be read by blueprint code
+	CPF_ExportObject					= 0x0000000000000008,	///< Object can be exported with actor.
+	CPF_BlueprintReadOnly				= 0x0000000000000010,	///< This property cannot be modified by blueprint code
+	CPF_Net								= 0x0000000000000020,	///< Property is relevant to network replication.
+	CPF_EditFixedSize					= 0x0000000000000040,	///< Indicates that elements of an array can be modified, but its size cannot be changed.
+	CPF_Parm							= 0x0000000000000080,	///< Function/When call parameter.
+	CPF_OutParm							= 0x0000000000000100,	///< Value is copied out after function call.
+	CPF_ZeroConstructor					= 0x0000000000000200,	///< memset is fine for construction
+	CPF_ReturnParm						= 0x0000000000000400,	///< Return value.
+	CPF_DisableEditOnTemplate			= 0x0000000000000800,	///< Disable editing of this property on an archetype/sub-blueprint
+	//CPF_      						= 0x0000000000001000,	///< 
+	CPF_Transient   					= 0x0000000000002000,	///< Property is transient: shouldn't be saved or loaded, except for Blueprint CDOs.
+	CPF_Config      					= 0x0000000000004000,	///< Property should be loaded/saved as permanent profile.
+	//CPF_								= 0x0000000000008000,	///< 
+	CPF_DisableEditOnInstance			= 0x0000000000010000,	///< Disable editing on an instance of this class
+	CPF_EditConst   					= 0x0000000000020000,	///< Property is uneditable in the editor.
+	CPF_GlobalConfig					= 0x0000000000040000,	///< Load config from base class, not subclass.
+	CPF_InstancedReference				= 0x0000000000080000,	///< Property is a component references.
+	//CPF_								= 0x0000000000100000,	///<
+	CPF_DuplicateTransient				= 0x0000000000200000,	///< Property should always be reset to the default value during any type of duplication (copy/paste, binary duplication, etc.)
+	CPF_SubobjectReference				= 0x0000000000400000,	///< Property contains subobject references (TSubobjectPtr)
+	//CPF_    							= 0x0000000000800000,	///< 
+	CPF_SaveGame						= 0x0000000001000000,	///< Property should be serialized for save games
+	CPF_NoClear							= 0x0000000002000000,	///< Hide clear (and browse) button.
+	//CPF_  							= 0x0000000004000000,	///<
+	CPF_ReferenceParm					= 0x0000000008000000,	///< Value is passed by reference; CPF_OutParam and CPF_Param should also be set.
+	CPF_BlueprintAssignable				= 0x0000000010000000,	///< MC Delegates only.  Property should be exposed for assigning in blueprint code
+	CPF_Deprecated  					= 0x0000000020000000,	///< Property is deprecated.  Read it from an archive, but don't save it.
+	CPF_IsPlainOldData					= 0x0000000040000000,	///< If this is set, then the property can be memcopied instead of CopyCompleteValue / CopySingleValue
+	CPF_RepSkip							= 0x0000000080000000,	///< Not replicated. For non replicated properties in replicated structs 
+	CPF_RepNotify						= 0x0000000100000000,	///< Notify actors when a property is replicated
+	CPF_Interp							= 0x0000000200000000,	///< interpolatable property for use with matinee
+	CPF_NonTransactional				= 0x0000000400000000,	///< Property isn't transacted
+	CPF_EditorOnly						= 0x0000000800000000,	///< Property should only be loaded in the editor
+	CPF_NoDestructor					= 0x0000001000000000,	///< No destructor
+	//CPF_								= 0x0000002000000000,	///<
+	CPF_AutoWeak						= 0x0000004000000000,	///< Only used for weak pointers, means the export type is autoweak
+	CPF_ContainsInstancedReference		= 0x0000008000000000,	///< Property contains component references.
+	CPF_AssetRegistrySearchable			= 0x0000010000000000,	///< asset instances will add properties with this flag to the asset registry automatically
+	CPF_SimpleDisplay					= 0x0000020000000000,	///< The property is visible by default in the editor details view
+	CPF_AdvancedDisplay					= 0x0000040000000000,	///< The property is advanced and not visible by default in the editor details view
+	CPF_Protected						= 0x0000080000000000,	///< property is protected from the perspective of script
+	CPF_BlueprintCallable				= 0x0000100000000000,	///< MC Delegates only.  Property should be exposed for calling in blueprint code
+	CPF_BlueprintAuthorityOnly			= 0x0000200000000000,	///< MC Delegates only.  This delegate accepts (only in blueprint) only events with BlueprintAuthorityOnly.
+	CPF_TextExportTransient				= 0x0000400000000000,	///< Property shouldn't be exported to text format (e.g. copy/paste)
+	CPF_NonPIEDuplicateTransient		= 0x0000800000000000,	///< Property should only be copied in PIE
+	CPF_ExposeOnSpawn					= 0x0001000000000000,	///< Property is exposed on spawn
+	CPF_PersistentInstance				= 0x0002000000000000,	///< A object referenced by the property is duplicated like a component. (Each actor should have an own instance.)
+	CPF_UObjectWrapper					= 0x0004000000000000,	///< Property was parsed as a wrapper class like TSubclassOf<T>, FScriptInterface etc., rather than a USomething*
+	CPF_HasGetValueTypeHash				= 0x0008000000000000,	///< This property can generate a meaningful hash value.
+	CPF_NativeAccessSpecifierPublic		= 0x0010000000000000,	///< Public native access specifier
+	CPF_NativeAccessSpecifierProtected	= 0x0020000000000000,	///< Protected native access specifier
+	CPF_NativeAccessSpecifierPrivate	= 0x0040000000000000,	///< Private native access specifier
+	CPF_SkipSerialization				= 0x0080000000000000,	///< Property shouldn't be serialized, can still be exported to text
 };
 
-/** @name Combinations flags */
-//@{
+/** All Native Access Specifier flags */
 #define CPF_NativeAccessSpecifiers	(CPF_NativeAccessSpecifierPublic | CPF_NativeAccessSpecifierProtected | CPF_NativeAccessSpecifierPrivate)
 
+/** All parameter flags */
 #define CPF_ParmFlags				(CPF_Parm | CPF_OutParm | CPF_ReturnParm | CPF_ReferenceParm | CPF_ConstParm)
+
+/** Flags that are propagated to properties inside containers */
 #define CPF_PropagateToArrayInner	(CPF_ExportObject | CPF_PersistentInstance | CPF_InstancedReference | CPF_ContainsInstancedReference | CPF_Config | CPF_EditConst | CPF_Deprecated | CPF_EditorOnly | CPF_AutoWeak | CPF_UObjectWrapper )
 #define CPF_PropagateToMapValue		(CPF_ExportObject | CPF_PersistentInstance | CPF_InstancedReference | CPF_ContainsInstancedReference | CPF_Config | CPF_EditConst | CPF_Deprecated | CPF_EditorOnly | CPF_AutoWeak | CPF_UObjectWrapper | CPF_Edit )
 #define CPF_PropagateToMapKey		(CPF_ExportObject | CPF_PersistentInstance | CPF_InstancedReference | CPF_ContainsInstancedReference | CPF_Config | CPF_EditConst | CPF_Deprecated | CPF_EditorOnly | CPF_AutoWeak | CPF_UObjectWrapper | CPF_Edit )
 #define CPF_PropagateToSetElement	(CPF_ExportObject | CPF_PersistentInstance | CPF_InstancedReference | CPF_ContainsInstancedReference | CPF_Config | CPF_EditConst | CPF_Deprecated | CPF_EditorOnly | CPF_AutoWeak | CPF_UObjectWrapper | CPF_Edit )
 
-/** the flags that should never be set on interface properties */
+/** The flags that should never be set on interface properties */
 #define CPF_InterfaceClearMask		(CPF_ExportObject|CPF_InstancedReference|CPF_ContainsInstancedReference)
 
-/** all the properties that can be stripped for final release console builds */
+/** All the properties that can be stripped for final release console builds */
 #define CPF_DevelopmentAssets		(CPF_EditorOnly)
 
-/** all the properties that should never be loaded or saved */
+/** All the properties that should never be loaded or saved */
 #define CPF_ComputedFlags			(CPF_IsPlainOldData | CPF_NoDestructor | CPF_ZeroConstructor | CPF_HasGetValueTypeHash)
 
+/** Mask of all property flags */
 #define CPF_AllFlags				((EPropertyFlags)0xFFFFFFFFFFFFFFFF)
 
 ENUM_CLASS_FLAGS(EPropertyFlags)
 
-//@}
-
 /** 
  * Flags describing an object instance
  */
@@ -460,20 +449,20 @@ enum EObjectFlags
 {
 	// Do not add new flags unless they truly belong here. There are alternatives.
 	// if you change any the bit of any of the RF_Load flags, then you will need legacy serialization
-	RF_NoFlags						= 0x00000000,	///< No flags, used to avoid a cast
+	RF_NoFlags					= 0x00000000,	///< No flags, used to avoid a cast
 
 	// This first group of flags mostly has to do with what kind of object it is. Other than transient, these are the persistent object flags.
 	// The garbage collector also tends to look at these.
 	RF_Public					=0x00000001,	///< Object is visible outside its package.
 	RF_Standalone				=0x00000002,	///< Keep object around for editing even if unreferenced.
-	RF_MarkAsNative					=0x00000004,	///< Object (UField) will be marked as native on construction (DO NOT USE THIS FLAG in HasAnyFlags() etc)
+	RF_MarkAsNative				=0x00000004,	///< Object (UField) will be marked as native on construction (DO NOT USE THIS FLAG in HasAnyFlags() etc)
 	RF_Transactional			=0x00000008,	///< Object is transactional.
 	RF_ClassDefaultObject		=0x00000010,	///< This object is its class's default object
 	RF_ArchetypeObject			=0x00000020,	///< This object is a template for another object - treat like a class default object
 	RF_Transient				=0x00000040,	///< Don't save object.
 
 	// This group of flags is primarily concerned with garbage collection.
-	RF_MarkAsRootSet					=0x00000080,	///< Object will be marked as root set on construction and not be garbage collected, even if unreferenced (DO NOT USE THIS FLAG in HasAnyFlags() etc)
+	RF_MarkAsRootSet			=0x00000080,	///< Object will be marked as root set on construction and not be garbage collected, even if unreferenced (DO NOT USE THIS FLAG in HasAnyFlags() etc)
 	RF_TagGarbageTemp			=0x00000100,	///< This is a temp user flag for various utilities that need to use the garbage collector. The garbage collector itself does not interpret it.
 
 	// The group of flags tracks the stages of the lifetime of a uobject
@@ -493,29 +482,30 @@ enum EObjectFlags
 	RF_TextExportTransient		=0x00100000,	///< Do not export object to text form (e.g. copy/paste). Generally used for sub-objects that can be regenerated from data in their parent object.
 	RF_LoadCompleted			=0x00200000,	///< Object has been completely serialized by linkerload at least once. DO NOT USE THIS FLAG, It should be replaced with RF_WasLoaded.
 	RF_InheritableComponentTemplate = 0x00400000, ///< Archetype of the object can be in its super class
-	RF_DuplicateTransient = 0x00800000, ///< Object should not be included in any type of duplication (copy/paste, binary duplication, etc.)
-	RF_StrongRefOnFrame			= 0x01000000,	///< References to this object from persistent function frame are handled as strong ones.
-	RF_NonPIEDuplicateTransient		= 0x02000000,  ///< Object should not be included for duplication unless it's being duplicated for a PIE session
-	RF_Dynamic = 0x04000000, // Field Only. Dynamic field - doesn't get constructed during static initialization, can be constructed multiple times
-	RF_WillBeLoaded = 0x08000000, // This object was constructed during load and will be loaded shortly
+	RF_DuplicateTransient		=0x00800000,	///< Object should not be included in any type of duplication (copy/paste, binary duplication, etc.)
+	RF_StrongRefOnFrame			=0x01000000,	///< References to this object from persistent function frame are handled as strong ones.
+	RF_NonPIEDuplicateTransient	=0x02000000,	///< Object should not be included for duplication unless it's being duplicated for a PIE session
+	RF_Dynamic					=0x04000000,	///< Field Only. Dynamic field - doesn't get constructed during static initialization, can be constructed multiple times
+	RF_WillBeLoaded				=0x08000000,	///< This object was constructed during load and will be loaded shortly
 };
 
-	// Special all and none masks
+/** Mask for all object flags */
 #define RF_AllFlags				(EObjectFlags)0x0fffffff	///< All flags, used mainly for error checking
 
-	// Predefined groups of the above
-#define RF_Load						((EObjectFlags)(RF_Public | RF_Standalone | RF_Transactional | RF_ClassDefaultObject | RF_ArchetypeObject | RF_DefaultSubObject | RF_TextExportTransient | RF_InheritableComponentTemplate | RF_DuplicateTransient | RF_NonPIEDuplicateTransient)) // Flags to load from Unrealfiles.
-#define RF_PropagateToSubObjects	((EObjectFlags)(RF_Public | RF_ArchetypeObject | RF_Transactional | RF_Transient))		// Sub-objects will inherit these flags from their SuperObject.
+/** Flags to load from unreal asset files */
+#define RF_Load						((EObjectFlags)(RF_Public | RF_Standalone | RF_Transactional | RF_ClassDefaultObject | RF_ArchetypeObject | RF_DefaultSubObject | RF_TextExportTransient | RF_InheritableComponentTemplate | RF_DuplicateTransient | RF_NonPIEDuplicateTransient)) 
 
-ENUM_CLASS_FLAGS(EObjectFlags);
+/** Sub-objects will inherit these flags from their SuperObject */
+#define RF_PropagateToSubObjects	((EObjectFlags)(RF_Public | RF_ArchetypeObject | RF_Transactional | RF_Transient))
 
-//@}
+ENUM_CLASS_FLAGS(EObjectFlags);
 
 /** Objects flags for internal use (GC, low level UObject code) */
 enum class EInternalObjectFlags : int32
 {
 	None = 0,
 	//~ All the other bits are reserved, DO NOT ADD NEW FLAGS HERE!
+
 	ReachableInCluster = 1 << 23, ///< External reference to object in cluster exists
 	ClusterRoot = 1 << 24, ///< Root of a cluster
 	Native = 1 << 25, ///< Native (UClass only). 
@@ -527,6 +517,7 @@ enum class EInternalObjectFlags : int32
 	//~ UnusedFlag = 1 << 31,
 
 	GarbageCollectionKeepFlags = Native | Async | AsyncLoading,
+
 	//~ Make sure this is up to date!
 	AllFlags = ReachableInCluster | ClusterRoot | Native | Async | AsyncLoading | Unreachable | PendingKill | RootSet
 };
@@ -536,7 +527,6 @@ ENUM_CLASS_FLAGS(EInternalObjectFlags);
 	Core types.
 ----------------------------------------------------------------------------*/
 
-// forward declarations
 class UObject;
 class UProperty;
 class FObjectInitializer; 
@@ -730,7 +720,7 @@ namespace UC
 		/// Marks this class as 'experimental' (a totally unsupported and undocumented prototype)
 		Experimental,
 
-		// Marks this class as an 'early access' preview (while not considered production-ready, it's a step beyond 'experimental' and is being provided as a preview of things to come)
+		/// Marks this class as an 'early access' preview (while not considered production-ready, it's a step beyond 'experimental' and is being provided as a preview of things to come)
 		EarlyAccessPreview,
 	};
 }
@@ -1411,7 +1401,7 @@ namespace UM
 #define DECLARE_FSTRUCTUREDARCHIVE_SERIALIZER( TClass, API ) virtual API void Serialize(FStructuredArchive::FRecord Record) override;
 
 /*-----------------------------------------------------------------------------
-Class declaration macros.
+	Class declaration macros.
 -----------------------------------------------------------------------------*/
 
 #define DECLARE_CLASS( TClass, TSuperClass, TStaticFlags, TStaticCastFlags, TPackage, TRequiredAPI  ) \
@@ -1659,22 +1649,25 @@ public: \
 		return PrivateStaticClass; \
 	}
 
-/*-----------------------------------------------------------------------------
-	ERenameFlags.
-
-	Options to the UObject::Rename function, bit flag
------------------------------------------------------------------------------*/
-
+/** Options to the UObject::Rename() function, bit flag */
 typedef uint32 ERenameFlags;
 
-#define REN_None				(0x0000)
-#define REN_ForceNoResetLoaders	(0x0001) // Rename won't call ResetLoaders - most likely you should never specify this option (unless you are renaming a UPackage possibly)
-#define REN_Test				(0x0002) // Just test to make sure that the rename is guaranteed to succeed if an non test rename immediately follows
-#define REN_DoNotDirty			(0x0004) // Indicates that the object (and new outer) should not be dirtied.
-#define REN_DontCreateRedirectors (0x0010) // Don't create an object redirector, even if the class is marked RF_Public
-#define REN_NonTransactional	(0x0020) // Don't call Modify() on the objects, so they won't be stored in the transaction buffer
-#define REN_ForceGlobalUnique	(0x0040) // Force unique names across all packages not just while the scope of the new outer
-#define REN_SkipGeneratedClasses (0x0080) // Prevent renaming of any child generated classes and CDO's in blueprints
+/** Default rename behavior */
+#define REN_None					(0x0000)
+/** Rename won't call ResetLoaders or flush async loading. You should pass this if you are renaming a deep subobject and do not need to reset loading for the outer package */
+#define REN_ForceNoResetLoaders		(0x0001) 
+/** Just test to make sure that the rename is guaranteed to succeed if an non test rename immediately follows */
+#define REN_Test					(0x0002) 
+/** Indicates that the object (and new outer) should not be dirtied */
+#define REN_DoNotDirty				(0x0004) 
+/** Don't create an object redirector, even if the class is marked RF_Public */
+#define REN_DontCreateRedirectors	(0x0010) 
+/** Don't call Modify() on the objects, so they won't be stored in the transaction buffer */
+#define REN_NonTransactional		(0x0020) 
+/** Force unique names across all packages not just within the scope of the new outer */
+#define REN_ForceGlobalUnique		(0x0040) 
+/** Prevent renaming of any child generated classes and CDO's in blueprints */
+#define REN_SkipGeneratedClasses	(0x0080) 
 
 /*-----------------------------------------------------------------------------
 	Misc.
@@ -1694,6 +1687,7 @@ namespace UE4
 	enum ELoadConfigPropagationFlags
 	{
 		LCPF_None					=	0x0,
+
 		/**
 		 * Indicates that the object should read ini values from each section up its class's hierarchy chain;
 		 * Useful when calling LoadConfig on an object after it has already been initialized against its archetype
@@ -1715,7 +1709,9 @@ namespace UE4
 		 */
 		LCPF_ReloadingConfigData	=	0x8,
 
-		// Combination flags
+		/** 
+		 * All flags that should be persisted to propagated recursive calls 
+		 */
 		LCPF_PersistentFlags		=	LCPF_ReloadingConfigData,
 	};
 }
diff --git a/Engine/Source/Runtime/CoreUObject/Public/UObject/UObjectBase.h b/Engine/Source/Runtime/CoreUObject/Public/UObject/UObjectBase.h
--- a/Engine/Source/Runtime/CoreUObject/Public/UObject/UObjectBase.h
+++ b/Engine/Source/Runtime/CoreUObject/Public/UObject/UObjectBase.h
@@ -1,7 +1,7 @@
 // Copyright 1998-2019 Epic Games, Inc. All Rights Reserved.
 
 /*=============================================================================
-	UObjectBase.h: Unreal UObject base class
+	UObjectBase.h: Base class for UObject, defines low level functionality
 =============================================================================*/
 
 #pragma once
@@ -14,6 +14,9 @@
 
 DECLARE_DWORD_COUNTER_STAT_EXTERN(TEXT("STAT_UObjectsStatGroupTester"), STAT_UObjectsStatGroupTester, STATGROUP_UObjects, COREUOBJECT_API);
 
+/** 
+ * Low level implementation of UObject, should not be used directly in game code 
+ */
 class COREUOBJECT_API UObjectBase
 {
 	friend class UObjectBaseUtility;
@@ -36,12 +39,14 @@ protected:
 #endif
 	{
 	}
+
 	/**
 	 * Constructor used for bootstrapping
 	 * @param	InFlags			RF_Flags to assign
 	 */
 	UObjectBase( EObjectFlags InFlags );
 public:
+
 	/**
 	 * Constructor used by StaticAllocateObject
 	 * @param	InClass				non NULL, this gives the class of the new object, if known at this time
@@ -92,6 +97,7 @@ private:
 	 * @param InSetInternalFlags Internal object flags to be set on the object once it's been added to the array
 	 */
 	void AddObject(FName Name, EInternalObjectFlags InSetInternalFlags);
+
 public:
 	/**
 	 * Checks to see if the object appears to be valid
@@ -116,22 +122,30 @@ public:
 	{
 		return (uint32)InternalIndex;
 	}
+
+	/** Returns the UClass that defines the fields of this object */
 	FORCEINLINE UClass* GetClass() const
 	{
 		return ClassPrivate;
 	}
+	
+	/** Returns the UObject this object resides in */
 	FORCEINLINE UObject* GetOuter() const
 	{
 		return OuterPrivate;
 	}
+
+	/** Returns the logical name of this object */
 	FORCEINLINE FName GetFName() const
 	{
 		return NamePrivate;
 	}
 
 	/** 
-	 * Returns the stat ID of the object...
-	**/
+	 * Returns the stat ID of the object, used for profiling. This will create a stat ID if needed.
+	 *
+	 * @param bForDeferred If true, a stat ID will be created even if a group is disabled
+	 */
 	FORCEINLINE TStatId GetStatID(bool bForDeferredUse = false) const
 	{
 #if STATS
@@ -156,11 +170,8 @@ public:
 
 	
 private:
-
-	/** 
-	 * Creates this stat ID for the object...and handle a null this pointer
-	**/
 #if STATS || ENABLE_STATNAMEDEVENTS_UOBJECT
+	/** Creates a stat ID for this object */
 	void CreateStatID() const;
 #endif
 
@@ -174,6 +185,7 @@ protected:
 		checkfSlow((NewFlags & ~RF_AllFlags) == 0, TEXT("%s flagged as 0x%x but is trying to set flags to RF_AllFlags"), *GetFName().ToString(), (int)ObjectFlags);
 		ObjectFlags = NewFlags;
 	}
+
 public:
 	/**
 	 * Retrieve the object flags directly
@@ -248,8 +260,10 @@ private:
 #endif
 #endif // STATS || ENABLE_STATNAMEDEVENTS
 
-	// This is used by the reinstancer to re-class and re-archetype the current instances of a class before recompiling
+	
 	friend class FBlueprintCompileReinstancer;
+
+	/** This is used by the reinstancer to re-class and re-archetype the current instances of a class before recompiling */
 	void SetClass(UClass* NewClass);
 
 #if HACK_HEADER_GENERATOR
@@ -395,7 +409,10 @@ struct FCompiledInDeferEnum
  */
 COREUOBJECT_API class UEnum *GetStaticEnum(class UEnum *(*InRegister)(), UObject* EnumOuter, const TCHAR* EnumName);
 
+/** Called during HotReload to hook up an existing structure */
 COREUOBJECT_API class UScriptStruct* FindExistingStructIfHotReloadOrDynamic(UObject* Outer, const TCHAR* StructName, SIZE_T Size, uint32 Crc, bool bIsDynamic);
+
+/** Called during HotReload to hook up an existing enum */
 COREUOBJECT_API class UEnum* FindExistingEnumIfHotReloadOrDynamic(UObject* Outer, const TCHAR* EnumName, SIZE_T Size, uint32 Crc, bool bIsDynamic);
 
 /** Must be called after a module has been loaded that contains UObject classes */
diff --git a/Engine/Source/Runtime/CoreUObject/Public/UObject/UObjectBaseUtility.h b/Engine/Source/Runtime/CoreUObject/Public/UObject/UObjectBaseUtility.h
--- a/Engine/Source/Runtime/CoreUObject/Public/UObject/UObjectBaseUtility.h
+++ b/Engine/Source/Runtime/CoreUObject/Public/UObject/UObjectBaseUtility.h
@@ -19,7 +19,9 @@
 	#endif
 #endif
 
-
+/**
+ * Provides utility functions for UObject, this class should not be used directly
+ */
 class COREUOBJECT_API UObjectBaseUtility : public UObjectBase
 {
 public:
@@ -30,20 +32,25 @@ public:
 	{
 	}
 
-	/***********************/
-	/******** Flags ********/
-	/***********************/
 
+	/*-------------------
+			Flags
+	-------------------*/
+
+	/** Modifies object flags for a specific object */
 	FORCEINLINE void SetFlags( EObjectFlags NewFlags )
 	{
 		checkSlow(!(NewFlags & (RF_MarkAsNative | RF_MarkAsRootSet))); // These flags can't be used outside of constructors / internal code
 		SetFlagsTo(GetFlags() | NewFlags);
 	}
+
+	/** Clears subset of flags for a specific object */
 	FORCEINLINE void ClearFlags( EObjectFlags NewFlags )
 	{
 		checkSlow(!(NewFlags & (RF_MarkAsNative | RF_MarkAsRootSet)) || NewFlags == RF_AllFlags); // These flags can't be used outside of constructors / internal code
 		SetFlagsTo(GetFlags() & ~NewFlags);
 	}
+
 	/**
 	 * Used to safely check whether any of the passed in flags are set. 
 	 *
@@ -55,6 +62,7 @@ public:
 		checkSlow(!(FlagsToCheck & (RF_MarkAsNative | RF_MarkAsRootSet)) || FlagsToCheck == RF_AllFlags); // These flags can't be used outside of constructors / internal code
 		return (GetFlags() & FlagsToCheck) != 0;
 	}
+
 	/**
 	 * Used to safely check whether all of the passed in flags are set. 
 	 *
@@ -66,6 +74,7 @@ public:
 		checkSlow(!(FlagsToCheck & (RF_MarkAsNative | RF_MarkAsRootSet)) || FlagsToCheck == RF_AllFlags); // These flags can't be used outside of constructors / internal code
 		return ((GetFlags() & FlagsToCheck) == FlagsToCheck);
 	}
+
 	/**
 	 * Returns object flags that are both in the mask and set on the object.
 	 *
@@ -77,9 +86,9 @@ public:
 		return EObjectFlags(GetFlags() & Mask);
 	}
 
-	/***********************/
-	/******** Marks *******  UObjectMarks.cpp */
-	/***********************/
+	/*----------------------------------------------------
+			Marks, implemented in UObjectMarks.cpp
+	----------------------------------------------------*/
 
 	/**
 	 * Adds marks to an object
@@ -152,25 +161,23 @@ public:
 	}
 
 	/**
-	* Unmarks this object as PendingKill.
-	*/
+	 * Unmarks this object as PendingKill.
+	 */
 	FORCEINLINE void ClearPendingKill()
 	{
 		GUObjectArray.IndexToObject(InternalIndex)->ClearPendingKill();
 	}
 
-	//
-	// Add an object to the root set. This prevents the object and all
-	// its descendants from being deleted during garbage collection.
-	//
+	/**
+	 * Add an object to the root set. This prevents the object and all
+	 * its descendants from being deleted during garbage collection.
+	 */
 	FORCEINLINE void AddToRoot()
 	{
 		GUObjectArray.IndexToObject(InternalIndex)->SetRootSet();
 	}
 
-	//
-	// Remove an object from the root set.
-	//
+	/** Remove an object from the root set. */
 	FORCEINLINE void RemoveFromRoot()
 	{
 		GUObjectArray.IndexToObject(InternalIndex)->ClearRootSet();
@@ -186,45 +193,39 @@ public:
 		return GUObjectArray.IndexToObject(InternalIndex)->IsRootSet();
 	}
 
-	/**
-	* Atomically clear the unreachable flag
-	*
-	* @return true if we are the thread that cleared RF_Unreachable
-	**/
+	 /**
+	 * Atomically clear the unreachable flag
+	 *
+	 * @return true if we are the thread that cleared RF_Unreachable
+	 */
 	FORCEINLINE bool ThisThreadAtomicallyClearedRFUnreachable()
 	{
 		return GUObjectArray.IndexToObject(InternalIndex)->ThisThreadAtomicallyClearedRFUnreachable();
 	}
 
-	/**
-	* Checks if the object is unreachable.
-	**/
+	/** Checks if the object is unreachable. */
 	FORCEINLINE bool IsUnreachable() const
 	{
 		return GUObjectArray.IndexToObject(InternalIndex)->IsUnreachable();
 	}
 
-	/**
-	* Checks if the object is pending kill or unreachable.
-	**/
+	/** Checks if the object is pending kill or unreachable. */
 	FORCEINLINE bool IsPendingKillOrUnreachable() const
 	{
 		return GUObjectArray.IndexToObject(InternalIndex)->HasAnyFlags(EInternalObjectFlags::PendingKill | EInternalObjectFlags::Unreachable);
 	}
 
-	/**
-	* Checks if the object is native.
-	**/
+	/** Checks if the object is native. */
 	FORCEINLINE bool IsNative() const
 	{
 		return GUObjectArray.IndexToObject(InternalIndex)->HasAnyFlags(EInternalObjectFlags::Native);
 	}
 
 	/**
-	* Clears passed in internal flags .
+	 * Clears passed in internal flags.
 	 *
-	* @param FlagsToClear	Object flags to clear.
-	* @return				true if any of the passed in flags are set, false otherwise  (including no flags passed in).
+	 * @param FlagsToClear	Object flags to clear.
+	 * @return				true if any of the passed in flags are set, false otherwise  (including no flags passed in).
 	 */
 	FORCEINLINE void SetInternalFlags(EInternalObjectFlags FlagsToSet) const
 	{
@@ -232,10 +233,10 @@ public:
 	}
 
 	/**
-	* Gets internal flags.
+	 * Gets internal flags.
 	 *
-	* @param FlagsToClear	Object flags to clear.
-	* @return				true if any of the passed in flags are set, false otherwise  (including no flags passed in).
+	 * @param FlagsToClear	Object flags to clear.
+	 * @return				true if any of the passed in flags are set, false otherwise  (including no flags passed in).
 	 */
 	FORCEINLINE EInternalObjectFlags GetInternalFlags() const
 	{
@@ -243,10 +244,10 @@ public:
 	}
 
 	/**
-	* Used to safely check whether any of the passed in internal flags are set.
+	 * Used to safely check whether any of the passed in internal flags are set.
 	 *
-	* @param FlagsToCheck	Object flags to check for.
-	* @return				true if any of the passed in flags are set, false otherwise  (including no flags passed in).
+	 * @param FlagsToCheck	Object flags to check for.
+	 * @return				true if any of the passed in flags are set, false otherwise  (including no flags passed in).
 	 */
 	FORCEINLINE bool HasAnyInternalFlags(EInternalObjectFlags FlagsToCheck) const
 	{
@@ -254,10 +255,10 @@ public:
 	}
 
 	/**
-	* Clears passed in internal flags .
+	 * Clears passed in internal flags.
 	 *
-	* @param FlagsToClear	Object flags to clear.
-	* @return				true if any of the passed in flags are set, false otherwise  (including no flags passed in).
+	 * @param FlagsToClear	Object flags to clear.
+	 * @return				true if any of the passed in flags are set, false otherwise  (including no flags passed in).
 	 */
 	FORCEINLINE void ClearInternalFlags(EInternalObjectFlags FlagsToClear) const
 	{
@@ -265,19 +266,20 @@ public:
 	}
 
 	/**
-	* Atomically clears passed in internal flags .
-	*
-	* @param FlagsToClear	Object flags to clear.
-	* @return				true if any of the passed in flags are set, false otherwise  (including no flags passed in).
-	*/
+	 * Atomically clears passed in internal flags.
+	 *
+	 * @param FlagsToClear	Object flags to clear.
+	 * @return				true if any of the passed in flags are set, false otherwise  (including no flags passed in).
+	 */
 	FORCEINLINE bool AtomicallyClearInternalFlags(EInternalObjectFlags FlagsToClear) const
 	{
 		return GUObjectArray.IndexToObject(InternalIndex)->ThisThreadAtomicallyClearedFlag(FlagsToClear);
 	}
 
-	/***********************/
-	/******** Names ********/
-	/***********************/
+
+	/*-------------------
+			Names
+	-------------------*/
 
 	/**
 	 * Returns the fully qualified pathname for this object as well as the name of the class, in the format:
@@ -302,7 +304,6 @@ public:
 	FString GetPathName( const UObject* StopOuter=NULL ) const;
 
 public:
-
 	/**
 	* Called after load to determine if the object can be a cluster root
 	*
@@ -338,7 +339,6 @@ public:
 	void AddToCluster(UObjectBaseUtility* ClusterRootOrObjectFromCluster, bool bAddAsMutableObject = false);
 
 protected:
-
 	/** Helper function to create a cluster from UObject */
 	static void CreateClusterFromObject(UObjectBaseUtility* ClusterRootObject, UObjectBaseUtility* ReferencingObject);
 
@@ -365,19 +365,22 @@ public:
 	{
 		return GetFName().ToString();
 	}
-	// GetFullName optimization
+
+	/** Optimized version of GetName that overwries an existing string */
 	FORCEINLINE void GetName(FString &ResultString) const
 	{
 		GetFName().ToString(ResultString);
 	}
+	/** Optimized version of GetName that appends to an existing string */
 	FORCEINLINE void AppendName(FString& ResultString) const
 	{
 		GetFName().AppendString(ResultString);
 	}
 
-	/***********************/
-	/******** Outer ********/
-	/***********************/
+
+	/*-------------------
+			Outer
+	-------------------*/
 
 	/** 
 	 * Walks up the list of outers until it finds the highest one.
@@ -420,9 +423,7 @@ public:
 		return (T *)GetTypedOuter(T::StaticClass());
 	}
 
-	/**
-	 * @return	true if the specified object appears somewhere in this object's outer chain.
-	 */
+	/** Returns true if the specified object appears somewhere in this object's outer chain. */
 	bool IsIn( const UObject* SomeOuter ) const;
 
 	/**
@@ -441,13 +442,10 @@ public:
 	 */
 	bool RootPackageHasAnyFlags( uint32 CheckFlagMask ) const;
 
-	/***********************/
-	/******** Class ********/
-	/***********************/
 
-	/**
-	 * @return	true if this object is of the specified type.
-	 */
+	/*-------------------
+			Class
+	-------------------*/
 
 private:
 	template <typename ClassType>
@@ -457,6 +455,7 @@ private:
 	}
 
 public:
+	/** Returns true if this object is of the specified type. */
 	template <typename OtherClassType>
 	FORCEINLINE bool IsA( OtherClassType SomeBase ) const
 	{
@@ -477,9 +476,7 @@ public:
 		return IsChildOfWorkaround(ThisClass, SomeBaseClass);
 	}
 
-	/**
-	 * @return	true if this object is of the template type.
-	 */
+	/** Returns true if this object is of the template type. */
 	template<class T>
 	bool IsA() const
 	{
@@ -528,9 +525,9 @@ public:
 	bool IsDefaultSubobject() const;
 	
 
-	/***********************/
-	/******** Linker ****  UObjectLinker.cpp */
-	/***********************/
+	/*--------------------------------------------------
+			Linker, defined in UObjectLinker.cpp
+	--------------------------------------------------*/
 
 	/**
 	 * Returns the linker for this object.
@@ -538,6 +535,7 @@ public:
 	 * @return	a pointer to the linker for this object, or NULL if this object has no linker
 	 */
 	class FLinkerLoad* GetLinker() const;
+
 	/**
 	 * Returns this object's LinkerIndex.
 	 *
@@ -590,6 +588,7 @@ public:
 	}
 };
 
+/** Returns false if this pointer cannot be a valid pointer to a UObject */
 FORCEINLINE bool IsPossiblyAllocatedUObjectPointer(UObject* Ptr)
 {
 	auto CountByteValues = [](UPTRINT Val, UPTRINT ByteVal) -> int32
@@ -674,6 +673,8 @@ COREUOBJECT_API UClass* GetParentNativeClass(UClass* Class);
 #endif
 
 #if STATS
+
+/** Structure used to track time spent by a UObject */
 struct FScopeCycleCounterUObject : public FCycleCounter
 {
 #if USE_MALLOC_PROFILER
@@ -727,6 +728,7 @@ struct FScopeCycleCounterUObject : public FCycleCounter
 		}
 #endif
 	}
+
 	/**
 	 * Updates the stat with the time spent
 	 */
@@ -746,6 +748,7 @@ struct FScopeCycleCounterUObject : public FCycleCounter
 #endif
 };
 
+/** Declares a scope cycle counter for a specific object with a Name context */
 #define SCOPE_CYCLE_UOBJECT(Name, Object) \
 	FScopeCycleCounterUObject ObjCycleCount_##Name(Object);
 
@@ -764,6 +767,7 @@ struct FScopeCycleCounterUObject
 	}
 };
 
+/** Declares a scope cycle counter for a specific object with a Name context */
 #define SCOPE_CYCLE_UOBJECT(Name, Object) \
 	FScopeCycleCounterUObject ObjCycleCount_##Name(Object);
 #elif USE_LIGHTWEIGHT_STATS_FOR_HITCH_DETECTION && USE_HITCH_DETECTION && USE_LIGHTWEIGHT_UOBJECT_STATS_FOR_HITCH_DETECTION
@@ -788,6 +792,8 @@ public:
 
 	COREUOBJECT_API void ReportHitch();
 };
+
+/** Declares a scope cycle counter for a specific object with a Name context */
 #define SCOPE_CYCLE_UOBJECT(Name, Object) \
 	FScopeCycleCounterUObject ObjCycleCount_##Name(Object);
 #else
