Removed duplicate comments from cpp per coding standard

diff --git a/Engine/Source/Runtime/CoreUObject/Private/UObject/UObjectGlobals.cpp b/Engine/Source/Runtime/CoreUObject/Private/UObject/UObjectGlobals.cpp
--- a/Engine/Source/Runtime/CoreUObject/Private/UObject/UObjectGlobals.cpp
+++ b/Engine/Source/Runtime/CoreUObject/Private/UObject/UObjectGlobals.cpp
@@ -1,7 +1,7 @@
 // Copyright 1998-2019 Epic Games, Inc. All Rights Reserved.
 
 /*=============================================================================
-	UnObj.cpp: Unreal object global data and functions
+	UObjectGlobals.cpp: Unreal object global data and functions
 =============================================================================*/
 
 #include "UObject/UObjectGlobals.h"
@@ -270,18 +270,9 @@ int32 UpdateSuffixForNextNewObject(UObject* Parent, UClass* Class, TFunctionRef<
 	return Result;
 }
 
-/**
- * Fast version of StaticFindObject that relies on the passed in FName being the object name
- * without any group/ package qualifiers.
- *
- * @param	ObjectClass		The to be found object's class
- * @param	ObjectPackage	The to be found object's outer
- * @param	ObjectName		The to be found object's class
- * @param	ExactClass		Whether to require an exact match with the passed in class
- * @param	AnyPackage		Whether to look in any package
- * @param	ExclusiveFlags	Ignores objects that contain any of the specified exclusive flags
- * @return	Returns a pointer to the found object or NULL if none could be found
- */
+//
+// Find an object, path must unqualified
+//
 UObject* StaticFindObjectFast(UClass* ObjectClass, UObject* ObjectPackage, FName ObjectName, bool ExactClass, bool AnyPackage, EObjectFlags ExclusiveFlags, EInternalObjectFlags ExclusiveInternalFlags)
 {
 	if (GIsSavingPackage || IsGarbageCollectingOnGameThread())
@@ -448,79 +439,12 @@ void GlobalSetProperty( const TCHAR* Value, UClass* Class, UProperty* Property,
 	}
 }
 
-
-/**
- * Executes a delegate by calling the named function on the object bound to the delegate.  You should
- * always first verify that the delegate is safe to execute by calling IsBound() before calling this function.
- *
- * @param	Params		Parameter structure
- */
-//void FScriptDelegate::ProcessDelegate( void* Parameters ) const
-//{
-//	checkf( Object.IsValid() != false, TEXT( "ProcessDelegate() called with no object bound to delegate!" ) );
-//	checkf( FunctionName != NAME_None, TEXT( "ProcessDelegate() called with no function name set!" ) );
-//
-//	// Object was pending kill, so we cannot execute the delegate.  Note that it's important to assert
-//	// here and not simply continue execution, as memory may be left uninitialized if the delegate is
-//	// not able to execute, resulting in much harder-to-detect code errors.  Users should always make
-//	// sure IsBound() returns true before calling ProcessDelegate()!
-//	UObject* ObjectPtr = static_cast< UObject* >( Object.Get() );	// Down-cast
-//	checkSlow( !ObjectPtr->IsPendingKill() );
-//
-//	// Object *must* implement the specified function
-//	UFunction* Function = ObjectPtr->FindFunctionChecked( FunctionName );
-//
-//	// Execute the delegate!
-//	ObjectPtr->ProcessEvent(Function, Parameters);
-//}
-
-
-/**
- * Executes a multi-cast delegate by calling all functions on objects bound to the delegate.  Always
- * safe to call, even if when no objects are bound, or if objects have expired.
- *
- * @param	Params				Parameter structure
- */
-//void FMulticastScriptDelegate::ProcessMulticastDelegate(void* Parameters) const
-//{
-//	if( InvocationList.Num() > 0 )
-//	{
-//		// Create a copy of the invocation list, just in case the list is modified by one of the callbacks during the broadcast
-//		typedef TArray< FScriptDelegate, TInlineAllocator< 4 > > FInlineInvocationList;
-//		FInlineInvocationList InvocationListCopy = FInlineInvocationList(InvocationList);
-//	
-//		// Invoke each bound function
-//		for( FInlineInvocationList::TConstIterator FunctionIt( InvocationListCopy ); FunctionIt; ++FunctionIt )
-//		{
-//			if( FunctionIt->IsBound() )
-//			{
-//				// Invoke this delegate!
-//				FunctionIt->ProcessDelegate(Parameters);
-//			}
-//			else if ( FunctionIt->IsCompactable() )
-//			{
-//				// Function couldn't be executed, so remove it.  Note that because the original list could have been modified by one of the callbacks, we have to search for the function to remove here.
-//				RemoveInternal( *FunctionIt );
-//			}
-//		}
-//	}
-//}
-
-
 /*-----------------------------------------------------------------------------
 	UObject Tick.
 -----------------------------------------------------------------------------*/
 
-/**
- * Static UObject tick function, used to verify certain key assumptions and to tick the async loading code.
- *
- * @warning: The streaming stats rely on this function not doing any work besides calling ProcessAsyncLoading.
- * @todo: Move stats code into core?
- *
- * @param DeltaTime	Time in seconds since last call
- * @param bUseFullTimeLimit	If true, use the entire time limit even if blocked on I/O
- * @param AsyncLoadingTime Time in seconds to use for async loading limit
- */
+// @warning: The streaming stats rely on this function not doing any work besides calling ProcessAsyncLoading.
+// @todo: Move stats code into core?
 void StaticTick( float DeltaTime, bool bUseFullTimeLimit, float AsyncLoadingTime )
 {
 	check(!IsLoading());
@@ -607,13 +531,6 @@ void SafeLoadError( UObject* Outer, uint32 LoadFlags, const TCHAR* ErrorMessage)
 	}
 }
 
-/**
- * Find an existing package by name
- * @param InOuter		The Outer object to search inside
- * @param PackageName	The name of the package to find
- *
- * @return The package if it exists
- */
 UPackage* FindPackage( UObject* InOuter, const TCHAR* PackageName )
 {
 	FString InName;
@@ -846,18 +763,6 @@ bool ResolveName(UObject*& InPackage, FString& InOutName, bool Create, bool Thro
 	return true;
 }
 
-/**
- * Parse an object from a text representation
- *
- * @param Stream			String containing text to parse
- * @param Match				Tag to search for object representation within string
- * @param Class				The class of the object to be loaded.
- * @param DestRes			returned uobject pointer
- * @param InParent			Outer to search
- * @param bInvalidObject	[opt] Optional output.  If true, Tag was matched but the specified object wasn't found.
- *
- * @return true if the object parsed successfully
- */
 bool ParseObject( const TCHAR* Stream, const TCHAR* Match, UClass* Class, UObject*& DestRes, UObject* InParent, bool* bInvalidObject )
 {
 	TCHAR TempStr[1024];
@@ -892,20 +797,6 @@ bool ParseObject( const TCHAR* Stream, const TCHAR* Match, UClass* Class, UObjec
 	}
 }
 
-/**
- * Find or load an object by string name with optional outer and filename specifications.
- * These are optional because the InName can contain all of the necessary information.
- *
- * @param ObjectClass	The class (or a superclass) of the object to be loaded.
- * @param InOuter		An optional object to narrow where to find/load the object from
- * @param InName		String name of the object. If it's not fully qualified, InOuter and/or Filename will be needed
- * @param Filename		An optional file to load from (or find in the file's package object)
- * @param LoadFlags		Flags controlling how to handle loading from disk
- * @param Sandbox		A list of packages to restrict the search for the object
- * @param bAllowObjectReconciliation	Whether to allow the object to be found via FindObject in the case of seek free loading
- *
- * @return The object that was loaded or found. NULL for a failure.
- */
 UObject* StaticLoadObjectInternal(UClass* ObjectClass, UObject* InOuter, const TCHAR* InName, const TCHAR* Filename, uint32 LoadFlags, UPackageMap* Sandbox, bool bAllowObjectReconciliation, FUObjectSerializeContext* InSerializeContext)
 {
 	SCOPE_CYCLE_COUNTER(STAT_LoadObject);
@@ -1148,15 +1039,6 @@ public:
 
 #endif
 
-/**
-* Loads a package and all contained objects that match context flags.
-*
-* @param	InOuter		Package to load new package into (usually NULL or ULevel->GetOuter())
-* @param	Filename	Long package name to load.
-* @param	LoadFlags	Flags controlling loading behavior
-* @param	ImportLinker	Linker that requests this package through one of its imports
-* @return	Loaded package if successful, NULL otherwise
-*/
 UPackage* LoadPackageInternal(UPackage* InOuter, const TCHAR* InLongPackageNameOrFilename, uint32 LoadFlags, FLinkerLoad* ImportLinker, FArchive* InReaderOverride, FUObjectSerializeContext* InLoadContext)
 {
 	DECLARE_SCOPE_CYCLE_COUNTER(TEXT("LoadPackageInternal"), STAT_LoadPackageInternal, STATGROUP_ObjectVerbose);
@@ -1886,17 +1768,6 @@ void EndLoad(FUObjectSerializeContext* LoadContext)
 	Object name functions.
 -----------------------------------------------------------------------------*/
 
-/**
- * Create a unique name by combining a base name and an arbitrary number string.
- * The object name returned is guaranteed not to exist.
- *
- * @param	Parent		the outer for the object that needs to be named
- * @param	Class		the class for the object
- * @param	BaseName	optional base name to use when generating the unique object name; if not specified, the class's name is used
- *
- * @return	name is the form BaseName_##, where ## is the number of objects of this
- *			type that have been created since the last time the class was garbage collected.
- */
 FName MakeUniqueObjectName( UObject* Parent, UClass* Class, FName InBaseName/*=NAME_None*/ )
 {
 	CSV_SCOPED_TIMING_STAT(UObject, MakeUniqueObjectName);
@@ -1977,18 +1848,6 @@ FName MakeUniqueObjectName( UObject* Parent, UClass* Class, FName InBaseName/*=N
 	return TestName;
 }
 
-/**
- * Given a display label string, generates an FName slug that is a valid FName for that label.
- * If the object's current name is already satisfactory, then that name will be returned.
- * For example, "[MyObject]: Object Label" becomes "MyObjectObjectLabel" FName slug.
- * 
- * Note: The generated name isn't guaranteed to be unique.
- *
- * @param DisplayLabel The label string to convert to an FName
- * @param CurrentObjectName The object's current name, or NAME_None if it has no name yet
- *
- * @return	The generated object name
- */
 FName MakeObjectNameFromDisplayLabel(const FString& DisplayLabel, const FName CurrentObjectName)
 {
 	FString GeneratedName = SlugStringForValidName(DisplayLabel);
@@ -2292,18 +2151,6 @@ UObject* StaticDuplicateObjectEx( FObjectDuplicationParameters& Parameters )
 	return DupRootObject;
 }
 
-/**
- * Save a copy of this object into the transaction buffer if we are currently recording into
- * one (undo/redo). If bMarkDirty is true, will also mark the package as needing to be saved.
- *
- * @param	bMarkDirty	If true, marks the package dirty if we are currently recording into a
- *						transaction buffer
- * @param	Object		object to save.
-*
- * @return	true if a copy of the object was saved and the package potentially marked dirty; false
- *			if we are not recording into a transaction buffer, the package is a PIE/script package,
- *			or the object is not transactional (implies the package was not marked dirty)
- */
 bool SaveToTransactionBuffer(UObject* Object, bool bMarkDirty)
 {
 	bool bSavedToTransactionBuffer = false;
@@ -2330,13 +2177,6 @@ bool SaveToTransactionBuffer(UObject* Object, bool bMarkDirty)
 	return bSavedToTransactionBuffer;
 }
 
-/**
- * Causes the transaction system to emit a snapshot event for the given object if the following conditions are met:
- *  a) The object is currently transacting.
- *  b) The object has changed since it started transacting.
- *
- * @param	Object		object to snapshot.
- */
 void SnapshotTransactionBuffer(UObject* Object)
 {
 	// Script packages should not end up in the transaction buffer.
@@ -2351,15 +2191,6 @@ void SnapshotTransactionBuffer(UObject* Object)
 	}
 }
 
-/**
- * Check for StaticAllocateObject error; only for use with the editor, make or other commandlets.
- * 
- * @param	Class		the class of the object to create
- * @param	InOuter		the object to create this object within (the Outer property for the new object will be set to the value specified here).
- * @param	Name		the name to give the new object. If no value (NAME_None) is specified, the object will be given a unique name in the form of ClassName_#.
- * @param	SetFlags	the ObjectFlags to assign to the new object. some flags can affect the behavior of constructing the object.
- * @return	true if NULL should be returned; there was a problem reported 
- */
 bool StaticAllocateObjectErrorTests( UClass* InClass, UObject* InOuter, FName InName, EObjectFlags InFlags)
 {
 	// Validation checks.
@@ -3731,16 +3562,6 @@ private:
 	FReferencerInformation *CurrentReferenceInfo;
 };
 
-/**
- * Returns whether an object is referenced, not counting the one
- * reference at Obj.
- *
- * @param	Obj			Object to check
- * @param	KeepFlags	Objects with these flags will be considered as being referenced. If RF_NoFlags, check all objects
- * @param	bCheckSubObjects	Treat subobjects as if they are the same as passed in object
- * @param	FoundReferences		If non-NULL fill in with list of objects that hold references
- * @return true if object is referenced, false otherwise
- */
 bool IsReferenced(UObject*& Obj, EObjectFlags KeepFlags, EInternalObjectFlags InternalKeepFlags, bool bCheckSubObjects, FReferencerInformationList* FoundReferences)
 {
 	check(!Obj->IsUnreachable());
diff --git a/Engine/Source/Runtime/CoreUObject/Public/UObject/UObjectGlobals.h b/Engine/Source/Runtime/CoreUObject/Public/UObject/UObjectGlobals.h
--- a/Engine/Source/Runtime/CoreUObject/Public/UObject/UObjectGlobals.h
+++ b/Engine/Source/Runtime/CoreUObject/Public/UObject/UObjectGlobals.h
@@ -1,7 +1,7 @@
 // Copyright 1998-2019 Epic Games, Inc. All Rights Reserved.
 
 /*=============================================================================
-	UnObjGlobals.h: Unreal object system globals.
+	UObjectGlobals.h: Unreal object system globals.
 =============================================================================*/
 
 #pragma once
@@ -42,27 +42,26 @@ DECLARE_DWORD_COUNTER_STAT_EXTERN(TEXT("FindObjectFast"),STAT_FindObjectFast,STA
 DECLARE_CYCLE_STAT_EXTERN(TEXT("NetSerializeFast Array"),STAT_NetSerializeFastArray,STATGROUP_ServerCPU, COREUOBJECT_API);
 DECLARE_CYCLE_STAT_EXTERN(TEXT("NetSerializeFast Array BuildMap"),STAT_NetSerializeFastArray_BuildMap,STATGROUP_ServerCPU, COREUOBJECT_API);
 
-
-
 #define	INVALID_OBJECT	(UObject*)-1
 
 
-
 // Private system wide variables.
 
-/** set while in SavePackage() to detect certain operations that are illegal while saving */
+/** Set while in SavePackage() to detect certain operations that are illegal while saving */
 extern COREUOBJECT_API bool					GIsSavingPackage;
 /** This allows loading unversioned cooked content in the editor */
 extern COREUOBJECT_API bool					GAllowUnversionedContentInEditor;
 
+/** Enum used in StaticDuplicateObject() and related functions to describe why something is being duplicated */
 namespace EDuplicateMode
 {
 	enum Type
 	{
+		/** No specific information about the reason for duplication */
 		Normal,
-		// Object is being duplicated as part of a world duplication
+		/** Object is being duplicated as part of a world duplication */
 		World,
-		// Object is being duplicated as part of a world duplication for PIE
+		/** Object is being duplicated as part of the process for entering Play In Editor */
 		PIE
 	};
 };
@@ -93,27 +92,27 @@ struct FObjectDuplicationParameters
 	FName			DestName;
 
 	/**
-	 * a bitmask of EObjectFlags to propagate to the duplicate of SourceObject (and its subobjects).
+	 * A bitmask of EObjectFlags to propagate to the duplicate of SourceObject (and its subobjects).
 	 */
 	EObjectFlags	FlagMask;
 
 	/**
-	* a bitmask of EInternalObjectFlags to propagate to the duplicate of SourceObject (and its subobjects).
-	*/
+	 * A bitmask of EInternalObjectFlags to propagate to the duplicate of SourceObject (and its subobjects).
+	 */
 	EInternalObjectFlags InternalFlagMask;
 
 	/**
-	 * a bitmask of EObjectFlags to set on each duplicate object created.  Different from FlagMask in that only the bits
+	 * A bitmask of EObjectFlags to set on each duplicate object created.  Different from FlagMask in that only the bits
 	 * from FlagMask which are also set on the source object will be set on the duplicate, while the flags in this value
 	 * will always be set.
 	 */
 	EObjectFlags	ApplyFlags;
 
 	/**
-	* a bitmask of EInternalObjectFlags to set on each duplicate object created.  Different from FlagMask in that only the bits
-	* from FlagMask which are also set on the source object will be set on the duplicate, while the flags in this value
-	* will always be set.
-	*/
+	 * A bitmask of EInternalObjectFlags to set on each duplicate object created.  Different from FlagMask in that only the bits
+	 * from FlagMask which are also set on the source object will be set on the duplicate, while the flags in this value
+	 * will always be set.
+	 */
 	EInternalObjectFlags	ApplyInternalFlags;
 
 	/**
@@ -124,8 +123,8 @@ struct FObjectDuplicationParameters
 	EDuplicateMode::Type DuplicateMode;
 
 	/**
-	 * optional class to specify for the destination object.
-	 * @note: MUST BE SERIALIZATION COMPATIBLE WITH SOURCE OBJECT, AND DOES NOT WORK WELL FOR OBJECT WHICH HAVE COMPLEX COMPONENT HIERARCHIES!!!
+	 * Optional class to specify for the destination object.
+	 * @warning: MUST BE SERIALIZATION COMPATIBLE WITH SOURCE OBJECT, AND DOES NOT WORK WELL FOR OBJECT WHICH HAVE COMPLEX COMPONENT HIERARCHIES!!!
 	 */
 	UClass*			DestClass;
 
@@ -139,7 +138,7 @@ struct FObjectDuplicationParameters
 	TMap<UObject*,UObject*>	DuplicationSeed;
 
 	/**
-	 * If non-NULL, this will be filled with the list of objects created during the call to StaticDuplicateObject.
+	 * If non-null, this will be filled with the list of objects created during the call to StaticDuplicateObject.
 	 *
 	 * Key will be the source object; value will be the duplicated object
 	 */
@@ -151,15 +150,16 @@ struct FObjectDuplicationParameters
 	COREUOBJECT_API FObjectDuplicationParameters( UObject* InSourceObject, UObject* InDestOuter );
 };
 
+/** Parses a bit mask of property flags into an array of string literals that match the flags */
 COREUOBJECT_API TArray<const TCHAR*> ParsePropertyFlags(EPropertyFlags Flags);
 
+/** Returns the transient top-level package, which is useful for temporarily storing objects that should never be saved */
 COREUOBJECT_API UPackage* GetTransientPackage();
 
 /**
  * Gets INI file name from object's reference if it contains one. 
  *
- * @returns If object reference doesn't contain any INI reference the function
- *		returns nullptr. Otherwise a ptr to INI's file name.
+ * @returns If object reference doesn't contain any INI reference the function returns nullptr. Otherwise a ptr to INI's file name.
  */
 COREUOBJECT_API const FString* GetIniFilenameFromObjectsReference(const FString& ObjectsReferenceString);
 
@@ -174,42 +174,76 @@ COREUOBJECT_API const FString* GetIniFilenameFromObjectsReference(const FString&
  */
 COREUOBJECT_API FString ResolveIniObjectsReference(const FString& ObjectReference, const FString* IniFilename = nullptr, bool bThrow = false);
 
+/**
+ * Internal function that takes a fully qualified or relative object path string and converts it into a path relative to a package.
+ * Normally, you should call one of the FindObject or LoadObject functions instead.
+ *
+ * @param	Outer					The package to search within. If null, ObjectsReferenceString be a globally scoped path and this will be filled in with the actual package if found/created
+ * @param	ObjectsReferenceString	The object path string to resolve. If it is successfully resolved, this will be replaced with a path relative to Outer
+ * @param	Create					If true, it will try to load or create the required package if it is not in memory
+ * @param	Throw					If true, it will potentially raise an error if the object cannot be found
+ * @param	LoadFlags				Flags to use if Create is true and it needs to load a package, from the ELoadFlags enum
+ * @param	InLoadContext			Additional context when called during serialization
+ *
+ * @return	True if the name was successfully resolved
+ */
 COREUOBJECT_API bool ResolveName(UObject*& Outer, FString& ObjectsReferenceString, bool Create, bool Throw, uint32 LoadFlags = LOAD_None, FUObjectSerializeContext* InLoadContext = nullptr);
+
+/** Internal function used to possibly output an error message, taking into account the outer and LoadFlags */
 COREUOBJECT_API void SafeLoadError( UObject* Outer, uint32 LoadFlags, const TCHAR* ErrorMessage);
 
+/** Internal function used to update the suffix to be given to the next newly-created unnamed object. */
 COREUOBJECT_API int32 UpdateSuffixForNextNewObject(UObject* Parent, UClass* Class, TFunctionRef<void(int32&)> IndexMutator);
 
+
 /**
- * Fast version of StaticFindObject that relies on the passed in FName being the object name
- * without any group/ package qualifiers.
+ * Fast version of StaticFindObject that relies on the passed in FName being the object name without any group/package qualifiers.
+ * This will only find top level packages or subobjects nested directly within a passed in outer.
  *
  * @param	Class			The to be found object's class
- * @param	InOuter			The to be found object's outer
- * @param	InName			The to be found object's class
+ * @param	InOuter			Outer object to look inside, if null this will only look for top level packages
+ * @param	InName			Object name to look for relative to InOuter
  * @param	ExactClass		Whether to require an exact match with the passed in class
  * @param	AnyPackage		Whether to look in any package
  * @param	ExclusiveFlags	Ignores objects that contain any of the specified exclusive flags
- * @param ExclusiveInternalFlags  Ignores objects that contain any of the specified internal exclusive flags
- * @return	Returns a pointer to the found object or NULL if none could be found
+ * @param	ExclusiveInternalFlags  Ignores objects that contain any of the specified internal exclusive flags
+ *
+ * @return	Returns a pointer to the found object or null if none could be found
  */
 COREUOBJECT_API UObject* StaticFindObjectFast(UClass* Class, UObject* InOuter, FName InName, bool ExactClass = false, bool AnyPackage = false, EObjectFlags ExclusiveFlags = RF_NoFlags, EInternalObjectFlags ExclusiveInternalFlags = EInternalObjectFlags::None);
+
+/**
+ * Tries to find an object in memory. This will handle fully qualified paths of the form /path/packagename.object:subobject and resolve references for you.
+ *
+ * @param	Class			The to be found object's class
+ * @param	InOuter			Outer object to look inside. If this is ANY_PACKAGE it will search all in memory packages, if this is null then InName should start with a package name
+ * @param	InName			The object path to search for an object, relative to InOuter
+ * @param	ExactClass		Whether to require an exact match with the passed in class
+ *
+ * @return	Returns a pointer to the found object or nullptr if none could be found
+ */
 COREUOBJECT_API UObject* StaticFindObject( UClass* Class, UObject* InOuter, const TCHAR* Name, bool ExactClass=false );
+
+/** Version of StaticFindObject() that will assert if the object is not found */
 COREUOBJECT_API UObject* StaticFindObjectChecked( UClass* Class, UObject* InOuter, const TCHAR* Name, bool ExactClass=false );
+
+/** Internal version of StaticFindObject that will not assert on GIsSavingPackage or IsGarbageCollecting() */
 COREUOBJECT_API UObject* StaticFindObjectSafe( UClass* Class, UObject* InOuter, const TCHAR* Name, bool ExactClass=false );
 
 
 /**
- * Parse an object from a text representation
+ * Parse a reference to an object from a text representation
  *
- * @param Stream		String containing text to parse
- * @param Class			Tag to search for object representation within string
- * @param Class			The class of the object to be loaded.
- * @param DestRes		returned uobject pointer
- * @param InParent		Outer to search
+ * @param Stream			String containing text to parse
+ * @param Match				Tag to search for object representation within string
+ * @param Class				The class of the object to be loaded.
+ * @param DestRes			Returned object pointer
+ * @param InParent			Outer to search
+ * @param bInvalidObject	[opt] Optional output.  If true, Tag was matched but the specified object wasn't found.
  *
- * @return true if the object parsed successfully
+ * @return True if the object parsed successfully, even if object was not found
  */
-COREUOBJECT_API bool ParseObject( const TCHAR* Stream, const TCHAR* Match, UClass* Class, UObject*& DestRes, UObject* InParent, bool* bInvalidObject=NULL );
+COREUOBJECT_API bool ParseObject( const TCHAR* Stream, const TCHAR* Match, UClass* Class, UObject*& DestRes, UObject* InParent, bool* bInvalidObject=nullptr );
 
 /**
  * Find or load an object by string name with optional outer and filename specifications.
@@ -217,77 +251,104 @@ COREUOBJECT_API bool ParseObject( const TCHAR* Stream, const TCHAR* Match, UClas
  *
  * @param ObjectClass	The class (or a superclass) of the object to be loaded.
  * @param InOuter		An optional object to narrow where to find/load the object from
- * @param InName		String name of the object. If it's not fully qualified, InOuter and/or Filename will be needed
+ * @param Name			String name of the object. If it's not fully qualified, InOuter and/or Filename will be needed
  * @param Filename		An optional file to load from (or find in the file's package object)
- * @param LoadFlags		Flags controlling how to handle loading from disk
+ * @param LoadFlags		Flags controlling how to handle loading from disk, from the ELoadFlags enum
  * @param Sandbox		A list of packages to restrict the search for the object
  * @param bAllowObjectReconciliation	Whether to allow the object to be found via FindObject in the case of seek free loading
  * @param InSerializeContext	Additional context when called during serialization
  *
- * @return The object that was loaded or found. NULL for a failure.
+ * @return The object that was loaded or found. nullptr for a failure.
  */
-COREUOBJECT_API UObject* StaticLoadObject( UClass* Class, UObject* InOuter, const TCHAR* Name, const TCHAR* Filename = NULL, uint32 LoadFlags = LOAD_None, UPackageMap* Sandbox = NULL, bool bAllowObjectReconciliation = true, FUObjectSerializeContext* InSerializeContext = nullptr );
-COREUOBJECT_API UClass* StaticLoadClass(UClass* BaseClass, UObject* InOuter, const TCHAR* Name, const TCHAR* Filename = NULL, uint32 LoadFlags = LOAD_None, UPackageMap* Sandbox = NULL);
+COREUOBJECT_API UObject* StaticLoadObject( UClass* Class, UObject* InOuter, const TCHAR* Name, const TCHAR* Filename = nullptr, uint32 LoadFlags = LOAD_None, UPackageMap* Sandbox = nullptr, bool bAllowObjectReconciliation = true, FUObjectSerializeContext* InSerializeContext = nullptr );
+
+/** Version of StaticLoadObject() that will load classes */
+COREUOBJECT_API UClass* StaticLoadClass(UClass* BaseClass, UObject* InOuter, const TCHAR* Name, const TCHAR* Filename = nullptr, uint32 LoadFlags = LOAD_None, UPackageMap* Sandbox = nullptr);
 
 /**
-* Create a new instance of an object.  The returned object will be fully initialized.  If InFlags contains RF_NeedsLoad (indicating that the object still needs to load its object data from disk), components
-* are not instanced (this will instead occur in PostLoad()).  The different between StaticConstructObject and StaticAllocateObject is that StaticConstructObject will also call the class constructor on the object
-* and instance any components.
-*
-* @param	Class		the class of the object to create
-* @param	InOuter		the object to create this object within (the Outer property for the new object will be set to the value specified here).
-* @param	Name		the name to give the new object. If no value (NAME_None) is specified, the object will be given a unique name in the form of ClassName_#.
-* @param	SetFlags	the ObjectFlags to assign to the new object. some flags can affect the behavior of constructing the object.
-* @param	InternalSetFlags	the InternalObjectFlags to assign to the new object. some flags can affect the behavior of constructing the object.
-* @param	Template	if specified, the property values from this object will be copied to the new object, and the new object's ObjectArchetype value will be set to this object.
-*						If NULL, the class default object is used instead.
-* @param	bInCopyTransientsFromClassDefaults - if true, copy transient from the class defaults instead of the pass in archetype ptr (often these are the same)
-* @param	InstanceGraph
-*						contains the mappings of instanced objects and components to their templates
-*
-* @return	a pointer to a fully initialized object of the specified class.
-*/
-COREUOBJECT_API UObject* StaticConstructObject_Internal(UClass* Class, UObject* InOuter = (UObject*)GetTransientPackage(), FName Name = NAME_None, EObjectFlags SetFlags = RF_NoFlags, EInternalObjectFlags InternalSetFlags = EInternalObjectFlags::None, UObject* Template = NULL, bool bCopyTransientsFromClassDefaults = false, struct FObjectInstancingGraph* InstanceGraph = NULL, bool bAssumeTemplateIsArchetype = false);
+ * Create a new instance of an object.  The returned object will be fully initialized.  If InFlags contains RF_NeedsLoad (indicating that the object still needs to load its object data from disk), components
+ * are not instanced (this will instead occur in PostLoad()).  The different between StaticConstructObject and StaticAllocateObject is that StaticConstructObject will also call the class constructor on the object
+ * and instance any components.
+ *
+ * @param	Class		The class of the object to create
+ * @param	InOuter		The object to create this object within (the Outer property for the new object will be set to the value specified here).
+ * @param	Name		The name to give the new object. If no value (NAME_None) is specified, the object will be given a unique name in the form of ClassName_#.
+ * @param	SetFlags	The ObjectFlags to assign to the new object. some flags can affect the behavior of constructing the object.
+ * @param	InternalSetFlags	The InternalObjectFlags to assign to the new object. some flags can affect the behavior of constructing the object.
+ * @param	Template	If specified, the property values from this object will be copied to the new object, and the new object's ObjectArchetype value will be set to this object.
+ *						If nullptr, the class default object is used instead.
+ * @param	bInCopyTransientsFromClassDefaults	If true, copy transient from the class defaults instead of the pass in archetype ptr (often these are the same)
+ * @param	InstanceGraph	Contains the mappings of instanced objects and components to their templates
+ * @param	bAssumeTemplateIsArchetype	If true, Template is guaranteed to be an archetype
+ *
+ * @return	A pointer to a fully initialized object of the specified class.
+ */
+COREUOBJECT_API UObject* StaticConstructObject_Internal(UClass* Class, UObject* InOuter = (UObject*)GetTransientPackage(), FName Name = NAME_None, EObjectFlags SetFlags = RF_NoFlags, EInternalObjectFlags InternalSetFlags = EInternalObjectFlags::None, UObject* Template = nullptr, bool bCopyTransientsFromClassDefaults = false, struct FObjectInstancingGraph* InstanceGraph = nullptr, bool bAssumeTemplateIsArchetype = false);
 
 /**
  * Creates a copy of SourceObject using the Outer and Name specified, as well as copies of all objects contained by SourceObject.  
- * Any objects referenced by SourceOuter or RootObject and contained by SourceOuter are also copied, maintaining their name relative to SourceOuter.  Any
- * references to objects that are duplicated are automatically replaced with the copy of the object.
+ * Any objects referenced by SourceOuter or RootObject and contained by SourceOuter are also copied, maintaining their name relative to SourceOuter.
+ * Any references to objects that are duplicated are automatically replaced with the copy of the object.
  *
- * @param	SourceObject	the object to duplicate
- * @param	RootObject		should always be the same value as SourceObject (unused)
- * @param	DestOuter		the object to use as the Outer for the copy of SourceObject
- * @param	DestName		the name to use for the copy of SourceObject
- * @param	FlagMask		a bitmask of EObjectFlags that should be propagated to the object copies.  The resulting object copies will only have the object flags
+ * @param	SourceObject	The object to duplicate
+ * @param	DestOuter		The object to use as the Outer for the copy of SourceObject
+ * @param	DestName		The name to use for the copy of SourceObject, if none it will be autogenerated
+ * @param	FlagMask		A bitmask of EObjectFlags that should be propagated to the object copies.  The resulting object copies will only have the object flags
  *							specified copied from their source object.
- * @param	DestClass		optional class to specify for the destination object. MUST BE SERIALIZATION COMPATIBLE WITH SOURCE OBJECT!!!
- * @param	InternalFlagsMask  bitmask of EInternalObjectFlags that should be propagated to the object copies.
+ * @param	DestClass		Optional class to specify for the destination object. MUST BE SERIALIZATION COMPATIBLE WITH SOURCE OBJECT!!!
+ * @param	InternalFlagsMask  Bitmask of EInternalObjectFlags that should be propagated to the object copies.
  *
- * @return	the duplicate of SourceObject.
+ * @return	The duplicate of SourceObject.
  *
- * @note: this version is deprecated in favor of StaticDuplicateObjectEx
+ * @deprecated This version is deprecated in favor of StaticDuplicateObjectEx
  */
 COREUOBJECT_API UObject* StaticDuplicateObject(UObject const* SourceObject, UObject* DestOuter, const FName DestName = NAME_None, EObjectFlags FlagMask = RF_AllFlags, UClass* DestClass = nullptr, EDuplicateMode::Type DuplicateMode = EDuplicateMode::Normal, EInternalObjectFlags InternalFlagsMask = EInternalObjectFlags::AllFlags);
+
+/**
+ * Creates a copy of SourceObject using the Outer and Name specified, as well as copies of all objects contained by SourceObject.
+ * Any objects referenced by SourceOuter or RootObject and contained by SourceOuter are also copied, maintaining their name relative to SourceOuter.
+ * Any references to objects that are duplicated are automatically replaced with the copy of the object.
+ *
+ * @param	Parameters  Specific options to use when duplicating this object
+ *
+ * @return	The duplicate of SourceObject.
+ */
 COREUOBJECT_API UObject* StaticDuplicateObjectEx( struct FObjectDuplicationParameters& Parameters );
 
 /** 
- *   Iterate over all objects considered part of the root to setup GC optimizations
+ * Parses a global context system console or debug command and executes it.
+ *
+ * @param	InWorld		The world to use as a context, enables certain commands
+ * @param	Cmd			Command string to execute
+ * @param	Ar			Output device to write results of commands to
+ * 
+ * @return	True if the command was successfully parsed
  */
 COREUOBJECT_API bool StaticExec( UWorld* InWorld, const TCHAR* Cmd, FOutputDevice& Ar=*GLog );
+
+/**
+ * Static UObject tick function, used to verify certain key assumptions and to tick the async loading code.
+ *
+ * @param DeltaTime	Time in seconds since last call
+ * @param bUseFullTimeLimit	If true, use the entire time limit even if blocked on I/O
+ * @param AsyncLoadingTime Time in seconds to use for async loading limit
+ */
 COREUOBJECT_API void StaticTick( float DeltaTime, bool bUseFullTimeLimit = true, float AsyncLoadingTime = 0.005f );
 
 /**
  * Loads a package and all contained objects that match context flags.
  *
- * @param	InOuter		Package to load new package into (usually NULL or ULevel->GetOuter())
- * @param	Filename	Long package name to load
- * @param	LoadFlags	Flags controlling loading behavior
- * @param InLoadContext	Additional context when called during serialization
- * @return	Loaded package if successful, NULL otherwise
+ * @param	InOuter				Package to load new package into (usually nullptr or ULevel->GetOuter())
+ * @param	InLongPackageName	Long package name to load
+ * @param	LoadFlags			Flags controlling loading behavior, from the ELoadFlags enum
+ * @param	InReaderOverride	Optional archive to use for reading package data
+ * @param	InLoadContext		Additional context when called during serialization
+ *
+ * @return	Loaded package if successful, nullptr otherwise
  */
 COREUOBJECT_API UPackage* LoadPackage( UPackage* InOuter, const TCHAR* InLongPackageName, uint32 LoadFlags, FArchive* InReaderOverride = nullptr, FUObjectSerializeContext* InLoadContext = nullptr );
 
-/* Async package loading result */
+/** Async package loading result */
 namespace EAsyncLoadingResult
 {
 	enum Type
@@ -307,33 +368,38 @@ typedef int32 TAsyncLoadPriority;
 /**
  * Delegate called on completion of async package loading
  * @param	PackageName			Package name we were trying to load
- * @param	LoadedPackage		Loaded package if successful, NULL otherwise	
+ * @param	LoadedPackage		Loaded package if successful, nullptr otherwise	
  * @param	Result		Result of async loading.
  */
 DECLARE_DELEGATE_ThreeParams(FLoadPackageAsyncDelegate, const FName& /*PackageName*/, UPackage* /*LoadedPackage*/, EAsyncLoadingResult::Type /*Result*/)
 
 /**
-* Asynchronously load a package and all contained objects that match context flags. Non- blocking.
-*
-* @param	InName					Name of package to load
-* @param	InGuid					GUID of the package to load, or NULL for "don't care"
-* @param	InPackageToLoadFrom		If non-null, this is another package name. We load from this package name, into a (probably new) package named PackageName
-* @param	InCompletionDelegate	Delegate to be invoked when the packages has finished streaming
-* @param	InFlags					Package flags
-* @param	InPIEInstanceID			PIE instance ID
-* @param	InPackagePriority		Loading priority
-* @return Unique ID associated with this load request (the same package can be associated with multiple IDs).
-*/
+ * Asynchronously load a package and all contained objects that match context flags. Non-blocking.
+ * This version is useful when loading multiple copies of the same package.
+ *
+ * @param	InName					Name of package to load
+ * @param	InGuid					GUID of the package to load, or nullptr for "don't care"
+ * @param	InPackageToLoadFrom		If non-null, this is another package name. We load from this package name, into a (probably new) package named InName
+ * @param	InCompletionDelegate	Delegate to be invoked when the packages has finished streaming
+ * @param	InPackageFlags			Package flags used to construct loaded package in memory
+ * @param	InPIEInstanceID			Play in Editor instance ID
+ * @param	InPackagePriority		Loading priority
+ * @return Unique ID associated with this load request (the same package can be associated with multiple IDs).
+ */
 COREUOBJECT_API int32 LoadPackageAsync(const FString& InName, const FGuid* InGuid = nullptr, const TCHAR* InPackageToLoadFrom = nullptr, FLoadPackageAsyncDelegate InCompletionDelegate = FLoadPackageAsyncDelegate(), EPackageFlags InPackageFlags = PKG_None, int32 InPIEInstanceID = INDEX_NONE, TAsyncLoadPriority InPackagePriority = 0);
 
 /**
-* Asynchronously load a package and all contained objects that match context flags. Non- blocking.
-*
-* @param	InName					Name of package to load
-* @param	InCompletionDelegate	Delegate to be invoked when the packages has finished streaming
-* @param	InPackagePriority		Loading priority
-* @return Unique ID associated with this load request (the same package can be associated with multiple IDs).
-*/
+ * Asynchronously load a package and all contained objects that match context flags. Non-blocking.
+ *
+ * @param	InName					Name of package to load
+ * @param	InCompletionDelegate	Delegate to be invoked when the packages has finished streaming
+ * @param	InPackagePriority		Loading priority
+ * @param	InPackageFlags			Package flags used to construct loaded package in memory
+ * @param	InPIEInstanceID			Play in Editor instance ID
+ * @return Unique ID associated with this load request (the same package can be associated with multiple IDs).
+ *
+ * @see FStreamableManager for an engine-level wrapper
+ */
 COREUOBJECT_API int32 LoadPackageAsync(const FString& InName, FLoadPackageAsyncDelegate InCompletionDelegate, TAsyncLoadPriority InPackagePriority = 0, EPackageFlags InPackageFlags = PKG_None, int32 InPIEInstanceID = INDEX_NONE);
 
 /**
@@ -353,7 +419,7 @@ COREUOBJECT_API bool IsEventDrivenLoaderEnabled();
 
 /**
  * Returns the async load percentage for a package in flight with the passed in name or -1 if there isn't one.
- * THIS IS SLOW. MAY BLOCK ASYNC LOADING.
+ * @warning THIS IS SLOW. MAY BLOCK ASYNC LOADING.
  *
  * @param	PackageName			Name of package to query load percentage for
  * @return	Async load percentage if package is currently being loaded, -1 otherwise
@@ -372,6 +438,7 @@ COREUOBJECT_API bool IsGarbageCollecting();
  * @param	bPerformFullPurge	if true, perform a full purge after the mark pass
  */
 COREUOBJECT_API void CollectGarbage(EObjectFlags KeepFlags, bool bPerformFullPurge = true);
+
 /**
 * Performs garbage collection only if no other thread holds a lock on GC
 *
@@ -383,6 +450,9 @@ COREUOBJECT_API bool TryCollectGarbage(EObjectFlags KeepFlags, bool bPerformFull
 /**
 * Calls ConditionalBeginDestroy on unreachable objects
 *
+* @param	bUseTimeLimit	whether the time limit parameter should be used
+* @param	TimeLimit		soft time limit for this function call
+*
 * @return true if the time limit passed and there's still objects pending to be unhashed
 */
 COREUOBJECT_API bool UnhashUnreachableObjects(bool bUseTimeLimit, float TimeLimit = 0.0f);
@@ -440,17 +510,16 @@ COREUOBJECT_API FName MakeUniqueObjectName( UObject* Outer, UClass* Class, FName
 COREUOBJECT_API FName MakeObjectNameFromDisplayLabel(const FString& DisplayLabel, const FName CurrentObjectName);
 
 /**
- * Returns whether an object is referenced, not counting the one
- * reference at Obj.
+ * Returns whether an object is referenced, not counting references from itself
  *
  * @param	Obj			Object to check
  * @param	KeepFlags	Objects with these flags will be considered as being referenced
-* @param	InternalKeepFlags	Objects with these internal flags will be considered as being referenced
+ * @param	InternalKeepFlags	Objects with these internal flags will be considered as being referenced
  * @param	bCheckSubObjects	Treat subobjects as if they are the same as passed in object
- * @param	FoundReferences		If non-NULL fill in with list of objects that hold references
+ * @param	FoundReferences		If non-nullptr fill in with list of objects that hold references
  * @return true if object is referenced, false otherwise
  */
-COREUOBJECT_API bool IsReferenced( UObject*& Res, EObjectFlags KeepFlags, EInternalObjectFlags InternalKeepFlags, bool bCheckSubObjects = false, FReferencerInformationList* FoundReferences = NULL );
+COREUOBJECT_API bool IsReferenced( UObject*& Res, EObjectFlags KeepFlags, EInternalObjectFlags InternalKeepFlags, bool bCheckSubObjects = false, FReferencerInformationList* FoundReferences = nullptr );
 
 /**
  * Blocks till all pending package/ linker requests are fulfilled.
@@ -461,7 +530,7 @@ COREUOBJECT_API bool IsReferenced( UObject*& Res, EObjectFlags KeepFlags, EInter
 COREUOBJECT_API void FlushAsyncLoading(int32 PackageID = INDEX_NONE);
 
 /**
- * @return number of active async load package requests
+ * Return number of active async load package requests
  */
 COREUOBJECT_API int32 GetNumAsyncPackages();
 
@@ -529,6 +598,7 @@ COREUOBJECT_API UPackage* FindPackage(UObject* InOuter, const TCHAR* PackageName
  */
 COREUOBJECT_API UPackage* CreatePackage( UObject* InOuter, const TCHAR* PackageName );
 
+/** Internal function used to set a specific property value from debug/console code */
 void GlobalSetProperty( const TCHAR* Value, UClass* Class, UProperty* Property, bool bNotifyObjectOfChange );
 
 /**
@@ -561,7 +631,7 @@ COREUOBJECT_API void SnapshotTransactionBuffer(UObject* Object);
  * @param	InOuter		the object to create this object within (the Outer property for the new object will be set to the value specified here).
  * @param	Name		the name to give the new object. If no value (NAME_None) is specified, the object will be given a unique name in the form of ClassName_#.
  * @param	SetFlags	the ObjectFlags to assign to the new object. some flags can affect the behavior of constructing the object.
- * @return	true if NULL should be returned; there was a problem reported 
+ * @return	true if nullptr should be returned; there was a problem reported 
  */
 bool StaticAllocateObjectErrorTests( UClass* Class, UObject* InOuter, FName Name, EObjectFlags SetFlags);
 
@@ -578,9 +648,9 @@ bool StaticAllocateObjectErrorTests( UClass* Class, UObject* InOuter, FName Name
  * @param bOutReusedSubobject	flag indicating if the object is a subobject that has already been created (in which case further initialization is not necessary).
  * @return	a pointer to a fully initialized object of the specified class.
  */
-COREUOBJECT_API UObject* StaticAllocateObject(UClass* Class, UObject* InOuter, FName Name, EObjectFlags SetFlags, EInternalObjectFlags InternalSetFlags = EInternalObjectFlags::None, bool bCanReuseSubobjects = false, bool* bOutReusedSubobject = NULL);
+COREUOBJECT_API UObject* StaticAllocateObject(UClass* Class, UObject* InOuter, FName Name, EObjectFlags SetFlags, EInternalObjectFlags InternalSetFlags = EInternalObjectFlags::None, bool bCanReuseSubobjects = false, bool* bOutReusedSubobject = nullptr);
 
-/** Base class for TSubobjectPtr template. Holds the actual pointer and utility methods. */
+/** @deprecated Use raw pointers or TWeakObjectPtr instead */
 class COREUOBJECT_API FSubobjectPtr
 {
 protected:
@@ -606,17 +676,17 @@ protected:
 	void Set(UObject* InObject);
 
 public:
-	/** Resets the internal pointer to NULL. */
+	/** Resets the internal pointer to nullptr. */
 	FORCEINLINE void Reset()
 	{
-		Set(NULL);
+		Set(nullptr);
 	}
 	/** Gets the pointer to the subobject. */
 	FORCEINLINE UObject* Get() const
 	{
 		return Object;
 	}
-	/** Checks if the subobject != NULL. */
+	/** Checks if the subobject != nullptr. */
 	FORCEINLINE bool IsValid() const
 	{
 		return !!Object && Object != (UObject*)InvalidPtrValue;
@@ -626,7 +696,7 @@ public:
 	{
 		return IsValid();
 	}
-	/** Compare against NULL */
+	/** Compare against nullptr */
 	FORCEINLINE bool operator==(TYPE_OF_NULL Other) const
 	{
 		return !IsValid();
@@ -646,9 +716,7 @@ template<> struct TIsPODType<FSubobjectPtr> { enum { Value = true }; };
 template<> struct TIsZeroConstructType<FSubobjectPtr> { enum { Value = true }; };
 template<> struct TIsWeakPointerType<FSubobjectPtr> { enum { Value = false }; };
 
-/**
- * TSubobjectPtr - Sub-object smart pointer, soon to be deprecated and should no longer be used.
- */
+/** @deprecated Use raw pointers or TWeakObjectPtr instead */
 template <class SubobjectType>
 class TSubobjectPtrDeprecated : public FSubobjectPtr
 {
@@ -698,7 +766,6 @@ template<class T> struct TIsZeroConstructType< TSubobjectPtrDeprecated<T> > { en
 template<class T> struct TIsWeakPointerType< TSubobjectPtrDeprecated<T> > { enum { Value = false }; };
 
 
-
 /**
  * Internal class to finalize UObject creation (initialize properties) after the real C++ constructor is called.
  **/
@@ -718,7 +785,7 @@ public:
 	 * @param	bInShouldInitializeProps false is a special case for changing base classes in UCCMake
 	 * @param	InInstanceGraph passed instance graph
 	 */
-	FObjectInitializer(UObject* InObj, UObject* InObjectArchetype, bool bInCopyTransientsFromClassDefaults, bool bInShouldInitializeProps, struct FObjectInstancingGraph* InInstanceGraph = NULL);
+	FObjectInitializer(UObject* InObj, UObject* InObjectArchetype, bool bInCopyTransientsFromClassDefaults, bool bInShouldInitializeProps, struct FObjectInstancingGraph* InInstanceGraph = nullptr);
 
 	~FObjectInitializer();
 
@@ -866,7 +933,7 @@ public:
 	FObjectInitializer const& DoNotCreateDefaultSubobject(FName SubobjectName) const
 	{
 		AssertIfSubobjectSetupIsNotAllowed(*SubobjectName.GetPlainNameString());
-		ComponentOverrides.Add(SubobjectName, NULL, *this);
+		ComponentOverrides.Add(SubobjectName, nullptr, *this);
 		return *this;
 	}
 
@@ -877,7 +944,7 @@ public:
 	FORCEINLINE FObjectInitializer const& DoNotCreateDefaultSubobject(TCHAR const*SubobjectName) const
 	{
 		AssertIfSubobjectSetupIsNotAllowed(SubobjectName);
-		ComponentOverrides.Add(SubobjectName, NULL, *this);
+		ComponentOverrides.Add(SubobjectName, nullptr, *this);
 		return *this;
 	}
 
@@ -967,7 +1034,7 @@ private:
 				ObjectInitializer.IslegalOverride(InComponentName, Overrides[Index].ComponentClass, InComponentClass); // if a base class is asking for an override, the existing override (which we are going to use) had better be derived
 			}
 		}
-		/**  Retrieve an override, or TClassToConstructByDefault::StaticClass or NULL if this was removed by a derived class **/
+		/**  Retrieve an override, or TClassToConstructByDefault::StaticClass or nullptr if this was removed by a derived class **/
 		UClass* Get(FName InComponentName, UClass* ReturnType, UClass* ClassToConstructByDefault, FObjectInitializer const& ObjectInitializer)
 		{
 			int32 Index = Find(InComponentName);
@@ -982,9 +1049,9 @@ private:
 				{
 					return Overrides[Index].ComponentClass; // the override is of an acceptable class, so use it
 				}
-				// else return NULL; this is a unacceptable override
+				// else return nullptr; this is a unacceptable override
 			}
-			return NULL;  // the override is of NULL, which means "don't create this component"
+			return nullptr;  // the override is of nullptr, which means "don't create this component"
 		}
 private:
 		/**  Search for an override **/
@@ -1126,10 +1193,6 @@ public:
 	}
 };
 
-/**
- * Helper class for deferred execution of 
-*/
-
 #if DO_CHECK
 /** Called by NewObject to make sure Child is actually a child of Parent */
 COREUOBJECT_API void CheckIsClassChildOf_Internal(UClass* Parent, UClass* Child);
@@ -1202,15 +1265,15 @@ FUNCTION_NON_NULL_RETURN_END
 template< class T >
 T* DuplicateObject(T const* SourceObject,UObject* Outer, const FName Name = NAME_None)
 {
-	if (SourceObject != NULL)
+	if (SourceObject != nullptr)
 	{
-		if (Outer == NULL || Outer == INVALID_OBJECT)
+		if (Outer == nullptr || Outer == INVALID_OBJECT)
 		{
 			Outer = (UObject*)GetTransientPackage();
 		}
 		return (T*)StaticDuplicateObject(SourceObject,Outer,Name);
 	}
-	return NULL;
+	return nullptr;
 }
 
 /**
@@ -1227,79 +1290,102 @@ COREUOBJECT_API FString GetConfigFilename( UObject* SourceObject );
 	Core templates.
 ----------------------------------------------------------------------------*/
 
-// Parse an object name in the input stream.
+/** Parse a reference to an object from the input stream. */
 template< class T > 
-inline bool ParseObject( const TCHAR* Stream, const TCHAR* Match, T*& Obj, UObject* Outer, bool* bInvalidObject=NULL )
+inline bool ParseObject( const TCHAR* Stream, const TCHAR* Match, T*& Obj, UObject* Outer, bool* bInvalidObject=nullptr )
 {
 	return ParseObject( Stream, Match, T::StaticClass(), (UObject*&)Obj, Outer, bInvalidObject );
 }
 
-// Find an optional object, relies on the name being unqualified
+/** 
+ * Find an optional object, relies on the name being unqualified 
+ * @see StaticFindObjectFast()
+ */
 template< class T > 
 inline T* FindObjectFast( UObject* Outer, FName Name, bool ExactClass=false, bool AnyPackage=false, EObjectFlags ExclusiveFlags=RF_NoFlags )
 {
 	return (T*)StaticFindObjectFast( T::StaticClass(), Outer, Name, ExactClass, AnyPackage, ExclusiveFlags );
 }
 
-// Find an optional object.
+/**
+ * Find an optional object.
+ * @see StaticFindObject()
+ */
 template< class T > 
 inline T* FindObject( UObject* Outer, const TCHAR* Name, bool ExactClass=false )
 {
 	return (T*)StaticFindObject( T::StaticClass(), Outer, Name, ExactClass );
 }
 
-// Find an object, no failure allowed.
+/**
+ * Find an optional object, no failure allowed
+ * @see StaticFindObjectChecked()
+ */
 template< class T > 
 inline T* FindObjectChecked( UObject* Outer, const TCHAR* Name, bool ExactClass=false )
 {
 	return (T*)StaticFindObjectChecked( T::StaticClass(), Outer, Name, ExactClass );
 }
 
-// Find an object without asserting on GIsSavingPackage or IsGarbageCollecting()
+/**
+ * Find an object without asserting on GIsSavingPackage or IsGarbageCollecting()
+ * @see StaticFindObjectSafe()
+ */
 template< class T > 
 inline T* FindObjectSafe( UObject* Outer, const TCHAR* Name, bool ExactClass=false )
 {
 	return (T*)StaticFindObjectSafe( T::StaticClass(), Outer, Name, ExactClass );
 }
 
-
-// Load an object.
+/** 
+ * Load an object. 
+ * @see StaticLoadObject()
+ */
 template< class T > 
 inline T* LoadObject( UObject* Outer, const TCHAR* Name, const TCHAR* Filename=nullptr, uint32 LoadFlags=LOAD_None, UPackageMap* Sandbox=nullptr )
 {
 	return (T*)StaticLoadObject( T::StaticClass(), Outer, Name, Filename, LoadFlags, Sandbox );
 }
 
-// Load a class object.
+/**
+ * Load a class object
+ * @see StaticLoadClass
+ */
 template< class T > 
 inline UClass* LoadClass( UObject* Outer, const TCHAR* Name, const TCHAR* Filename=nullptr, uint32 LoadFlags=LOAD_None, UPackageMap* Sandbox=nullptr )
 {
 	return StaticLoadClass( T::StaticClass(), Outer, Name, Filename, LoadFlags, Sandbox );
 }
 
-// Get default object of a class.
+/** 
+ * Get default object of a class.
+ * @see UClass::GetDefaultObject()
+ */
 template< class T > 
 inline const T* GetDefault()
 {
 	return (const T*)T::StaticClass()->GetDefaultObject();
 }
 
-// Get default object of a class.
+/**
+ * Get default object of a class.
+ * @see Class.h
+ */
 template< class T > 
 inline const T* GetDefault(UClass *Class);
 
-// Get the default object of a class (mutable).
+/** Version of GetDefault() that allows modification */
 template< class T >
 inline T* GetMutableDefault()
 {
 	return (T*)T::StaticClass()->GetDefaultObject();
 }
 
-// Get default object of a class (mutable).
+/** Version of GetDefault() that allows modification */
 template< class T > 
 inline T* GetMutableDefault(UClass *Class);
 
-// Returns true if a class has been loaded (e.g. it has a CDO)
+/** Returns true if a class has been loaded (e.g. it has a CDO) */
 template< class T >
 inline bool IsClassLoaded()
 {
@@ -1321,12 +1407,12 @@ COREUOBJECT_API UFunction* FindDelegateSignature(FName DelegateSignatureName);
  * @param	out_Objects		if specified, any objects that match the SearchClass will be added to this array
  */
 template <class T>
-bool ContainsObjectOfClass( const TArray<T*>& ObjectArray, UClass* ClassToCheck, bool bExactClass=false, TArray<T*>* out_Objects=NULL )
+bool ContainsObjectOfClass( const TArray<T*>& ObjectArray, UClass* ClassToCheck, bool bExactClass=false, TArray<T*>* out_Objects=nullptr )
 {
 	bool bResult = false;
 	for ( int32 ArrayIndex = 0; ArrayIndex < ObjectArray.Num(); ArrayIndex++ )
 	{
-		if ( ObjectArray[ArrayIndex] != NULL )
+		if ( ObjectArray[ArrayIndex] != nullptr )
 		{
 			bool bMatchesSearchCriteria = bExactClass
 				? ObjectArray[ArrayIndex]->GetClass() == ClassToCheck
@@ -1335,7 +1421,7 @@ bool ContainsObjectOfClass( const TArray<T*>& ObjectArray, UClass* ClassToCheck,
 			if ( bMatchesSearchCriteria )
 			{
 				bResult = true;
-				if ( out_Objects != NULL )
+				if ( out_Objects != nullptr )
 				{
 					out_Objects->Add(ObjectArray[ArrayIndex]);
 				}
@@ -1369,6 +1455,7 @@ class FScopedObjectFlagMarker
 		EObjectFlags Flags;
 		EInternalObjectFlags InternalFlags;
 	};
+
 	/**
 	 * Map that tracks the ObjectFlags set on all objects; we use a map rather than iterating over all objects twice because FObjectIterator
 	 * won't return objects that have RF_Unreachable set, and we may want to actually unset that flag.
@@ -1423,7 +1510,7 @@ public:
 		* @param	InClass			if non-null, will only iterate on items IsA this class
 		* @param	InbExactClass	if true, will only iterate on exact matches
 		*/
-	FORCEINLINE TObjectArrayIterator( TArray<TObjectClass*>& InArray, UClass* InClassToCheck = NULL, bool InbExactClass = false) :	
+	FORCEINLINE TObjectArrayIterator( TArray<TObjectClass*>& InArray, UClass* InClassToCheck = nullptr, bool InbExactClass = false) :	
 		Array(InArray),
 		Index(-1),
 		ClassToCheck(InClassToCheck),
@@ -1448,6 +1535,7 @@ public:
 	{
 		return !(bool)*this;
 	}
+
 	/**
 		* Dereferences the iterator 
 		* @return	the UObject at the iterator
@@ -1823,12 +1911,12 @@ public:
 	 * @param	InOuter					value for LimitOuter
 	 * @param	bInRequireDirectOuter	value for bRequireDirectOuter
 	 * @param	bShouldIgnoreArchetype	whether to disable serialization of ObjectArchetype references
-	 * @param	bInSerializeRecursively	only applicable when LimitOuter != NULL && bRequireDirectOuter==true;
+	 * @param	bInSerializeRecursively	only applicable when LimitOuter != nullptr && bRequireDirectOuter==true;
 	 *									serializes each object encountered looking for subobjects of referenced
 	 *									objects that have LimitOuter for their Outer (i.e. nested subobjects/components)
 	 * @param	bShouldIgnoreTransient	true to skip serialization of transient properties
 	 */
-	FReferenceFinder(TArray<UObject*>& InObjectArray, UObject* InOuter = NULL, bool bInRequireDirectOuter = true, bool bInShouldIgnoreArchetype = false, bool bInSerializeRecursively = false, bool bInShouldIgnoreTransient = false);
+	FReferenceFinder(TArray<UObject*>& InObjectArray, UObject* InOuter = nullptr, bool bInRequireDirectOuter = true, bool bInShouldIgnoreArchetype = false, bool bInSerializeRecursively = false, bool bInShouldIgnoreTransient = false);
 
 	/**
 	 * Finds all objects referenced by Object.
@@ -1857,7 +1945,7 @@ protected:
 	TArray<UObject*>&		ObjectArray;
 	/** List of objects that have been recursively serialized. */
 	TSet<const UObject*>	SerializedObjects;
-	/** Only objects within this outer will be considered, NULL value indicates that outers are disregarded. */
+	/** Only objects within this outer will be considered, nullptr value indicates that outers are disregarded. */
 	UObject*		LimitOuter;
 	/** Property that is referencing the current object */
 	class UProperty* SerializedProperty;
@@ -2005,6 +2093,7 @@ struct COREUOBJECT_API FCoreUObjectDelegates
 	UE_DEPRECATED(4.17, "RedirectorFollowed is deprecated, FixeupRedirects was replaced with ResavePackages -FixupRedirect")
 	static FOnRedirectorFollowed RedirectorFollowed;
 
+	/** Called during cooking to see if a specific package should be cooked for a given target platform */
 	DECLARE_DELEGATE_RetVal_TwoParams(bool, FShouldCookPackageForPlatform, const UPackage*, const ITargetPlatform*);
 	static FShouldCookPackageForPlatform ShouldCookPackageForPlatform;
 };
@@ -2019,6 +2108,7 @@ enum class EConstructDynamicType : uint8
 	OnlyAllocateClassObject,
 	CallZConstructor
 };
+
 /** Constructs dynamic type of a given class. */
 COREUOBJECT_API UObject* ConstructDynamicType(FName TypePathName, EConstructDynamicType ConstructionSpecifier);
 
@@ -2041,6 +2131,7 @@ struct COREUOBJECT_API FDynamicClassStaticData
 	TMap<FName, FName> SelectedSearchableValues;
 };
 
+/** Returns map of all dynamic/nativized classes */
 COREUOBJECT_API TMap<FName, FDynamicClassStaticData>& GetDynamicClassMap();
 
 /**
@@ -2106,6 +2197,7 @@ COREUOBJECT_API bool IsEditorOnlyObject(const UObject* InObject, bool bCheckRecu
 struct FClassFunctionLinkInfo;
 struct FCppClassTypeInfoStatic;
 
+/// @cond DOXYGEN_IGNORE
 namespace UE4CodeGen_Private
 {
 	enum class EPropertyGenFlags : uint32
@@ -2495,6 +2587,7 @@ namespace UE4CodeGen_Private
 	COREUOBJECT_API void ConstructUPackage(UPackage*& OutPackage, const FPackageParams& Params);
 	COREUOBJECT_API void ConstructUClass(UClass*& OutClass, const FClassParams& Params);
 }
+/// @endcond
 
 // METADATA_PARAMS(x, y) expands to x, y, if WITH_METADATA is set, otherwise expands to nothing
 #if WITH_METADATA
@@ -2517,10 +2610,13 @@ namespace UE4CodeGen_Private
 	#define IF_WITH_EDITORONLY_DATA(x, y) y
 #endif
 
+/** Enum used by UDataValidationManager to see if an asset has been validated for correctness */
 enum class EDataValidationResult : uint8
 {
+	/** Asset has failed validation */
 	Invalid,
+	/** Asset has passed validation */
 	Valid,
-
+	/** Asset has not yet been validated */
 	NotValidated
 };
