Fixed a rare case where async loading would not serialize all imports before postloading exports that depend on them

diff --git a/Engine/Source/Runtime/CoreUObject/Private/Serialization/AsyncLoading.cpp b/Engine/Source/Runtime/CoreUObject/Private/Serialization/AsyncLoading.cpp
--- a/Engine/Source/Runtime/CoreUObject/Private/Serialization/AsyncLoading.cpp
+++ b/Engine/Source/Runtime/CoreUObject/Private/Serialization/AsyncLoading.cpp
@@ -2421,10 +2421,7 @@ EAsyncPackageState::Type FAsyncPackage::SetupImports_Event()
 
 				if (ImportPackage)
 				{
-					if (!ImportedPackages.Contains(ImportPackage))
-					{
-						ImportedPackages.Add(ImportPackage);
-					}
+					ImportedPackages.Add(ImportPackage);
 
 					FLinkerLoad* ImportLinker = ImportPackage->LinkerLoad;
 					if (ImportLinker && ImportLinker->AsyncRoot)
@@ -5819,6 +5816,25 @@ EAsyncPackageState::Type FAsyncPackage::TickAsyncPackage(bool InbUseTimeLimit, b
 				SCOPED_LOADTIMER(Package_PreLoadObjects);
 				LoadingState = PreLoadObjects();
 			}
+
+			if (LoadingState == EAsyncPackageState::Complete || bLoadHasFailed)
+			{
+				const bool bInternalCallbacks = true;
+				CallCompletionCallbacks(bInternalCallbacks, bLoadHasFailed ? EAsyncLoadingResult::Failed : EAsyncLoadingResult::Succeeded);
+			}
+
+			if (LoadingState == EAsyncPackageState::Complete)
+			{
+				// We can only continue to PostLoad if all imported packages finished serializing their exports
+				for (UPackage* ImportedPackage : ImportedPackages)
+				{
+					if (ImportedPackage->LinkerLoad && ImportedPackage->LinkerLoad->AsyncRoot && !ImportedPackage->LinkerLoad->AsyncRoot->bAllExportsSerialized)
+					{
+						LoadingState = EAsyncPackageState::PendingImports;
+						break;
+					}
+				}
+			}
 		} // !GEventDrivenLoaderEnabled
 
 		if (LoadingState == EAsyncPackageState::Complete && !bLoadHasFailed)
@@ -6371,6 +6387,15 @@ EAsyncPackageState::Type FAsyncPackage::CreateImports()
 
 		// Make sure this object is not claimed by GC if it's triggered while streaming.
 		AddObjectReference(Object);
+
+		// Keep track of all imported packages that are also being loaded so that we can wait until they also finished serializing their exports
+		if (UPackage* ImportedPackage = Cast<UPackage>(Object))
+		{
+			if (ImportedPackage->LinkerLoad && ImportedPackage->LinkerLoad->AsyncRoot)
+			{
+				ImportedPackages.Add(ImportedPackage);
+			}
+		}
 	}
 
 	return ImportIndex == Linker->ImportMap.Num() ? EAsyncPackageState::Complete : EAsyncPackageState::TimeOut;
@@ -6464,7 +6489,13 @@ EAsyncPackageState::Type FAsyncPackage::CreateExports()
 	// We no longer need the referenced packages.
 	FreeReferencedImports();
 
-	return ExportIndex == Linker->ExportMap.Num() ? EAsyncPackageState::Complete : EAsyncPackageState::TimeOut;
+	EAsyncPackageState::Type Result = (ExportIndex == Linker->ExportMap.Num() ? EAsyncPackageState::Complete : EAsyncPackageState::TimeOut);
+	if (Result == EAsyncPackageState::Complete)
+	{
+		bAllExportsSerialized = true;
+	}
+
+	return Result;
 }
 
 /**
