Fixed cases where some objects would not get PostLoaded when playing standalone game or running uncooked server

diff --git a/Engine/Source/Runtime/CoreUObject/Private/Serialization/AsyncLoading.cpp b/Engine/Source/Runtime/CoreUObject/Private/Serialization/AsyncLoading.cpp
--- a/Engine/Source/Runtime/CoreUObject/Private/Serialization/AsyncLoading.cpp
+++ b/Engine/Source/Runtime/CoreUObject/Private/Serialization/AsyncLoading.cpp
@@ -5959,10 +5959,17 @@ EAsyncPackageState::Type FAsyncPackage::CreateLinker()
 		if (Linker)
 		{
 			// Swap the load context to the currently associated with the existing linker
-			check(Linker->GetSerializeContext());
-			LoadContext->DecrementBeginLoadCount();
-			LoadContext = Linker->GetSerializeContext();
-			LoadContext->IncrementBeginLoadCount();
+			if (Linker->GetSerializeContext())
+			{
+				LoadContext->DecrementBeginLoadCount();
+				LoadContext = Linker->GetSerializeContext();
+				LoadContext->IncrementBeginLoadCount();
+			}
+			else
+			{
+				check(!GEventDrivenLoaderEnabled);
+				Linker->SetSerializeContext(LoadContext);
+			}
 
 			if (GEventDrivenLoaderEnabled)
 			{
@@ -6949,10 +6956,15 @@ EAsyncPackageState::Type FAsyncPackage::FinishObjects()
 		Linker->FlushCache();
 	}
 
+	if (GEventDrivenLoaderEnabled)
 	{
 		const bool bInternalCallbacks = true;
 		CallCompletionCallbacks(bInternalCallbacks, LoadingResult);
 	}
+	else
+	{
+		LoadContext->DetachFromLinkers();
+	}
 
 	return EAsyncPackageState::Complete;
 }
diff --git a/Engine/Source/Runtime/CoreUObject/Private/UObject/Linker.cpp b/Engine/Source/Runtime/CoreUObject/Private/UObject/Linker.cpp
--- a/Engine/Source/Runtime/CoreUObject/Private/UObject/Linker.cpp
+++ b/Engine/Source/Runtime/CoreUObject/Private/UObject/Linker.cpp
@@ -726,7 +726,7 @@ FLinkerLoad* GetPackageLinker
 	else if (InExistingContext)
 	{
 		if ((Result->GetSerializeContext() && Result->GetSerializeContext()->HasStartedLoading() && InExistingContext->GetBeginLoadCount() == 1) ||
-			  IsInAsyncLoadingThread())
+			  (IsInAsyncLoadingThread() && Result->GetSerializeContext()))
 		{
 			// Use the context associated with the linker because it has already started loading objects (or we're in ALT where each package needs its own context)
 			*InOutLoadContext = Result->GetSerializeContext();
@@ -737,7 +737,6 @@ FLinkerLoad* GetPackageLinker
 			{
 				// Make sure the objects already loaded with the context associated with the existing linker
 				// are copied to the context provided for this function call to make sure they all get loaded ASAP
-				check(!IsInAsyncLoadingThread());
 				InExistingContext->AddUniqueLoadedObjects(Result->GetSerializeContext()->PRIVATE_GetObjectsLoadedInternalUseOnly());
 			}
 			// Replace the linker context with the one passed into this function
diff --git a/Engine/Source/Runtime/CoreUObject/Private/UObject/LinkerLoad.cpp b/Engine/Source/Runtime/CoreUObject/Private/UObject/LinkerLoad.cpp
--- a/Engine/Source/Runtime/CoreUObject/Private/UObject/LinkerLoad.cpp
+++ b/Engine/Source/Runtime/CoreUObject/Private/UObject/LinkerLoad.cpp
@@ -640,7 +640,15 @@ FLinkerLoad* FLinkerLoad::CreateLinkerAsync(FUObjectSerializeContext* LoadContex
 
 void FLinkerLoad::SetSerializeContext(FUObjectSerializeContext* InLoadContext)
 {
+	if (!GEventDrivenLoaderEnabled && CurrentLoadContext)
+	{
+		CurrentLoadContext->DetachLinker(this);
+	}
 	CurrentLoadContext = InLoadContext;
+	if (!GEventDrivenLoaderEnabled && CurrentLoadContext)
+	{
+		CurrentLoadContext->AttachLinker(this);
+	}
 }
 FUObjectSerializeContext* FLinkerLoad::GetSerializeContext()
 {
@@ -863,6 +871,12 @@ FLinkerLoad::~FLinkerLoad()
 	// Detaches linker.
 	Detach();
 
+	// Detach the serialize context
+	if (GetSerializeContext())
+	{
+		SetSerializeContext(nullptr);
+	}
+
 	DEC_DWORD_STAT(STAT_LiveLinkerCount);
 
 #if WITH_EDITOR
@@ -3011,6 +3025,12 @@ bool FLinkerLoad::VerifyImportInner(const int32 ImportIndex, FString& WarningSuf
 			return true;
 		}
 	}
+
+	if (!GEventDrivenLoaderEnabled && Import.SourceLinker && !Import.SourceLinker->GetSerializeContext())
+	{
+		Import.SourceLinker->SetSerializeContext(GetSerializeContext());
+	}
+
 	return false;
 }
 
@@ -4522,6 +4542,10 @@ UObject* FLinkerLoad::CreateImport( int32 Index )
 			{
 				VerifyImportResult = VerifyImport(Index);
 			}
+			else if (!GEventDrivenLoaderEnabled && !Import.SourceLinker->GetSerializeContext())
+			{
+				Import.SourceLinker->SetSerializeContext(GetSerializeContext());
+			}
 			if(Import.SourceIndex != INDEX_NONE)
 			{
 				check(Import.SourceLinker);
diff --git a/Engine/Source/Runtime/CoreUObject/Private/UObject/UObjectGlobals.cpp b/Engine/Source/Runtime/CoreUObject/Private/UObject/UObjectGlobals.cpp
--- a/Engine/Source/Runtime/CoreUObject/Private/UObject/UObjectGlobals.cpp
+++ b/Engine/Source/Runtime/CoreUObject/Private/UObject/UObjectGlobals.cpp
@@ -1209,7 +1209,7 @@ UPackage* LoadPackageInternal(UPackage* InOuter, const TCHAR* InLongPackageNameO
 		{
 			FUObjectSerializeContext* InOutLoadContext = LoadContext;
 			Linker = GetPackageLinker(InOuter, *FileToLoad, LoadFlags, nullptr, nullptr, InReaderOverride, &InOutLoadContext);
-			if (InOutLoadContext != LoadContext)
+			if (InOutLoadContext != LoadContext && InOutLoadContext)
 			{
 				// The linker already existed and was associated with another context
 				LoadContext->DecrementBeginLoadCount();
@@ -1752,6 +1752,11 @@ void EndLoad(FUObjectSerializeContext* LoadContext)
 
 	if (LoadContext->GetBeginLoadCount() == 0)
 	{
+		if (!GEventDrivenLoaderEnabled)
+		{
+			LoadContext->DetachFromLinkers();
+		}
+
 		FUObjectThreadContext& ThreadContext = FUObjectThreadContext::Get();
 		if (ThreadContext.SerializeContext == LoadContext)
 		{
diff --git a/Engine/Source/Runtime/CoreUObject/Private/UObject/UObjectThreadContext.cpp b/Engine/Source/Runtime/CoreUObject/Private/UObject/UObjectThreadContext.cpp
--- a/Engine/Source/Runtime/CoreUObject/Private/UObject/UObjectThreadContext.cpp
+++ b/Engine/Source/Runtime/CoreUObject/Private/UObject/UObjectThreadContext.cpp
@@ -5,6 +5,8 @@
 =============================================================================*/
 
 #include "UObject/UObjectThreadContext.h"
+#include "UObject/Object.h"
+#include "UObject/LinkerLoad.h"
 
 DEFINE_LOG_CATEGORY(LogUObjectThreadContext);
 
@@ -35,6 +37,7 @@ FUObjectSerializeContext::FUObjectSerializeContext()
 FUObjectSerializeContext::~FUObjectSerializeContext()
 {
 	checkf(!HasLoadedObjects(), TEXT("FUObjectSerializeContext is being destroyed but it still has pending loaded objects in its ObjectsLoaded list."));
+	check(AttachedLinkers.Num() == 0);
 }
 
 int32 FUObjectSerializeContext::IncrementBeginLoadCount()
@@ -54,4 +57,36 @@ void FUObjectSerializeContext::AddUniqueLoadedObjects(const TArray<UObject*>& In
 		ObjectsLoaded.AddUnique(NewLoadedObject);
 	}
 	
+}
+
+bool FUObjectSerializeContext::PRIVATE_PatchNewObjectIntoExport(UObject* OldObject, UObject* NewObject)
+{
+	const int32 ObjLoadedIdx = ObjectsLoaded.Find(OldObject);
+	if (ObjLoadedIdx != INDEX_NONE)
+	{
+		ObjectsLoaded[ObjLoadedIdx] = NewObject;
+		return true;
+	}
+	else
+	{
+		return false;
+	}
+}
+void FUObjectSerializeContext::AttachLinker(FLinkerLoad* InLinker)
+{
+	check(!GEventDrivenLoaderEnabled);
+	AttachedLinkers.Add(InLinker);
+}
+
+void FUObjectSerializeContext::DetachFromLinkers()
+{
+	check(!GEventDrivenLoaderEnabled);
+	check(ObjectsLoaded.Num() == 0);
+	TArray<FLinkerLoad*> LinkersToDetach = AttachedLinkers.Array();
+	for (FLinkerLoad* Linker : LinkersToDetach)
+	{
+		check(Linker->GetSerializeContext() == this);
+		Linker->SetSerializeContext(nullptr);
+	}
+	check(AttachedLinkers.Num() == 0);
 }
\ No newline at end of file
diff --git a/Engine/Source/Runtime/CoreUObject/Public/UObject/UObjectThreadContext.h b/Engine/Source/Runtime/CoreUObject/Public/UObject/UObjectThreadContext.h
--- a/Engine/Source/Runtime/CoreUObject/Public/UObject/UObjectThreadContext.h
+++ b/Engine/Source/Runtime/CoreUObject/Public/UObject/UObjectThreadContext.h
@@ -12,6 +12,7 @@
 
 class FObjectInitializer;
 struct FUObjectSerializeContext;
+class FLinkerLoad;
 
 DECLARE_LOG_CATEGORY_EXTERN(LogUObjectThreadContext, Log, All);
 
@@ -107,7 +108,9 @@ private:
 	/** Objects that might need preloading. */
 	TArray<UObject*> ObjectsLoaded;
 	/** List of linkers that we want to close the loaders for (to free file handles) - needs to be delayed until EndLoad is called with GObjBeginLoadCount of 0 */
-	TArray<class FLinkerLoad*> DelayedLinkerClosePackages;
+	TArray<FLinkerLoad*> DelayedLinkerClosePackages;
+	/** List of linkers associated with this context */
+	TSet<FLinkerLoad*> AttachedLinkers;
 
 public:
 
@@ -133,8 +136,10 @@ public:
 	/** Adds a new loaded object */
 	void AddLoadedObject(UObject* InObject)
 	{
+		check(AttachedLinkers.Num() || GEventDrivenLoaderEnabled);
 		ObjectsLoaded.Add(InObject);
 	}
+
 	void AddUniqueLoadedObjects(const TArray<UObject*>& InObjects);
 
 	/** Checks if object loading has started */
@@ -178,19 +183,7 @@ public:
 		return !!ObjectsLoaded.Num();
 	}
 
-	bool PRIVATE_PatchNewObjectIntoExport(UObject* OldObject, UObject* NewObject)
-	{
-		const int32 ObjLoadedIdx = ObjectsLoaded.Find(OldObject);
-		if (ObjLoadedIdx != INDEX_NONE)
-		{
-			ObjectsLoaded[ObjLoadedIdx] = NewObject;
-			return true;
-		}
-		else
-		{
-			return false;
-		}
-	}
+	bool PRIVATE_PatchNewObjectIntoExport(UObject* OldObject, UObject* NewObject);
 
 	/** This is only meant to be used by FAsyncPackage for performance reasons. The ObjectsLoaded array should not be manipulated directly! */
 	TArray<UObject*>& PRIVATE_GetObjectsLoadedInternalUseOnly()
@@ -229,6 +222,18 @@ public:
 		OutDelayedLinkerClosePackages = MoveTemp(DelayedLinkerClosePackages);
 	}
 
+	/** Attaches a linker to this context */
+	void AttachLinker(FLinkerLoad* InLinker);
+	
+	/** Detaches a linker from this context */
+	void DetachLinker(FLinkerLoad* InLinker)
+	{
+		AttachedLinkers.Remove(InLinker);
+	}
+
+	/** Detaches all linkers from this context */
+	void DetachFromLinkers();
+
 	//~ TRefCountPtr interface
 	int32 AddRef()
 	{
