Added an Error log when DOREPLIFETIME is not called on a replicated property inside GetLifetimeReplicatedProps

If you intentionally don't want to replicate this property, use the DISABLE_REPLICATED_PROPERTY to silence the Error (see cl 4962290 for the macros)

diff --git a/Engine/Source/Runtime/CoreUObject/Private/UObject/Class.cpp b/Engine/Source/Runtime/CoreUObject/Private/UObject/Class.cpp
--- a/Engine/Source/Runtime/CoreUObject/Private/UObject/Class.cpp
+++ b/Engine/Source/Runtime/CoreUObject/Private/UObject/Class.cpp
@@ -13,6 +13,7 @@
 #include "Misc/FeedbackContext.h"
 #include "Misc/OutputDeviceConsole.h"
 #include "Misc/AutomationTest.h"
+#include "Misc/EnumClassFlags.h"
 #include "UObject/ErrorException.h"
 #include "Modules/ModuleManager.h"
 #include "UObject/UObjectAllocator.h"
@@ -3387,6 +3388,14 @@ void UClass::Link(FArchive& Ar, bool bRelinkExistingProperties)
 	Super::Link(Ar, bRelinkExistingProperties);
 }
 
+#if (UE_BUILD_SHIPPING)
+static int32 GValidateReplicatedProperties = 0;
+#else 
+static int32 GValidateReplicatedProperties = 1;
+#endif
+
+static FAutoConsoleVariable CVarValidateReplicatedPropertyRegistration(TEXT("net.ValidateReplicatedPropertyRegistration"), GValidateReplicatedProperties, TEXT("Warns if replicated properties were not registered in GetLifetimeReplicatedProps."));
+
 void UClass::SetUpRuntimeReplicationData()
 {
 	if (!HasAnyClassFlags(CLASS_ReplicationDataIsSetUp) && PropertyLink != NULL)
@@ -3470,6 +3479,69 @@ void UClass::SetUpRuntimeReplicationData()
 		Sort(NetFields.GetData(), NetFields.Num(), FCompareUFieldNames());
 
 		ClassFlags |= CLASS_ReplicationDataIsSetUp;
+
+		if (GValidateReplicatedProperties != 0)
+		{
+			ValidateRuntimeReplicationData();
+		}
+	}
+}
+
+void UClass::ValidateRuntimeReplicationData()
+{
+	DECLARE_SCOPE_CYCLE_COUNTER(TEXT("Class ValidateRuntimeReplicationData"), STAT_Class_ValidateRuntimeReplicationData, STATGROUP_Game);
+
+	if (HasAnyClassFlags(CLASS_CompiledFromBlueprint))
+	{
+		// Blueprint classes don't always generate a GetLifetimeReplicatedProps function. 
+		// Assume the Blueprint compiler was ok to do this.
+		return;
+	}
+
+	if (HasAnyClassFlags(CLASS_ReplicationDataIsSetUp) == false)
+	{
+		UE_LOG(LogClass, Warning, TEXT("ValidateRuntimeReplicationData for class %s called before ReplicationData was setup."), *GetName());
+		return;
+	}
+
+	// Let's compare the CDO's registered lifetime properties with the Class's net properties
+	TArray<FLifetimeProperty> LifetimeProps;
+	LifetimeProps.Reserve(ClassReps.Num());
+
+	const UObject* Object = GetDefaultObject();
+	Object->GetLifetimeReplicatedProps(LifetimeProps);
+
+	if (LifetimeProps.Num() == ClassReps.Num())
+	{
+		// All replicated properties were registered for this class
+		return;
+	}
+
+	// Find which properties where not registered by the user code
+	for (int32 RepIndex = 0; RepIndex < ClassReps.Num(); ++RepIndex)
+	{
+		const UProperty* RepProp = ClassReps[RepIndex].Property;
+
+		const FLifetimeProperty* LifetimeProp = LifetimeProps.FindByPredicate([&RepIndex](const FLifetimeProperty& Var) { return Var.RepIndex == RepIndex; });
+
+		if (LifetimeProp == nullptr)
+		{
+			// Check if this unregistered property type uses a custom delta serializer
+			if (const UStructProperty* StructProperty = Cast<UStructProperty>(RepProp))
+			{
+				const UScriptStruct* Struct = StructProperty->Struct;
+
+				if (EnumHasAnyFlags(Struct->StructFlags, STRUCT_NetDeltaSerializeNative))
+				{
+					UE_LOG(LogClass, Warning, TEXT("Property %s::%s (SourceClass: %s) with custom net delta serializer was not registered in GetLifetimeReplicatedProps. This property will replicate but you should still register it."),
+						*GetName(), *RepProp->GetName(), *RepProp->GetOwnerClass()->GetName());
+					continue;
+				}
+			}
+
+			UE_LOG(LogClass, Error, TEXT("Property %s::%s (SourceClass: %s) was not registered in GetLifetimeReplicatedProps. This property will not be replicated."),
+				*GetName(), *RepProp->GetName(), *RepProp->GetOwnerClass()->GetName());
+		}
 	}
 }
 
diff --git a/Engine/Source/Runtime/CoreUObject/Public/UObject/Class.h b/Engine/Source/Runtime/CoreUObject/Public/UObject/Class.h
--- a/Engine/Source/Runtime/CoreUObject/Public/UObject/Class.h
+++ b/Engine/Source/Runtime/CoreUObject/Public/UObject/Class.h
@@ -2777,6 +2777,11 @@ private:
 		return UObject::FindFunctionChecked(InName);
 	}
 
+	/**
+	 * Tests if all properties tagged with Replicate were registered in GetLifetimeReplicatedProps
+	 */
+	void ValidateRuntimeReplicationData();
+
 protected:
 	/**
 	 * Get the default object from the class, creating it if missing, if requested or under a few other circumstances
