Fix LoadPackageInternal returning the wrong package in the EDL case if an Outer package was supplied. Broken in CL# 3354737.
Confirmed to fix issues attempted to be fixed by PR# 5068

diff --git a/Engine/Source/Runtime/CoreUObject/Private/UObject/UObjectGlobals.cpp b/Engine/Source/Runtime/CoreUObject/Private/UObject/UObjectGlobals.cpp
--- a/Engine/Source/Runtime/CoreUObject/Private/UObject/UObjectGlobals.cpp
+++ b/Engine/Source/Runtime/CoreUObject/Private/UObject/UObjectGlobals.cpp
@@ -1088,7 +1088,7 @@ UPackage* LoadPackageInternal(UPackage* InOuter, const TCHAR* InLongPackageNameO
 			}
 		}
 
-		Result = FindObjectFast<UPackage>(nullptr, PackageFName);
+		Result = (InOuter ? InOuter : FindObjectFast<UPackage>(nullptr, PackageFName));
 		return Result;
 	}
 
