Corrected assertion to account for INVALID_OBJECT when looking for instanced subobjects

diff --git a/Engine/Source/Runtime/CoreUObject/Private/UObject/CoreNative.cpp b/Engine/Source/Runtime/CoreUObject/Private/UObject/CoreNative.cpp
--- a/Engine/Source/Runtime/CoreUObject/Private/UObject/CoreNative.cpp
+++ b/Engine/Source/Runtime/CoreUObject/Private/UObject/CoreNative.cpp
@@ -197,7 +197,7 @@ UObject* FObjectInstancingGraph::GetInstancedSubobject( UObject* SourceSubobject
 							{
 								SubobjectOuter = GetInstancedSubobject(SourceSubobject->GetOuter(), SourceSubobject->GetOuter(), CurrentObject, bDoNotCreateNewInstance, bAllowSelfReference);
 
-								checkf(SubobjectOuter, TEXT("No corresponding destination object found for '%s' while attempting to instance component '%s'"), *SourceSubobject->GetOuter()->GetFullName(), *SourceSubobject->GetFullName());
+								checkf(SubobjectOuter && SubobjectOuter != INVALID_OBJECT, TEXT("No corresponding destination object found for '%s' while attempting to instance component '%s'"), *SourceSubobject->GetOuter()->GetFullName(), *SourceSubobject->GetFullName());
 							}
 
 							FName SubobjectName = SourceSubobject->GetFName();
