Reduced cached archetye annotation array memory usage by 50% on 64-bit platforms  by storing the archetype index instead of a pointer

diff --git a/Engine/Source/Runtime/CoreUObject/Private/UObject/UObjectArchetype.cpp b/Engine/Source/Runtime/CoreUObject/Private/UObject/UObjectArchetype.cpp
--- a/Engine/Source/Runtime/CoreUObject/Private/UObject/UObjectArchetype.cpp
+++ b/Engine/Source/Runtime/CoreUObject/Private/UObject/UObjectArchetype.cpp
@@ -114,28 +114,28 @@ struct FArchetypeInfo
 	* Default constructor must be the default item
 	*/
 	FArchetypeInfo() 
-		: Archetype(nullptr)
+		: ArchetypeIndex(INDEX_NONE)
 	{
 	}
 	/**
 	* Determine if this linker pair is the default
 	* @return true is this is a default pair. We only check the linker because CheckInvariants rules out bogus combinations
 	*/
-	FORCEINLINE bool IsDefault()
+	FORCEINLINE bool IsDefault() const
 	{
-		return Archetype == nullptr;
+		return ArchetypeIndex == INDEX_NONE;
 	}
 
 	/**
 	* Constructor
 	* @param InArchetype Archetype to assign
 	*/
-	FArchetypeInfo(UObject* InArchetype) 
-		: Archetype(InArchetype)
+	FArchetypeInfo(int32 InArchetypeIndex) 
+		: ArchetypeIndex(InArchetypeIndex)
 	{
 	}
 
-	UObject* Archetype;
+	int32 ArchetypeIndex;
 };
 
 static FUObjectAnnotationDense<FArchetypeInfo, true> ArchetypeAnnotation;
@@ -149,25 +149,29 @@ UObject* UObject::GetArchetype() const
 	//SCOPE_SECONDS_ACCUMULATOR(STAT_FArchiveRealtimeGC_GetArchetype);
 
 #if UE_CACHE_ARCHETYPE
-	UObject* Archetype = ArchetypeAnnotation.GetAnnotation(this).Archetype;
-	if (Archetype == nullptr)
+	UObject* Archetype = nullptr;
+	int32 ArchetypeIndex = ArchetypeAnnotation.GetAnnotation(this).ArchetypeIndex;
+	if (ArchetypeIndex == INDEX_NONE)
 	{
 		Archetype = GetArchetypeFromRequiredInfo(GetClass(), GetOuter(), GetFName(), GetFlags());
 		if (Archetype)
 		{
-			ArchetypeAnnotation.AddAnnotation(this, Archetype);
+			ArchetypeAnnotation.AddAnnotation(this, GUObjectArray.ObjectToIndex(Archetype));
 		}
 	}		
-#if UE_VERIFY_CACHED_ARCHETYPE
 	else
 	{
+		FUObjectItem* ArchetypeItem = GUObjectArray.IndexToObject(ArchetypeIndex);
+		check(ArchetypeItem != nullptr);
+		Archetype = static_cast<UObject*>(ArchetypeItem->Object);
+#if UE_VERIFY_CACHED_ARCHETYPE
 		UObject* CurrentArchetype = GetArchetypeFromRequiredInfo(GetClass(), GetOuter(), GetFName(), GetFlags());
 		if (CurrentArchetype != Archetype)
 		{
 			UE_LOG(LogClass, Fatal, TEXT("Cached archetype mismatch: %s vs current: %s"), *Archetype->GetFullName(), *CurrentArchetype->GetFullName());
 		}
-	}
 #endif // UE_VERIFY_CACHED_ARCHETYPE
+	}
 	// Note that IsValidLowLevelFast check may fail during initial load as not all classes are initialized at this point so skip it
 	check(Archetype == nullptr || GIsInitialLoad || Archetype->IsValidLowLevelFast());
 
