Reduced Linker annotations memory usage by introducing new annotation array type that can reclaim memory from object ranges that no longer require the annotation

diff --git a/Engine/Source/Runtime/CoreUObject/Private/UObject/UObjectLinker.cpp b/Engine/Source/Runtime/CoreUObject/Private/UObject/UObjectLinker.cpp
--- a/Engine/Source/Runtime/CoreUObject/Private/UObject/UObjectLinker.cpp
+++ b/Engine/Source/Runtime/CoreUObject/Private/UObject/UObjectLinker.cpp
@@ -85,7 +85,7 @@ template <> struct TIsPODType<FLinkerIndexPair> { enum { Value = false }; };
  * NOTE: We're currently using dense annotations for linkers to emphasize speed over memory
  * usage, but might want to revisit this decision on platforms that are memory limited.
  */
-static FUObjectAnnotationDense<FLinkerIndexPair,false> LinkerAnnotation;
+static FUObjectAnnotationChunked<FLinkerIndexPair, false> LinkerAnnotation;
 
 /** Remove all annotations on exit. This is to prevent issues with the order of static destruction of singletons. */
 void CleanupLinkerAnnotations()
diff --git a/Engine/Source/Runtime/CoreUObject/Public/UObject/UObjectAnnotation.h b/Engine/Source/Runtime/CoreUObject/Public/UObject/UObjectAnnotation.h
--- a/Engine/Source/Runtime/CoreUObject/Public/UObject/UObjectAnnotation.h
+++ b/Engine/Source/Runtime/CoreUObject/Public/UObject/UObjectAnnotation.h
@@ -451,6 +451,432 @@ public:
 	}
 };
 
+/**
+* FUObjectAnnotationChunked is a helper class that is used to store dense, fast and temporary, editor only, external
+* or other tangential information about subsets of UObjects.
+*
+* There is a notion of a default annotation and UObjects default to this annotation.
+*
+* Annotations are automatically returned to the default when UObjects are destroyed.
+* Annotation are not "garbage collection aware", so it isn't safe to store pointers to other UObjects in an
+* annotation unless external guarantees are made such that destruction of the other object removes the
+* annotation.
+* The advantage of FUObjectAnnotationChunked is that it can reclaim memory if subsets of UObjects within predefined chunks
+* no longer have any annotations associated with them.
+* @param TAnnotation type of the annotation
+* @param bAutoRemove if true, annotation will automatically be removed, otherwise in non-final builds it will verify that the annotation was removed by other means prior to destruction.
+**/
+template<typename TAnnotation, bool bAutoRemove>
+class FUObjectAnnotationChunked : public FUObjectArray::FUObjectDeleteListener
+{
+	enum
+	{
+		NumAnnotationsPerChunk = 64 * 1024,
+	};
+
+	struct TAnnotationChunk
+	{
+		int32 Num;
+		TAnnotation* Items;
+
+		TAnnotationChunk()
+			: Num(0)
+			: Items(nullptr)
+		{}
+	};
+
+
+	/** Master table to chunks of pointers **/
+	TArray<TAnnotationChunk> Chunks;
+	/** Number of elements we currently have **/
+	int32 NumAnnotations;
+	/** Number of elements we can have **/
+	int32 MaxAnnotations;
+	/** Current allocated memory */
+	uint32 CurrentAllocatedMemory;
+	/** Max allocated memory */
+	uint32 MaxAllocatedMemory;
+
+	/** Mutex */
+	FRWLock AnnotationArrayCritical;
+
+	/**
+	* Makes sure we have enough chunks to fit the new index
+	**/
+	void ExpandChunksToIndex(int32 Index) TSAN_SAFE
+	{
+		check(Index >= 0);
+		int32 ChunkIndex = Index / NumAnnotationsPerChunk;
+		if (ChunkIndex >= Chunks.Num())
+		{
+			Chunks.AddZeroed(ChunkIndex + 1 - Chunks.Num());
+		}
+		check(ChunkIndex < Chunks.Num());
+		MaxAnnotations = Chunks.Num() * NumAnnotationsPerChunk;
+	}
+
+	/**
+	* Initializes an annotation for the specified index, makes sure the chunk it resides in is allocated
+	**/
+	TAnnotation& AllocateAnnotation(int32 Index) TSAN_SAFE
+	{
+		ExpandChunksToIndex(Index);
+
+		const int32 ChunkIndex = Index / NumAnnotationsPerChunk;
+		const int32 WithinChunkIndex = Index % NumAnnotationsPerChunk;
+
+		TAnnotationChunk& Chunk = Chunks[ChunkIndex];
+		if (!Chunk.Items)
+		{
+			Chunk.Items = new TAnnotation[NumAnnotationsPerChunk];
+			CurrentAllocatedMemory += NumAnnotationsPerChunk * sizeof(TAnnotation);
+			MaxAllocatedMemory = FMath::Max(CurrentAllocatedMemory, MaxAllocatedMemory);
+		}
+		check(Chunk.Items[WithinChunkIndex].IsDefault());
+		Chunk.Num++;
+		check(Chunk.Num <= NumAnnotationsPerChunk);
+		NumAnnotations++;
+
+		return Chunk.Items[WithinChunkIndex];
+	}
+
+	/**
+	* Frees the annotation for the specified index
+	**/
+	void FreeAnnotation(int32 Index) TSAN_SAFE
+	{
+		const int32 ChunkIndex = Index / NumAnnotationsPerChunk;
+		const int32 WithinChunkIndex = Index % NumAnnotationsPerChunk;
+
+		TAnnotationChunk& Chunk = Chunks[ChunkIndex];
+		if (Chunk.Items != nullptr)
+		{
+			if (!Chunk.Items[WithinChunkIndex].IsDefault())
+			{
+				Chunk.Items[WithinChunkIndex] = TAnnotation();
+				Chunk.Num--;
+				check(Chunk.Num >= 0);
+				if (Chunk.Num == 0)
+				{
+					delete[] Chunk.Items;
+					Chunk.Items = nullptr;
+					const uint32 ChunkMemory = NumAnnotationsPerChunk * sizeof(TAnnotation);
+					check(CurrentAllocatedMemory >= ChunkMemory);
+					CurrentAllocatedMemory -= ChunkMemory;
+				}
+				NumAnnotations--;
+			}
+			check(NumAnnotations >= 0);
+		}
+	}
+
+	/**
+	* Releases all allocated memory and resets the annotation array
+	**/
+	void FreeAllAnnotations() TSAN_SAFE
+	{
+		for (TAnnotationChunk& Chunk : Chunks)
+		{
+			delete[] Chunk.Items;
+		}
+		Chunks.Empty();
+		NumAnnotations = 0;
+		MaxAnnotations = 0;
+		CurrentAllocatedMemory = 0;
+		MaxAllocatedMemory = 0;
+	}
+
+	/**
+	* Adds a new annotation for the specified index
+	**/
+	void AddAnnotationInternal(int32 Index, TAnnotation Annotation)
+	{
+		check(Index >= 0);
+		if (Annotation.IsDefault())
+		{
+			FreeAnnotation(Index); // adding the default annotation is the same as removing an annotation
+		}
+		else
+		{
+			if (NumAnnotations == 0 && Chunks.Num() == 0)
+			{
+				// we are adding the first one, so if we are auto removing or verifying removal, register now
+#if (UE_BUILD_SHIPPING || UE_BUILD_TEST)
+				if (bAutoRemove)
+#endif
+				{
+					GUObjectArray.AddUObjectDeleteListener(this);
+				}
+			}
+
+			TAnnotation& NewAnnotation = AllocateAnnotation(Index);
+			NewAnnotation = Annotation;
+		}
+	}
+
+public:
+
+	/** Constructor : Probably not thread safe **/
+	FUObjectAnnotationChunked() TSAN_SAFE
+		: NumAnnotations(0)
+		, MaxAnnotations(0)
+		, CurrentAllocatedMemory(0)
+		, MaxAllocatedMemory(0)
+	{
+	}
+
+	virtual ~FUObjectAnnotationChunked()
+	{
+		RemoveAllAnnotations();
+	}
+
+public:
+
+	/**
+	 * Add an annotation to the annotation list. If the Annotation is the default, then the annotation is removed.
+	 *
+	 * @param Object		Object to annotate.
+	 * @param Annotation	Annotation to associate with Object.
+	 */
+	void AddAnnotation(const UObjectBase *Object, TAnnotation Annotation)
+	{
+		check(Object);
+		AddAnnotation(GUObjectArray.ObjectToIndex(Object), Annotation);
+	}
+	/**
+	 * Add an annotation to the annotation list. If the Annotation is the default, then the annotation is removed.
+	 *
+	 * @param Index			Index of object to annotate.
+	 * @param Annotation	Annotation to associate with Object.
+	 */
+	void AddAnnotation(int32 Index, TAnnotation Annotation)
+	{
+		FRWScopeLock AnnotationArrayLock(AnnotationArrayCritical, SLT_Write);
+		AddAnnotationInternal(Index, Annotation);
+	}
+
+	/**
+	 * Add an annotation to the annotation list. If the Annotation is the default, then the annotation is removed.
+	 *
+	 * @param Object		Object to annotate.
+	 * @param Annotation	Annotation to associate with Object.
+	 */
+	TAnnotation& AddOrGetAnnotation(const UObjectBase *Object, TFunctionRef<TAnnotation()> NewAnnotationFn)
+	{
+		check(Object);
+		AddOrGetAnnotation(GUObjectArray.ObjectToIndex(Object), NewAnnotationFn);
+	}
+	/**
+	 * Add an annotation to the annotation list. If the Annotation is the default, then the annotation is removed.
+	 *
+	 * @param Index			Index of object to annotate.
+	 * @param Annotation	Annotation to associate with Object.
+	 */
+	TAnnotation& AddOrGetAnnotation(int32 Index, TFunctionRef<TAnnotation()> NewAnnotationFn)
+	{		
+		FRWScopeLock AnnotationArrayLock(AnnotationArrayCritical, SLT_Write);
+		
+		if (NumAnnotations == 0 && Chunks.Num() == 0)
+		{
+			// we are adding the first one, so if we are auto removing or verifying removal, register now
+#if (UE_BUILD_SHIPPING || UE_BUILD_TEST)
+			if (bAutoRemove)
+#endif
+			{
+				GUObjectArray.AddUObjectDeleteListener(this);
+			}
+		}
+
+		ExpandChunksToIndex(Index);
+
+		const int32 ChunkIndex = Index / NumAnnotationsPerChunk;
+		const int32 WithinChunkIndex = Index % NumAnnotationsPerChunk;
+
+		TAnnotationChunk& Chunk = Chunks[ChunkIndex];
+		if (!Chunk.Items)
+		{
+			Chunk.Items = new TAnnotation[NumAnnotationsPerChunk];
+			CurrentAllocatedMemory += NumAnnotationsPerChunk * sizeof(TAnnotation);
+			MaxAllocatedMemory = FMath::Max(CurrentAllocatedMemory, MaxAllocatedMemory);
+		}
+		if (Chunk.Items[WithinChunkIndex].IsDefault())
+		{
+			Chunk.Num++;
+			check(Chunk.Num <= NumAnnotationsPerChunk);
+			NumAnnotations++;
+			Chunk.Items[WithinChunkIndex] = NewAnnotationFn();
+			check(!Chunk.Items[WithinChunkIndex].IsDefault());
+		}
+		return Chunk.Items[WithinChunkIndex];
+	}
+
+	/**
+	 * Removes an annotation from the annotation list.
+	 *
+	 * @param Object		Object to de-annotate.
+	 */
+	void RemoveAnnotation(const UObjectBase *Object)
+	{
+		check(Object);
+		RemoveAnnotation(GUObjectArray.ObjectToIndex(Object));
+	}
+	/**
+	 * Removes an annotation from the annotation list.
+	 *
+	 * @param Object		Object to de-annotate.
+	 */
+	void RemoveAnnotation(int32 Index)
+	{
+		FRWScopeLock AnnotationArrayLock(AnnotationArrayCritical, SLT_Write);
+		FreeAnnotation(Index);
+	}
+
+	/**
+	 * Return the annotation associated with a uobject
+	 *
+	 * @param Object		Object to return the annotation for
+	 */
+	FORCEINLINE TAnnotation GetAnnotation(const UObjectBase *Object)
+	{
+		check(Object);
+		return GetAnnotation(GUObjectArray.ObjectToIndex(Object));
+	}
+
+	/**
+	 * Return the annotation associated with a uobject
+	 *
+	 * @param Index		Index of the annotation to return
+	 */
+	FORCEINLINE TAnnotation GetAnnotation(int32 Index)
+	{
+		check(Index >= 0);
+		FRWScopeLock AnnotationArrayLock(AnnotationArrayCritical, SLT_ReadOnly);
+
+		const int32 ChunkIndex = Index / NumAnnotationsPerChunk;
+		if (ChunkIndex < Chunks.Num())
+		{
+			const int32 WithinChunkIndex = Index % NumAnnotationsPerChunk;
+
+			TAnnotationChunk& Chunk = Chunks[ChunkIndex];
+			if (Chunk.Items != nullptr)
+			{
+				return Chunk.Items[WithinChunkIndex];
+			}
+		}
+		return TAnnotation();
+	}
+
+	/**
+	* Return the number of elements in the array
+	* Thread safe, but you know, someone might have added more elements before this even returns
+	* @return	the number of elements in the array
+	**/
+	FORCEINLINE int32 GetAnnotationCount() const
+	{
+		return NumAnnotations;
+	}
+
+	/**
+	* Return the number max capacity of the array
+	* Thread safe, but you know, someone might have added more elements before this even returns
+	* @return	the maximum number of elements in the array
+	**/
+	FORCEINLINE int32 GetMaxAnnottations() const TSAN_SAFE
+	{
+		return MaxAnnotations;
+	}
+
+	/**
+	* Return if this index is valid
+	* Thread safe, if it is valid now, it is valid forever. Other threads might be adding during this call.
+	* @param	Index	Index to test
+	* @return	true, if this is a valid
+	**/
+	FORCEINLINE bool IsValidIndex(int32 Index) const
+	{
+		return Index >= 0 && Index < MaxAnnotations;
+	}
+
+	/**
+	 * Removes all annotation from the annotation list.
+	 */
+	void RemoveAllAnnotations()
+	{
+		bool bHadAnnotations = (NumAnnotations > 0);	
+		FRWScopeLock AnnotationArrayLock(AnnotationArrayCritical, SLT_Write);
+		FreeAllAnnotations();
+		if (bHadAnnotations)
+		{
+			// we are removing the last one, so if we are auto removing or verifying removal, unregister now
+#if (UE_BUILD_SHIPPING || UE_BUILD_TEST)
+			if (bAutoRemove)
+#endif
+			{
+				GUObjectArray.RemoveUObjectDeleteListener(this);
+			}
+		}
+	}
+
+	/**
+	 * Frees chunk memory from empty chunks.
+	 */
+	void TrimAnnotations()
+	{
+		FRWScopeLock AnnotationArrayLock(AnnotationArrayCritical, SLT_Write);
+		for (TAnnotationChunk& Chunk : Chunks)
+		{
+			if (Chunk.Num == 0 && Chunk.Items)
+			{
+				delete[] Chunk.Items;
+				Chunk.Items = nullptr;
+				const uint32 ChunkMemory = NumAnnotationsPerChunk * sizeof(TAnnotationChunk);
+				check(CurrentAllocatedMemory >= ChunkMemory);
+				CurrentAllocatedMemory -= ChunkMemory;
+			}
+		}
+	}
+
+	/** Returns the memory allocated by the internal array */
+	uint32 GetAllocatedSize() const
+	{
+		uint32 AllocatedSize = Chunks.GetAllocatedSize();
+		for (const TAnnotationChunk& Chunk : Chunks)
+		{
+			if (Chunk.Items)
+			{
+				AllocatedSize += NumAnnotationsPerChunk * sizeof(TAnnotation);
+			}
+		}
+		return AllocatedSize;
+	}
+
+	/** Returns the maximum memory allocated by the internal arrays */
+	uint32 GetMaxAllocatedSize() const
+	{
+		return Chunks.GetAllocatedSize() + MaxAllocatedMemory;
+	}
+
+	/**
+	 * Interface for FUObjectAllocator::FUObjectDeleteListener
+	 *
+	 * @param Object object that has been destroyed
+	 * @param Index	index of object that is being deleted
+	 */
+	virtual void NotifyUObjectDeleted(const UObjectBase *Object, int32 Index)
+	{
+#if !(UE_BUILD_SHIPPING || UE_BUILD_TEST)
+		if (!bAutoRemove)
+		{
+			// in this case we are only verifying that the external assurances of removal are met
+			check(Index >= MaxAnnotations || GetAnnotation(Index).IsDefault());
+		}
+		else
+#endif
+		{
+			RemoveAnnotation(Index);
+		}
+	}
+};
 
 /**
 * FUObjectAnnotationDense is a helper class that is used to store dense, fast, temporary, editor only, external 
@@ -665,6 +1091,12 @@ public:
 		return AnnotationArray[Index];
 	}
 
+	/** Returns the memory allocated by the internal array */
+	uint32 GetAllocatedSize() const
+	{
+		return AnnotationArray.GetAllocatedSize();
+	}
+
 private:
 
 	/**
