Fixing CIS

diff --git a/Engine/Source/Runtime/CoreUObject/Public/UObject/UObjectAnnotation.h b/Engine/Source/Runtime/CoreUObject/Public/UObject/UObjectAnnotation.h
--- a/Engine/Source/Runtime/CoreUObject/Public/UObject/UObjectAnnotation.h
+++ b/Engine/Source/Runtime/CoreUObject/Public/UObject/UObjectAnnotation.h
@@ -481,7 +481,7 @@ class FUObjectAnnotationChunked : public FUObjectArray::FUObjectDeleteListener
 
 		TAnnotationChunk()
 			: Num(0)
-			: Items(nullptr)
+			, Items(nullptr)
 		{}
 	};
 
