Downgrade missing GetLifetimeReplicatedProps property from Error to Warning since it gets picked up during the Cooking process.

diff --git a/Engine/Source/Runtime/CoreUObject/Private/UObject/Class.cpp b/Engine/Source/Runtime/CoreUObject/Private/UObject/Class.cpp
--- a/Engine/Source/Runtime/CoreUObject/Private/UObject/Class.cpp
+++ b/Engine/Source/Runtime/CoreUObject/Private/UObject/Class.cpp
@@ -3539,7 +3539,7 @@ void UClass::ValidateRuntimeReplicationData()
 				}
 			}
 
-			UE_LOG(LogClass, Error, TEXT("Property %s::%s (SourceClass: %s) was not registered in GetLifetimeReplicatedProps. This property will not be replicated."),
+			UE_LOG(LogClass, Warning, TEXT("Property %s::%s (SourceClass: %s) was not registered in GetLifetimeReplicatedProps. This property will not be replicated."),
 				*GetName(), *RepProp->GetName(), *RepProp->GetOwnerClass()->GetName());
 		}
 	}
