UE4 - Moved mapped ubulk files to .m.ubulk so that they can be aligned and not compressed without affecting other bulk data such as textures.

diff --git a/Engine/Source/Developer/PakFileUtilities/Private/PakFileUtilities.cpp b/Engine/Source/Developer/PakFileUtilities/Private/PakFileUtilities.cpp
--- a/Engine/Source/Developer/PakFileUtilities/Private/PakFileUtilities.cpp
+++ b/Engine/Source/Developer/PakFileUtilities/Private/PakFileUtilities.cpp
@@ -845,12 +845,13 @@ void ProcessCommandLine(const TCHAR* CmdLine, const TArray<FString>& NonOptionAr
 
 			UE_LOG(LogPakFile, Log, TEXT("Added file Source: %s Dest: %s"), *Input.Source, *Input.Dest);
 
-			bool bIsUBulk = FPaths::GetExtension(Input.Source, true) == TEXT(".ubulk");
-			if (bIsUBulk && CmdLineParameters.AlignForMemoryMapping > 0 && Input.bNeedsCompression)
+			bool bIsMappedBulk = Input.Source.EndsWith(TEXT(".m.ubulk"));
+			if (bIsMappedBulk && CmdLineParameters.AlignForMemoryMapping > 0 && Input.bNeedsCompression 
+				&& !Input.bNeedEncryption) // if it is encrypted, we will compress it anyway since it won't be mapped at runtime
 			{
 				// no compression for bulk aligned files because they are memory mapped
 				Input.bNeedsCompression = false;
-				UE_LOG(LogPakFile, Warning, TEXT("Stripped compression from %s for memory mapping."), *Input.Dest);
+				UE_LOG(LogPakFile, Log, TEXT("Stripped compression from %s for memory mapping."), *Input.Dest);
 			}
 			Entries.Add(Input);
 		}			
@@ -1529,7 +1530,7 @@ bool CreatePakFile(const TCHAR* Filename, TArray<FPakInputPair>& FilesToAdd, con
 		bool bDeleted = FilesToAdd[FileIndex].bIsDeleteRecord;
 		bool bIsUAssetUExpPairUAsset = false;
 		bool bIsUAssetUExpPairUExp = false;
-		bool bIsUBulk = FPaths::GetExtension(FilesToAdd[FileIndex].Source, true) == TEXT(".ubulk");
+		bool bIsMappedBulk = FilesToAdd[FileIndex].Source.EndsWith(TEXT(".m.ubulk"));
 
 		if (FileIndex)
 		{
@@ -1655,7 +1656,7 @@ bool CreatePakFile(const TCHAR* Filename, TArray<FPakInputPair>& FilesToAdd, con
 				}
 			}
 			// Align bulk data
-			if (bIsUBulk && CmdLineParameters.AlignForMemoryMapping > 0 && OriginalFileSize != INDEX_NONE && !bDeleted)
+			if (bIsMappedBulk && CmdLineParameters.AlignForMemoryMapping > 0 && OriginalFileSize != INDEX_NONE && !bDeleted)
 			{
 				if (!IsAligned(NewEntryOffset + NewEntry.Info.GetSerializedSize(FPakInfo::PakFile_Version_Latest), CmdLineParameters.AlignForMemoryMapping))
 				{
@@ -1701,7 +1702,7 @@ bool CreatePakFile(const TCHAR* Filename, TArray<FPakInputPair>& FilesToAdd, con
 		if (bCopiedToPak)
 		{
 			if (RequiredPatchPadding > 0 &&
-				!(bIsUBulk && CmdLineParameters.AlignForMemoryMapping > 0) // don't wreck the bulk padding with patch padding
+				!(bIsMappedBulk && CmdLineParameters.AlignForMemoryMapping > 0) // don't wreck the bulk padding with patch padding
 				)
 			{
 				//if the next file is going to cross a patch-block boundary then pad out the current set of files with 0's
@@ -1856,7 +1857,7 @@ bool CreatePakFile(const TCHAR* Filename, TArray<FPakInputPair>& FilesToAdd, con
 				int64 RealStart = Entry.Info.Offset;
 				if ((RealStart % RequiredPatchPadding) != 0 && 
 					!Entry.Filename.EndsWith(TEXT("uexp")) && // these are export sections of larger files and may be packed with uasset/umap and so we don't need a warning here
-					!(Entry.Filename.EndsWith(TEXT("ubulk")) && CmdLineParameters.AlignForMemoryMapping > 0)) // Bulk padding unaligns patch padding and so we don't need a warning here
+					!(Entry.Filename.EndsWith(TEXT(".m.ubulk")) && CmdLineParameters.AlignForMemoryMapping > 0)) // Bulk padding unaligns patch padding and so we don't need a warning here
 				{
 					UE_LOG(LogPakFile, Warning, TEXT("File at offset %lld of size %lld not aligned to patch size %i"), RealStart, Entry.Info.Size, RequiredPatchPadding);
 				}
diff --git a/Engine/Source/Runtime/CoreUObject/Private/Serialization/BulkData.cpp b/Engine/Source/Runtime/CoreUObject/Private/Serialization/BulkData.cpp
--- a/Engine/Source/Runtime/CoreUObject/Private/Serialization/BulkData.cpp
+++ b/Engine/Source/Runtime/CoreUObject/Private/Serialization/BulkData.cpp
@@ -103,8 +103,11 @@ FOwnedBulkDataPtr::~FOwnedBulkDataPtr()
 	}
 	else
 	{
-		delete MappedHandle;
-		delete MappedRegion;
+		if (MappedRegion || MappedHandle)
+		{
+			delete MappedRegion;
+			delete MappedHandle;
+		}
 	}
 }
 
@@ -1012,7 +1015,14 @@ void FUntypedBulkData::Serialize( FArchive& Ar, UObject* Owner, int32 Idx, bool
 				}
 				else if (BulkDataFlags & BULKDATA_PayloadInSeperateFile)
 				{
-					Filename = FPaths::ChangeExtension(Filename, TEXT(".ubulk"));
+					if (BulkDataFlags & BULKDATA_MemoryMappedPayload)
+					{
+						Filename = FPaths::ChangeExtension(Filename, TEXT(".m.ubulk"));
+					}
+					else
+					{
+						Filename = FPaths::ChangeExtension(Filename, TEXT(".ubulk"));
+					}
 
 
 					if (bAttemptFileMapping)
@@ -1496,7 +1506,15 @@ void FUntypedBulkData::LoadDataIntoMemory( void* Dest )
 		(BulkDataFlags & BULKDATA_PayloadInSeperateFile))
 	{
 		// The attached archive is a package cooked for EDL loaded in the editor so the actual bulk data sits in a separate ubulk file.
-		const FString BulkDataFilename = FPaths::ChangeExtension(Filename, TEXT(".ubulk"));
+		FString BulkDataFilename;
+		if (BulkDataFlags & BULKDATA_MemoryMappedPayload)
+		{
+			BulkDataFilename = FPaths::ChangeExtension(Filename, TEXT(".m.ubulk"));
+		}
+		else
+		{
+			BulkDataFilename = FPaths::ChangeExtension(Filename, TEXT(".ubulk"));
+		}
 		BulkDataArchive = IFileManager::Get().CreateFileReader(*BulkDataFilename, FILEREAD_Silent);
 	}
 
@@ -1697,7 +1715,7 @@ void FFloatBulkData::SerializeElement( FArchive& Ar, void* Data, int32 ElementIn
 	Ar << FloatData;
 }
 
-void FFormatContainer::Serialize(FArchive& Ar, UObject* Owner, const TArray<FName>* FormatsToSave, bool bSingleUse, uint32 InAlignment, bool bInline)
+void FFormatContainer::Serialize(FArchive& Ar, UObject* Owner, const TArray<FName>* FormatsToSave, bool bSingleUse, uint32 InAlignment, bool bInline, bool bMapped)
 {
 	if (Ar.IsLoading())
 	{
@@ -1741,11 +1759,15 @@ void FFormatContainer::Serialize(FArchive& Ar, UObject* Owner, const TArray<FNam
 				if (bInline)
 				{
 					Bulk->SetBulkDataFlags(BULKDATA_ForceInlinePayload);
-					Bulk->ClearBulkDataFlags(BULKDATA_PayloadAtEndOfFile | BULKDATA_PayloadInSeperateFile | BULKDATA_Force_NOT_InlinePayload);
+					Bulk->ClearBulkDataFlags(BULKDATA_PayloadAtEndOfFile | BULKDATA_PayloadInSeperateFile | BULKDATA_Force_NOT_InlinePayload | BULKDATA_MemoryMappedPayload);
 				}
 				else
 				{
 					Bulk->SetBulkDataFlags(BULKDATA_PayloadAtEndOfFile | BULKDATA_PayloadInSeperateFile | BULKDATA_Force_NOT_InlinePayload);
+					if (bMapped)
+					{
+						Bulk->SetBulkDataFlags(BULKDATA_MemoryMappedPayload);
+					}
 					Bulk->ClearBulkDataFlags(BULKDATA_ForceInlinePayload);
 
 				}
diff --git a/Engine/Source/Runtime/CoreUObject/Private/UObject/SavePackage.cpp b/Engine/Source/Runtime/CoreUObject/Private/UObject/SavePackage.cpp
--- a/Engine/Source/Runtime/CoreUObject/Private/UObject/SavePackage.cpp
+++ b/Engine/Source/Runtime/CoreUObject/Private/UObject/SavePackage.cpp
@@ -5564,6 +5564,7 @@ FSavePackageResultStruct UPackage::Save(UPackage* InOuter, UObject* Base, EObjec
 
 					FArchive* BulkArchive = nullptr;
 					FArchive* OptionalBulkArchive = nullptr;
+					FArchive* MappedBulkArchive = nullptr;
 					uint32 ExtraBulkDataFlags = 0;
 
 					static struct FUseSeperateBulkDataFiles
@@ -5591,6 +5592,10 @@ FSavePackageResultStruct UPackage::Save(UPackage* InOuter, UObject* Base, EObjec
 
 					const FString OptionalBulkFilename = FPaths::ChangeExtension(Filename, OptionalBulkFileExtension);
 
+					const static TCHAR* MappedBulkFileExtension = TEXT(".m.ubulk");
+
+					const FString MappedBulkFilename = FPaths::ChangeExtension(Filename, MappedBulkFileExtension);
+
 					if (bShouldUseSeparateBulkFile)
 					{
 						ExtraBulkDataFlags = BULKDATA_PayloadInSeperateFile;
@@ -5598,12 +5603,13 @@ FSavePackageResultStruct UPackage::Save(UPackage* InOuter, UObject* Base, EObjec
 						{
 							BulkArchive = new FBufferArchive(true);
 							OptionalBulkArchive = new FBufferArchive(true);
-
+							MappedBulkArchive = new FBufferArchive(true);
 						}
 						else
 						{
 							BulkArchive = IFileManager::Get().CreateFileWriter(*BulkFilename);
 							OptionalBulkArchive = IFileManager::Get().CreateFileWriter(*OptionalBulkFilename);
+							MappedBulkArchive = IFileManager::Get().CreateFileWriter(*MappedBulkFilename);
 						}
 					}
 
@@ -5625,6 +5631,16 @@ FSavePackageResultStruct UPackage::Save(UPackage* InOuter, UObject* Base, EObjec
 						// Set bulk data flags to what they were during initial serialization (they might have changed after that)
 						const uint32 OldBulkDataFlags = BulkDataStorageInfo.BulkData->GetBulkDataFlags();
 						uint32 ModifiedBulkDataFlags = BulkDataStorageInfo.BulkDataFlags | ExtraBulkDataFlags;
+						bool bBulkItemIsOptional = (ModifiedBulkDataFlags & BULKDATA_OptionalPayload);
+						bool bBulkItemIsMapped = bAlignBulkData && (ModifiedBulkDataFlags & BULKDATA_MemoryMappedPayload);
+
+						if (bBulkItemIsMapped && bBulkItemIsOptional)
+						{
+							UE_LOG(LogSavePackage, Warning, TEXT("%s has bulk data that is both mapped and optional. This is not currently supported. Will not be mapped."), Filename);
+							ModifiedBulkDataFlags &= ~BULKDATA_MemoryMappedPayload;
+							bBulkItemIsMapped = false;
+						}
+
 						BulkDataStorageInfo.BulkData->ClearBulkDataFlags(0xFFFFFFFF);
 						BulkDataStorageInfo.BulkData->SetBulkDataFlags(ModifiedBulkDataFlags);
 
@@ -5632,12 +5648,12 @@ FSavePackageResultStruct UPackage::Save(UPackage* InOuter, UObject* Base, EObjec
 						if ( bShouldUseSeparateBulkFile )
 						{
 							check( OptionalBulkArchive && BulkArchive );
-							TargetArchive = ModifiedBulkDataFlags & BULKDATA_OptionalPayload ? OptionalBulkArchive : BulkArchive;
+							TargetArchive = bBulkItemIsOptional ? OptionalBulkArchive : (bBulkItemIsMapped ? MappedBulkArchive : BulkArchive);
 						}
 
 						int64 BulkStartOffset = TargetArchive->Tell();
 
-						if (bAlignBulkData && BulkDataAlignment > 0 && !IsAligned(BulkStartOffset, BulkDataAlignment))
+						if (bBulkItemIsMapped && BulkDataAlignment > 0 && !IsAligned(BulkStartOffset, BulkDataAlignment))
 						{
 							int64 AlignedOffset = Align(BulkStartOffset, BulkDataAlignment);
 							int64 Padding = AlignedOffset - BulkStartOffset;
@@ -5692,24 +5708,24 @@ FSavePackageResultStruct UPackage::Save(UPackage* InOuter, UObject* Base, EObjec
 							{
 								TotalPackageSizeUncompressed += Archive->TotalSize();
 								Archive->Close();
-						if ( bSaveAsync )
-						{
+								if ( bSaveAsync )
+								{
 									FBufferArchive* Buffer = (FBufferArchive*)(Archive);
 
 									if ( Buffer->TotalSize() > 0 )
 									{
 										int64 DataSize = Buffer->TotalSize();
 
-							// TODO - update the BulkBuffer code to write into a FLargeMemoryWriter so we can 
-							// take ownership of the data
-							uint8* Data = new uint8[DataSize];
+										// TODO - update the BulkBuffer code to write into a FLargeMemoryWriter so we can 
+										// take ownership of the data
+										uint8* Data = new uint8[DataSize];
 										FMemory::Memcpy(Data, Buffer->GetData(), DataSize);
 
-							FLargeMemoryPtr DataPtr = FLargeMemoryPtr(Data);
+										FLargeMemoryPtr DataPtr = FLargeMemoryPtr(Data);
 		
 										AsyncWriteFile(MoveTemp(DataPtr), DataSize, *ArchiveFilename, FDateTime::MinValue(), false);
-							}
-						}
+									}
+								}
 								delete Archive;
 							};
 
@@ -5725,10 +5741,15 @@ FSavePackageResultStruct UPackage::Save(UPackage* InOuter, UObject* Base, EObjec
 								{
 									CookedPackageHash.Update(((FBufferArchive*)OptionalBulkArchive)->GetData(), OptionalBulkArchive->TotalSize());
 								}
+								if (MappedBulkArchive->TotalSize())
+								{
+									CookedPackageHash.Update(((FBufferArchive*)MappedBulkArchive)->GetData(), MappedBulkArchive->TotalSize());
+								}
 							}
 						}
 						FinalizeBulkDataFile(BulkArchive, BulkFilename);
 						FinalizeBulkDataFile(OptionalBulkArchive, OptionalBulkFilename);
+						FinalizeBulkDataFile(MappedBulkArchive, MappedBulkFilename);
 
 						if (SaveFlags & SAVE_ComputeHash)
 						{
@@ -5736,6 +5757,7 @@ FSavePackageResultStruct UPackage::Save(UPackage* InOuter, UObject* Base, EObjec
 							{
 								AddFileToHash(BulkFilename, CookedPackageHash);
 								AddFileToHash(OptionalBulkFilename, CookedPackageHash);
+								AddFileToHash(MappedBulkFilename, CookedPackageHash);
 							}
 						}
 						
diff --git a/Engine/Source/Runtime/CoreUObject/Public/Serialization/BulkData.h b/Engine/Source/Runtime/CoreUObject/Public/Serialization/BulkData.h
--- a/Engine/Source/Runtime/CoreUObject/Public/Serialization/BulkData.h
+++ b/Engine/Source/Runtime/CoreUObject/Public/Serialization/BulkData.h
@@ -38,6 +38,8 @@ enum EBulkDataFlags
 	BULKDATA_Force_NOT_InlinePayload			= 1 << 10,
 	/** This payload is optional and may not be on device */
 	BULKDATA_OptionalPayload					= 1 << 11,
+	/** This payload will be memory mapped, this requires alignment, no compression etc. */
+	BULKDATA_MemoryMappedPayload = 1 << 12,
 };
 
 /**
@@ -769,7 +771,7 @@ public:
 		}
 		Formats.Empty();
 	}
-	COREUOBJECT_API void Serialize(FArchive& Ar, UObject* Owner, const TArray<FName>* FormatsToSave = nullptr, bool bSingleUse = true, uint32 InAlignment = DEFAULT_ALIGNMENT, bool bInline = true);
+	COREUOBJECT_API void Serialize(FArchive& Ar, UObject* Owner, const TArray<FName>* FormatsToSave = nullptr, bool bSingleUse = true, uint32 InAlignment = DEFAULT_ALIGNMENT, bool bInline = true, bool bMapped = false);
 	COREUOBJECT_API void SerializeAttemptMappedLoad(FArchive& Ar, UObject* Owner);
 };
 
diff --git a/Engine/Source/Runtime/Engine/Private/Animation/AnimSequence.cpp b/Engine/Source/Runtime/Engine/Private/Animation/AnimSequence.cpp
--- a/Engine/Source/Runtime/Engine/Private/Animation/AnimSequence.cpp
+++ b/Engine/Source/Runtime/Engine/Private/Animation/AnimSequence.cpp
@@ -2784,7 +2784,7 @@ void UAnimSequence::SerializeCompressedData(FArchive& Ar, bool bDDCData)
 			void* Dest = OptionalBulk.Realloc(Num);
 			FMemory::Memcpy(Dest, &(SerializedData[0]), Num);
 			OptionalBulk.Unlock();
-			OptionalBulk.SetBulkDataFlags(BULKDATA_PayloadAtEndOfFile | BULKDATA_PayloadInSeperateFile | BULKDATA_Force_NOT_InlinePayload);
+			OptionalBulk.SetBulkDataFlags(BULKDATA_PayloadAtEndOfFile | BULKDATA_PayloadInSeperateFile | BULKDATA_Force_NOT_InlinePayload | BULKDATA_MemoryMappedPayload);
 			OptionalBulk.ClearBulkDataFlags(BULKDATA_ForceInlinePayload);
 			OptionalBulk.Serialize(Ar, this);
 
diff --git a/Engine/Source/Runtime/Engine/Private/SoundWave.cpp b/Engine/Source/Runtime/Engine/Private/SoundWave.cpp
--- a/Engine/Source/Runtime/Engine/Private/SoundWave.cpp
+++ b/Engine/Source/Runtime/Engine/Private/SoundWave.cpp
@@ -357,8 +357,10 @@ void USoundWave::Serialize( FArchive& Ar )
 						ActualFormatsToSave.Add(Format);
 					}
 				}
-				bool bInline = !(CookingTarget->SupportsFeature(ETargetPlatformFeatures::MemoryMappedFiles) && CookingTarget->SupportsFeature(ETargetPlatformFeatures::MemoryMappedAudio));
-				CompressedFormatData.Serialize(Ar, this, &ActualFormatsToSave, true, DEFAULT_ALIGNMENT, bInline);
+				bool bMapped = CookingTarget->SupportsFeature(ETargetPlatformFeatures::MemoryMappedFiles) && CookingTarget->SupportsFeature(ETargetPlatformFeatures::MemoryMappedAudio);
+				CompressedFormatData.Serialize(Ar, this, &ActualFormatsToSave, true, DEFAULT_ALIGNMENT, 
+					!bMapped, // inline if not mapped
+					bMapped);
 #endif
 			}
 			else
