Support printing all results for obj refs command.

diff --git a/Engine/Source/Runtime/CoreUObject/Private/UObject/Obj.cpp b/Engine/Source/Runtime/CoreUObject/Private/UObject/Obj.cpp
--- a/Engine/Source/Runtime/CoreUObject/Private/UObject/Obj.cpp
+++ b/Engine/Source/Runtime/CoreUObject/Private/UObject/Obj.cpp
@@ -3828,6 +3828,10 @@ bool StaticExec( UWorld* InWorld, const TCHAR* Cmd, FOutputDevice& Ar )
 						}
 						SearchModeFlags |= EReferenceChainSearchMode::Longest;
 					}
+					else if (FCString::Stricmp(*Tok, TEXT("all")) == 0)
+					{
+						SearchModeFlags |= EReferenceChainSearchMode::PrintAllResults;
+					}
 					else if (FCString::Stricmp(*Tok, TEXT("external")) == 0)
 					{
 						SearchModeFlags |= EReferenceChainSearchMode::ExternalOnly;
