Support for BulkData with 64-bit sizes. This is implemented via an additional BULKDATA flag, rather than with versioning, as BulkData is often loaded from unversioned archives, e.g. from DDC via FMemoryReader.
This addresses the issue of enormous meshes >2Gb being unimportable (since source model data is stored as bulk data and loaded by the editor on demand).
This *won't* address the related problem that reading a large BulkData into a TArray is still not possible due to the 31-bit index limit on TArray.

diff --git a/Engine/Source/Runtime/CoreUObject/Private/Serialization/BulkData.cpp b/Engine/Source/Runtime/CoreUObject/Private/Serialization/BulkData.cpp
--- a/Engine/Source/Runtime/CoreUObject/Private/Serialization/BulkData.cpp
+++ b/Engine/Source/Runtime/CoreUObject/Private/Serialization/BulkData.cpp
@@ -5,8 +5,8 @@
 #include "HAL/FileManager.h"
 #include "Misc/ScopeLock.h"
 #include "Misc/Paths.h"
-#include "Serialization/MemoryWriter.h"
-#include "Serialization/MemoryReader.h"
+#include "Serialization/LargeMemoryWriter.h"
+#include "Serialization/LargeMemoryReader.h"
 #include "HAL/IConsoleManager.h"
 #include "UObject/Package.h"
 #include "Templates/Casts.h"
@@ -75,7 +75,7 @@ protected:
  */
 struct FObjectAndSize
 {
-	FObjectAndSize( const UObject* InObject, int32 InSize )
+	FObjectAndSize( const UObject* InObject, int64 InSize )
 	:	Object( InObject )
 	,	Size( InSize )
 	{}
@@ -83,7 +83,7 @@ struct FObjectAndSize
 	/** Object associated with size. */
 	const UObject*	Object;
 	/** Size associated with object. */
-	int32			Size;
+	int64			Size;
 };
 
 /** Hash function required for TMap support */
@@ -327,14 +327,14 @@ void FUntypedBulkData::DumpBulkDataUsage( FOutputDevice& Log )
 	for( int32 PerClassIndex=0; PerClassIndex<PerClassSizeArray.Num(); PerClassIndex++ )
 	{
 		const FObjectAndSize& PerClassSize = PerClassSizeArray[ PerClassIndex ];
-		Log.Logf( TEXT("  %5d KByte of bulk data for Class %s"), PerClassSize.Size / 1024, *PerClassSize.Object->GetPathName() );
+		Log.Logf( TEXT("  %5lld KByte of bulk data for Class %s"), PerClassSize.Size / 1024, *PerClassSize.Object->GetPathName() );
 	}
 	UE_LOG(LogSerialization, Log, TEXT(""));
 	UE_LOG(LogSerialization, Log, TEXT("Detailed per object stats of bulk data use:"));
 	for( int32 PerObjectIndex=0; PerObjectIndex<PerObjectSizeArray.Num(); PerObjectIndex++ )
 	{
 		const FObjectAndSize& PerObjectSize = PerObjectSizeArray[ PerObjectIndex ];
-		Log.Logf( TEXT("  %5d KByte of bulk data for %s"), PerObjectSize.Size / 1024, *PerObjectSize.Object->GetFullName() );
+		Log.Logf( TEXT("  %5lld KByte of bulk data for %s"), PerObjectSize.Size / 1024, *PerObjectSize.Object->GetFullName() );
 	}
 	UE_LOG(LogSerialization, Log, TEXT(""));
 #else
@@ -352,7 +352,7 @@ void FUntypedBulkData::DumpBulkDataUsage( FOutputDevice& Log )
  *
  * @return Number of elements in this bulk data array
  */
-int32 FUntypedBulkData::GetElementCount() const
+int64 FUntypedBulkData::GetElementCount() const
 {
 	return ElementCount;
 }
@@ -361,7 +361,7 @@ int32 FUntypedBulkData::GetElementCount() const
  *
  * @return Size of the bulk data in bytes
  */
-int32 FUntypedBulkData::GetBulkDataSize() const
+int64 FUntypedBulkData::GetBulkDataSize() const
 {
 	return GetElementCount() * GetElementSize();
 }
@@ -371,7 +371,7 @@ int32 FUntypedBulkData::GetBulkDataSize() const
  *
  * @return Size of the bulk data on disk or INDEX_NONE in case there's no association
  */
-int32 FUntypedBulkData::GetBulkDataSizeOnDisk() const
+int64 FUntypedBulkData::GetBulkDataSizeOnDisk() const
 {
 	return BulkDataSizeOnDisk;
 }
@@ -511,7 +511,7 @@ void FUntypedBulkData::GetCopy( void** Dest, bool bDiscardInternalCopy )
 			// Can't/ Don't discard so we need to allocate and copy.
 			else
 			{
-				int32 BulkDataSize = GetBulkDataSize();
+				int64 BulkDataSize = GetBulkDataSize();
 				if (BulkDataSize != 0)
 				{
 					// Allocate enough memory for data...
@@ -529,7 +529,7 @@ void FUntypedBulkData::GetCopy( void** Dest, bool bDiscardInternalCopy )
 		// Data isn't currently loaded so we need to load it from disk.
 		else
 		{
-			int32 BulkDataSize = GetBulkDataSize();
+			int64 BulkDataSize = GetBulkDataSize();
 			if (BulkDataSize != 0)
 			{
 				// Allocate enough memory for data...
@@ -608,7 +608,7 @@ const void* FUntypedBulkData::LockReadOnly() const
  *
  * @param InElementCount	Number of elements array should be resized to
  */
-void* FUntypedBulkData::Realloc( int32 InElementCount )
+void* FUntypedBulkData::Realloc( int64 InElementCount )
 {
 	DECLARE_SCOPE_CYCLE_COUNTER(TEXT("FUntypedBulkData::Realloc"), STAT_UBD_Realloc, STATGROUP_Memory);
 
@@ -913,13 +913,29 @@ void FUntypedBulkData::Serialize( FArchive& Ar, UObject* Owner, int32 Idx, bool
 		int64 SavedBulkDataFlagsPos = Ar.Tell();
 		{
 			FArchive::FScopeSetDebugSerializationFlags S(Ar, DSF_IgnoreDiff);
+
+			// If the bulk data size is greater than can be held in an int32, then potentially the ElementCount
+			// and BulkDataSizeOnDisk need to be held as int64s, so set a flag indicating the new format.
+			if (Ar.IsSaving() && GetBulkDataSize() >= (1LL << 31))
+			{
+				BulkDataFlags |= BULKDATA_Size64Bit;
+			}
 			Ar << BulkDataFlags;
 		}
 
 
 		// Number of elements in array.
-		Ar << ElementCount;
-		
+		if (BulkDataFlags & BULKDATA_Size64Bit)
+		{
+			Ar << ElementCount;
+		}
+		else
+		{
+			int32 ElementCountAsInt32 = ElementCount;
+			Ar << ElementCountAsInt32;
+			ElementCount = ElementCountAsInt32;
+		}
+
 		// We're loading from the persistent archive.
 		if( Ar.IsLoading() )
 		{
@@ -941,8 +957,17 @@ void FUntypedBulkData::Serialize( FArchive& Ar, UObject* Owner, int32 Idx, bool
 			}
 
 			// Size on disk, which in the case of compression is != GetBulkDataSize()
-			Ar << BulkDataSizeOnDisk;
-			
+			if (BulkDataFlags & BULKDATA_Size64Bit)
+			{
+				Ar << BulkDataSizeOnDisk;
+			}
+			else
+			{
+				int32 BulkDataSizeOnDiskAsInt32;
+				Ar << BulkDataSizeOnDiskAsInt32;
+				BulkDataSizeOnDisk = BulkDataSizeOnDiskAsInt32;
+			}
+
 			Ar << BulkDataOffsetInFile;
 
 			// determine whether the payload is stored inline or at the end of the file
@@ -1118,7 +1143,15 @@ void FUntypedBulkData::Serialize( FArchive& Ar, UObject* Owner, int32 Idx, bool
 				FArchive::FScopeSetDebugSerializationFlags S(Ar, DSF_IgnoreDiff);
 
 				// And serialize the placeholder which is going to be overwritten later.
-				Ar << BulkDataSizeOnDisk;
+				if (BulkDataFlags & BULKDATA_Size64Bit)
+				{
+					Ar << BulkDataSizeOnDisk;
+				}
+				else
+				{
+					int32 BulkDataSizeOnDiskAsInt32 = BulkDataSizeOnDisk;
+					Ar << BulkDataSizeOnDiskAsInt32;
+				}
 				// Keep track of position we are going to serialize placeholder BulkDataOffsetInFile.
 				SavedBulkDataOffsetInFilePos = Ar.Tell();
 				BulkDataOffsetInFile = INDEX_NONE;
@@ -1191,7 +1224,15 @@ void FUntypedBulkData::Serialize( FArchive& Ar, UObject* Owner, int32 Idx, bool
 
 				// Seek back and overwrite placeholder for BulkDataSizeOnDisk
 				Ar.Seek(SavedBulkDataSizeOnDiskPos);
-				Ar << BulkDataSizeOnDisk;
+				if (BulkDataFlags & BULKDATA_Size64Bit)
+				{
+					Ar << BulkDataSizeOnDisk;
+				}
+				else
+				{
+					int32 BulkDataSizeOnDiskAsInt32 = BulkDataSizeOnDisk;
+					Ar << BulkDataSizeOnDiskAsInt32;
+				}
 
 				// Seek back and overwrite placeholder for BulkDataOffsetInFile
 				Ar.Seek(SavedBulkDataOffsetInFilePos);
@@ -1322,7 +1363,7 @@ void FUntypedBulkData::InitializeMemberVariables()
 void FUntypedBulkData::SerializeElements(FArchive& Ar, void* Data)
 {
 	// Serialize each element individually.				
-	for (int32 ElementIndex = 0; ElementIndex < ElementCount; ElementIndex++)
+	for (int64 ElementIndex = 0; ElementIndex < ElementCount; ElementIndex++)
 	{
 		SerializeElement(Ar, Data, ElementIndex);
 	}
@@ -1343,7 +1384,7 @@ void FUntypedBulkData::SerializeBulkData( FArchive& Ar, void* Data )
 	}
 
 	// Skip serialization for bulk data of zero length
-	const int32 BulkDataSize = GetBulkDataSize();
+	const int64 BulkDataSize = GetBulkDataSize();
 	if(BulkDataSize == 0)
 	{
 		return;
@@ -1383,21 +1424,16 @@ void FUntypedBulkData::SerializeBulkData( FArchive& Ar, void* Data )
 		// Serialize data compressed.
 		if( BulkDataFlags & BULKDATA_SerializeCompressed )
 		{
-			// Placeholder for to be serialized data.
-			TArray<uint8> SerializedData;
-			
 			// Loading, data is compressed in archive and needs to be decompressed.
 			if( Ar.IsLoading() )
 			{
-				// Create space for uncompressed data.
-				SerializedData.Empty( GetBulkDataSize() );
-				SerializedData.AddUninitialized( GetBulkDataSize() );
+				TUniquePtr<uint8[]> SerializedData = MakeUnique<uint8[]>( GetBulkDataSize() );
 
 				// Serialize data with passed in archive and compress.
-				Ar.SerializeCompressed( SerializedData.GetData(), SerializedData.Num(), GetDecompressionFormat(), COMPRESS_NoFlags, false );
+				Ar.SerializeCompressed( SerializedData.Get(), GetBulkDataSize(), GetDecompressionFormat(), COMPRESS_NoFlags, false );
 				
 				// Initialize memory reader with uncompressed data array and propagate forced byte swapping
-				FMemoryReader MemoryReader( SerializedData, true );
+				FLargeMemoryReader MemoryReader( SerializedData.Get(), GetBulkDataSize(), ELargeMemoryReaderFlags::Persistent );
 				MemoryReader.SetByteSwapping( Ar.ForceByteSwapping() );
 
 				// Serialize each element individually via memory reader.
@@ -1407,14 +1443,14 @@ void FUntypedBulkData::SerializeBulkData( FArchive& Ar, void* Data )
 			else if( Ar.IsSaving() )
 			{			
 				// Initialize memory writer with blank data array and propagate forced byte swapping
-				FMemoryWriter MemoryWriter( SerializedData, true );
+				FLargeMemoryWriter MemoryWriter( GetBulkDataSize(), true );
 				MemoryWriter.SetByteSwapping( Ar.ForceByteSwapping() );
 
 				// Serialize each element individually via memory writer.			
 				SerializeElements(MemoryWriter, Data);
 
 				// Serialize data with passed in archive and compress.
-				Ar.SerializeCompressed( SerializedData.GetData(), SerializedData.Num(), GetDecompressionFormat(), COMPRESS_NoFlags, false );
+				Ar.SerializeCompressed( MemoryWriter.GetData(), GetBulkDataSize(), GetDecompressionFormat(), COMPRESS_NoFlags, false );
 			}
 		}
 		// Uncompressed/ regular serialization.
@@ -1445,7 +1481,7 @@ void FUntypedBulkData::MakeSureBulkDataIsLoaded()
 		}
 		else
 		{
-			const int32 BytesNeeded = GetBulkDataSize();
+			const int64 BytesNeeded = GetBulkDataSize();
 			// Allocate memory for bulk data.
 			BulkData.Reallocate(BytesNeeded, BulkDataAlignment);
 
@@ -1465,7 +1501,7 @@ void FUntypedBulkData::WaitForAsyncLoading()
 	DECLARE_SCOPE_CYCLE_COUNTER(TEXT("FUntypedBulkData::WaitForAsyncLoading"), STAT_UBD_WaitForAsyncLoading, STATGROUP_Memory);
 	while (!SerializeFuture.WaitFor(FTimespan::FromMilliseconds(1000.0)))
 	{
-		UE_LOG(LogSerialization, Warning, TEXT("Waiting for %s bulk data (%d) to be loaded longer than 1000ms"), *Filename, GetBulkDataSize());
+		UE_LOG(LogSerialization, Warning, TEXT("Waiting for %s bulk data (%lld) to be loaded longer than 1000ms"), *Filename, GetBulkDataSize());
 	}
 	check(BulkDataAsync);
 }
@@ -1628,7 +1664,7 @@ int32 FByteBulkData::GetElementSize() const
  * @param Data			Base pointer to data
  * @param ElementIndex	Element index to serialize
  */
-void FByteBulkData::SerializeElement( FArchive& Ar, void* Data, int32 ElementIndex )
+void FByteBulkData::SerializeElement( FArchive& Ar, void* Data, int64 ElementIndex )
 {
 	uint8& ByteData = *((uint8*)Data + ElementIndex);
 	Ar << ByteData;
@@ -1655,7 +1691,7 @@ int32 FWordBulkData::GetElementSize() const
  * @param Data			Base pointer to data
  * @param ElementIndex	Element index to serialize
  */
-void FWordBulkData::SerializeElement( FArchive& Ar, void* Data, int32 ElementIndex )
+void FWordBulkData::SerializeElement( FArchive& Ar, void* Data, int64 ElementIndex )
 {
 	uint16& WordData = *((uint16*)Data + ElementIndex);
 	Ar << WordData;
@@ -1682,7 +1718,7 @@ int32 FIntBulkData::GetElementSize() const
  * @param Data			Base pointer to data
  * @param ElementIndex	Element index to serialize
  */
-void FIntBulkData::SerializeElement( FArchive& Ar, void* Data, int32 ElementIndex )
+void FIntBulkData::SerializeElement( FArchive& Ar, void* Data, int64 ElementIndex )
 {
 	int32& IntData = *((int32*)Data + ElementIndex);
 	Ar << IntData;
@@ -1709,7 +1745,7 @@ int32 FFloatBulkData::GetElementSize() const
  * @param Data			Base pointer to data
  * @param ElementIndex	Element index to serialize
  */
-void FFloatBulkData::SerializeElement( FArchive& Ar, void* Data, int32 ElementIndex )
+void FFloatBulkData::SerializeElement( FArchive& Ar, void* Data, int64 ElementIndex )
 {
 	float& FloatData = *((float*)Data + ElementIndex);
 	Ar << FloatData;
diff --git a/Engine/Source/Runtime/CoreUObject/Private/UObject/SavePackage.cpp b/Engine/Source/Runtime/CoreUObject/Private/UObject/SavePackage.cpp
--- a/Engine/Source/Runtime/CoreUObject/Private/UObject/SavePackage.cpp
+++ b/Engine/Source/Runtime/CoreUObject/Private/UObject/SavePackage.cpp
@@ -5682,7 +5682,7 @@ FSavePackageResultStruct UPackage::Save(UPackage* InOuter, UObject* Base, EObjec
 						int64 BulkEndOffset = TargetArchive->Tell();
 						int64 LinkerEndOffset = Linker->Tell();
 
-						int32 SizeOnDisk = (int32)(BulkEndOffset - BulkStartOffset);
+						int64 SizeOnDisk = (int32)(BulkEndOffset - BulkStartOffset);
 				
 						Linker->Seek(BulkDataStorageInfo.BulkDataFlagsPos);
 						*Linker << ModifiedBulkDataFlags;
@@ -5691,7 +5691,16 @@ FSavePackageResultStruct UPackage::Save(UPackage* InOuter, UObject* Base, EObjec
 						*Linker << StoredBulkStartOffset;
 
 						Linker->Seek(BulkDataStorageInfo.BulkDataSizeOnDiskPos);
-						*Linker << SizeOnDisk;
+						if (ModifiedBulkDataFlags & BULKDATA_Size64Bit)
+						{
+							*Linker << SizeOnDisk;
+						}
+						else
+						{
+							check(SizeOnDisk < (1LL << 31));
+							int32 SizeOnDiskAsInt32 = SizeOnDisk;
+							*Linker << SizeOnDiskAsInt32;
+						}
 
 						Linker->Seek(LinkerEndOffset);
 
diff --git a/Engine/Source/Runtime/CoreUObject/Public/Serialization/BulkData.h b/Engine/Source/Runtime/CoreUObject/Public/Serialization/BulkData.h
--- a/Engine/Source/Runtime/CoreUObject/Public/Serialization/BulkData.h
+++ b/Engine/Source/Runtime/CoreUObject/Public/Serialization/BulkData.h
@@ -40,6 +40,8 @@ enum EBulkDataFlags
 	BULKDATA_OptionalPayload					= 1 << 11,
 	/** This payload will be memory mapped, this requires alignment, no compression etc. */
 	BULKDATA_MemoryMappedPayload = 1 << 12,
+	/** Bulk data size is 64 bits long */
+	BULKDATA_Size64Bit							= 1 << 13
 };
 
 /**
@@ -177,7 +179,7 @@ private:
 			return bAllocated;
 		}
 
-		void Reallocate(int32 Count, int32 Alignment = DEFAULT_ALIGNMENT)
+		void Reallocate(int64 Count, int32 Alignment = DEFAULT_ALIGNMENT)
 		{
 			check(!MappedHandle && !MappedRegion); // not legal for mapped bulk data
 			if (Count)
@@ -303,7 +305,7 @@ public:
 	 *
 	 * @return Number of elements in this bulk data array
 	 */
-	int32 GetElementCount() const;
+	int64 GetElementCount() const;
 	/**
 	 * Returns size in bytes of single element.
 	 *
@@ -317,14 +319,14 @@ public:
 	 *
 	 * @return Size of the bulk data in bytes
 	 */
-	int32 GetBulkDataSize() const;
+	int64 GetBulkDataSize() const;
 	/**
 	 * Returns the size of the bulk data on disk. This can differ from GetBulkDataSize if
 	 * BULKDATA_SerializeCompressed is set.
 	 *
 	 * @return Size of the bulk data on disk or INDEX_NONE in case there's no association
 	 */
-	int32 GetBulkDataSizeOnDisk() const;
+	int64 GetBulkDataSizeOnDisk() const;
 	/**
 	 * Returns the offset into the file the bulk data is located at.
 	 *
@@ -442,7 +444,7 @@ public:
 	 *
 	 * @param InElementCount	Number of elements array should be resized to
 	 */
-	void* Realloc( int32 InElementCount );
+	void* Realloc( int64 InElementCount );
 
 	/** 
 	 * Unlocks bulk data after which point the pointer returned by Lock no longer is valid.
@@ -539,7 +541,7 @@ protected:
 	 * @param Data			Base pointer to data
 	 * @param ElementIndex	Index of element to serialize
 	 */
-	virtual void SerializeElement( FArchive& Ar, void* Data, int32 ElementIndex ) = 0;
+	virtual void SerializeElement( FArchive& Ar, void* Data, int64 ElementIndex ) = 0;
 
 	/**
 	 * Returns whether single element serialization is required given an archive. This e.g.
@@ -613,21 +615,21 @@ private:
 
 	/** Serialized flags for bulk data																					*/
 	uint32					BulkDataFlags;
+	/** Alignment of bulk data																							*/
+	uint16					BulkDataAlignment;
+	/** Current lock status																								*/
+	uint16					LockStatus;
 	/** Number of elements in bulk data array																			*/
-	int32					ElementCount;
+	int64					ElementCount;
 	/** Offset of bulk data into file or INDEX_NONE if no association													*/
 	int64					BulkDataOffsetInFile;
 	/** Size of bulk data on disk or INDEX_NONE if no association														*/
-	int32					BulkDataSizeOnDisk;
-	/** Alignment of bulk data																							*/
-	int32					BulkDataAlignment;
+	int64					BulkDataSizeOnDisk;
 
 	/** Pointer to cached bulk data																						*/
 	FAllocatedPtr		BulkData;
 	/** Pointer to cached async bulk data																				*/
 	FAllocatedPtr		BulkDataAsync;
-	/** Current lock status																								*/
-	uint32				LockStatus;
 	/** Async helper for loading bulk data on a separate thread */
 	TFuture<bool> SerializeFuture;
 
@@ -666,7 +668,7 @@ struct COREUOBJECT_API FByteBulkData : public FUntypedBulkData
 	 * @param Data			Base pointer to data
 	 * @param ElementIndex	Element index to serialize
 	 */
-	virtual void SerializeElement( FArchive& Ar, void* Data, int32 ElementIndex );
+	virtual void SerializeElement( FArchive& Ar, void* Data, int64 ElementIndex );
 };
 
 /*-----------------------------------------------------------------------------
@@ -689,7 +691,7 @@ struct COREUOBJECT_API FWordBulkData : public FUntypedBulkData
 	 * @param Data			Base pointer to data
 	 * @param ElementIndex	Element index to serialize
 	 */
-	virtual void SerializeElement( FArchive& Ar, void* Data, int32 ElementIndex );
+	virtual void SerializeElement( FArchive& Ar, void* Data, int64 ElementIndex );
 };
 
 /*-----------------------------------------------------------------------------
@@ -712,7 +714,7 @@ struct COREUOBJECT_API FIntBulkData : public FUntypedBulkData
 	 * @param Data			Base pointer to data
 	 * @param ElementIndex	Element index to serialize
 	 */
-	virtual void SerializeElement( FArchive& Ar, void* Data, int32 ElementIndex );
+	virtual void SerializeElement( FArchive& Ar, void* Data, int64 ElementIndex );
 };
 
 /*-----------------------------------------------------------------------------
@@ -735,7 +737,7 @@ struct COREUOBJECT_API FFloatBulkData : public FUntypedBulkData
 	 * @param Data			Base pointer to data
 	 * @param ElementIndex	Element index to serialize
 	 */
-	virtual void SerializeElement( FArchive& Ar, void* Data, int32 ElementIndex );
+	virtual void SerializeElement( FArchive& Ar, void* Data, int64 ElementIndex );
 };
 
 class FFormatContainer
diff --git a/Engine/Source/Runtime/Engine/Private/LightMap.cpp b/Engine/Source/Runtime/Engine/Private/LightMap.cpp
--- a/Engine/Source/Runtime/Engine/Private/LightMap.cpp
+++ b/Engine/Source/Runtime/Engine/Private/LightMap.cpp
@@ -2509,7 +2509,7 @@ int32 TQuantizedLightSampleBulkData<QuantizedLightSampleType>::GetElementSize()
  * @param ElementIndex	Element index to serialize
  */
 template<class QuantizedLightSampleType>
-void TQuantizedLightSampleBulkData<QuantizedLightSampleType>::SerializeElement( FArchive& Ar, void* Data, int32 ElementIndex )
+void TQuantizedLightSampleBulkData<QuantizedLightSampleType>::SerializeElement( FArchive& Ar, void* Data, int64 ElementIndex )
 {
 	QuantizedLightSampleType* QuantizedLightSample = (QuantizedLightSampleType*)Data + ElementIndex;
 	// serialize as colors
diff --git a/Engine/Source/Runtime/Engine/Private/StaticMeshBuild.cpp b/Engine/Source/Runtime/Engine/Private/StaticMeshBuild.cpp
--- a/Engine/Source/Runtime/Engine/Private/StaticMeshBuild.cpp
+++ b/Engine/Source/Runtime/Engine/Private/StaticMeshBuild.cpp
@@ -511,7 +511,7 @@ struct FStaticMeshTriangleBulkData : public FUntypedBulkData
 		return sizeof(FStaticMeshTriangle);
 	}
 
-	virtual void SerializeElement( FArchive& Ar, void* Data, int32 ElementIndex )
+	virtual void SerializeElement( FArchive& Ar, void* Data, int64 ElementIndex )
 	{
 		FStaticMeshTriangle& StaticMeshTriangle = *((FStaticMeshTriangle*)Data + ElementIndex);
 		Ar << StaticMeshTriangle.Vertices[0];
diff --git a/Engine/Source/Runtime/Engine/Public/LightMap.h b/Engine/Source/Runtime/Engine/Public/LightMap.h
--- a/Engine/Source/Runtime/Engine/Public/LightMap.h
+++ b/Engine/Source/Runtime/Engine/Public/LightMap.h
@@ -399,7 +399,7 @@ struct TQuantizedLightSampleBulkData : public FUntypedBulkData
 	 * @param Data			Base pointer to data
 	 * @param ElementIndex	Element index to serialize
 	 */
-	virtual void SerializeElement( FArchive& Ar, void* Data, int32 ElementIndex );
+	virtual void SerializeElement( FArchive& Ar, void* Data, int64 ElementIndex );
 };
 
 /** A 1D array of incident lighting data. */
