Modified FAsyncPackage::GetLoadedAssets to add unique pacakges to the list of loaded assets to prevent duplicated asset loaded notifications

diff --git a/Engine/Source/Runtime/CoreUObject/Private/Serialization/AsyncLoading.cpp b/Engine/Source/Runtime/CoreUObject/Private/Serialization/AsyncLoading.cpp
--- a/Engine/Source/Runtime/CoreUObject/Private/Serialization/AsyncLoading.cpp
+++ b/Engine/Source/Runtime/CoreUObject/Private/Serialization/AsyncLoading.cpp
@@ -5642,7 +5642,7 @@ void FAsyncPackage::GetLoadedAssets(TArray<FWeakObjectPtr>& AssetList)
 	{
 		if (Obj && !Obj->IsPendingKill() && Obj->IsAsset())
 		{
-			AssetList.Add(Obj);
+			AssetList.AddUnique(Obj);
 		}
 	}
 }
