Fixed crashes in cooked projects when EDL is disabled

diff --git a/Engine/Source/Runtime/CoreUObject/Private/UObject/LinkerLoad.cpp b/Engine/Source/Runtime/CoreUObject/Private/UObject/LinkerLoad.cpp
--- a/Engine/Source/Runtime/CoreUObject/Private/UObject/LinkerLoad.cpp
+++ b/Engine/Source/Runtime/CoreUObject/Private/UObject/LinkerLoad.cpp
@@ -4735,9 +4735,8 @@ void FLinkerLoad::Detach()
 
 	// Remove from object manager, if it has been added.
 	FLinkerManager::Get().RemoveLoaderFromObjectLoadersAndLoadersWithNewImports(this);
-	if (!FPlatformProperties::HasEditorOnlyData())
+	if (!FPlatformProperties::HasEditorOnlyData() && CurrentLoadContext)
 	{
-		check(CurrentLoadContext);
 		CurrentLoadContext->RemoveDelayedLinkerClosePackage(this);
 	}
 
diff --git a/Engine/Source/Runtime/CoreUObject/Private/UObject/UObjectGlobals.cpp b/Engine/Source/Runtime/CoreUObject/Private/UObject/UObjectGlobals.cpp
--- a/Engine/Source/Runtime/CoreUObject/Private/UObject/UObjectGlobals.cpp
+++ b/Engine/Source/Runtime/CoreUObject/Private/UObject/UObjectGlobals.cpp
@@ -1039,6 +1039,10 @@ public:
 
 #endif
 
+// Temporary load counter for the game thread, used mostly for checking if we're still loading
+// @todo: remove this in the new loader
+static int32 GGameThreadLoadCounter = 0;
+
 UPackage* LoadPackageInternal(UPackage* InOuter, const TCHAR* InLongPackageNameOrFilename, uint32 LoadFlags, FLinkerLoad* ImportLinker, FArchive* InReaderOverride, FUObjectSerializeContext* InLoadContext)
 {
 	DECLARE_SCOPE_CYCLE_COUNTER(TEXT("LoadPackageInternal"), STAT_LoadPackageInternal, STATGROUP_ObjectVerbose);
@@ -1371,10 +1375,10 @@ UPackage* LoadPackageInternal(UPackage* InOuter, const TCHAR* InLongPackageNameO
 		if (FPlatformProperties::RequiresCookedData())
 		{
 			if (!IsInAsyncLoadingThread())
-			{
-				check(Linker->GetSerializeContext());
-				if (!Linker->GetSerializeContext()->HasStartedLoading())
+			{				
+				if (GGameThreadLoadCounter == 0)
 				{
+					check(!Linker->GetSerializeContext());
 					// Sanity check to make sure that Linker is the linker that loaded our Result package or the linker has already been detached
 					check(!Result || Result->LinkerLoad == Linker || Result->LinkerLoad == nullptr);
 					if (Result && Linker->HasLoader())
@@ -1431,8 +1435,6 @@ UPackage* LoadPackage(UPackage* InOuter, const TCHAR* InLongPackageName, uint32
 	return LoadPackageInternal(InOuter, InLongPackageName, LoadFlags, /*ImportLinker =*/ nullptr, InReaderOverride, InLoadContext);
 }
 
-static int32 GGameThreadLoadCounter = 0;
-
 /**
  * Returns whether we are currently loading a package (sync or async)
  *
