UE-70402 Mobile projects crash with Event Driven Loader disabled

#iOS
#Android
#4.22

- when running the editor, changing the EDL flag is not reflected in the cooker (it will continue with the value read when the editor started)

diff --git a/Engine/Source/Runtime/CoreUObject/Private/Serialization/AsyncLoading.cpp b/Engine/Source/Runtime/CoreUObject/Private/Serialization/AsyncLoading.cpp
--- a/Engine/Source/Runtime/CoreUObject/Private/Serialization/AsyncLoading.cpp
+++ b/Engine/Source/Runtime/CoreUObject/Private/Serialization/AsyncLoading.cpp
@@ -7415,6 +7415,11 @@ bool IsEventDrivenLoaderEnabledInCookedBuilds()
 		bool bEventDrivenLoaderEnabled;
 		FEventDrivenLoaderEnabledInCookedBuildsInit()
 			: bEventDrivenLoaderEnabled(false)
+		{
+			SetEventDrivenLoaderEnabled();
+		}
+
+		void SetEventDrivenLoaderEnabled()
 		{
 			check(GConfig || GIsRequestingExit);
 			if (GConfig)
@@ -7429,6 +7434,14 @@ bool IsEventDrivenLoaderEnabledInCookedBuilds()
 			}
 		}
 	} EventDrivenLoaderEnabledInCookedBuilds;
+
+#if WITH_EDITOR	
+	// when building from the UE4 Editor, s.EventDrivenLoaderEnabled can be changed from Project Settings, so we need to test it at every call
+	if (GIsEditor)
+	{
+		EventDrivenLoaderEnabledInCookedBuilds.SetEventDrivenLoaderEnabled();
+	}
+#endif
 	return EventDrivenLoaderEnabledInCookedBuilds.bEventDrivenLoaderEnabled;
 }
 
