Core: Fix pop-up dialogs when loading corrupt assets in unattended sessions, causing jobs to hang on build machines.

diff --git a/Engine/Source/Runtime/CoreUObject/Private/UObject/LinkerLoad.cpp b/Engine/Source/Runtime/CoreUObject/Private/UObject/LinkerLoad.cpp
--- a/Engine/Source/Runtime/CoreUObject/Private/UObject/LinkerLoad.cpp
+++ b/Engine/Source/Runtime/CoreUObject/Private/UObject/LinkerLoad.cpp
@@ -4595,7 +4595,7 @@ UObject* FLinkerLoad::IndexToObject( FPackageIndex Index )
 	{
 		#if PLATFORM_DESKTOP
 			// Show a message box indicating, possible, corrupt data (desktop platforms only)
-			if ( !ExportMap.IsValidIndex( Index.ToExport() ) )
+			if ( !ExportMap.IsValidIndex( Index.ToExport() ) && !FApp::IsUnattended() )
 			{
 				FText ErrorMessage, ErrorCaption;
 				GConfig->GetText(TEXT("/Script/Engine.Engine"),
@@ -4630,7 +4630,7 @@ UObject* FLinkerLoad::IndexToObject( FPackageIndex Index )
 	{
 		#if PLATFORM_DESKTOP
 			// Show a message box indicating, possible, corrupt data (desktop platforms only)
-			if ( !ImportMap.IsValidIndex( Index.ToImport() ) )
+			if ( !ImportMap.IsValidIndex( Index.ToImport() ) && !FApp::IsUnattended() )
 			{
 				FText ErrorMessage, ErrorCaption;
 				GConfig->GetText(TEXT("/Script/Engine.Engine"),
