Fixing a crash when the engine fails to open a package for comparison in the deterministic cooking commandlet

diff --git a/Engine/Source/Runtime/CoreUObject/Private/Serialization/ArchiveStackTrace.cpp b/Engine/Source/Runtime/CoreUObject/Private/Serialization/ArchiveStackTrace.cpp
--- a/Engine/Source/Runtime/CoreUObject/Private/Serialization/ArchiveStackTrace.cpp
+++ b/Engine/Source/Runtime/CoreUObject/Private/Serialization/ArchiveStackTrace.cpp
@@ -1305,14 +1305,14 @@ void FArchiveStackTrace::DumpPackageHeaderDiffs(const FPackageData& SourcePackag
 		TRefCountPtr<FUObjectSerializeContext> LinkerLoadContext(new FUObjectSerializeContext());
 		BeginLoad(LinkerLoadContext);
 		SourceLinker = CreateLinkerForPackage(LinkerLoadContext, SourceAssetPackageName, AssetFilename, SourcePackage);
-		EndLoad(SourceLinker->GetSerializeContext());
+		EndLoad(SourceLinker ? SourceLinker->GetSerializeContext() : LinkerLoadContext);
 	}
 	
 	{
 		TRefCountPtr<FUObjectSerializeContext> LinkerLoadContext(new FUObjectSerializeContext());
 		BeginLoad(LinkerLoadContext);
 		DestLinker = CreateLinkerForPackage(LinkerLoadContext, DestAssetPackageName, AssetFilename, DestPackage);
-		EndLoad(DestLinker->GetSerializeContext());
+		EndLoad(DestLinker ? DestLinker->GetSerializeContext() : LinkerLoadContext);
 	}
 
 	if (SourceLinker && DestLinker)
