Allow the deterministic cook commandlet to open cooked content in the editor regardless of the 'cook.AllowCookedDataInEditorBuilds' ini setting

diff --git a/Engine/Source/Runtime/CoreUObject/Private/Serialization/ArchiveStackTrace.cpp b/Engine/Source/Runtime/CoreUObject/Private/Serialization/ArchiveStackTrace.cpp
--- a/Engine/Source/Runtime/CoreUObject/Private/Serialization/ArchiveStackTrace.cpp
+++ b/Engine/Source/Runtime/CoreUObject/Private/Serialization/ArchiveStackTrace.cpp
@@ -1287,6 +1287,7 @@ static void DumpTableDifferences(
 #endif // !NO_LOGGING
 }
 
+extern int32 GAllowCookedDataInEditorBuilds;
 
 void FArchiveStackTrace::DumpPackageHeaderDiffs(const FPackageData& SourcePackage, const FPackageData& DestPackage, const FString& AssetFilename, const int32 MaxDiffsToLog)
 {
@@ -1297,6 +1298,7 @@ void FArchiveStackTrace::DumpPackageHeaderDiffs(const FPackageData& SourcePackag
 
 	TGuardValue<bool> GuardIsSavingPackage(GIsSavingPackage, false);
 	TGuardValue<bool> GuardAllowUnversionedContentInEditor(GAllowUnversionedContentInEditor, true);
+	TGuardValue<int32> GuardAllowCookedDataInEditorBuilds(GAllowCookedDataInEditorBuilds, 1);
 
 	FLinkerLoad* SourceLinker = nullptr;
 	FLinkerLoad* DestLinker = nullptr;
diff --git a/Engine/Source/Runtime/CoreUObject/Private/UObject/LinkerLoad.cpp b/Engine/Source/Runtime/CoreUObject/Private/UObject/LinkerLoad.cpp
--- a/Engine/Source/Runtime/CoreUObject/Private/UObject/LinkerLoad.cpp
+++ b/Engine/Source/Runtime/CoreUObject/Private/UObject/LinkerLoad.cpp
@@ -82,7 +82,7 @@ bool FLinkerLoad::ShouldCreateThrottledSlowTask() const
 #endif
 
 
-static int32 GAllowCookedDataInEditorBuilds = 0;
+int32 GAllowCookedDataInEditorBuilds = 0;
 static FAutoConsoleVariableRef CVarAllowCookedDataInEditorBuilds(
 	TEXT("cook.AllowCookedDataInEditorBuilds"),
 		GAllowCookedDataInEditorBuilds,
