Fix clang builds.

diff --git a/Engine/Source/Runtime/CoreUObject/Private/Serialization/ArchiveStackTrace.cpp b/Engine/Source/Runtime/CoreUObject/Private/Serialization/ArchiveStackTrace.cpp
--- a/Engine/Source/Runtime/CoreUObject/Private/Serialization/ArchiveStackTrace.cpp
+++ b/Engine/Source/Runtime/CoreUObject/Private/Serialization/ArchiveStackTrace.cpp
@@ -1307,14 +1307,14 @@ void FArchiveStackTrace::DumpPackageHeaderDiffs(const FPackageData& SourcePackag
 		TRefCountPtr<FUObjectSerializeContext> LinkerLoadContext(new FUObjectSerializeContext());
 		BeginLoad(LinkerLoadContext);
 		SourceLinker = CreateLinkerForPackage(LinkerLoadContext, SourceAssetPackageName, AssetFilename, SourcePackage);
-		EndLoad(SourceLinker ? SourceLinker->GetSerializeContext() : LinkerLoadContext);
+		EndLoad(SourceLinker ? SourceLinker->GetSerializeContext() : LinkerLoadContext.GetReference());
 	}
 	
 	{
 		TRefCountPtr<FUObjectSerializeContext> LinkerLoadContext(new FUObjectSerializeContext());
 		BeginLoad(LinkerLoadContext);
 		DestLinker = CreateLinkerForPackage(LinkerLoadContext, DestAssetPackageName, AssetFilename, DestPackage);
-		EndLoad(DestLinker ? DestLinker->GetSerializeContext() : LinkerLoadContext);
+		EndLoad(DestLinker ? DestLinker->GetSerializeContext() : LinkerLoadContext.GetReference());
 	}
 
 	if (SourceLinker && DestLinker)
@@ -1387,4 +1387,4 @@ FArchiveStackTraceReader* FArchiveStackTraceReader::CreateFromFile(const TCHAR*
 		Reader = new FArchiveStackTraceReader(InFilename, PackageData.Data, PackageData.Size);
 	}
 	return Reader;
-}
\ No newline at end of file
+}
