Fixed containers failing to gather required instance overrides

If a container element is overridden on an instance, then the entire container is serialized into the instance with different keys than the archetype. This means that all elements of a container need to be gathered if any elements within the container have changed.

#rnx

diff --git a/Engine/Source/Runtime/CoreUObject/Private/Serialization/PropertyLocalizationDataGathering.cpp b/Engine/Source/Runtime/CoreUObject/Private/Serialization/PropertyLocalizationDataGathering.cpp
--- a/Engine/Source/Runtime/CoreUObject/Private/Serialization/PropertyLocalizationDataGathering.cpp
+++ b/Engine/Source/Runtime/CoreUObject/Private/Serialization/PropertyLocalizationDataGathering.cpp
@@ -216,7 +216,7 @@ void FPropertyLocalizationDataGatherer::GatherLocalizationDataFromChildTextPrope
 	const UStructProperty* const StructProperty = Cast<const UStructProperty>(Property);
 	const UObjectPropertyBase* const ObjectProperty = Cast<const UObjectPropertyBase>(Property);
 
-	const EPropertyLocalizationGathererTextFlags ChildPropertyGatherTextFlags = GatherTextFlags | (Property->HasAnyPropertyFlags(CPF_EditorOnly) ? EPropertyLocalizationGathererTextFlags::ForceEditorOnly : EPropertyLocalizationGathererTextFlags::None);
+	const EPropertyLocalizationGathererTextFlags FixedChildPropertyGatherTextFlags = GatherTextFlags | (Property->HasAnyPropertyFlags(CPF_EditorOnly) ? EPropertyLocalizationGathererTextFlags::ForceEditorOnly : EPropertyLocalizationGathererTextFlags::None);
 
 	auto CanGatherFromInnerProperty = [](UProperty* InnerProperty)
 	{
@@ -233,12 +233,22 @@ void FPropertyLocalizationDataGatherer::GatherLocalizationDataFromChildTextPrope
 		const void* const ElementValueAddress = reinterpret_cast<const uint8*>(ValueAddress) + Property->ElementSize * i;
 		const void* const DefaultElementValueAddress = DefaultValueAddress ? (reinterpret_cast<const uint8*>(DefaultValueAddress) + Property->ElementSize * i) : nullptr;
 
+		EPropertyLocalizationGathererTextFlags ElementChildPropertyGatherTextFlags = FixedChildPropertyGatherTextFlags;
+		if (!EnumHasAnyFlags(ElementChildPropertyGatherTextFlags, EPropertyLocalizationGathererTextFlags::ForceIsDefaultValue))
+		{
+			const bool bIsDefaultValue = DefaultElementValueAddress && Property->Identical(ElementValueAddress, DefaultElementValueAddress, PPF_None);
+			if (bIsDefaultValue)
+			{
+				ElementChildPropertyGatherTextFlags |= EPropertyLocalizationGathererTextFlags::ForceIsDefaultValue;
+			}
+		}
+
 		// Property is a text property.
 		if (TextProperty)
 		{
 			const FText* const TextElementValueAddress = static_cast<const FText*>(ElementValueAddress);
 
-			const bool bIsDefaultValue = DefaultElementValueAddress && Property->Identical(ElementValueAddress, DefaultElementValueAddress, PPF_None);
+			const bool bIsDefaultValue = EnumHasAnyFlags(ElementChildPropertyGatherTextFlags, EPropertyLocalizationGathererTextFlags::ForceIsDefaultValue);
 			if (bIsDefaultValue)
 			{
 				MarkDefaultTextInstance(*TextElementValueAddress);
@@ -269,17 +279,7 @@ void FPropertyLocalizationDataGatherer::GatherLocalizationDataFromChildTextPrope
 			for(int32 j = 0; j < ElementCount; ++j)
 			{
 				const uint8* ElementPtr = ScriptArrayHelper.GetRawPtr(j);
-				const uint8* DefaultElementPtr = nullptr;
-				if (DefaultElementValueAddress)
-				{
-					FScriptArrayHelper DefaultScriptArrayHelper(ArrayProperty, DefaultElementValueAddress);
-					if (DefaultScriptArrayHelper.IsValidIndex(j))
-					{
-						DefaultElementPtr = DefaultScriptArrayHelper.GetRawPtr(j);
-					}
-				}
-
-				GatherLocalizationDataFromChildTextProperties(PathToElement + FString::Printf(TEXT("(%d)"), j), ArrayProperty->Inner, ElementPtr, DefaultElementPtr, ChildPropertyGatherTextFlags);
+				GatherLocalizationDataFromChildTextProperties(PathToElement + FString::Printf(TEXT("(%d)"), j), ArrayProperty->Inner, ElementPtr, nullptr, ElementChildPropertyGatherTextFlags);
 			}
 		}
 		// Property is a map property.
@@ -304,26 +304,19 @@ void FPropertyLocalizationDataGatherer::GatherLocalizationDataFromChildTextPrope
 				}
 
 				const uint8* MapPairPtr = ScriptMapHelper.GetPairPtr(j);
-				const uint8* DefaultPairPtr = nullptr;
-				if (DefaultElementValueAddress)
-				{
-					const uint8* MapKeyPtr = MapPairPtr;
-					FScriptMapHelper DefaultScriptMapHelper(MapProperty, DefaultElementValueAddress);
-					DefaultPairPtr = DefaultScriptMapHelper.FindMapPairPtrFromHash(MapKeyPtr);
-				}
 
 				if (bGatherMapKey)
 				{
 					const uint8* MapKeyPtr = MapPairPtr;
-					const uint8* DefaultMapKeyPtr = DefaultPairPtr ? DefaultPairPtr : nullptr;
-					GatherLocalizationDataFromChildTextProperties(PathToElement + FString::Printf(TEXT("(%d - Key)"), ElementIndex), MapProperty->KeyProp, MapKeyPtr, DefaultMapKeyPtr, ChildPropertyGatherTextFlags);
+					GatherLocalizationDataFromChildTextProperties(PathToElement + FString::Printf(TEXT("(%d - Key)"), ElementIndex), MapProperty->KeyProp, MapKeyPtr, nullptr, ElementChildPropertyGatherTextFlags);
 				}
+
 				if (bGatherMapValue)
 				{
 					const uint8* MapValuePtr = MapPairPtr + MapProperty->MapLayout.ValueOffset;
-					const uint8* DefaultMapValuePtr = DefaultPairPtr ? DefaultPairPtr + MapProperty->MapLayout.ValueOffset : nullptr;
-					GatherLocalizationDataFromChildTextProperties(PathToElement + FString::Printf(TEXT("(%d - Value)"), ElementIndex), MapProperty->ValueProp, MapValuePtr, DefaultMapValuePtr, ChildPropertyGatherTextFlags);
+					GatherLocalizationDataFromChildTextProperties(PathToElement + FString::Printf(TEXT("(%d - Value)"), ElementIndex), MapProperty->ValueProp, MapValuePtr, nullptr, ElementChildPropertyGatherTextFlags);
 				}
+
 				++ElementIndex;
 			}
 		}
@@ -347,14 +340,14 @@ void FPropertyLocalizationDataGatherer::GatherLocalizationDataFromChildTextPrope
 				}
 
 				const uint8* ElementPtr = ScriptSetHelper.GetElementPtr(j);
-				GatherLocalizationDataFromChildTextProperties(PathToElement + FString::Printf(TEXT("(%d)"), ElementIndex), SetProperty->ElementProp, ElementPtr, nullptr, ChildPropertyGatherTextFlags);
+				GatherLocalizationDataFromChildTextProperties(PathToElement + FString::Printf(TEXT("(%d)"), ElementIndex), SetProperty->ElementProp, ElementPtr, nullptr, ElementChildPropertyGatherTextFlags);
 				++ElementIndex;
 			}
 		}
 		// Property is a struct property.
 		else if (StructProperty)
 		{
-			GatherLocalizationDataFromStructFields(PathToElement, StructProperty->Struct, ElementValueAddress, DefaultElementValueAddress, ChildPropertyGatherTextFlags);
+			GatherLocalizationDataFromStructFields(PathToElement, StructProperty->Struct, ElementValueAddress, DefaultElementValueAddress, ElementChildPropertyGatherTextFlags);
 		}
 		// Property is an object property.
 		else if (ObjectProperty && !(GatherTextFlags & EPropertyLocalizationGathererTextFlags::SkipSubObjects))
@@ -362,7 +355,7 @@ void FPropertyLocalizationDataGatherer::GatherLocalizationDataFromChildTextPrope
 			const UObject* InnerObject = ObjectProperty->GetObjectPropertyValue(ElementValueAddress);
 			if (InnerObject && IsObjectValidForGather(InnerObject))
 			{
-				GatherLocalizationDataFromObjectWithCallbacks(InnerObject, ChildPropertyGatherTextFlags);
+				GatherLocalizationDataFromObjectWithCallbacks(InnerObject, FixedChildPropertyGatherTextFlags);
 			}
 		}
 	}
diff --git a/Engine/Source/Runtime/CoreUObject/Public/Serialization/PropertyLocalizationDataGathering.h b/Engine/Source/Runtime/CoreUObject/Public/Serialization/PropertyLocalizationDataGathering.h
--- a/Engine/Source/Runtime/CoreUObject/Public/Serialization/PropertyLocalizationDataGathering.h
+++ b/Engine/Source/Runtime/CoreUObject/Public/Serialization/PropertyLocalizationDataGathering.h
@@ -35,10 +35,15 @@ enum class EPropertyLocalizationGathererTextFlags : uint8
 	 */
 	ForceEditorOnly = ForceEditorOnlyProperties | ForceEditorOnlyScriptData,
 
+	/**
+	 * Force all gathered text to be considered "default" (matching its archetype value).
+	 */
+	ForceIsDefaultValue = 1<<3,
+
 	/**
 	 * Don't process any sub-objects (either inner objects or object pointers).
 	 */
-	SkipSubObjects = 1<<3,
+	SkipSubObjects = 1<<4,
 };
 ENUM_CLASS_FLAGS(EPropertyLocalizationGathererTextFlags);
 
