Fix sparse multicast delegates operating on the wrong object when binding, removing, or clearing on an object from another object's blueprint.
#rnx

diff --git a/Engine/Source/Runtime/CoreUObject/Private/UObject/PropertyMulticastDelegate.cpp b/Engine/Source/Runtime/CoreUObject/Private/UObject/PropertyMulticastDelegate.cpp
--- a/Engine/Source/Runtime/CoreUObject/Private/UObject/PropertyMulticastDelegate.cpp
+++ b/Engine/Source/Runtime/CoreUObject/Private/UObject/PropertyMulticastDelegate.cpp
@@ -374,8 +374,20 @@ const TCHAR* UMulticastInlineDelegateProperty::ImportText_Internal(const TCHAR*
 	return ImportDelegateFromText(MulticastDelegate, Buffer, PropertyValue, PortFlags, Parent, ErrorText);
 }
 
+void ResolveDelegateReference(const UMulticastInlineDelegateProperty* InlineProperty, UObject*& Parent, void*& PropertyValue)
+{
+	if (PropertyValue == nullptr)
+	{
+		checkf(Parent, TEXT("Must specify at least one of Parent or PropertyValue"));
+		PropertyValue = InlineProperty->GetPropertyValuePtr_InContainer(Parent);
+	}
+	// Owner doesn't matter for inline delegates, so we don't worry about the Owner == nullptr case
+}
+
 void UMulticastInlineDelegateProperty::AddDelegate(FScriptDelegate ScriptDelegate, UObject* Parent, void* PropertyValue) const
 {
+	ResolveDelegateReference(this, Parent, PropertyValue);
+
 	FMulticastScriptDelegate& MulticastDelegate = (*(FMulticastScriptDelegate*)PropertyValue);
 
 	// Remove this delegate to our multicast delegate's invocation list
@@ -384,6 +396,8 @@ void UMulticastInlineDelegateProperty::AddDelegate(FScriptDelegate ScriptDelegat
 
 void UMulticastInlineDelegateProperty::RemoveDelegate(const FScriptDelegate& ScriptDelegate, UObject* Parent, void* PropertyValue) const
 {
+	ResolveDelegateReference(this, Parent, PropertyValue);
+
 	FMulticastScriptDelegate& MulticastDelegate = (*(FMulticastScriptDelegate*)PropertyValue);
 
 	// Remove this delegate to our multicast delegate's invocation list
@@ -392,6 +406,8 @@ void UMulticastInlineDelegateProperty::RemoveDelegate(const FScriptDelegate& Scr
 
 void UMulticastInlineDelegateProperty::ClearDelegate(UObject* Parent, void* PropertyValue) const
 {
+	ResolveDelegateReference(this, Parent, PropertyValue);
+
 	FMulticastScriptDelegate& MulticastDelegate = (*(FMulticastScriptDelegate*)PropertyValue);
 	MulticastDelegate.Clear();
 }
@@ -519,17 +535,36 @@ const TCHAR* UMulticastSparseDelegateProperty::ImportText_Internal(const TCHAR*
 	return Result;
 }
 
+void ResolveDelegateReference(const UMulticastSparseDelegateProperty* SparseProperty, UObject*& Parent, void*& PropertyValue)
+{
+	USparseDelegateFunction* SparseDelegateFunc = CastChecked<USparseDelegateFunction>(SparseProperty->SignatureFunction);
+
+	if (Parent == nullptr)
+	{
+		checkf(PropertyValue, TEXT("Must specify at least one of Parent or PropertyValue"));
+		Parent = FSparseDelegateStorage::ResolveSparseOwner(*(FSparseDelegate*)PropertyValue, SparseDelegateFunc->OwningClassName, SparseDelegateFunc->DelegateName);
+	}
+	else if (PropertyValue)
+	{
+		checkSlow(Parent == FSparseDelegateStorage::ResolveSparseOwner(*(FSparseDelegate*)PropertyValue, SparseDelegateFunc->OwningClassName, SparseDelegateFunc->DelegateName));
+	}
+	else
+	{
+		PropertyValue = SparseProperty->GetPropertyValuePtr_InContainer(Parent);
+	}
+}
 
 void UMulticastSparseDelegateProperty::AddDelegate(FScriptDelegate ScriptDelegate, UObject* Parent, void* PropertyValue) const
 {
+	ResolveDelegateReference(this, Parent, PropertyValue);
 	USparseDelegateFunction* SparseDelegateFunc = CastChecked<USparseDelegateFunction>(SignatureFunction);
 	FSparseDelegate& SparseDelegate = *(FSparseDelegate*)PropertyValue;
 	SparseDelegate.__Internal_AddUnique(Parent, SparseDelegateFunc->DelegateName, MoveTemp(ScriptDelegate));
 }
 
-
 void UMulticastSparseDelegateProperty::RemoveDelegate(const FScriptDelegate& ScriptDelegate, UObject* Parent, void* PropertyValue) const
 {
+	ResolveDelegateReference(this, Parent, PropertyValue);
 	USparseDelegateFunction* SparseDelegateFunc = CastChecked<USparseDelegateFunction>(SignatureFunction);
 	FSparseDelegate& SparseDelegate = *(FSparseDelegate*)PropertyValue;
 	SparseDelegate.__Internal_Remove(Parent, SparseDelegateFunc->DelegateName, ScriptDelegate);
@@ -537,6 +572,7 @@ void UMulticastSparseDelegateProperty::RemoveDelegate(const FScriptDelegate& Scr
 
 void UMulticastSparseDelegateProperty::ClearDelegate(UObject* Parent, void* PropertyValue) const
 {
+	ResolveDelegateReference(this, Parent, PropertyValue);
 	USparseDelegateFunction* SparseDelegateFunc = CastChecked<USparseDelegateFunction>(SignatureFunction);
 	FSparseDelegate& SparseDelegate = *(FSparseDelegate*)PropertyValue;
 	SparseDelegate.__Internal_Clear(Parent, SparseDelegateFunc->DelegateName);
diff --git a/Engine/Source/Runtime/CoreUObject/Private/UObject/ScriptCore.cpp b/Engine/Source/Runtime/CoreUObject/Private/UObject/ScriptCore.cpp
--- a/Engine/Source/Runtime/CoreUObject/Private/UObject/ScriptCore.cpp
+++ b/Engine/Source/Runtime/CoreUObject/Private/UObject/ScriptCore.cpp
@@ -2516,7 +2516,7 @@ DEFINE_FUNCTION(UObject::execAddMulticastDelegate)
 
 	if (DelegateProp && DelegateAddr)
 	{
-		DelegateProp->AddDelegate(MoveTemp(Delegate), Stack.Object, DelegateAddr);
+		DelegateProp->AddDelegate(MoveTemp(Delegate), nullptr, DelegateAddr);
 	}
 }
 IMPLEMENT_VM_FUNCTION( EX_AddMulticastDelegate, execAddMulticastDelegate );
@@ -2536,7 +2536,7 @@ DEFINE_FUNCTION(UObject::execRemoveMulticastDelegate)
 
 	if (DelegateProp && DelegateAddr)
 	{
-		DelegateProp->RemoveDelegate(Delegate, Stack.Object, DelegateAddr);
+		DelegateProp->RemoveDelegate(Delegate, nullptr, DelegateAddr);
 	}
 }
 IMPLEMENT_VM_FUNCTION( EX_RemoveMulticastDelegate, execRemoveMulticastDelegate );
@@ -2553,7 +2553,7 @@ DEFINE_FUNCTION(UObject::execClearMulticastDelegate)
 
 	if (DelegateProp && DelegateAddr)
 	{
-		DelegateProp->ClearDelegate(Stack.Object, DelegateAddr);
+		DelegateProp->ClearDelegate(nullptr, DelegateAddr);
 	}
 }
 IMPLEMENT_VM_FUNCTION( EX_ClearMulticastDelegate, execClearMulticastDelegate );
diff --git a/Engine/Source/Runtime/CoreUObject/Public/UObject/UnrealType.h b/Engine/Source/Runtime/CoreUObject/Public/UObject/UnrealType.h
--- a/Engine/Source/Runtime/CoreUObject/Public/UObject/UnrealType.h
+++ b/Engine/Source/Runtime/CoreUObject/Public/UObject/UnrealType.h
@@ -4224,12 +4224,9 @@ public:
 	virtual const FMulticastScriptDelegate* GetMulticastDelegate(const void* PropertyValue) const PURE_VIRTUAL(UMulticastDelegateProperty::GetMulticastDelegate, return nullptr;);
 	virtual void SetMulticastDelegate(void* PropertyValue, FMulticastScriptDelegate ScriptDelegate) const PURE_VIRTUAL(UMulticastDelegateProperty::SetMulticastDelegate, );
 
-	virtual void AddDelegate(FScriptDelegate ScriptDelegate, UObject* Parent) const PURE_VIRTUAL(UMulticastDelegateProperty::AddDelegate, );
-	virtual void AddDelegate(FScriptDelegate ScriptDelegate, UObject* Parent, void* PropertyValue) const PURE_VIRTUAL(UMulticastDelegateProperty::AddDelegate, );
-	virtual void RemoveDelegate(const FScriptDelegate& ScriptDelegate, UObject* Parent) const PURE_VIRTUAL(UMulticastDelegateProperty::RemoveDelegate, );
-	virtual void RemoveDelegate(const FScriptDelegate& ScriptDelegate, UObject* Parent, void* PropertyValue) const PURE_VIRTUAL(UMulticastDelegateProperty::RemoveDelegate, );
-	virtual void ClearDelegate(UObject* Parent) const PURE_VIRTUAL(UMulticastDelegateProperty::ClearDelegate, );
-	virtual void ClearDelegate(UObject* Parent, void* PropertyValue)  const PURE_VIRTUAL(UMulticastDelegateProperty::ClearDelegate, );
+	virtual void AddDelegate(FScriptDelegate ScriptDelegate, UObject* Parent = nullptr, void* PropertyValue = nullptr) const PURE_VIRTUAL(UMulticastDelegateProperty::AddDelegate, );
+	virtual void RemoveDelegate(const FScriptDelegate& ScriptDelegate, UObject* Parent = nullptr, void* PropertyValue = nullptr) const PURE_VIRTUAL(UMulticastDelegateProperty::RemoveDelegate, );
+	virtual void ClearDelegate(UObject* Parent = nullptr, void* PropertyValue = nullptr)  const PURE_VIRTUAL(UMulticastDelegateProperty::ClearDelegate, );
 
 protected:
 	friend class UProperty;
@@ -4305,24 +4302,10 @@ public:
 	virtual const FMulticastScriptDelegate* GetMulticastDelegate(const void* PropertyValue) const override;
 	virtual void SetMulticastDelegate(void* PropertyValue, FMulticastScriptDelegate ScriptDelegate) const override;
 
-	virtual void AddDelegate(FScriptDelegate ScriptDelegate, UObject* Parent) const override
-	{
-		AddDelegate(MoveTemp(ScriptDelegate), Parent, GetPropertyValuePtr_InContainer(Parent));
-	}
-
-	virtual void RemoveDelegate(const FScriptDelegate& ScriptDelegate, UObject* Parent) const override
-	{
-		RemoveDelegate(ScriptDelegate, Parent, GetPropertyValuePtr_InContainer(Parent));
-	}
-
-	virtual void ClearDelegate(UObject* Parent) const override
-	{
-		ClearDelegate(Parent, GetPropertyValuePtr_InContainer(Parent));
-	}
+	virtual void AddDelegate(FScriptDelegate ScriptDelegate, UObject* Parent = nullptr, void* PropertyValue = nullptr) const override;
+	virtual void RemoveDelegate(const FScriptDelegate& ScriptDelegate, UObject* Parent = nullptr, void* PropertyValue = nullptr) const override;
+	virtual void ClearDelegate(UObject* Parent = nullptr, void* PropertyValue = nullptr) const override;
 
-	virtual void AddDelegate(FScriptDelegate ScriptDelegate, UObject* Parent, void* PropertyValue) const override final;
-	virtual void RemoveDelegate(const FScriptDelegate& ScriptDelegate, UObject* Parent, void* PropertyValue) const override final;
-	virtual void ClearDelegate(UObject* Parent, void* PropertyValue) const override final;
 protected:
 	virtual FMulticastScriptDelegate::FInvocationList& GetInvocationList(const void* PropertyValue) const;
 	// End of UMulticastDelegateProperty interface
@@ -4353,24 +4336,10 @@ public:
 	virtual const FMulticastScriptDelegate* GetMulticastDelegate(const void* PropertyValue) const override;
 	virtual void SetMulticastDelegate(void* PropertyValue, FMulticastScriptDelegate ScriptDelegate) const override;
 
-	virtual void AddDelegate(FScriptDelegate ScriptDelegate, UObject* Parent) const override
-	{
-		AddDelegate(MoveTemp(ScriptDelegate), Parent, GetPropertyValuePtr_InContainer(Parent));
-	}
-
-	virtual void RemoveDelegate(const FScriptDelegate& ScriptDelegate, UObject* Parent) const override
-	{
-		RemoveDelegate(ScriptDelegate, Parent, GetPropertyValuePtr_InContainer(Parent));
-	}
-
-	virtual void ClearDelegate(UObject* Parent) const override
-	{
-		ClearDelegate(Parent, GetPropertyValuePtr_InContainer(Parent));
-	}
+	virtual void AddDelegate(FScriptDelegate ScriptDelegate, UObject* Parent = nullptr, void* PropertyValue = nullptr) const override;
+	virtual void RemoveDelegate(const FScriptDelegate& ScriptDelegate, UObject* Parent = nullptr, void* PropertyValue = nullptr) const override;
+	virtual void ClearDelegate(UObject* Parent = nullptr, void* PropertyValue = nullptr) const override;
 
-	virtual void AddDelegate(FScriptDelegate ScriptDelegate, UObject* Parent, void* PropertyValue) const override final;
-	virtual void RemoveDelegate(const FScriptDelegate& ScriptDelegate, UObject* Parent, void* PropertyValue) const override final;
-	virtual void ClearDelegate(UObject* Parent, void* PropertyValue) const override final;
 protected:
 	virtual FMulticastScriptDelegate::FInvocationList& GetInvocationList(const void* PropertyValue) const;
 	// End of UMulticastDelegateProperty interface
