#UE4 Disabling an assert as a workaround until Core team can investigate

diff --git a/Engine/Source/Runtime/CoreUObject/Private/UObject/Linker.cpp b/Engine/Source/Runtime/CoreUObject/Private/UObject/Linker.cpp
--- a/Engine/Source/Runtime/CoreUObject/Private/UObject/Linker.cpp
+++ b/Engine/Source/Runtime/CoreUObject/Private/UObject/Linker.cpp
@@ -716,7 +716,7 @@ FLinkerLoad* GetPackageLinker
 	if( !Result )
 	{
 		// @todo: RM so never create the context below?
-		check(InExistingContext && InExistingContext->HasStartedLoading());
+		//check(InExistingContext && InExistingContext->HasStartedLoading());
 
 		// we will already have found the filename above
 		check(NewFilename.Len() > 0);
