Fixing issues with objects not being serialized before they are being ConditionallyPostLoaded by other objects and LoadContext related crashes

diff --git a/Engine/Source/Editor/UnrealEd/Private/Commandlets/PackageUtilities.cpp b/Engine/Source/Editor/UnrealEd/Private/Commandlets/PackageUtilities.cpp
--- a/Engine/Source/Editor/UnrealEd/Private/Commandlets/PackageUtilities.cpp
+++ b/Engine/Source/Editor/UnrealEd/Private/Commandlets/PackageUtilities.cpp
@@ -1553,7 +1553,7 @@ int32 UPkgInfoCommandlet::Main( const FString& Params )
 		if (!bDumpProperties)
 		{
 			TGuardValue<bool> GuardAllowUnversionedContentInEditor(GAllowUnversionedContentInEditor, true);
-			TRefCountPtr<FUObjectSerializeContext> LoadContext(new FUObjectSerializeContext());
+			TRefCountPtr<FUObjectSerializeContext> LoadContext(FUObjectThreadContext::Get().GetSerializeContext());
 			BeginLoad(LoadContext);
 			Linker = CreateLinkerForFilename(LoadContext, Filename);
 			EndLoad(Linker->GetSerializeContext());
diff --git a/Engine/Source/Runtime/CoreUObject/Private/Serialization/ArchiveStackTrace.cpp b/Engine/Source/Runtime/CoreUObject/Private/Serialization/ArchiveStackTrace.cpp
--- a/Engine/Source/Runtime/CoreUObject/Private/Serialization/ArchiveStackTrace.cpp
+++ b/Engine/Source/Runtime/CoreUObject/Private/Serialization/ArchiveStackTrace.cpp
@@ -1304,14 +1304,14 @@ void FArchiveStackTrace::DumpPackageHeaderDiffs(const FPackageData& SourcePackag
 	FLinkerLoad* DestLinker = nullptr;
 	// Create linkers. Note there's no need to clean them up here since they will be removed by the package associated with them
 	{
-		TRefCountPtr<FUObjectSerializeContext> LinkerLoadContext(new FUObjectSerializeContext());
+		TRefCountPtr<FUObjectSerializeContext> LinkerLoadContext(FUObjectThreadContext::Get().GetSerializeContext());
 		BeginLoad(LinkerLoadContext);
 		SourceLinker = CreateLinkerForPackage(LinkerLoadContext, SourceAssetPackageName, AssetFilename, SourcePackage);
 		EndLoad(SourceLinker ? SourceLinker->GetSerializeContext() : LinkerLoadContext.GetReference());
 	}
 	
 	{
-		TRefCountPtr<FUObjectSerializeContext> LinkerLoadContext(new FUObjectSerializeContext());
+		TRefCountPtr<FUObjectSerializeContext> LinkerLoadContext(FUObjectThreadContext::Get().GetSerializeContext());
 		BeginLoad(LinkerLoadContext);
 		DestLinker = CreateLinkerForPackage(LinkerLoadContext, DestAssetPackageName, AssetFilename, DestPackage);
 		EndLoad(DestLinker ? DestLinker->GetSerializeContext() : LinkerLoadContext.GetReference());
diff --git a/Engine/Source/Runtime/CoreUObject/Private/Serialization/AsyncLoading.cpp b/Engine/Source/Runtime/CoreUObject/Private/Serialization/AsyncLoading.cpp
--- a/Engine/Source/Runtime/CoreUObject/Private/Serialization/AsyncLoading.cpp
+++ b/Engine/Source/Runtime/CoreUObject/Private/Serialization/AsyncLoading.cpp
@@ -485,6 +485,11 @@ void FAsyncPackage::PopulateFlushTree(struct FFlushTree* FlushTree)
 	}
 }
 
+FUObjectSerializeContext* FAsyncPackage::GetSerializeContext()
+{
+	return FUObjectThreadContext::Get().GetSerializeContext();
+}
+
 FAsyncPackage* FAsyncLoadingThread::FindExistingPackageAndAddCompletionCallback(FAsyncPackageDesc* PackageRequest, TMap<FName, FAsyncPackage*>& PackageList, FFlushTree* FlushTree)
 {
 	checkSlow(IsInAsyncLoadThread());
@@ -3149,6 +3154,7 @@ UObject* FAsyncPackage::EventDrivenIndexToObject(FPackageIndex Index, bool bChec
 		MyDependentNode.Phase = EEventLoadNode::ImportOrExport_Create;
 		if (EventNodeArray.GetNode(MyDependentNode, false).bAddedToGraph || !EventNodeArray.GetNode(MyDependentNode, false).bFired)
 		{
+			FUObjectSerializeContext* LoadContext = GetSerializeContext();
 			UClass* SerClass = Cast<UClass>(LoadContext->SerializedObject);
 			if (!SerClass || Linker->ImpExp(Index).ObjectName != SerClass->GetDefaultObjectName())
 			{
@@ -3166,6 +3172,7 @@ UObject* FAsyncPackage::EventDrivenIndexToObject(FPackageIndex Index, bool bChec
 
 		if (DumpIndex.IsNull())
 		{
+			FUObjectSerializeContext* LoadContext = GetSerializeContext();
 			DumpDependencies(TEXT("Dependencies"), LoadContext->SerializedObject);
 		}
 		else
@@ -3210,6 +3217,7 @@ void FAsyncPackage::EventDrivenCreateExport(int32 LocalExportIndex)
 	check(!Export.Object); // we should not have this yet
 	if (!Export.Object && !Export.bExportLoadFailed)
 	{
+		FUObjectSerializeContext* LoadContext = GetSerializeContext();
 		FScopedAddObjectreference OnExitAddReference(*this, Export.Object);
 
 		if (!Linker->FilterExport(Export)) // for some acceptable position, it was not "not for" 
@@ -3581,6 +3589,7 @@ void FAsyncPackage::EventDrivenSerializeExport(int32 LocalExportIndex)
 
 		Object->ClearFlags(RF_NeedLoad);
 
+		FUObjectSerializeContext* LoadContext = GetSerializeContext();
 		UObject* PrevSerializedObject = LoadContext->SerializedObject;
 		LoadContext->SerializedObject = Object;
 		Linker->bForceSimpleIndexToObject = true;
@@ -4122,6 +4131,7 @@ void FAsyncPackage::Event_StartPostload()
 	AsyncPackageLoadingState = EAsyncPackageLoadingState::PostLoad_Etc;
 	EventDrivenLoadingComplete();
 	{
+		FUObjectSerializeContext* LoadContext = GetSerializeContext();
 		LoadContext->ReserveObjectsLoaded(LoadContext->GetNumObjectsLoaded() + Linker->ExportMap.Num());
 		for (int32 LocalExportIndex = 0; LocalExportIndex < Linker->ExportMap.Num(); LocalExportIndex++)
 		{
@@ -5485,7 +5495,6 @@ FAsyncPackage::FAsyncPackage(const FAsyncPackageDesc& InDesc)
 , FinishObjectsTime(0.0)
 #endif // PERF_TRACK_DETAILED_ASYNC_STATS
 {
-	LoadContext = new FUObjectSerializeContext();
 	AddRequestID(InDesc.RequestID);
 }
 
@@ -5675,6 +5684,7 @@ void FAsyncPackage::BeginAsyncLoad()
 	}
 
 	// this won't do much during async loading except increase the load count which causes IsLoading to return true
+	FUObjectSerializeContext* LoadContext = GetSerializeContext();
 	BeginLoad(LoadContext);
 }
 
@@ -5687,6 +5697,7 @@ void FAsyncPackage::EndAsyncLoad()
 	check(IsAsyncLoading());
 
 	// this won't do much during async loading except decrease the load count which causes IsLoading to return false
+	FUObjectSerializeContext* LoadContext = GetSerializeContext();
 	EndLoad(LoadContext);
 
 	if (IsInGameThread())
@@ -5958,19 +5969,6 @@ EAsyncPackageState::Type FAsyncPackage::CreateLinker()
 		Linker = FLinkerLoad::FindExistingLinkerForPackage(Package);
 		if (Linker)
 		{
-			// Swap the load context to the currently associated with the existing linker
-			if (Linker->GetSerializeContext())
-			{
-				LoadContext->DecrementBeginLoadCount();
-				LoadContext = Linker->GetSerializeContext();
-				LoadContext->IncrementBeginLoadCount();
-			}
-			else
-			{
-				check(!GEventDrivenLoaderEnabled);
-				Linker->SetSerializeContext(LoadContext);
-			}
-
 			if (GEventDrivenLoaderEnabled)
 			{
 				// this almost works, but the EDL does not tolerate redoing steps it already did
@@ -6064,6 +6062,7 @@ EAsyncPackageState::Type FAsyncPackage::CreateLinker()
 				LinkerFlags |= LOAD_PackageForPIE;
 			}
 #endif
+			FUObjectSerializeContext* LoadContext = GetSerializeContext();
 			if (GEventDrivenLoaderEnabled)
 			{
 				FWeakAsyncPackagePtr WeakPtr(this);
@@ -6530,6 +6529,7 @@ EAsyncPackageState::Type FAsyncPackage::PreLoadObjects()
 	// GC can't run in here
 	FGCScopeGuard GCGuard;
 
+	FUObjectSerializeContext* LoadContext = GetSerializeContext();
 	TArray<UObject*>& ThreadObjLoaded = LoadContext->PRIVATE_GetObjectsLoadedInternalUseOnly();
 	PackageObjLoaded.Append(ThreadObjLoaded);
 	ThreadObjLoaded.Reset();
@@ -6629,6 +6629,7 @@ EAsyncPackageState::Type FAsyncPackage::PostLoadObjects()
 	FUObjectThreadContext& ThreadContext = FUObjectThreadContext::Get();
 	TGuardValue<bool> GuardIsRoutingPostLoad(ThreadContext.IsRoutingPostLoad, true);
 
+	FUObjectSerializeContext* LoadContext = GetSerializeContext();
 	TArray<UObject*>& ThreadObjLoaded = LoadContext->PRIVATE_GetObjectsLoadedInternalUseOnly();
 	if (ThreadObjLoaded.Num())
 	{
@@ -6709,6 +6710,7 @@ EAsyncPackageState::Type FAsyncPackage::PostLoadDeferredObjects(double InTickSta
 	TGuardValue<bool> GuardIsRoutingPostLoad(PackageScope.ThreadContext.IsRoutingPostLoad, true);
 	FAsyncLoadingTickScope InAsyncLoadingTick;
 
+	FUObjectSerializeContext* LoadContext = GetSerializeContext();
 	TArray<UObject*>& ObjLoadedInPostLoad = LoadContext->PRIVATE_GetObjectsLoadedInternalUseOnly();
 	TArray<UObject*> ObjLoadedInPostLoadLocal;
 
@@ -6898,6 +6900,7 @@ EAsyncPackageState::Type FAsyncPackage::FinishObjects()
 	LastObjectWorkWasPerformedOn	= nullptr;
 	LastTypeOfWorkPerformed			= TEXT("finishing all objects");
 
+	FUObjectSerializeContext* LoadContext = GetSerializeContext();
 	check(!Linker || LoadContext == Linker->GetSerializeContext());		
 	TArray<UObject*>& ThreadObjLoaded = LoadContext->PRIVATE_GetObjectsLoadedInternalUseOnly();
 
diff --git a/Engine/Source/Runtime/CoreUObject/Private/UObject/Linker.cpp b/Engine/Source/Runtime/CoreUObject/Private/UObject/Linker.cpp
--- a/Engine/Source/Runtime/CoreUObject/Private/UObject/Linker.cpp
+++ b/Engine/Source/Runtime/CoreUObject/Private/UObject/Linker.cpp
@@ -720,7 +720,7 @@ FLinkerLoad* GetPackageLinker
 
 		// we will already have found the filename above
 		check(NewFilename.Len() > 0);
-		TRefCountPtr<FUObjectSerializeContext> LoadContext(InExistingContext ? InExistingContext : new FUObjectSerializeContext());
+		TRefCountPtr<FUObjectSerializeContext> LoadContext(FUObjectThreadContext::Get().GetSerializeContext());
 		Result = FLinkerLoad::CreateLinker(LoadContext, InOuter, *NewFilename, LoadFlags, InReaderOverride);
 	}
 	else if (InExistingContext)
@@ -772,7 +772,7 @@ FLinkerLoad* GetPackageLinker
 FLinkerLoad* LoadPackageLinker(UPackage* InOuter, const TCHAR* InLongPackageName, uint32 LoadFlags, UPackageMap* Sandbox, FGuid* CompatibleGuid, FArchive* InReaderOverride, TFunctionRef<void(FLinkerLoad* LoadedLinker)> LinkerLoadedCallback)
 {
 	FLinkerLoad* Linker = nullptr;
-	TRefCountPtr<FUObjectSerializeContext> LoadContext(new FUObjectSerializeContext());
+	TRefCountPtr<FUObjectSerializeContext> LoadContext(FUObjectThreadContext::Get().GetSerializeContext());
 	BeginLoad(LoadContext);
 	{
 		FUObjectSerializeContext* InOutLoadContext = LoadContext;
diff --git a/Engine/Source/Runtime/CoreUObject/Private/UObject/LinkerLoad.cpp b/Engine/Source/Runtime/CoreUObject/Private/UObject/LinkerLoad.cpp
--- a/Engine/Source/Runtime/CoreUObject/Private/UObject/LinkerLoad.cpp
+++ b/Engine/Source/Runtime/CoreUObject/Private/UObject/LinkerLoad.cpp
@@ -642,19 +642,11 @@ FLinkerLoad* FLinkerLoad::CreateLinkerAsync(FUObjectSerializeContext* LoadContex
 
 void FLinkerLoad::SetSerializeContext(FUObjectSerializeContext* InLoadContext)
 {
-	if (!GEventDrivenLoaderEnabled && CurrentLoadContext)
-	{
-		CurrentLoadContext->DetachLinker(this);
-	}
-	CurrentLoadContext = InLoadContext;
-	if (!GEventDrivenLoaderEnabled && CurrentLoadContext)
-	{
-		CurrentLoadContext->AttachLinker(this);
-	}
 }
+
 FUObjectSerializeContext* FLinkerLoad::GetSerializeContext()
 {
-	return CurrentLoadContext;
+	return FUObjectThreadContext::Get().GetSerializeContext();
 }
 
 /**
@@ -2429,6 +2421,8 @@ FLinkerLoad::EVerifyResult FLinkerLoad::VerifyImport(int32 ImportIndex)
 	// these checks find out if the VerifyImportInner was successful or not 
 	if (Import.SourceLinker && Import.SourceIndex == INDEX_NONE && Import.XObject == NULL && !Import.OuterIndex.IsNull() && Import.ObjectName != NAME_ObjectRedirector)
 	{
+		FUObjectSerializeContext* CurrentLoadContext = GetSerializeContext();
+
 		// if we found the package, but not the object, look for a redirector
 		FObjectImport OriginalImport = Import;
 		Import.ClassName = NAME_ObjectRedirector;
@@ -2949,6 +2943,7 @@ bool FLinkerLoad::VerifyImportInner(const int32 ImportIndex, FString& WarningSuf
 			{
 				// except if we are looking for _the_ package...in which case we are looking for TmpPkg, so we are done
 				Import.XObject = TmpPkg;
+				FUObjectSerializeContext* CurrentLoadContext = GetSerializeContext();
 				check(CurrentLoadContext);
 				CurrentLoadContext->IncrementImportCount();
 				FLinkerManager::Get().AddLoaderWithNewImports(this);
@@ -3003,6 +2998,7 @@ bool FLinkerLoad::VerifyImportInner(const int32 ImportIndex, FString& WarningSuf
 				if (FindObject != NULL && ((LoadFlags & LOAD_FindIfFail) || bIsInMemoryOnlyOrNativeTransient))
 				{
 					Import.XObject = FindObject;
+					FUObjectSerializeContext* CurrentLoadContext = GetSerializeContext();
 					check(CurrentLoadContext);
 					CurrentLoadContext->IncrementImportCount();
 					FLinkerManager::Get().AddLoaderWithNewImports(this);
@@ -3424,6 +3420,8 @@ void FLinkerLoad::Preload( UObject* Object )
 	// Preload the object if necessary.
 	if (Object->HasAnyFlags(RF_NeedLoad))
 	{
+		FUObjectSerializeContext* CurrentLoadContext = GetSerializeContext();
+
 		if (Object->GetLinker() == this)
 		{
 			check(!GEventDrivenLoaderEnabled || !bLockoutLegacyOperations || !EVENT_DRIVEN_ASYNC_LOAD_ACTIVE_AT_RUNTIME);
@@ -3556,7 +3554,7 @@ void FLinkerLoad::Preload( UObject* Object )
 							return;
 						}
 #endif // USE_CIRCULAR_DEPENDENCY_LOAD_DEFERRING
-
+						
 						check(CurrentLoadContext);
 						// Maintain the current SerializedObjects.
 						UObject* PrevSerializedObject = CurrentLoadContext->SerializedObject;
@@ -3853,6 +3851,7 @@ UObject* FLinkerLoad::CreateExport( int32 Index )
 	// Check whether we already loaded the object and if not whether the context flags allow loading it.
 	if( !Export.Object && !FilterExport(Export) ) // for some acceptable position, it was not "not for" 
 	{
+		FUObjectSerializeContext* CurrentLoadContext = GetSerializeContext();
 		check(!GEventDrivenLoaderEnabled || !bLockoutLegacyOperations || !EVENT_DRIVEN_ASYNC_LOAD_ACTIVE_AT_RUNTIME);
 		check(Export.ObjectName!=NAME_None || !(Export.ObjectFlags&RF_Public));
 		check(IsLoading());
@@ -4453,6 +4452,7 @@ UObject* FLinkerLoad::CreateImport( int32 Index )
 	// Imports can have no name if they were filtered out due to package redirects, skip in that case
 	if (Import.XObject == nullptr && Import.ObjectName != NAME_None)
 	{
+		FUObjectSerializeContext* CurrentLoadContext = GetSerializeContext();
 		if (!GIsEditor && !IsRunningCommandlet())
 		{
 			// Try to find existing version in memory first.
@@ -4735,8 +4735,10 @@ void FLinkerLoad::Detach()
 
 	// Remove from object manager, if it has been added.
 	FLinkerManager::Get().RemoveLoaderFromObjectLoadersAndLoadersWithNewImports(this);
-	if (!FPlatformProperties::HasEditorOnlyData() && CurrentLoadContext)
+	if (!FPlatformProperties::HasEditorOnlyData())
 	{
+		FUObjectSerializeContext* CurrentLoadContext = GetSerializeContext();
+		check(CurrentLoadContext);
 		CurrentLoadContext->RemoveDelayedLinkerClosePackage(this);
 	}
 
@@ -4834,7 +4836,6 @@ FArchive& FLinkerLoad::operator<<( UObject*& Object )
 		return *this;
 	}
 
-	check(CurrentLoadContext != nullptr);
 	UObject* Temporary = NULL;
 	Temporary = IndexToObject( Index );
 
@@ -4946,6 +4947,7 @@ UObject* FLinkerLoad::GetArchetypeFromLoader(const UObject* Obj)
 {
 	if (GEventDrivenLoaderEnabled)
 	{
+		FUObjectSerializeContext* CurrentLoadContext = GetSerializeContext();
 		check(CurrentLoadContext);
 		check(!TemplateForGetArchetypeFromLoader || CurrentLoadContext->SerializedObject == Obj);
 		return TemplateForGetArchetypeFromLoader;
diff --git a/Engine/Source/Runtime/CoreUObject/Private/UObject/SavePackage.cpp b/Engine/Source/Runtime/CoreUObject/Private/UObject/SavePackage.cpp
--- a/Engine/Source/Runtime/CoreUObject/Private/UObject/SavePackage.cpp
+++ b/Engine/Source/Runtime/CoreUObject/Private/UObject/SavePackage.cpp
@@ -3553,7 +3553,7 @@ FSavePackageResultStruct UPackage::Save(UPackage* InOuter, UObject* Base, EObjec
 
 	if (FPlatformProperties::HasEditorOnlyData())
 	{
-		TRefCountPtr<FUObjectSerializeContext> SaveContext(new FUObjectSerializeContext());
+		TRefCountPtr<FUObjectSerializeContext> SaveContext(FUObjectThreadContext::Get().GetSerializeContext());
 
 #if WITH_EDITOR
 		struct FDiffSettings
diff --git a/Engine/Source/Runtime/CoreUObject/Private/UObject/UObjectGlobals.cpp b/Engine/Source/Runtime/CoreUObject/Private/UObject/UObjectGlobals.cpp
--- a/Engine/Source/Runtime/CoreUObject/Private/UObject/UObjectGlobals.cpp
+++ b/Engine/Source/Runtime/CoreUObject/Private/UObject/UObjectGlobals.cpp
@@ -1155,30 +1155,7 @@ UPackage* LoadPackageInternal(UPackage* InOuter, const TCHAR* InLongPackageNameO
 	}
 	
 	// Set up a load context
-	TRefCountPtr<FUObjectSerializeContext> LoadContext;
-	if (InLoadContext)
-	{
-		// Use the privided context
-		LoadContext = InLoadContext;
-	}
-	else
-	{
-		// Try to get the context from the callstack
-		FUObjectThreadContext& ThreadContext = FUObjectThreadContext::Get();
-		if (ThreadContext.SerializeContext)
-		{
-			LoadContext = ThreadContext.SerializeContext;
-		}
-		else
-		{
-			// Create a new context
-			LoadContext = new FUObjectSerializeContext();
-			if (!IsInAsyncLoadingThread())
-			{
-				ThreadContext.SerializeContext = LoadContext;
-			}
-		}
-	}
+	TRefCountPtr<FUObjectSerializeContext> LoadContext = FUObjectThreadContext::Get().GetSerializeContext();
 
 	// Try to load.
 	BeginLoad(LoadContext, InLongPackageNameOrFilename);
@@ -1759,14 +1736,8 @@ void EndLoad(FUObjectSerializeContext* LoadContext)
 			LoadContext->DetachFromLinkers();
 		}
 
-		FUObjectThreadContext& ThreadContext = FUObjectThreadContext::Get();
-		if (ThreadContext.SerializeContext == LoadContext)
 		{
-			ThreadContext.SerializeContext = nullptr;
-		}
-		else
-		{
-			check(!ThreadContext.SerializeContext);
+			LoadContext->DetachFromLinkers();
 		}
 	}
 }
@@ -2073,7 +2044,7 @@ UObject* StaticDuplicateObjectEx( FObjectDuplicationParameters& Parameters )
 		SerializedObjects.Add(Object);
 	};
 
-	TRefCountPtr<FUObjectSerializeContext> LoadContext(new FUObjectSerializeContext());
+	TRefCountPtr<FUObjectSerializeContext> LoadContext(FUObjectThreadContext::Get().GetSerializeContext());
 	FDuplicateDataReader Reader(DuplicatedObjectAnnotation, ObjectData, Parameters.PortFlags, Parameters.DestOuter);
 	Reader.SetSerializeContext(LoadContext);
 	for(int32 ObjectIndex = 0;ObjectIndex < SerializedObjects.Num();ObjectIndex++)
diff --git a/Engine/Source/Runtime/CoreUObject/Private/UObject/UObjectThreadContext.cpp b/Engine/Source/Runtime/CoreUObject/Private/UObject/UObjectThreadContext.cpp
--- a/Engine/Source/Runtime/CoreUObject/Private/UObject/UObjectThreadContext.cpp
+++ b/Engine/Source/Runtime/CoreUObject/Private/UObject/UObjectThreadContext.cpp
@@ -17,9 +17,12 @@ FUObjectThreadContext::FUObjectThreadContext()
 , IsInConstructor(0)
 , ConstructedObject(nullptr)
 , AsyncPackage(nullptr)
-, SerializeContext(nullptr)
+, SerializeContext(new FUObjectSerializeContext())
 {}
 
+FUObjectThreadContext::~FUObjectThreadContext()
+{
+}
 
 FUObjectSerializeContext::FUObjectSerializeContext()
 	: RefCount(0)
@@ -59,6 +62,11 @@ void FUObjectSerializeContext::AddUniqueLoadedObjects(const TArray<UObject*>& In
 	
 }
 
+void FUObjectSerializeContext::AddLoadedObject(UObject* InObject)
+{
+	ObjectsLoaded.Add(InObject);
+}
+
 bool FUObjectSerializeContext::PRIVATE_PatchNewObjectIntoExport(UObject* OldObject, UObject* NewObject)
 {
 	const int32 ObjLoadedIdx = ObjectsLoaded.Find(OldObject);
@@ -78,6 +86,11 @@ void FUObjectSerializeContext::AttachLinker(FLinkerLoad* InLinker)
 	AttachedLinkers.Add(InLinker);
 }
 
+void FUObjectSerializeContext::DetachLinker(FLinkerLoad* InLinker)
+{
+	AttachedLinkers.Remove(InLinker);
+}
+
 void FUObjectSerializeContext::DetachFromLinkers()
 {
 	check(!GEventDrivenLoaderEnabled);
diff --git a/Engine/Source/Runtime/CoreUObject/Public/Serialization/AsyncLoading.h b/Engine/Source/Runtime/CoreUObject/Public/Serialization/AsyncLoading.h
--- a/Engine/Source/Runtime/CoreUObject/Public/Serialization/AsyncLoading.h
+++ b/Engine/Source/Runtime/CoreUObject/Public/Serialization/AsyncLoading.h
@@ -945,7 +945,7 @@ private:
 public:
 
 	/** Serialization context for this package */
-	TRefCountPtr<FUObjectSerializeContext> LoadContext;
+	FUObjectSerializeContext* GetSerializeContext();
 };
 
 struct FScopedAsyncPackageEvent
diff --git a/Engine/Source/Runtime/CoreUObject/Public/UObject/LinkerLoad.h b/Engine/Source/Runtime/CoreUObject/Public/UObject/LinkerLoad.h
--- a/Engine/Source/Runtime/CoreUObject/Public/UObject/LinkerLoad.h
+++ b/Engine/Source/Runtime/CoreUObject/Public/UObject/LinkerLoad.h
@@ -1185,11 +1185,6 @@ private:
 	// FLinkerLoad creation helpers END
 	//
 
-private:
-
-	/** Current UObject serialization context */
-	TRefCountPtr<FUObjectSerializeContext> CurrentLoadContext;
-
 public:
 
 	//~ FArchive interface
diff --git a/Engine/Source/Runtime/CoreUObject/Public/UObject/UObjectThreadContext.h b/Engine/Source/Runtime/CoreUObject/Public/UObject/UObjectThreadContext.h
--- a/Engine/Source/Runtime/CoreUObject/Public/UObject/UObjectThreadContext.h
+++ b/Engine/Source/Runtime/CoreUObject/Public/UObject/UObjectThreadContext.h
@@ -9,6 +9,7 @@
 #include "CoreMinimal.h"
 #include "HAL/ThreadSingleton.h"
 #include "HAL/ThreadSafeCounter.h"
+#include "Templates/RefCounting.h"
 
 class FObjectInitializer;
 struct FUObjectSerializeContext;
@@ -21,6 +22,7 @@ class COREUOBJECT_API FUObjectThreadContext : public TThreadSingleton<FUObjectTh
 	friend TThreadSingleton<FUObjectThreadContext>;
 
 	FUObjectThreadContext();
+	virtual ~FUObjectThreadContext();
 
 	/** Stack of currently used FObjectInitializers for this thread */
 	TArray<FObjectInitializer*> InitializerStack;
@@ -76,8 +78,6 @@ public:
 	UObject* ConstructedObject;
 	/** Async Package currently processing objects */
 	struct FAsyncPackage* AsyncPackage;
-	/** Current serialization context (only valid for non-async loads) */
-	FUObjectSerializeContext* SerializeContext;
 
 #if !(UE_BUILD_SHIPPING || UE_BUILD_TEST)
 	/** Stack to ensure that PostInitProperties is routed through Super:: calls. **/
@@ -89,13 +89,31 @@ public:
 	/** Maps a package name to all packages marked as editor-only due to the fact it was marked as editor-only */
 	TMap<FName, TSet<FName>> PackagesMarkedEditorOnlyByOtherPackage;
 #endif
+
+	/** Gets the current serialization context */
+	FUObjectSerializeContext* GetSerializeContext()
+	{
+		return SerializeContext;
+	}
+
+private:
+	/** Current serialization context */
+	TRefCountPtr<FUObjectSerializeContext> SerializeContext;
 };
 
 /** Structure that holds the current serialization state of UObjects */
 struct COREUOBJECT_API FUObjectSerializeContext
 {
+	friend class FUObjectThreadContext;
+
 private:
 
+	/** Constructor */
+	FUObjectSerializeContext();
+
+	/** Destructor */
+	~FUObjectSerializeContext();
+
 	/** Reference count of this context */
 	int32 RefCount;
 
@@ -127,19 +145,8 @@ public:
 	/** Points to the most recently used Linker for serialization by CreateExport() */
 	FLinkerLoad* SerializedExportLinker;
 
-	/** Constructor */
-	FUObjectSerializeContext();
-
-	/** Destructor */
-	~FUObjectSerializeContext();
-
 	/** Adds a new loaded object */
-	void AddLoadedObject(UObject* InObject)
-	{
-		check(AttachedLinkers.Num() || GEventDrivenLoaderEnabled);
-		ObjectsLoaded.Add(InObject);
-	}
-
+	void AddLoadedObject(UObject* InObject);
 	void AddUniqueLoadedObjects(const TArray<UObject*>& InObjects);
 
 	/** Checks if object loading has started */
@@ -226,10 +233,7 @@ public:
 	void AttachLinker(FLinkerLoad* InLinker);
 	
 	/** Detaches a linker from this context */
-	void DetachLinker(FLinkerLoad* InLinker)
-	{
-		AttachedLinkers.Remove(InLinker);
-	}
+	void DetachLinker(FLinkerLoad* InLinker);
 
 	/** Detaches all linkers from this context */
 	void DetachFromLinkers();
