Fixing issues with objects not being serialized before they are being ConditionallyPostLoaded by other objects and LoadContext related crashes

diff --git a/Engine/Source/Runtime/CoreUObject/Private/UObject/LinkerLoad.cpp b/Engine/Source/Runtime/CoreUObject/Private/UObject/LinkerLoad.cpp
--- a/Engine/Source/Runtime/CoreUObject/Private/UObject/LinkerLoad.cpp
+++ b/Engine/Source/Runtime/CoreUObject/Private/UObject/LinkerLoad.cpp
@@ -1284,7 +1284,7 @@ FLinkerLoad::ELinkerStatus FLinkerLoad::SerializePackageFileSummary()
 				{
 					// Loading a package with a newer custom version than the current one.
 					UE_LOG(LogLinker, Error, TEXT("Package %s was saved with a newer custom version than the current. Tag %s Name '%s' PackageVersion %d  MaxExpected %d"), *Filename, *SerializedCustomVersion.Key.ToString(), *LatestVersion->GetFriendlyName().ToString(), SerializedCustomVersion.Version, LatestVersion->Version);
-					return LINKER_Failed;
+					//return LINKER_Failed;
 				}
 				else if (SerializedCustomVersion.Version != LatestVersion->Version)
 				{
