Fixing a crash when detaching linkers from serialize context

diff --git a/Engine/Source/Runtime/CoreUObject/Private/UObject/UObjectThreadContext.cpp b/Engine/Source/Runtime/CoreUObject/Private/UObject/UObjectThreadContext.cpp
--- a/Engine/Source/Runtime/CoreUObject/Private/UObject/UObjectThreadContext.cpp
+++ b/Engine/Source/Runtime/CoreUObject/Private/UObject/UObjectThreadContext.cpp
@@ -94,7 +94,7 @@ void FUObjectSerializeContext::DetachLinker(FLinkerLoad* InLinker)
 void FUObjectSerializeContext::DetachFromLinkers()
 {
 	check(!GEventDrivenLoaderEnabled);
-	check(ObjectsLoaded.Num() == 0);
+	check(ObjectsLoaded.Num() == 0 || AttachedLinkers.Num() == 0);
 	TArray<FLinkerLoad*> LinkersToDetach = AttachedLinkers.Array();
 	for (FLinkerLoad* Linker : LinkersToDetach)
 	{
