#UE4 Remove redundant call to DetachFromLinkers. It was causing a crash when using EDL

diff --git a/Engine/Source/Runtime/CoreUObject/Private/UObject/UObjectGlobals.cpp b/Engine/Source/Runtime/CoreUObject/Private/UObject/UObjectGlobals.cpp
--- a/Engine/Source/Runtime/CoreUObject/Private/UObject/UObjectGlobals.cpp
+++ b/Engine/Source/Runtime/CoreUObject/Private/UObject/UObjectGlobals.cpp
@@ -1735,10 +1735,6 @@ void EndLoad(FUObjectSerializeContext* LoadContext)
 		{
 			LoadContext->DetachFromLinkers();
 		}
-
-		{
-			LoadContext->DetachFromLinkers();
-		}
 	}
 }
 
