Fix to stop UClass CDOs trying to have their token streams reassembled on hot reload.

diff --git a/Engine/Source/Runtime/CoreUObject/Private/UObject/Class.cpp b/Engine/Source/Runtime/CoreUObject/Private/UObject/Class.cpp
--- a/Engine/Source/Runtime/CoreUObject/Private/UObject/Class.cpp
+++ b/Engine/Source/Runtime/CoreUObject/Private/UObject/Class.cpp
@@ -4381,7 +4381,7 @@ void UClass::AssembleReferenceTokenStreams()
 				Class->GetDefaultObject(); // Force the default object to be constructed if it isn't already
 			}
 			// Assemble reference token stream for garbage collection/ RTGC.
-			if (!Class->HasAnyClassFlags(CLASS_TokenStreamAssembled))
+			if (!Class->HasAnyFlags(RF_ClassDefaultObject) && !Class->HasAnyClassFlags(CLASS_TokenStreamAssembled))
 			{
 				Class->AssembleReferenceTokenStream();
 			}
