Fix obj refs code by merging parts of CL 5365249 and 5536455 as 6542056 conflicted with a change in framework. This fixes an infinite recursion crash in obj refs and changes it so you don't need to pass the GC collector to AddReferencedObjects but doesn't include the deprecation changes or setting the debug name on streamablemanager

diff --git a/Engine/Source/Runtime/CoreUObject/Private/Misc/GCObjectReferencer.cpp b/Engine/Source/Runtime/CoreUObject/Private/Misc/GCObjectReferencer.cpp
--- a/Engine/Source/Runtime/CoreUObject/Private/Misc/GCObjectReferencer.cpp
+++ b/Engine/Source/Runtime/CoreUObject/Private/Misc/GCObjectReferencer.cpp
@@ -16,17 +16,23 @@ extern bool GObjIncrementalPurgeIsInProgress;
 extern bool GObjUnhashUnreachableIsInProgress;
 
 void UGCObjectReferencer::AddReferencedObjects(UObject* InThis, FReferenceCollector& Collector)
-{	
+{
 	UGCObjectReferencer* This = CastChecked<UGCObjectReferencer>(InThis);
+
+	checkSlow(!This->bIsAddingReferencedObjects);
+	This->bIsAddingReferencedObjects = true;
 	// Note we're not locking ReferencedObjectsCritical here because we guard
 	// against adding new references during GC in AddObject and RemoveObject.
 	// Let each registered object handle its AddReferencedObjects call
 	for (FGCObject* Object : This->ReferencedObjects)
 	{
 		check(Object);
+		This->CurrentlySerializingObject = Object;
 		Object->AddReferencedObjects(Collector);
 	}
-	Super::AddReferencedObjects( This, Collector );
+	This->CurrentlySerializingObject = nullptr;
+	Super::AddReferencedObjects(This, Collector);
+	This->bIsAddingReferencedObjects = false;
 }
 
 void UGCObjectReferencer::AddObject(FGCObject* Object)
@@ -46,8 +52,23 @@ void UGCObjectReferencer::RemoveObject(FGCObject* Object)
 	check(NumRemoved == 1);
 }
 
-bool UGCObjectReferencer::GetReferencerName(UObject* Object, FString& OutName) const
+bool UGCObjectReferencer::GetReferencerName(UObject* Object, FString& OutName, bool bOnlyIfAddingReferenced) const
 {
+	if (bOnlyIfAddingReferenced)
+	{
+		if (!bIsAddingReferencedObjects || !CurrentlySerializingObject)
+		{
+			return false;
+		}
+		OutName = CurrentlySerializingObject->GetReferencerName();
+		FString ReferencerProperty;
+		if (CurrentlySerializingObject->GetReferencerPropertyName(Object, ReferencerProperty))
+		{
+			OutName += TEXT(":") + ReferencerProperty;
+		}
+		return true;
+	}
+
 	// Let each registered object handle its AddReferencedObjects call
 	for (int32 i = 0; i < ReferencedObjects.Num(); i++)
 	{
diff --git a/Engine/Source/Runtime/CoreUObject/Private/Serialization/AsyncLoading.cpp b/Engine/Source/Runtime/CoreUObject/Private/Serialization/AsyncLoading.cpp
--- a/Engine/Source/Runtime/CoreUObject/Private/Serialization/AsyncLoading.cpp
+++ b/Engine/Source/Runtime/CoreUObject/Private/Serialization/AsyncLoading.cpp
@@ -5531,9 +5531,9 @@ FAsyncPackage::~FAsyncPackage()
 
 void FAsyncPackage::AddReferencedObjects(FReferenceCollector& Collector)
 {
-	Collector.AddReferencedObjects(ReferencedObjects, FGCObject::GGCObjectReferencer);
-	Collector.AddReferencedObjects(DeferredFinalizeObjects, FGCObject::GGCObjectReferencer);
-	Collector.AddReferencedObjects(PackageObjLoaded, FGCObject::GGCObjectReferencer);
+	Collector.AddReferencedObjects(ReferencedObjects);
+	Collector.AddReferencedObjects(DeferredFinalizeObjects);
+	Collector.AddReferencedObjects(PackageObjLoaded);
 }
 
 void FAsyncPackage::AddObjectReference(UObject* InObject)
diff --git a/Engine/Source/Runtime/CoreUObject/Private/UObject/ReferenceChainSearch.cpp b/Engine/Source/Runtime/CoreUObject/Private/UObject/ReferenceChainSearch.cpp
--- a/Engine/Source/Runtime/CoreUObject/Private/UObject/ReferenceChainSearch.cpp
+++ b/Engine/Source/Runtime/CoreUObject/Private/UObject/ReferenceChainSearch.cpp
@@ -279,6 +279,29 @@ void FReferenceChainSearch::DumpChain(FReferenceChain* Chain)
 	}
 }
 
+void FReferenceChainSearch::WriteChain(FReferenceChain* Chain, FString& OutString)
+{
+	if (Chain->Num())
+	{
+		// Roots are at the end so iterate from the last to the first node
+		for (int32 NodeIndex = Chain->Num() - 1; NodeIndex >= 0; --NodeIndex)
+		{
+			UObject* Object = Chain->GetNode(NodeIndex)->Object;
+			const FNodeReferenceInfo& ReferenceInfo = Chain->GetReferenceInfo(NodeIndex);
+
+			if (NodeIndex != Chain->Num() - 1)
+			{
+				OutString.Append(LINE_TERMINATOR);
+				OutString.Append(FCString::Spc(Chain->Num() - NodeIndex - 1));
+			}
+
+			OutString.Append(*GetObjectFlags(Object));
+			OutString.Append(*Object->GetFullName());
+			OutString.Append(*ReferenceInfo.ToString());
+		}
+	}
+}
+
 void FReferenceChainSearch::FReferenceChain::FillReferenceInfo()
 {
 	// The first entry is the object we were looking for references to so add an empty entry for it
@@ -350,15 +373,14 @@ public:
 			else
 			{
 				RefInfo.Type = FReferenceChainSearch::EReferenceType::AddReferencedObjects;
-				if (ReferencingObject == FGCObject::GGCObjectReferencer)
+				if (FGCObject::GGCObjectReferencer && (!ReferencingObject || ReferencingObject == FGCObject::GGCObjectReferencer))
 				{
 					FString RefName;
-					if (FGCObject::GGCObjectReferencer->GetReferencerName(Object, RefName))
+					if (FGCObject::GGCObjectReferencer->GetReferencerName(Object, RefName, true))
 					{
-						RefInfo.Type = FReferenceChainSearch::EReferenceType::Property;
 						RefInfo.ReferencerName = FName(*RefName);
 					}
-					else
+					else if (ReferencingObject)
 					{
 						RefInfo.ReferencerName = ReferencingObject->GetFName();
 					}
@@ -447,7 +469,7 @@ void FReferenceChainSearch::FindDirectReferencesForObjects()
 	}
 }
 
-void FReferenceChainSearch::PrintResults(bool bDumpAllChains /*= false*/)
+void FReferenceChainSearch::PrintResults(bool bDumpAllChains /*= false*/) const
 {
 	if (ReferenceChains.Num())
 	{
@@ -480,6 +502,25 @@ void FReferenceChainSearch::PrintResults(bool bDumpAllChains /*= false*/)
 	}
 }
 
+FString FReferenceChainSearch::GetRootPath() const
+{
+	if (ReferenceChains.Num())
+	{
+		FReferenceChain* Chain = ReferenceChains[0];
+		FString OutString;
+
+		WriteChain(Chain, OutString);
+		return OutString;
+	}
+	else
+	{
+		return FString::Printf(TEXT("%s%s is not currently reachable."),
+			*GetObjectFlags(ObjectToFindReferencesTo),
+			*ObjectToFindReferencesTo->GetFullName()
+		);
+	}
+}
+
 void FReferenceChainSearch::Cleanup()
 {
 	for (int32 ChainIndex = 0; ChainIndex < ReferenceChains.Num(); ++ChainIndex)
diff --git a/Engine/Source/Runtime/CoreUObject/Public/UObject/GCObject.h b/Engine/Source/Runtime/CoreUObject/Public/UObject/GCObject.h
--- a/Engine/Source/Runtime/CoreUObject/Public/UObject/GCObject.h
+++ b/Engine/Source/Runtime/CoreUObject/Public/UObject/GCObject.h
@@ -29,6 +29,10 @@ class COREUOBJECT_API UGCObjectReferencer : public UObject
 	TArray<FGCObject*> ReferencedObjects;
 	/** Critical section used when adding and removing objects */
 	FCriticalSection ReferencedObjectsCritical;
+	/** True if we are currently inside AddReferencedObjects */
+	bool bIsAddingReferencedObjects = false;
+	/** Currently serializing FGCObject*, only valid if bIsAddingReferencedObjects */
+	FGCObject* CurrentlySerializingObject = nullptr;
 
 public:
 	DECLARE_CASTED_CLASS_INTRINSIC_WITH_API(UGCObjectReferencer, UObject, CLASS_Transient, TEXT("/Script/CoreUObject"), CASTCLASS_None, NO_API);
@@ -52,9 +56,10 @@ public:
 	 *
 	 * @param Object The object that we're looking for.
 	 * @param OutName the name of the FGCObject that reports this object.
+	 * @param bOnlyIfAddingReferenced Only try to find the name if we are currently inside AddReferencedObjects
 	 * @return true if the object was found.
 	 */
-	bool GetReferencerName(UObject* Object, FString& OutName) const;
+	bool GetReferencerName(UObject* Object, FString& OutName, bool bOnlyIfAddingReferenced = false) const;
 
 	/**
 	 * Forwards this call to all registered objects so they can reference
diff --git a/Engine/Source/Runtime/CoreUObject/Public/UObject/ReferenceChainSearch.h b/Engine/Source/Runtime/CoreUObject/Public/UObject/ReferenceChainSearch.h
--- a/Engine/Source/Runtime/CoreUObject/Public/UObject/ReferenceChainSearch.h
+++ b/Engine/Source/Runtime/CoreUObject/Public/UObject/ReferenceChainSearch.h
@@ -91,7 +91,14 @@ public:
 			}
 			else if (Type == EReferenceType::AddReferencedObjects)
 			{
-				return TEXT("::AddReferencedObjects()");
+				if (!ReferencerName.IsNone())
+				{
+					return FString::Printf(TEXT("::AddReferencedObjects(): %s"), *ReferencerName.ToString());
+				}
+				else
+				{
+					return TEXT("::AddReferencedObjects()");
+				}
 			}
 			return FString();
 		}
@@ -183,19 +190,22 @@ public:
 	};
 
 	/** Constructs a new search engine and finds references to the specified object */
-	FReferenceChainSearch(UObject* InObjectToFindReferencesTo, EReferenceChainSearchMode Mode = EReferenceChainSearchMode::PrintResults);
+	COREUOBJECT_API FReferenceChainSearch(UObject* InObjectToFindReferencesTo, EReferenceChainSearchMode Mode = EReferenceChainSearchMode::PrintResults);
 
 	/** Destructor */
-	~FReferenceChainSearch();
+	COREUOBJECT_API ~FReferenceChainSearch();
 
-	/** 
-	 * Dumps results to log 
+	/**
+	 * Dumps results to log
 	 * @param bDumpAllChains - if set to false, the output will be trimmed to the first 100 reference chains
 	 */
-	void PrintResults(bool bDumpAllChains = false);
+	COREUOBJECT_API void PrintResults(bool bDumpAllChains = false) const;
+
+	/** Returns a string with a short report explaining the root path, will contain newlines */
+	COREUOBJECT_API FString GetRootPath() const;
 
 	/** Returns all reference chains */
-	const TArray<FReferenceChain*>& GetReferenceChains() const
+	COREUOBJECT_API const TArray<FReferenceChain*>& GetReferenceChains() const
 	{
 		return ReferenceChains;
 	}
@@ -235,5 +245,7 @@ private:
 	static FString GetObjectFlags(UObject* InObject);
 	/** Dumps a reference chain to log */
 	static void DumpChain(FReferenceChain* Chain);
+	/** Writes a reference chain to a string */
+	static void WriteChain(FReferenceChain* Chain, FString& OutString);
 };
 
diff --git a/Engine/Source/Runtime/Engine/Private/Slate/DeferredCleanupSlateBrush.cpp b/Engine/Source/Runtime/Engine/Private/Slate/DeferredCleanupSlateBrush.cpp
--- a/Engine/Source/Runtime/Engine/Private/Slate/DeferredCleanupSlateBrush.cpp
+++ b/Engine/Source/Runtime/Engine/Private/Slate/DeferredCleanupSlateBrush.cpp
@@ -43,7 +43,7 @@ TSharedRef<FDeferredCleanupSlateBrush> FDeferredCleanupSlateBrush::CreateBrush(
 
 void FDeferredCleanupSlateBrush::AddReferencedObjects(FReferenceCollector& Collector)
 {
-	InternalBrush.AddReferencedObjects(Collector, FGCObject::GGCObjectReferencer);
+	InternalBrush.AddReferencedObjects(Collector);
 }
 
 FString FDeferredCleanupSlateBrush::GetReferencerName() const
diff --git a/Engine/Source/Runtime/Engine/Private/StreamableManager.cpp b/Engine/Source/Runtime/Engine/Private/StreamableManager.cpp
--- a/Engine/Source/Runtime/Engine/Private/StreamableManager.cpp
+++ b/Engine/Source/Runtime/Engine/Private/StreamableManager.cpp
@@ -782,7 +782,7 @@ void FStreamableManager::AddReferencedObjects(FReferenceCollector& Collector)
 		FStreamable* Existing = It.Value();
 		if (Existing->Target)
 		{
-			Collector.AddReferencedObject(Existing->Target, FGCObject::GGCObjectReferencer);
+			Collector.AddReferencedObject(Existing->Target);
 		}
 	}
 
@@ -791,7 +791,7 @@ void FStreamableManager::AddReferencedObjects(FReferenceCollector& Collector)
 		FRedirectedPath& Existing = It.Value();
 		if (Existing.LoadedRedirector)
 		{
-			Collector.AddReferencedObject(Existing.LoadedRedirector, FGCObject::GGCObjectReferencer);
+			Collector.AddReferencedObject(Existing.LoadedRedirector);
 		}
 	}
 }
diff --git a/Engine/Source/Runtime/UMG/Private/Slate/SObjectWidget.cpp b/Engine/Source/Runtime/UMG/Private/Slate/SObjectWidget.cpp
--- a/Engine/Source/Runtime/UMG/Private/Slate/SObjectWidget.cpp
+++ b/Engine/Source/Runtime/UMG/Private/Slate/SObjectWidget.cpp
@@ -86,7 +86,7 @@ FString SObjectWidget::GetReferencerName() const
 
 void SObjectWidget::AddReferencedObjects(FReferenceCollector& Collector)
 {
-	Collector.AddReferencedObject(WidgetObject, FGCObject::GGCObjectReferencer);
+	Collector.AddReferencedObject(WidgetObject);
 }
 
 void SObjectWidget::SetPadding(const TAttribute<FMargin>& InMargin)
