Removed unnecessary validation logic from FAsyncLoadingThread::ProcessLoadedPackages - On behalf of stefan.boberg

diff --git a/Engine/Source/Runtime/CoreUObject/Private/Serialization/AsyncLoading.cpp b/Engine/Source/Runtime/CoreUObject/Private/Serialization/AsyncLoading.cpp
--- a/Engine/Source/Runtime/CoreUObject/Private/Serialization/AsyncLoading.cpp
+++ b/Engine/Source/Runtime/CoreUObject/Private/Serialization/AsyncLoading.cpp
@@ -4779,13 +4779,6 @@ EAsyncPackageState::Type FAsyncLoadingThread::ProcessLoadedPackages(bool bUseTim
 				{
 					FScopeLock LoadedLock(&LoadedPackagesToProcessCritical);
 					LoadedPackagesToProcess.RemoveAt(PackageIndex--);
-					if (LoadedPackagesToProcess.FindByPredicate([Package](const FAsyncPackage* Pkg)
-					{
-						return Pkg->GetPackageName() == Package->GetPackageName();
-					}))
-					{
-						UE_LOG(LogStreaming, Warning, TEXT("Package %s has already been loaded"), *Package->GetPackageName().ToString());
-					}
 					LoadedPackagesToProcessNameLookup.Remove(Package->GetPackageName());
 
 					if (FPlatformProperties::RequiresCookedData())
