Removed inserting packages by priority when EDL is enabled (on behalf of Per.Larsson)

#rnx

diff --git a/Engine/Source/Runtime/CoreUObject/Private/Serialization/AsyncLoading.cpp b/Engine/Source/Runtime/CoreUObject/Private/Serialization/AsyncLoading.cpp
--- a/Engine/Source/Runtime/CoreUObject/Private/Serialization/AsyncLoading.cpp
+++ b/Engine/Source/Runtime/CoreUObject/Private/Serialization/AsyncLoading.cpp
@@ -4345,36 +4345,44 @@ void FAsyncLoadingThread::InsertPackage(FAsyncPackage* Package, bool bReinsert,
 		{
 			AsyncPackages.RemoveSingle(Package);
 		}
-		int32 InsertIndex = -1;
 
-		switch (InsertMode)
+		if (GEventDrivenLoaderEnabled)
 		{
-			case EAsyncPackageInsertMode::InsertAfterMatchingPriorities:
+			AsyncPackages.Add(Package);
+		}
+		else
+		{
+			int32 InsertIndex = -1;
+
+			switch (InsertMode)
 			{
-				InsertIndex = AsyncPackages.IndexOfByPredicate([Package](const FAsyncPackage* Element)
+				case EAsyncPackageInsertMode::InsertAfterMatchingPriorities:
 				{
-					return Element->GetPriority() < Package->GetPriority();
-				});
+					InsertIndex = AsyncPackages.IndexOfByPredicate([Package](const FAsyncPackage* Element)
+					{
+						return Element->GetPriority() < Package->GetPriority();
+					});
 
-				break;
-			}
+					break;
+				}
 
-			case EAsyncPackageInsertMode::InsertBeforeMatchingPriorities:
-			{
-				// Insert new package keeping descending priority order in AsyncPackages
-				InsertIndex = AsyncPackages.IndexOfByPredicate([Package](const FAsyncPackage* Element)
+				case EAsyncPackageInsertMode::InsertBeforeMatchingPriorities:
 				{
-					return Element->GetPriority() <= Package->GetPriority();
-				});
+					// Insert new package keeping descending priority order in AsyncPackages
+					InsertIndex = AsyncPackages.IndexOfByPredicate([Package](const FAsyncPackage* Element)
+					{
+						return Element->GetPriority() <= Package->GetPriority();
+					});
 
-				break;
-			}
-		};
+					break;
+				}
+			};
 
-		InsertIndex = InsertIndex == INDEX_NONE ? AsyncPackages.Num() : InsertIndex;
+			InsertIndex = InsertIndex == INDEX_NONE ? AsyncPackages.Num() : InsertIndex;
 
-		AsyncPackages.InsertUninitialized(InsertIndex);
-		AsyncPackages[InsertIndex] = Package;
+			AsyncPackages.InsertUninitialized(InsertIndex);
+			AsyncPackages[InsertIndex] = Package;
+		}
 
 		if (!bReinsert)
 		{
