Avoid unnecessary alloca call when not calling remote functions, remove unused parameter from UObject::GetFunctionCallspace
Upgrade notes: second parameter of GetFunctionCallspace overrides should be removed - it served no purpose

diff --git a/Engine/Plugins/NetcodeUnitTest/NetcodeUnitTest/Source/NetcodeUnitTest/Private/MinimalClient.cpp b/Engine/Plugins/NetcodeUnitTest/NetcodeUnitTest/Source/NetcodeUnitTest/Private/MinimalClient.cpp
--- a/Engine/Plugins/NetcodeUnitTest/NetcodeUnitTest/Source/NetcodeUnitTest/Private/MinimalClient.cpp
+++ b/Engine/Plugins/NetcodeUnitTest/NetcodeUnitTest/Source/NetcodeUnitTest/Private/MinimalClient.cpp
@@ -524,7 +524,7 @@ bool UMinimalClient::PostSendRPC(FString RPCName, UObject* Target/*=nullptr*/)
 			}
 			else
 			{
-				int32 Callspace = TargetActor->GetFunctionCallspace(TargetFunc, nullptr, nullptr);
+				int32 Callspace = TargetActor->GetFunctionCallspace(TargetFunc, nullptr);
 
 				if (!(Callspace & FunctionCallspace::Remote))
 				{
diff --git a/Engine/Plugins/Runtime/GameplayAbilities/Source/GameplayAbilities/Private/Abilities/GameplayAbility.cpp b/Engine/Plugins/Runtime/GameplayAbilities/Source/GameplayAbilities/Private/Abilities/GameplayAbility.cpp
--- a/Engine/Plugins/Runtime/GameplayAbilities/Source/GameplayAbilities/Private/Abilities/GameplayAbility.cpp
+++ b/Engine/Plugins/Runtime/GameplayAbilities/Source/GameplayAbilities/Private/Abilities/GameplayAbility.cpp
@@ -87,14 +87,14 @@ UWorld* UGameplayAbility::GetWorld() const
 	return GetOuter()->GetWorld();
 }
 
-int32 UGameplayAbility::GetFunctionCallspace(UFunction* Function, void* Parameters, FFrame* Stack)
+int32 UGameplayAbility::GetFunctionCallspace(UFunction* Function, FFrame* Stack)
 {
 	if (HasAnyFlags(RF_ClassDefaultObject) || !IsSupportedForNetworking())
 	{
 		return FunctionCallspace::Local;
 	}
 	check(GetOuter() != nullptr);
-	return GetOuter()->GetFunctionCallspace(Function, Parameters, Stack);
+	return GetOuter()->GetFunctionCallspace(Function, Stack);
 }
 
 bool UGameplayAbility::CallRemoteFunction(UFunction* Function, void* Parameters, FOutParmRec* OutParms, FFrame* Stack)
diff --git a/Engine/Plugins/Runtime/GameplayAbilities/Source/GameplayAbilities/Public/Abilities/GameplayAbility.h b/Engine/Plugins/Runtime/GameplayAbilities/Source/GameplayAbilities/Public/Abilities/GameplayAbility.h
--- a/Engine/Plugins/Runtime/GameplayAbilities/Source/GameplayAbilities/Public/Abilities/GameplayAbility.h
+++ b/Engine/Plugins/Runtime/GameplayAbilities/Source/GameplayAbilities/Public/Abilities/GameplayAbility.h
@@ -475,7 +475,7 @@ public:
 	//	UObject overrides
 	// --------------------------------------	
 	virtual UWorld* GetWorld() const override;
-	virtual int32 GetFunctionCallspace(UFunction* Function, void* Parameters, FFrame* Stack) override;
+	virtual int32 GetFunctionCallspace(UFunction* Function, FFrame* Stack) override;
 	virtual bool CallRemoteFunction(UFunction* Function, void* Parameters, FOutParmRec* OutParms, FFrame* Stack) override;
 	virtual bool IsSupportedForNetworking() const override;
 
diff --git a/Engine/Source/Runtime/CoreUObject/Private/UObject/ScriptCore.cpp b/Engine/Source/Runtime/CoreUObject/Private/UObject/ScriptCore.cpp
--- a/Engine/Source/Runtime/CoreUObject/Private/UObject/ScriptCore.cpp
+++ b/Engine/Source/Runtime/CoreUObject/Private/UObject/ScriptCore.cpp
@@ -857,12 +857,12 @@ void UObject::CallFunction( FFrame& Stack, RESULT_DECL, UFunction* Function )
 
 	if (Function->FunctionFlags & FUNC_Native)
 	{
-		uint8* Buffer = (uint8*)FMemory_Alloca(Function->ParmsSize);
-		int32 FunctionCallspace = GetFunctionCallspace( Function, Buffer, &Stack );
+		int32 FunctionCallspace = GetFunctionCallspace( Function, &Stack );
 		uint8* SavedCode = NULL;
 		if (FunctionCallspace & FunctionCallspace::Remote)
 		{
 			// Call native networkable function.
+			uint8* Buffer = (uint8*)FMemory_Alloca(Function->ParmsSize);
 
 			SavedCode = Stack.Code; // Since this is native, we need to rollback the stack if we are calling both remotely and locally
 
@@ -1047,7 +1047,7 @@ DEFINE_FUNCTION(UObject::ProcessInternal)
 #endif
 
 	UFunction* Function = (UFunction*)Stack.Node;
-	int32 FunctionCallspace = P_THIS->GetFunctionCallspace(Function, Stack.Locals, NULL);
+	int32 FunctionCallspace = P_THIS->GetFunctionCallspace(Function, NULL);
 	if (FunctionCallspace & FunctionCallspace::Remote)
 	{
 		P_THIS->CallRemoteFunction(Function, Stack.Locals, Stack.OutParms, NULL);
@@ -1339,7 +1339,7 @@ void UObject::ProcessEvent( UFunction* Function, void* Parms )
 
 	if ((Function->FunctionFlags & FUNC_Native) != 0)
 	{
-		int32 FunctionCallspace = GetFunctionCallspace(Function, Parms, NULL);
+		int32 FunctionCallspace = GetFunctionCallspace(Function, NULL);
 		if (FunctionCallspace & FunctionCallspace::Remote)
 		{
 			CallRemoteFunction(Function, Parms, NULL, NULL);
diff --git a/Engine/Source/Runtime/CoreUObject/Public/UObject/Object.h b/Engine/Source/Runtime/CoreUObject/Public/UObject/Object.h
--- a/Engine/Source/Runtime/CoreUObject/Public/UObject/Object.h
+++ b/Engine/Source/Runtime/CoreUObject/Public/UObject/Object.h
@@ -1101,11 +1101,10 @@ public:
 	 * be called locally, remotely, or simply absorbed under the given conditions
 	 *
 	 * @param Function function to call
-	 * @param Parameters arguments to the function call
 	 * @param Stack stack frame for the function call
 	 * @return bitmask representing all callspaces that apply to this UFunction in the given context
 	 */
-	virtual int32 GetFunctionCallspace( UFunction* Function, void* Parameters, FFrame* Stack )
+	virtual int32 GetFunctionCallspace( UFunction* Function, FFrame* Stack )
 	{
 		return FunctionCallspace::Local;
 	}
diff --git a/Engine/Source/Runtime/Engine/Classes/Components/ActorComponent.h b/Engine/Source/Runtime/Engine/Classes/Components/ActorComponent.h
--- a/Engine/Source/Runtime/Engine/Classes/Components/ActorComponent.h
+++ b/Engine/Source/Runtime/Engine/Classes/Components/ActorComponent.h
@@ -784,7 +784,7 @@ public:
 	virtual bool IsNameStableForNetworking() const override;
 	virtual bool IsSupportedForNetworking() const override;
 	virtual void GetLifetimeReplicatedProps(TArray<FLifetimeProperty>& OutLifetimeProps) const override;
-	virtual int32 GetFunctionCallspace( UFunction* Function, void* Parameters, FFrame* Stack ) override;
+	virtual int32 GetFunctionCallspace( UFunction* Function, FFrame* Stack ) override;
 	virtual bool CallRemoteFunction( UFunction* Function, void* Parameters, FOutParmRec* OutParms, FFrame* Stack ) override;
 	virtual void PostInitProperties() override;
 	virtual void PostLoad() override;
diff --git a/Engine/Source/Runtime/Engine/Classes/GameFramework/Actor.h b/Engine/Source/Runtime/Engine/Classes/GameFramework/Actor.h
--- a/Engine/Source/Runtime/Engine/Classes/GameFramework/Actor.h
+++ b/Engine/Source/Runtime/Engine/Classes/GameFramework/Actor.h
@@ -1496,7 +1496,7 @@ public:
 	virtual void PostInitProperties() override;
 	virtual bool Modify( bool bAlwaysMarkDirty=true ) override;
 	virtual void ProcessEvent( UFunction* Function, void* Parameters ) override;
-	virtual int32 GetFunctionCallspace( UFunction* Function, void* Parameters, FFrame* Stack ) override;
+	virtual int32 GetFunctionCallspace( UFunction* Function, FFrame* Stack ) override;
 	virtual bool CallRemoteFunction( UFunction* Function, void* Parameters, FOutParmRec* OutParms, FFrame* Stack ) override;
 	virtual void Serialize(FArchive& Ar) override;
 	virtual void PostLoad() override;
diff --git a/Engine/Source/Runtime/Engine/Classes/Kismet/BlueprintFunctionLibrary.h b/Engine/Source/Runtime/Engine/Classes/Kismet/BlueprintFunctionLibrary.h
--- a/Engine/Source/Runtime/Engine/Classes/Kismet/BlueprintFunctionLibrary.h
+++ b/Engine/Source/Runtime/Engine/Classes/Kismet/BlueprintFunctionLibrary.h
@@ -17,7 +17,7 @@ class UBlueprintFunctionLibrary : public UObject
 	GENERATED_UCLASS_BODY()
 
 	// UObject interface
-	ENGINE_API virtual int32 GetFunctionCallspace(UFunction* Function, void* Parms, FFrame* Stack) override;
+	ENGINE_API virtual int32 GetFunctionCallspace(UFunction* Function, FFrame* Stack) override;
 	// End of UObject interface
 
 protected:
diff --git a/Engine/Source/Runtime/Engine/Private/Actor.cpp b/Engine/Source/Runtime/Engine/Private/Actor.cpp
--- a/Engine/Source/Runtime/Engine/Private/Actor.cpp
+++ b/Engine/Source/Runtime/Engine/Private/Actor.cpp
@@ -3902,7 +3902,7 @@ void AActor::SetNetDriverName(FName NewNetDriverName)
 //
 // Return whether a function should be executed remotely.
 //
-int32 AActor::GetFunctionCallspace( UFunction* Function, void* Parameters, FFrame* Stack )
+int32 AActor::GetFunctionCallspace( UFunction* Function, FFrame* Stack )
 {
 	// Quick reject 1.
 	if ((Function->FunctionFlags & FUNC_Static))
diff --git a/Engine/Source/Runtime/Engine/Private/BlueprintFunctionLibrary.cpp b/Engine/Source/Runtime/Engine/Private/BlueprintFunctionLibrary.cpp
--- a/Engine/Source/Runtime/Engine/Private/BlueprintFunctionLibrary.cpp
+++ b/Engine/Source/Runtime/Engine/Private/BlueprintFunctionLibrary.cpp
@@ -17,7 +17,7 @@ UBlueprintFunctionLibrary::UBlueprintFunctionLibrary(const FObjectInitializer& O
 {
 }
 
-int32 UBlueprintFunctionLibrary::GetFunctionCallspace(UFunction* Function, void* Parameters, FFrame* Stack)
+int32 UBlueprintFunctionLibrary::GetFunctionCallspace(UFunction* Function, FFrame* Stack)
 {
 	const bool bIsAuthoritativeFunc = Function->HasAnyFunctionFlags(FUNC_BlueprintAuthorityOnly);
 	const bool bIsCosmeticFunc      = Function->HasAnyFunctionFlags(FUNC_BlueprintCosmetic);
diff --git a/Engine/Source/Runtime/Engine/Private/Components/ActorComponent.cpp b/Engine/Source/Runtime/Engine/Private/Components/ActorComponent.cpp
--- a/Engine/Source/Runtime/Engine/Private/Components/ActorComponent.cpp
+++ b/Engine/Source/Runtime/Engine/Private/Components/ActorComponent.cpp
@@ -518,10 +518,10 @@ bool UActorComponent::NeedsLoadForEditorGame() const
 	return !IsEditorOnly() && Super::NeedsLoadForEditorGame();
 }
 
-int32 UActorComponent::GetFunctionCallspace( UFunction* Function, void* Parameters, FFrame* Stack )
+int32 UActorComponent::GetFunctionCallspace( UFunction* Function, FFrame* Stack )
 {
 	AActor* MyOwner = GetOwner();
-	return (MyOwner ? MyOwner->GetFunctionCallspace(Function, Parameters, Stack) : FunctionCallspace::Local);
+	return (MyOwner ? MyOwner->GetFunctionCallspace(Function, Stack) : FunctionCallspace::Local);
 }
 
 bool UActorComponent::CallRemoteFunction( UFunction* Function, void* Parameters, FOutParmRec* OutParms, FFrame* Stack )
diff --git a/Engine/Source/Runtime/MovieScene/Private/MovieSceneSequencePlayer.cpp b/Engine/Source/Runtime/MovieScene/Private/MovieSceneSequencePlayer.cpp
--- a/Engine/Source/Runtime/MovieScene/Private/MovieSceneSequencePlayer.cpp
+++ b/Engine/Source/Runtime/MovieScene/Private/MovieSceneSequencePlayer.cpp
@@ -1137,7 +1137,7 @@ void UMovieSceneSequencePlayer::PostNetReceive()
 	}
 }
 
-int32 UMovieSceneSequencePlayer::GetFunctionCallspace(UFunction* Function, void* Parameters, FFrame* Stack)
+int32 UMovieSceneSequencePlayer::GetFunctionCallspace(UFunction* Function, FFrame* Stack)
 {
 	if (HasAnyFlags(RF_ClassDefaultObject))
 	{
@@ -1145,7 +1145,7 @@ int32 UMovieSceneSequencePlayer::GetFunctionCallspace(UFunction* Function, void*
 	}
 
 	check(GetOuter());
-	return GetOuter()->GetFunctionCallspace(Function, Parameters, Stack);
+	return GetOuter()->GetFunctionCallspace(Function, Stack);
 }
 
 bool UMovieSceneSequencePlayer::CallRemoteFunction(UFunction* Function, void* Parameters, FOutParmRec* OutParms, FFrame* Stack)
diff --git a/Engine/Source/Runtime/MovieScene/Public/MovieSceneSequencePlayer.h b/Engine/Source/Runtime/MovieScene/Public/MovieSceneSequencePlayer.h
--- a/Engine/Source/Runtime/MovieScene/Public/MovieSceneSequencePlayer.h
+++ b/Engine/Source/Runtime/MovieScene/Public/MovieSceneSequencePlayer.h
@@ -551,7 +551,7 @@ protected:
 	/*~ Begin UObject interface */
 	virtual void BeginDestroy() override;
 	virtual bool IsSupportedForNetworking() const { return true; }
-	virtual int32 GetFunctionCallspace(UFunction* Function, void* Parameters, FFrame* Stack) override;
+	virtual int32 GetFunctionCallspace(UFunction* Function, FFrame* Stack) override;
 	virtual bool CallRemoteFunction(UFunction* Function, void* Parameters, FOutParmRec* OutParms, FFrame* Stack) override;
 	virtual void PostNetReceive() override;
 	/*~ End UObject interface */
