Moving PRIVATE_PatchNewObjectIntoExport to cpp file for better debugging experience

diff --git a/Engine/Source/Runtime/CoreUObject/Private/UObject/UObjectThreadContext.cpp b/Engine/Source/Runtime/CoreUObject/Private/UObject/UObjectThreadContext.cpp
--- a/Engine/Source/Runtime/CoreUObject/Private/UObject/UObjectThreadContext.cpp
+++ b/Engine/Source/Runtime/CoreUObject/Private/UObject/UObjectThreadContext.cpp
@@ -5,6 +5,7 @@
 =============================================================================*/
 
 #include "UObject/UObjectThreadContext.h"
+#include "UObject/Object.h"
 
 DEFINE_LOG_CATEGORY(LogUObjectThreadContext);
 
@@ -54,4 +55,18 @@ void FUObjectSerializeContext::AddUniqueLoadedObjects(const TArray<UObject*>& In
 		ObjectsLoaded.AddUnique(NewLoadedObject);
 	}
 	
+}
+
+bool FUObjectSerializeContext::PRIVATE_PatchNewObjectIntoExport(UObject* OldObject, UObject* NewObject)
+{
+	const int32 ObjLoadedIdx = ObjectsLoaded.Find(OldObject);
+	if (ObjLoadedIdx != INDEX_NONE)
+	{
+		ObjectsLoaded[ObjLoadedIdx] = NewObject;
+		return true;
+	}
+	else
+	{
+		return false;
+	}
 }
\ No newline at end of file
diff --git a/Engine/Source/Runtime/CoreUObject/Public/UObject/UObjectThreadContext.h b/Engine/Source/Runtime/CoreUObject/Public/UObject/UObjectThreadContext.h
--- a/Engine/Source/Runtime/CoreUObject/Public/UObject/UObjectThreadContext.h
+++ b/Engine/Source/Runtime/CoreUObject/Public/UObject/UObjectThreadContext.h
@@ -178,19 +178,7 @@ public:
 		return !!ObjectsLoaded.Num();
 	}
 
-	bool PRIVATE_PatchNewObjectIntoExport(UObject* OldObject, UObject* NewObject)
-	{
-		const int32 ObjLoadedIdx = ObjectsLoaded.Find(OldObject);
-		if (ObjLoadedIdx != INDEX_NONE)
-		{
-			ObjectsLoaded[ObjLoadedIdx] = NewObject;
-			return true;
-		}
-		else
-		{
-			return false;
-		}
-	}
+	bool PRIVATE_PatchNewObjectIntoExport(UObject* OldObject, UObject* NewObject);
 
 	/** This is only meant to be used by FAsyncPackage for performance reasons. The ObjectsLoaded array should not be manipulated directly! */
 	TArray<UObject*>& PRIVATE_GetObjectsLoadedInternalUseOnly()
