Moving code to reset linker on an object that's being destroyed before the object gets renamed to get better log output when debugging

diff --git a/Engine/Source/Runtime/CoreUObject/Private/UObject/Obj.cpp b/Engine/Source/Runtime/CoreUObject/Private/UObject/Obj.cpp
--- a/Engine/Source/Runtime/CoreUObject/Private/UObject/Obj.cpp
+++ b/Engine/Source/Runtime/CoreUObject/Private/UObject/Obj.cpp
@@ -752,10 +752,10 @@ void UObject::BeginDestroy()
 			);
 	}
 
-	LowLevelRename(NAME_None);
-
 	// Remove from linker's export table.
-	SetLinker( NULL, INDEX_NONE );
+	SetLinker(NULL, INDEX_NONE);
+
+	LowLevelRename(NAME_None);
 
 	// ensure BeginDestroy has been routed back to UObject::BeginDestroy.
 #if !(UE_BUILD_SHIPPING || UE_BUILD_TEST)
