Build fix for MallocProfiler

diff --git a/Engine/Source/Runtime/CoreUObject/Private/UObject/UObjectBaseUtility.cpp b/Engine/Source/Runtime/CoreUObject/Private/UObject/UObjectBaseUtility.cpp
--- a/Engine/Source/Runtime/CoreUObject/Private/UObject/UObjectBaseUtility.cpp
+++ b/Engine/Source/Runtime/CoreUObject/Private/UObject/UObjectBaseUtility.cpp
@@ -438,7 +438,9 @@ void FScopeCycleCounterUObject::TrackObjectForMallocProfiling(const FName InPack
 		const FNameEntry* NameEntry = InName.GetDisplayNameEntry();
 		if (NameEntry->IsWide())
 		{
-			const WIDECHAR* NameCharPtr = NameEntry->GetWideName();
+			WIDECHAR WideName[NAME_SIZE];
+			NameEntry->GetWideName(WideName);
+			const WIDECHAR* NameCharPtr = WideName;
 			while (*NameCharPtr != 0)
 			{
 				ScratchSpaceBuffer.Add((TCHAR)*NameCharPtr++);
@@ -446,7 +448,9 @@ void FScopeCycleCounterUObject::TrackObjectForMallocProfiling(const FName InPack
 		}
 		else
 		{
-			const ANSICHAR* NameCharPtr = NameEntry->GetAnsiName();
+			ANSICHAR AnsiName[NAME_SIZE];
+			NameEntry->GetAnsiName(AnsiName);
+			const ANSICHAR* NameCharPtr = AnsiName;
 			while (*NameCharPtr != 0)
 			{
 				ScratchSpaceBuffer.Add((TCHAR)*NameCharPtr++);
