Fixes for commas causing bad documention generation.

diff --git a/Engine/Source/Runtime/Core/Public/UObject/WeakObjectPtrTemplates.h b/Engine/Source/Runtime/Core/Public/UObject/WeakObjectPtrTemplates.h
--- a/Engine/Source/Runtime/Core/Public/UObject/WeakObjectPtrTemplates.h
+++ b/Engine/Source/Runtime/Core/Public/UObject/WeakObjectPtrTemplates.h
@@ -144,7 +144,7 @@ public:
 
 	/**  
 	 * Dereference the weak pointer
-	 * @param bEvenIfPendingKill, if this is true, pendingkill objects are considered valid	
+	 * @param bEvenIfPendingKill if this is true, pendingkill objects are considered valid
 	 * @return NULL if this object is gone or the weak pointer was NULL, otherwise a valid uobject pointer
 	**/
 	FORCEINLINE T* Get(bool bEvenIfPendingKill) const
@@ -184,8 +184,8 @@ public:
 
 	/**  
 	 * Test if this points to a live UObject
-	 * @param bEvenIfPendingKill, if this is true, pendingkill objects are considered valid
-	 * @param bThreadsafeTest, if true then function will just give you information whether referenced 
+	 * @param bEvenIfPendingKill if this is true, pendingkill objects are considered valid
+	 * @param bThreadsafeTest if true then function will just give you information whether referenced
 	 *							UObject is gone forever (@return false) or if it is still there (@return true, no object flags checked).
 	 * @return true if Get() would return a valid non-null pointer
 	**/
@@ -205,8 +205,8 @@ public:
 
 	/**  
 	 * Slightly different than !IsValid(), returns true if this used to point to a UObject, but doesn't any more and has not been assigned or reset in the mean time.
-	 * @param bIncludingIfPendingKill, if this is true, pendingkill objects are considered stale
-	 * @param bThreadsafeTest, set it to true when testing outside of Game Thread. Results in false if WeakObjPtr point to an existing object (no flags checked) 
+	 * @param bIncludingIfPendingKill if this is true, pendingkill objects are considered stale
+	 * @param bThreadsafeTest set it to true when testing outside of Game Thread. Results in false if WeakObjPtr point to an existing object (no flags checked)
 	 * @return true if this used to point at a real object but no longer does.
 	**/
 	FORCEINLINE bool IsStale(bool bIncludingIfPendingKill = true, bool bThreadsafeTest = false) const
diff --git a/Engine/Source/Runtime/CoreUObject/Public/UObject/WeakObjectPtr.h b/Engine/Source/Runtime/CoreUObject/Public/UObject/WeakObjectPtr.h
--- a/Engine/Source/Runtime/CoreUObject/Public/UObject/WeakObjectPtr.h
+++ b/Engine/Source/Runtime/CoreUObject/Public/UObject/WeakObjectPtr.h
@@ -100,7 +100,7 @@ public:
 
 	/**  
 	 * Dereference the weak pointer.
-	 * @param bEvenIfPendingKill, if this is true, pendingkill objects are considered valid
+	 * @param bEvenIfPendingKill if this is true, pendingkill objects are considered valid
 	 * @return NULL if this object is gone or the weak pointer was NULL, otherwise a UObject pointer
 	 */
 	COREUOBJECT_API class UObject *Get(bool bEvenIfPendingKill) const;
@@ -116,8 +116,8 @@ public:
 
 	/**  
 	 * Test if this points to a live UObject
-	 * @param bEvenIfPendingKill, if this is true, pendingkill are not considered invalid
-	 * @param bThreadsafeTest, if true then function will just give you information whether referenced 
+	 * @param bEvenIfPendingKill if this is true, pendingkill are not considered invalid
+	 * @param bThreadsafeTest if true then function will just give you information whether referenced
 	 *							UObject is gone forever (@return false) or if it is still there (@return true, no object flags checked).
 	 * @return true if Get() would return a valid non-null pointer, if bThreadsafeTest == true then @see @param bThreadsafeTest
 	**/
@@ -131,8 +131,8 @@ public:
 
 	/**  
 	 * Slightly different than !IsValid(), returns true if this used to point to a UObject, but doesn't any more and has not been assigned or reset in the mean time.
-	 * @param bIncludingIfPendingKill, if this is false, pendingkill objects are not considered stale
-	 * @param bThreadsafeTest, set it to true when testing outside of Game Thread. Results in false if WeakObjPtr point to an existing object (no flags checked)
+	 * @param bIncludingIfPendingKill if this is false, pendingkill objects are not considered stale
+	 * @param bThreadsafeTest set it to true when testing outside of Game Thread. Results in false if WeakObjPtr point to an existing object (no flags checked)
 	 * @return true if this used to point at a real object but no longer does.
 	**/
 	COREUOBJECT_API bool IsStale(bool bIncludingIfPendingKill = true, bool bThreadsafeTest = false) const;
