Fix for native UStruct and UClass alignments.

diff --git a/Engine/Source/Runtime/CoreUObject/Private/UObject/Class.cpp b/Engine/Source/Runtime/CoreUObject/Private/UObject/Class.cpp
--- a/Engine/Source/Runtime/CoreUObject/Private/UObject/Class.cpp
+++ b/Engine/Source/Runtime/CoreUObject/Private/UObject/Class.cpp
@@ -498,12 +498,12 @@ IMPLEMENT_CORE_INTRINSIC_CLASS(UField, UObject,
 //
 // Constructors.
 //
-UStruct::UStruct( EStaticConstructor, int32 InSize, EObjectFlags InFlags )
+UStruct::UStruct( EStaticConstructor, int32 InSize, int32 InMinAlignment, EObjectFlags InFlags )
 :	UField			( EC_StaticConstructor, InFlags )
 ,	SuperStruct		( nullptr )
 ,	Children		( NULL )
 ,	PropertiesSize	( InSize )
-,	MinAlignment	( 1 )
+,	MinAlignment	( InMinAlignment )
 ,	PropertyLink	( NULL )
 ,	RefLink			( NULL )
 ,	DestructorLink	( NULL )
@@ -1900,8 +1900,8 @@ bool FindConstructorUninitialized(UStruct* BaseClass,uint8* Data,uint8* Defaults
 #endif
 
 
-UScriptStruct::UScriptStruct( EStaticConstructor, int32 InSize, EObjectFlags InFlags )
-	: UStruct( EC_StaticConstructor, InSize, InFlags )
+UScriptStruct::UScriptStruct( EStaticConstructor, int32 InSize, int32 InAlignment, EObjectFlags InFlags )
+	: UStruct( EC_StaticConstructor, InSize, InAlignment, InFlags )
 	, StructFlags(STRUCT_NoFlags)
 #if HACK_HEADER_GENERATOR
 	, StructMacroDeclaredLineNumber(INDEX_NONE)
@@ -4090,6 +4090,7 @@ UClass::UClass
 	EStaticConstructor,
 	FName			InName,
 	uint32			InSize,
+	uint32			InAlignment,
 	EClassFlags		InClassFlags,
 	EClassCastFlags	InClassCastFlags,
 	const TCHAR*    InConfigName,
@@ -4098,7 +4099,7 @@ UClass::UClass
 	ClassVTableHelperCtorCallerType InClassVTableHelperCtorCaller,
 	ClassAddReferencedObjectsType InClassAddReferencedObjects
 )
-:	UStruct					( EC_StaticConstructor, InSize, InFlags )
+:	UStruct					( EC_StaticConstructor, InSize, InAlignment, InFlags )
 ,	ClassConstructor		( InClassConstructor )
 ,	ClassVTableHelperCtorCaller(InClassVTableHelperCtorCaller)
 ,	ClassAddReferencedObjects( InClassAddReferencedObjects )
@@ -4552,6 +4553,7 @@ void GetPrivateStaticClassBody(
 	UClass*& ReturnClass,
 	void(*RegisterNativeFunc)(),
 	uint32 InSize,
+	uint32 InAlignment,
 	EClassFlags InClassFlags,
 	EClassCastFlags InClassCastFlags,
 	const TCHAR* InConfigName,
@@ -4611,6 +4613,7 @@ void GetPrivateStaticClassBody(
 			EC_StaticConstructor,
 			Name,
 			InSize,
+			InAlignment,
 			InClassFlags,
 			InClassCastFlags,
 			InConfigName,
@@ -4630,6 +4633,7 @@ void GetPrivateStaticClassBody(
 			EC_StaticConstructor,
 			Name,
 			InSize,
+			InAlignment,
 			InClassFlags|CLASS_CompiledFromBlueprint,
 			InClassCastFlags,
 			InConfigName,
@@ -5191,6 +5195,7 @@ UDynamicClass::UDynamicClass(
 	EStaticConstructor,
 	FName			InName,
 	uint32			InSize,
+	uint32			InAlignment,
 	EClassFlags		InClassFlags,
 	EClassCastFlags	InClassCastFlags,
 	const TCHAR*    InConfigName,
@@ -5202,6 +5207,7 @@ UDynamicClass::UDynamicClass(
   EC_StaticConstructor
 , InName
 , InSize
+, InAlignment
 , InClassFlags
 , InClassCastFlags
 , InConfigName
diff --git a/Engine/Source/Runtime/CoreUObject/Public/UObject/Class.h b/Engine/Source/Runtime/CoreUObject/Public/UObject/Class.h
--- a/Engine/Source/Runtime/CoreUObject/Public/UObject/Class.h
+++ b/Engine/Source/Runtime/CoreUObject/Public/UObject/Class.h
@@ -284,7 +284,7 @@ public:
 
 public:
 	// Constructors.
-	UStruct( EStaticConstructor, int32 InSize, EObjectFlags InFlags );
+	UStruct( EStaticConstructor, int32 InSize, int32 InAlignment, EObjectFlags InFlags );
 	explicit UStruct(UStruct* InSuperStruct, SIZE_T ParamsSize = 0, SIZE_T Alignment = 0);
 	explicit UStruct(const FObjectInitializer& ObjectInitializer, UStruct* InSuperStruct, SIZE_T ParamsSize = 0, SIZE_T Alignment = 0 );
 
@@ -1189,7 +1189,7 @@ public:
 
 	DECLARE_CASTED_CLASS_INTRINSIC_NO_CTOR(UScriptStruct, UStruct, CLASS_MatchedSerializers, TEXT("/Script/CoreUObject"), CASTCLASS_UScriptStruct, COREUOBJECT_API)
 
-	COREUOBJECT_API UScriptStruct( EStaticConstructor, int32 InSize, EObjectFlags InFlags );
+	COREUOBJECT_API UScriptStruct( EStaticConstructor, int32 InSize, int32 InAlignment, EObjectFlags InFlags );
 	COREUOBJECT_API explicit UScriptStruct(const FObjectInitializer& ObjectInitializer, UScriptStruct* InSuperStruct, ICppStructOps* InCppStructOps = nullptr, EStructFlags InStructFlags = STRUCT_NoFlags, SIZE_T ExplicitSize = 0, SIZE_T ExplicitAlignment = 0);
 	COREUOBJECT_API explicit UScriptStruct(const FObjectInitializer& ObjectInitializer = FObjectInitializer::Get());
 
@@ -2316,7 +2316,7 @@ public:
 	// Constructors
 	UClass(const FObjectInitializer& ObjectInitializer = FObjectInitializer::Get());
 	explicit UClass(const FObjectInitializer& ObjectInitializer, UClass* InSuperClass);
-	UClass( EStaticConstructor, FName InName, uint32 InSize, EClassFlags InClassFlags, EClassCastFlags InClassCastFlags,
+	UClass( EStaticConstructor, FName InName, uint32 InSize, uint32 InAlignment, EClassFlags InClassFlags, EClassCastFlags InClassCastFlags,
 		const TCHAR* InClassConfigName, EObjectFlags InFlags, ClassConstructorType InClassConstructor,
 		ClassVTableHelperCtorCallerType InClassVTableHelperCtorCaller,
 		ClassAddReferencedObjectsType InClassAddReferencedObjects);
@@ -2807,7 +2807,7 @@ public:
 
 	UDynamicClass(const FObjectInitializer& ObjectInitializer = FObjectInitializer::Get());
 	explicit UDynamicClass(const FObjectInitializer& ObjectInitializer, UClass* InSuperClass);
-	UDynamicClass(EStaticConstructor, FName InName, uint32 InSize, EClassFlags InClassFlags, EClassCastFlags InClassCastFlags,
+	UDynamicClass(EStaticConstructor, FName InName, uint32 InSize, uint32 InAlignment, EClassFlags InClassFlags, EClassCastFlags InClassCastFlags,
 		const TCHAR* InClassConfigName, EObjectFlags InFlags, ClassConstructorType InClassConstructor,
 		ClassVTableHelperCtorCallerType InClassVTableHelperCtorCaller,
 		ClassAddReferencedObjectsType InClassAddReferencedObjects);
@@ -2876,6 +2876,7 @@ COREUOBJECT_API void InitializePrivateStaticClass(
  * @param ReturnClass reference to pointer to result. This must be PrivateStaticClass.
  * @param RegisterNativeFunc Native function registration function pointer.
  * @param InSize Size of the class
+ * @param InAlignment Alignment of the class
  * @param InClassFlags Class flags
  * @param InClassCastFlags Class cast flags
  * @param InConfigName Class config name
@@ -2892,6 +2893,7 @@ COREUOBJECT_API void GetPrivateStaticClassBody(
 	UClass*& ReturnClass,
 	void(*RegisterNativeFunc)(),
 	uint32 InSize,
+	uint32 InAlignment,
 	EClassFlags InClassFlags,
 	EClassCastFlags InClassCastFlags,
 	const TCHAR* InConfigName,
diff --git a/Engine/Source/Runtime/CoreUObject/Public/UObject/ObjectMacros.h b/Engine/Source/Runtime/CoreUObject/Public/UObject/ObjectMacros.h
--- a/Engine/Source/Runtime/CoreUObject/Public/UObject/ObjectMacros.h
+++ b/Engine/Source/Runtime/CoreUObject/Public/UObject/ObjectMacros.h
@@ -1570,6 +1570,7 @@ public: \
 				PrivateStaticClass, \
 				StaticRegisterNatives##TClass, \
 				sizeof(TClass), \
+				alignof(TClass), \
 				(EClassFlags)TClass::StaticClassFlags, \
 				TClass::StaticClassCastFlags(), \
 				TClass::StaticConfigName(), \
@@ -1638,6 +1639,7 @@ public: \
 			PrivateStaticClass, \
 			StaticRegisterNatives##TClass, \
 			sizeof(TClass), \
+			alignof(TClass), \
 			(EClassFlags)TClass::StaticClassFlags, \
 			TClass::StaticClassCastFlags(), \
 			TClass::StaticConfigName(), \
