Optimize cook processing of soft object paths.
Replace multimap in RedirectCollector with a map and set with package as key.
Use new package key to find data instead of iterating the complete multi map every time.

Cook -iterate Test Scenario:
~ 8600 startup packages
~ 2700 referenced startup soft object packages
=> ProcessSoftObjectPathPackageList called from StartCookByTheBook down from ~5 seconds to ~0.02 seconds (~99% reduction)

Full Cook Test Scenario:
~ 183 000 saved packages
=> ProcessSoftObjectPathPackageList called from StartCookByTheBook down from ~8 seconds to ~0.02 seconds (VTune user-mode sampling cpu time)
=> ProcessSoftObjectPathPackageList called from SaveCookedPackage down from ~67 seconds to ~0.02 seconds (VTune user-mode sampling cpu time)

diff --git a/Engine/Source/Runtime/CoreUObject/Private/Misc/RedirectCollector.cpp b/Engine/Source/Runtime/CoreUObject/Private/Misc/RedirectCollector.cpp
--- a/Engine/Source/Runtime/CoreUObject/Private/Misc/RedirectCollector.cpp
+++ b/Engine/Source/Runtime/CoreUObject/Private/Misc/RedirectCollector.cpp
@@ -25,7 +25,6 @@ void FRedirectCollector::OnSoftObjectPathLoaded(const FSoftObjectPath& InPath, F
 		return;
 	}
 
-	FPackagePropertyPair ContainingPackageAndProperty;
 	FSoftObjectPathThreadContext& ThreadContext = FSoftObjectPathThreadContext::Get();
 
 	FName PackageName, PropertyName;
@@ -40,20 +39,12 @@ void FRedirectCollector::OnSoftObjectPathLoaded(const FSoftObjectPath& InPath, F
 		return;
 	}
 
-	if (PackageName != NAME_None)
-	{
-		ContainingPackageAndProperty.SetPackage(PackageName);
-		if (PropertyName != NAME_None)
-		{
-			ContainingPackageAndProperty.SetProperty(PropertyName);
-		}	
-	}
-
-	ContainingPackageAndProperty.SetReferencedByEditorOnlyProperty(CollectType == ESoftObjectPathCollectType::EditorOnlyCollect);
+	const bool bReferencedByEditorOnlyProperty = (CollectType == ESoftObjectPathCollectType::EditorOnlyCollect);
+	FSoftObjectPathProperty SoftObjectPathProperty(InPath.GetAssetPathName(), PropertyName, bReferencedByEditorOnlyProperty);
 
 	FScopeLock ScopeLock(&CriticalSection);
 
-	SoftObjectPathMap.AddUnique(InPath.GetAssetPathName(), ContainingPackageAndProperty);
+	SoftObjectPathMap.FindOrAdd(PackageName).Add(SoftObjectPathProperty);
 }
 
 void FRedirectCollector::OnStringAssetReferenceLoaded(const FString& InString)
@@ -77,69 +68,79 @@ FString FRedirectCollector::OnStringAssetReferenceSaved(const FString& InString)
 
 void FRedirectCollector::ResolveAllSoftObjectPaths(FName FilterPackage)
 {	
+	auto LoadSoftObjectPathLambda = [this](const FSoftObjectPathProperty& SoftObjectPathProperty)
+	{
+		const FName& ToLoadFName = SoftObjectPathProperty.GetAssetPathName();
+		const FString ToLoad = ToLoadFName.ToString();
+
+		if (ToLoad.Len() > 0 )
+		{
+			UE_LOG(LogRedirectors, Verbose, TEXT("Resolving Soft Object Path '%s'"), *ToLoad);
+			UE_CLOG(SoftObjectPathProperty.GetPropertyName().ToString().Len(), LogRedirectors, Verbose, TEXT("    Referenced by '%s'"), *SoftObjectPathProperty.GetPropertyName().ToString());
+
+			int32 DotIndex = ToLoad.Find(TEXT("."));
+			FString PackageName = DotIndex != INDEX_NONE ? ToLoad.Left(DotIndex) : ToLoad;
+
+			// If is known missing don't try
+			if (FLinkerLoad::IsKnownMissingPackage(FName(*PackageName)))
+			{
+				return;
+			}
+
+			UObject *Loaded = LoadObject<UObject>(NULL, *ToLoad, NULL, SoftObjectPathProperty.GetReferencedByEditorOnlyProperty() ? LOAD_EditorOnly | LOAD_NoWarn : LOAD_NoWarn, NULL);
+
+			if (Loaded)
+			{
+				FString Dest = Loaded->GetPathName();
+				UE_LOG(LogRedirectors, Verbose, TEXT("    Resolved to '%s'"), *Dest);
+				if (Dest != ToLoad)
+				{
+					AssetPathRedirectionMap.Add(ToLoadFName, FName(*Dest));
+				}
+			}
+			else
+			{
+				const FString Referencer = SoftObjectPathProperty.GetPropertyName().ToString().Len() ? SoftObjectPathProperty.GetPropertyName().ToString() : TEXT("Unknown");
+				UE_LOG(LogRedirectors, Warning, TEXT("Soft Object Path '%s' was not found when resolving paths! (Referencer '%s')"), *ToLoad, *Referencer);
+			}
+		}
+	};
+
 	FScopeLock ScopeLock(&CriticalSection);
 
-	TMultiMap<FName, FPackagePropertyPair> SkippedReferences;
-	SkippedReferences.Empty(SoftObjectPathMap.Num());
+	FSoftObjectPathMap KeepSoftObjectPathMap;
+	KeepSoftObjectPathMap.Reserve(SoftObjectPathMap.Num());
 	while (SoftObjectPathMap.Num())
 	{
-		TMultiMap<FName, FPackagePropertyPair> CurrentReferences;
-		Swap(SoftObjectPathMap, CurrentReferences);
+		FSoftObjectPathMap LocalSoftObjectPathMap;
+		Swap(SoftObjectPathMap, LocalSoftObjectPathMap);
 
-		for (const TPair<FName, FPackagePropertyPair>& CurrentReference : CurrentReferences)
+		for (TPair<FName, FSoftObjectPathPropertySet>& CurrentPackage : LocalSoftObjectPathMap)
 		{
-			const FName& ToLoadFName = CurrentReference.Key;
-			const FPackagePropertyPair& RefFilenameAndProperty = CurrentReference.Value;
+			const FName& CurrentPackageName = CurrentPackage.Key;
+			FSoftObjectPathPropertySet& SoftObjectPathProperties = CurrentPackage.Value;
 
 			if ((FilterPackage != NAME_None) && // not using a filter
-				(FilterPackage != RefFilenameAndProperty.GetCachedPackageName()) && // this is the package we are looking for
-				(RefFilenameAndProperty.GetCachedPackageName() != NAME_None) // if we have an empty package name then process it straight away
+				(FilterPackage != CurrentPackageName) && // this is the package we are looking for
+				(CurrentPackageName != NAME_None) // if we have an empty package name then process it straight away
 				)
 			{
-				// If we have a valid filter and it doesn't match, skip this reference
-				SkippedReferences.Add(ToLoadFName, RefFilenameAndProperty);
+				// If we have a valid filter and it doesn't match, skip processing of this package and keep it
+				KeepSoftObjectPathMap.FindOrAdd(CurrentPackageName).Append(MoveTemp(SoftObjectPathProperties));
 				continue;
 			}
 
-			const FString ToLoad = ToLoadFName.ToString();
-
-			if (ToLoad.Len() > 0 )
+			// This will call LoadObject which may trigger OnSoftObjectPathLoaded and add new soft object paths to the SoftObjectPathMap
+			for (const FSoftObjectPathProperty& SoftObjecPathProperty : SoftObjectPathProperties)
 			{
-				UE_LOG(LogRedirectors, Verbose, TEXT("Resolving Soft Object Path '%s'"), *ToLoad);
-				UE_CLOG(RefFilenameAndProperty.GetProperty().ToString().Len(), LogRedirectors, Verbose, TEXT("    Referenced by '%s'"), *RefFilenameAndProperty.GetProperty().ToString());
-
-				int32 DotIndex = ToLoad.Find(TEXT("."));
-				FString PackageName = DotIndex != INDEX_NONE ? ToLoad.Left(DotIndex) : ToLoad;
-
-				// If is known missing don't try
-				if (FLinkerLoad::IsKnownMissingPackage(FName(*PackageName)))
-				{
-					continue;
-				}
-
-				UObject *Loaded = LoadObject<UObject>(NULL, *ToLoad, NULL, RefFilenameAndProperty.GetReferencedByEditorOnlyProperty() ? LOAD_EditorOnly | LOAD_NoWarn : LOAD_NoWarn, NULL);
-
-				if (Loaded)
-				{
-					FString Dest = Loaded->GetPathName();
-					UE_LOG(LogRedirectors, Verbose, TEXT("    Resolved to '%s'"), *Dest);
-					if (Dest != ToLoad)
-					{
-						AssetPathRedirectionMap.Add(ToLoadFName, FName(*Dest));
-					}
-				}
-				else
-				{
-					const FString Referencer = RefFilenameAndProperty.GetProperty().ToString().Len() ? RefFilenameAndProperty.GetProperty().ToString() : TEXT("Unknown");
-					UE_LOG(LogRedirectors, Warning, TEXT("Soft Object Path '%s' was not found when resolving paths! (Referencer '%s')"), *ToLoad, *Referencer);
-				}
+				LoadSoftObjectPathLambda(SoftObjecPathProperty);
 			}
 		}
 	}
 
 	check(SoftObjectPathMap.Num() == 0);
-	// Add any skipped references back into the map for the next time this is called
-	Swap(SoftObjectPathMap, SkippedReferences);
+	// Add any non processed packages back into the global map for the next time this is called
+	Swap(SoftObjectPathMap, KeepSoftObjectPathMap);
 	// we shouldn't have any references left if we decided to resolve them all
 	check((SoftObjectPathMap.Num() == 0) || (FilterPackage != NAME_None));
 }
@@ -148,24 +149,26 @@ void FRedirectCollector::ProcessSoftObjectPathPackageList(FName FilterPackage, b
 {
 	FScopeLock ScopeLock(&CriticalSection);
 
-	// Iterate map, remove all matching and potentially add to OutReferencedPackages
-	for (auto It = SoftObjectPathMap.CreateIterator(); It; ++It)
+	TSet<FSoftObjectPathProperty>* SoftObjectPathProperties = SoftObjectPathMap.Find(FilterPackage);
+	if (!SoftObjectPathProperties)
 	{
-		const FName& ToLoadFName = It->Key;
-		const FPackagePropertyPair& RefFilenameAndProperty = It->Value;
+		return;
+	}
 
-		// Package name may be null, if so return the set of things not associated with a package
-		if (FilterPackage == RefFilenameAndProperty.GetCachedPackageName())
+	// potentially add soft object path package names to OutReferencedPackages
+	OutReferencedPackages.Reserve(SoftObjectPathProperties->Num());
+	for (const FSoftObjectPathProperty& SoftObjecPathProperty : *SoftObjectPathProperties)
+	{
+		if (!SoftObjecPathProperty.GetReferencedByEditorOnlyProperty() || bGetEditorOnly)
 		{
-			if (!RefFilenameAndProperty.GetReferencedByEditorOnlyProperty() || bGetEditorOnly)
-			{
-				FString PackageNameString = FPackageName::ObjectPathToPackageName(ToLoadFName.ToString());
-				OutReferencedPackages.Add(FName(*PackageNameString));
-			}
-
-			It.RemoveCurrent();
+			const FName& ToLoadFName = SoftObjecPathProperty.GetAssetPathName();
+			FString PackageNameString = FPackageName::ObjectPathToPackageName(ToLoadFName.ToString());
+			OutReferencedPackages.Add(FName(*PackageNameString));
 		}
 	}
+
+	// always remove all data for the processed FilterPackage
+	SoftObjectPathMap.Remove(FilterPackage);
 }
 
 void FRedirectCollector::AddAssetPathRedirection(FName OriginalPath, FName RedirectedPath)
diff --git a/Engine/Source/Runtime/CoreUObject/Public/Misc/RedirectCollector.h b/Engine/Source/Runtime/CoreUObject/Public/Misc/RedirectCollector.h
--- a/Engine/Source/Runtime/CoreUObject/Public/Misc/RedirectCollector.h
+++ b/Engine/Source/Runtime/CoreUObject/Public/Misc/RedirectCollector.h
@@ -16,49 +16,38 @@ class COREUOBJECT_API FRedirectCollector
 private:
 	
 	/** Helper struct for soft object path tracking */
-	struct FPackagePropertyPair
+	struct FSoftObjectPathProperty
 	{
-		FPackagePropertyPair() : bReferencedByEditorOnlyProperty(false) {}
-
-		FPackagePropertyPair(const FName& InPackage, const FName& InProperty)
-		: Package(InPackage)
-		, Property(InProperty)
-		, bReferencedByEditorOnlyProperty(false)
+		FSoftObjectPathProperty(FName InAssetPathName, FName InProperty, bool bInReferencedByEditorOnlyProperty)
+			: AssetPathName(InAssetPathName)
+			, PropertyName(InProperty)
+			, bReferencedByEditorOnlyProperty(bInReferencedByEditorOnlyProperty)
 		{}
 
-		bool operator==(const FPackagePropertyPair& Other) const
-		{
-			return Package == Other.Package &&
-				Property == Other.Property &&
-				bReferencedByEditorOnlyProperty == Other.bReferencedByEditorOnlyProperty;
-		}
+		 bool operator==(const FSoftObjectPathProperty& Other) const
+		 {
+		 	return AssetPathName == Other.AssetPathName &&
+		 		PropertyName == Other.PropertyName &&
+		 		bReferencedByEditorOnlyProperty == Other.bReferencedByEditorOnlyProperty;
+		 }
 
-		const FName& GetCachedPackageName() const
+		friend inline uint32 GetTypeHash(const FSoftObjectPathProperty& Key)
 		{
-			return Package;
+			uint32 Hash = 0;
+			Hash = HashCombine(Hash, GetTypeHash(Key.AssetPathName));
+			Hash = HashCombine(Hash, GetTypeHash(Key.PropertyName));
+			Hash = HashCombine(Hash, (uint32)Key.bReferencedByEditorOnlyProperty);
+			return Hash;
 		}
 
-		void SetPackage(const FName& InPackage)
-		{
-			Package = InPackage;
-		}
-		void SetProperty(const FName& InProperty)
+		const FName& GetAssetPathName() const
 		{
-			Property = InProperty;
-		}
-		
-		void SetReferencedByEditorOnlyProperty(bool InReferencedByEditorOnlyProperty)
-		{
-			bReferencedByEditorOnlyProperty = InReferencedByEditorOnlyProperty;
+			return AssetPathName;
 		}
 
-		const FName& GetPackage() const
-		{
-			return Package;
-		}
-		const FName& GetProperty() const
+		const FName& GetPropertyName() const
 		{
-			return Property;
+			return PropertyName;
 		}
 
 		bool GetReferencedByEditorOnlyProperty() const
@@ -67,8 +56,8 @@ private:
 		}
 
 	private:
-		FName Package;
-		FName Property;
+		FName AssetPathName;
+		FName PropertyName;
 		bool bReferencedByEditorOnlyProperty;
 	};
 
@@ -127,8 +116,10 @@ public:
 
 private:
 
-	/** A map of assets referenced by soft object paths, with the key being the asset being referenced and the value equal to the package with the reference */
-	TMultiMap<FName, FPackagePropertyPair> SoftObjectPathMap;
+	/** A map of assets referenced by soft object paths, with the key being the package with the reference */
+	typedef TSet<FSoftObjectPathProperty> FSoftObjectPathPropertySet;
+	typedef TMap<FName, FSoftObjectPathPropertySet> FSoftObjectPathMap;
+	FSoftObjectPathMap SoftObjectPathMap;
 
 	/** When saving, apply this remapping to all soft object paths */
 	TMap<FName, FName> AssetPathRedirectionMap;
@@ -140,4 +131,4 @@ private:
 // global redirect collector callback structure
 COREUOBJECT_API extern FRedirectCollector GRedirectCollector;
 
-#endif // WITH_EDITOR
\ No newline at end of file
+#endif // WITH_EDITOR
