Made sure all UObject delete and create listeners are removed from the global UObject before static exit to avoid crashes caused by removing listeners after the global UObject array has already been destroyed.

diff --git a/Engine/Source/Editor/UnrealEd/Private/CookOnTheFlyServer.cpp b/Engine/Source/Editor/UnrealEd/Private/CookOnTheFlyServer.cpp
--- a/Engine/Source/Editor/UnrealEd/Private/CookOnTheFlyServer.cpp
+++ b/Engine/Source/Editor/UnrealEd/Private/CookOnTheFlyServer.cpp
@@ -1249,6 +1249,12 @@ struct FPackageTracker : public FUObjectArray::FUObjectCreateListener, public FU
 		}
 	}
 
+	virtual void OnUObjectArrayShutdown() override
+	{
+		GUObjectArray.RemoveUObjectDeleteListener(this);
+		GUObjectArray.RemoveUObjectCreateListener(this);
+	}
+
 	// This is the set of packages which have already had PostLoadFixup called 
 	TSet<UPackage*>			PostLoadFixupPackages;
 
diff --git a/Engine/Source/Runtime/CoreUObject/Private/UObject/SparseDelegate.cpp b/Engine/Source/Runtime/CoreUObject/Private/UObject/SparseDelegate.cpp
--- a/Engine/Source/Runtime/CoreUObject/Private/UObject/SparseDelegate.cpp
+++ b/Engine/Source/Runtime/CoreUObject/Private/UObject/SparseDelegate.cpp
@@ -29,6 +29,13 @@ void FSparseDelegateStorage::FObjectListener::NotifyUObjectDeleted(const UObject
 	}
 }
 
+void FSparseDelegateStorage::FObjectListener::OnUObjectArrayShutdown()
+{
+	FScopeLock SparseDelegateMapLock(&FSparseDelegateStorage::SparseDelegateMapCritical);
+	FSparseDelegateStorage::SparseDelegates.Empty();
+	DisableListener();
+}
+
 void FSparseDelegateStorage::FObjectListener::EnableListener()
 {
 	GUObjectArray.AddUObjectDeleteListener(this);
diff --git a/Engine/Source/Runtime/CoreUObject/Private/UObject/UObjectArray.cpp b/Engine/Source/Runtime/CoreUObject/Private/UObject/UObjectArray.cpp
--- a/Engine/Source/Runtime/CoreUObject/Private/UObject/UObjectArray.cpp
+++ b/Engine/Source/Runtime/CoreUObject/Private/UObject/UObjectArray.cpp
@@ -335,4 +335,23 @@ int32 FUObjectArray::AllocateSerialNumber(int32 Index)
  */
 void FUObjectArray::ShutdownUObjectArray()
 {
+	{
+#if THREADSAFE_UOBJECTS
+		FScopeLock UObjectDeleteListenersLock(&UObjectDeleteListenersCritical);
+#endif
+		for (int32 Index = UObjectDeleteListeners.Num() - 1; Index >= 0; --Index)
+		{
+			FUObjectDeleteListener* Listener = UObjectDeleteListeners[Index];
+			Listener->OnUObjectArrayShutdown();
+		}
+		UE_CLOG(UObjectDeleteListeners.Num(), LogUObjectArray, Fatal, TEXT("All UObject delete listeners should be unregistered when shutting down the UObject array"));
+	}
+	{
+		for (int32 Index = UObjectCreateListeners.Num() - 1; Index >= 0; --Index)
+		{
+			FUObjectCreateListener* Listener = UObjectCreateListeners[Index];
+			Listener->OnUObjectArrayShutdown();
+		}
+		UE_CLOG(UObjectCreateListeners.Num(), LogUObjectArray, Fatal, TEXT("All UObject delete listeners should be unregistered when shutting down the UObject array"));
+	}
 }
diff --git a/Engine/Source/Runtime/CoreUObject/Private/UObject/UObjectMarks.cpp b/Engine/Source/Runtime/CoreUObject/Private/UObject/UObjectMarks.cpp
--- a/Engine/Source/Runtime/CoreUObject/Private/UObject/UObjectMarks.cpp
+++ b/Engine/Source/Runtime/CoreUObject/Private/UObject/UObjectMarks.cpp
@@ -38,6 +38,12 @@ public:
 		RemoveAnnotation(Object);
 	}
 
+	virtual void OnUObjectArrayShutdown() override
+	{
+		RemoveAllAnnotations();
+		GUObjectArray.RemoveUObjectDeleteListener(this);
+	}
+
 	FUObjectAnnotationSparseNoSync() :
 		AnnotationCacheKey(NULL)
 	{
diff --git a/Engine/Source/Runtime/CoreUObject/Public/UObject/SparseDelegate.h b/Engine/Source/Runtime/CoreUObject/Public/UObject/SparseDelegate.h
--- a/Engine/Source/Runtime/CoreUObject/Public/UObject/SparseDelegate.h
+++ b/Engine/Source/Runtime/CoreUObject/Public/UObject/SparseDelegate.h
@@ -70,6 +70,7 @@ private:
 	{
 		virtual ~FObjectListener();
 		virtual void NotifyUObjectDeleted(const UObjectBase* Object, int32 Index) override;
+		virtual void OnUObjectArrayShutdown();
 		void EnableListener();
 		void DisableListener();
 	};
diff --git a/Engine/Source/Runtime/CoreUObject/Public/UObject/UObjectAnnotation.h b/Engine/Source/Runtime/CoreUObject/Public/UObject/UObjectAnnotation.h
--- a/Engine/Source/Runtime/CoreUObject/Public/UObject/UObjectAnnotation.h
+++ b/Engine/Source/Runtime/CoreUObject/Public/UObject/UObjectAnnotation.h
@@ -35,7 +35,7 @@ public:
 	 * @param Object object that has been destroyed
 	 * @param Index	index of object that is being deleted
 	 */
-	virtual void NotifyUObjectDeleted(const UObjectBase *Object, int32 Index)
+	virtual void NotifyUObjectDeleted(const UObjectBase *Object, int32 Index) override
 	{
 #if !(UE_BUILD_SHIPPING || UE_BUILD_TEST)
 		if (!bAutoRemove)
@@ -51,6 +51,12 @@ public:
 		}
 	}
 
+	virtual void OnUObjectArrayShutdown() override
+	{
+		RemoveAllAnnotations();
+		GUObjectArray.RemoveUObjectDeleteListener(this);
+	}
+
 	/**
 	 * Constructor, initializes to nothing
 	 */
@@ -66,7 +72,7 @@ public:
 	 */
 	virtual ~FUObjectAnnotationSparse()
 	{
-		RemoveAllAnnotations();
+		RemoveAllAnnotations();		
 	}
 
 	/**
@@ -259,11 +265,17 @@ public:
 	 * @param Object object that has been destroyed
 	 * @param Index	index of object that is being deleted
 	 */
-	virtual void NotifyUObjectDeleted(const UObjectBase *Object, int32 Index)
+	virtual void NotifyUObjectDeleted(const UObjectBase *Object, int32 Index) override
 	{
 		RemoveAnnotation(Object);
 	}
 
+	virtual void OnUObjectArrayShutdown() override
+	{
+		RemoveAllAnnotations();
+		GUObjectArray.RemoveUObjectDeleteListener(this);
+	}
+
 	/**
 	 * Destructor, removes all annotations, which removes the annotation as a uobject destruction listener
 	 */
@@ -466,14 +478,9 @@ public:
 * @param TAnnotation type of the annotation
 * @param bAutoRemove if true, annotation will automatically be removed, otherwise in non-final builds it will verify that the annotation was removed by other means prior to destruction.
 **/
-template<typename TAnnotation, bool bAutoRemove>
+template<typename TAnnotation, bool bAutoRemove, int32 NumAnnotationsPerChunk = 64 * 1024>
 class FUObjectAnnotationChunked : public FUObjectArray::FUObjectDeleteListener
 {
-	enum
-	{
-		NumAnnotationsPerChunk = 64 * 1024,
-	};
-
 	struct TAnnotationChunk
 	{
 		int32 Num;
@@ -876,6 +883,12 @@ public:
 			RemoveAnnotation(Index);
 		}
 	}
+
+	virtual void OnUObjectArrayShutdown() override
+	{
+		RemoveAllAnnotations();
+		GUObjectArray.RemoveUObjectDeleteListener(this);
+	}
 };
 
 /**
@@ -917,6 +930,12 @@ public:
 		}
 	}
 
+	virtual void OnUObjectArrayShutdown() override
+	{
+		RemoveAllAnnotations();
+		GUObjectArray.RemoveUObjectDeleteListener(this);
+	}
+
 	/**
 	 * Destructor, removes all annotations, which removes the annotation as a uobject destruction listener
 	 */
@@ -1127,6 +1146,12 @@ public:
 		RemoveAnnotation(Index);
 	}
 
+	virtual void OnUObjectArrayShutdown() override
+	{
+		RemoveAllAnnotations();
+		GUObjectArray.RemoveUObjectDeleteListener(this);
+	}
+
 	/**
 	 * Destructor, removes all annotations, which removes the annotation as a uobject destruction listener
 	 */
diff --git a/Engine/Source/Runtime/CoreUObject/Public/UObject/UObjectArray.h b/Engine/Source/Runtime/CoreUObject/Public/UObject/UObjectArray.h
--- a/Engine/Source/Runtime/CoreUObject/Public/UObject/UObjectArray.h
+++ b/Engine/Source/Runtime/CoreUObject/Public/UObject/UObjectArray.h
@@ -561,6 +561,11 @@ public:
 		 * @param Index	index of object that is being deleted
 		 */
 		virtual void NotifyUObjectCreated(const class UObjectBase *Object, int32 Index)=0;
+
+		/**
+		 * Called when UObject Array is being shut down, this is where all listeners should be removed from it 
+		 */
+		virtual void OnUObjectArrayShutdown()=0;
 	};
 
 	/**
@@ -578,6 +583,11 @@ public:
 		 * @param Index	index of object that is being deleted
 		 */
 		virtual void NotifyUObjectDeleted(const class UObjectBase *Object, int32 Index)=0;
+
+		/**
+		 * Called when UObject Array is being shut down, this is where all listeners should be removed from it
+		 */
+		virtual void OnUObjectArrayShutdown() = 0;
 	};
 
 	/**
