Moved GAllowCookedDataInEditorBuilds declaration to a global header to allow it to be set programatically by commandlets

diff --git a/Engine/Source/Runtime/CoreUObject/Public/UObject/UObjectGlobals.h b/Engine/Source/Runtime/CoreUObject/Public/UObject/UObjectGlobals.h
--- a/Engine/Source/Runtime/CoreUObject/Public/UObject/UObjectGlobals.h
+++ b/Engine/Source/Runtime/CoreUObject/Public/UObject/UObjectGlobals.h
@@ -44,6 +44,8 @@ DECLARE_DWORD_COUNTER_STAT_EXTERN(TEXT("FindObjectFast"),STAT_FindObjectFast,STA
 extern COREUOBJECT_API bool					GIsSavingPackage;
 /** This allows loading unversioned cooked content in the editor */
 extern COREUOBJECT_API bool					GAllowUnversionedContentInEditor;
+/** This allows loading cooked content in the editor */
+extern COREUOBJECT_API int32				GAllowCookedDataInEditorBuilds;
 
 /** Enum used in StaticDuplicateObject() and related functions to describe why something is being duplicated */
 namespace EDuplicateMode
