Removed redundant declaration of GAllowCookedDataInEditorBuilds as it's been moved to a global header

diff --git a/Engine/Source/Runtime/CoreUObject/Private/Serialization/ArchiveStackTrace.cpp b/Engine/Source/Runtime/CoreUObject/Private/Serialization/ArchiveStackTrace.cpp
--- a/Engine/Source/Runtime/CoreUObject/Private/Serialization/ArchiveStackTrace.cpp
+++ b/Engine/Source/Runtime/CoreUObject/Private/Serialization/ArchiveStackTrace.cpp
@@ -19,6 +19,7 @@
 #include "UObject/LinkerManager.h"
 #include "Misc/PackageName.h"
 #include "Templates/UniquePtr.h"
+#include "UObject/UObjectGlobals.h"
 
 DEFINE_LOG_CATEGORY_STATIC(LogArchiveDiff, Log, All);
 
@@ -1287,8 +1288,6 @@ static void DumpTableDifferences(
 #endif // !NO_LOGGING
 }
 
-extern int32 GAllowCookedDataInEditorBuilds;
-
 void FArchiveStackTrace::DumpPackageHeaderDiffs(const FPackageData& SourcePackage, const FPackageData& DestPackage, const FString& AssetFilename, const int32 MaxDiffsToLog)
 {
 #if !NO_LOGGING
