Minor optimization in FLinkerLoad::SerializeNameMap when using binary format

diff --git a/Engine/Source/Runtime/CoreUObject/Private/UObject/LinkerLoad.cpp b/Engine/Source/Runtime/CoreUObject/Private/UObject/LinkerLoad.cpp
--- a/Engine/Source/Runtime/CoreUObject/Private/UObject/LinkerLoad.cpp
+++ b/Engine/Source/Runtime/CoreUObject/Private/UObject/LinkerLoad.cpp
@@ -1420,17 +1420,28 @@ FLinkerLoad::ELinkerStatus FLinkerLoad::SerializeNameMap()
 	}
 
 	FStructuredArchive::FStream NameStream = StructuredArchiveRootRecord->EnterStream(FIELD_NAME_TEXT("Names"));
+	const bool bIsTextFormat = StructuredArchiveRootRecord->GetUnderlyingArchive().IsTextFormat();
 
-	while( bFinishedPrecaching && NameMapIndex < Summary.NameCount && !IsTimeLimitExceeded(TEXT("serializing name map"),100))
+	NameMap.Reserve(Summary.NameCount);
+
+	while (bFinishedPrecaching && NameMapIndex < Summary.NameCount && !IsTimeLimitExceeded(TEXT("serializing name map"), 100))
 	{
 		SCOPED_LOADTIMER(LinkerLoad_SerializeNameMap_ProcessingEntries);
 
-		// Read the name entry from the file.
 		FNameEntrySerialized NameEntry(ENAME_LinkerConstructor);
-		NameStream.EnterElement() << NameEntry;
+
+		if (bIsTextFormat)
+		{
+			NameStream.EnterElement() << NameEntry;
+		}
+		else
+		{
+			// Read the name from the underlying Archive when the format is binary to avoid the overhead from ArchiveProxy
+			NameStream.GetUnderlyingArchive() << NameEntry;
+		}
 
 		// Add it to the name table with no splitting and no hash calculations
-		NameMap.Add(FName(NameEntry));
+		NameMap.Emplace(NameEntry);
 
 		NameMapIndex++;
 	}
