Remove virtual from a bunch of internal serialization functions that take FArchive parameters to stop external code overriding them. Some serialization paths were changed to only call the structured archive version and as such wouldn't reach these overloads. Note that this will be a compile error for licensees who have overloaded these functions, but there is no way round that. These functions are pretty low level internal things though, so it shouldn't affect many people.

diff --git a/Engine/Source/Runtime/CoreUObject/Private/UObject/Class.cpp b/Engine/Source/Runtime/CoreUObject/Private/UObject/Class.cpp
--- a/Engine/Source/Runtime/CoreUObject/Private/UObject/Class.cpp
+++ b/Engine/Source/Runtime/CoreUObject/Private/UObject/Class.cpp
@@ -850,10 +850,6 @@ void UStruct::DestroyStruct(void* Dest, int32 ArrayDim) const
 	}
 }
 
-void UStruct::SerializeBin(FArchive& Ar, void* Data) const
-{
-	SerializeBin(FStructuredArchiveFromArchive(Ar).GetSlot(), Data);
-}
 //
 // Serialize all of the class's data that belongs in a particular
 // bin and resides in Data.
@@ -912,11 +908,6 @@ void UStruct::SerializeBinEx( FStructuredArchive::FSlot Slot, void* Data, void c
 	}
 }
 
-void UStruct::SerializeTaggedProperties(FArchive& Ar, uint8* Data, UStruct* DefaultsStruct, uint8* Defaults, const UObject* BreakRecursionIfFullyLoad) const
-{
-	SerializeTaggedProperties(FStructuredArchiveFromArchive(Ar).GetSlot(), Data, DefaultsStruct, Defaults, BreakRecursionIfFullyLoad);
-}
-
 void UStruct::SerializeTaggedProperties(FStructuredArchive::FSlot Slot, uint8* Data, UStruct* DefaultsStruct, uint8* Defaults, const UObject* BreakRecursionIfFullyLoad) const
 {
 	FArchive& UnderlyingArchive = Slot.GetUnderlyingArchive();
@@ -3886,11 +3877,6 @@ bool UClass::ImplementsInterface( const class UClass* SomeInterface ) const
 	return false;
 }
 
-void UClass::SerializeDefaultObject(UObject* Object, FArchive& Ar)
-{
-	SerializeDefaultObject(Object, FStructuredArchiveFromArchive(Ar).GetSlot());
-}
-
 /** serializes the passed in object as this class's default object using the given archive
  * @param Object the object to serialize as default
  * @param Ar the archive to serialize from
diff --git a/Engine/Source/Runtime/CoreUObject/Public/UObject/Class.h b/Engine/Source/Runtime/CoreUObject/Public/UObject/Class.h
--- a/Engine/Source/Runtime/CoreUObject/Public/UObject/Class.h
+++ b/Engine/Source/Runtime/CoreUObject/Public/UObject/Class.h
@@ -323,9 +323,14 @@ public:
 	/** Creates the field/property links and gets structure ready for use at runtime */
 	virtual void Link(FArchive& Ar, bool bRelinkExistingProperties);
 
+	/** Serializes struct properties, does not handle defaults*/
+	virtual void SerializeBin(FArchive& Ar, void* Data) const final 
+	{
+		SerializeBin(FStructuredArchiveFromArchive(Ar).GetSlot(), Data);
+	}
+
 	/** Serializes struct properties, does not handle defaults */
 	virtual void SerializeBin(FStructuredArchive::FSlot Slot, void* Data) const;
-	virtual void SerializeBin(FArchive& Ar, void* Data) const;
 
 	/**
 	 * Serializes the class properties that reside in Data if they differ from the corresponding values in DefaultData
@@ -338,8 +343,13 @@ public:
 	void SerializeBinEx( FStructuredArchive::FSlot Slot, void* Data, void const* DefaultData, UStruct* DefaultStruct ) const;
 
 	/** Serializes list of properties, using property tags to handle mismatches */
-	virtual void SerializeTaggedProperties( FArchive& Ar, uint8* Data, UStruct* DefaultsStruct, uint8* Defaults, const UObject* BreakRecursionIfFullyLoad=nullptr) const;
-	virtual void SerializeTaggedProperties( FStructuredArchive::FSlot Slot, uint8* Data, UStruct* DefaultsStruct, uint8* Defaults, const UObject* BreakRecursionIfFullyLoad = nullptr) const;
+	virtual void SerializeTaggedProperties(FArchive& Ar, uint8* Data, UStruct* DefaultsStruct, uint8* Defaults, const UObject* BreakRecursionIfFullyLoad = nullptr) const final 
+	{
+		SerializeTaggedProperties(FStructuredArchiveFromArchive(Ar).GetSlot(), Data, DefaultsStruct, Defaults, BreakRecursionIfFullyLoad);
+	}
+
+	/** Serializes list of properties, using property tags to handle mismatches */
+	virtual void SerializeTaggedProperties(FStructuredArchive::FSlot Slot, uint8* Data, UStruct* DefaultsStruct, uint8* Defaults, const UObject* BreakRecursionIfFullyLoad = nullptr) const;
 
 	/**
 	 * Initialize a struct over uninitialized memory. This may be done by calling the native constructor or individually initializing properties
@@ -2686,10 +2696,18 @@ public:
 
 	/** serializes the passed in object as this class's default object using the given archive
 	 * @param Object the object to serialize as default
-	 * @param Ar the archive to serialize from
+	 * @param Slot the structured archive slot to serialize from
 	 */
 	virtual void SerializeDefaultObject(UObject* Object, FStructuredArchive::FSlot Slot);
-	virtual void SerializeDefaultObject(UObject* Object, FArchive& Ar);
+
+	/** serializes the passed in object as this class's default object using the given archive
+	 * @param Object the object to serialize as default
+	 * @param Ar the archive to serialize from
+	 */
+	virtual void SerializeDefaultObject(UObject* Object, FArchive& Ar) final
+	{
+		SerializeDefaultObject(Object, FStructuredArchiveFromArchive(Ar).GetSlot());
+	}
 
 	/** Wraps the PostLoad() call for the class default object.
 	 * @param Object the default object to call PostLoad() on
diff --git a/Engine/Source/Runtime/Engine/Classes/Engine/BlueprintGeneratedClass.h b/Engine/Source/Runtime/Engine/Classes/Engine/BlueprintGeneratedClass.h
--- a/Engine/Source/Runtime/Engine/Classes/Engine/BlueprintGeneratedClass.h
+++ b/Engine/Source/Runtime/Engine/Classes/Engine/BlueprintGeneratedClass.h
@@ -702,7 +702,7 @@ public:
 	virtual void FlushCompilationQueueForLevel() override;
 	virtual UObject* GetArchetypeForCDO() const override;
 #endif //WITH_EDITOR
-	virtual void SerializeDefaultObject(UObject* Object, FArchive& Ar) override;
+	virtual void SerializeDefaultObject(UObject* Object, FStructuredArchive::FSlot Slot) override;
 	virtual void PostLoadDefaultObject(UObject* Object) override;
 	virtual bool IsFunctionImplementedInBlueprint(FName InFunctionName) const override;
 	virtual uint8* GetPersistentUberGraphFrame(UObject* Obj, UFunction* FuncToCheck) const override;
diff --git a/Engine/Source/Runtime/Engine/Private/BlueprintGeneratedClass.cpp b/Engine/Source/Runtime/Engine/Private/BlueprintGeneratedClass.cpp
--- a/Engine/Source/Runtime/Engine/Private/BlueprintGeneratedClass.cpp
+++ b/Engine/Source/Runtime/Engine/Private/BlueprintGeneratedClass.cpp
@@ -304,13 +304,14 @@ UObject* UBlueprintGeneratedClass::GetArchetypeForCDO() const
 }
 #endif //WITH_EDITOR
 
-void UBlueprintGeneratedClass::SerializeDefaultObject(UObject* Object, FArchive& Ar)
+void UBlueprintGeneratedClass::SerializeDefaultObject(UObject* Object, FStructuredArchive::FSlot Slot)
 {
 	FScopeLock SerializeAndPostLoadLock(&SerializeAndPostLoadCritical);
+	FArchive& UnderlyingArchive = Slot.GetUnderlyingArchive();
 
-	Super::SerializeDefaultObject(Object, Ar);
+	Super::SerializeDefaultObject(Object, Slot);
 
-	if (Ar.IsLoading() && !Ar.IsObjectReferenceCollector() && Object == ClassDefaultObject)
+	if (UnderlyingArchive.IsLoading() && !UnderlyingArchive.IsObjectReferenceCollector() && Object == ClassDefaultObject)
 	{
 		// On load, build the custom property list used in post-construct initialization logic. Note that in the editor, this will be refreshed during compile-on-load.
 		// @TODO - Potentially make this serializable (or cooked data) to eliminate the slight load time cost we'll incur below to generate this list in a cooked build. For now, it's not serialized since the raw UProperty references cannot be saved out.
