Attempt to please static analyzer

diff --git a/Engine/Source/Runtime/Core/Private/UObject/UnrealNames.cpp b/Engine/Source/Runtime/Core/Private/UObject/UnrealNames.cpp
--- a/Engine/Source/Runtime/Core/Private/UObject/UnrealNames.cpp
+++ b/Engine/Source/Runtime/Core/Private/UObject/UnrealNames.cpp
@@ -41,8 +41,7 @@ const TCHAR* LexToString(EName Ename)
 
 int32 FNameEntry::GetDataOffset()
 {
-	const FNameEntry* Dummy = nullptr;
-	return (intptr_t)Dummy->AnsiName - (intptr_t)Dummy;
+	return STRUCT_OFFSET(FNameEntry, AnsiName);
 }
 
 /*-----------------------------------------------------------------------------
diff --git a/Engine/Source/Runtime/CoreUObject/Private/UObject/CoreNet.cpp b/Engine/Source/Runtime/CoreUObject/Private/UObject/CoreNet.cpp
--- a/Engine/Source/Runtime/CoreUObject/Private/UObject/CoreNet.cpp
+++ b/Engine/Source/Runtime/CoreUObject/Private/UObject/CoreNet.cpp
@@ -314,7 +314,7 @@ bool UPackageMap::StaticSerializeName(FArchive& Ar, FName& InName)
 		const EName* InEName = InName.ToEName();
 		uint8 bHardcoded = InEName && ShouldReplicateAsInteger(*InEName);
 		Ar.SerializeBits(&bHardcoded, 1);
-		if (bHardcoded)
+		if (bHardcoded && /* silence static analyzer */ InEName)
 		{
 			// send by hardcoded index
 			checkSlow(InName.GetNumber() <= 0); // hardcoded names should never have a Number
