Remove STRUCT_SerializeNativeStructured flag. It wasn't being set in some cases, and was only used to piggy back the STRUCT_SerializeNative flag which represents the same thing

diff --git a/Engine/Source/Runtime/CoreUObject/Private/UObject/Class.cpp b/Engine/Source/Runtime/CoreUObject/Private/UObject/Class.cpp
--- a/Engine/Source/Runtime/CoreUObject/Private/UObject/Class.cpp
+++ b/Engine/Source/Runtime/CoreUObject/Private/UObject/Class.cpp
@@ -1999,7 +1999,7 @@ void UScriptStruct::PrepareCppStructOps()
 	}
 
 	check(!(StructFlags & STRUCT_ComputedFlags));
-	if (CppStructOps->HasSerializer())
+	if (CppStructOps->HasSerializer() || CppStructOps->HasStructuredSerializer())
 	{
 		UE_LOG(LogClass, Verbose, TEXT("Native struct %s has a custom serializer."),*GetName());
 		StructFlags = EStructFlags(StructFlags | STRUCT_SerializeNative );
diff --git a/Engine/Source/Runtime/CoreUObject/Public/UObject/Class.h b/Engine/Source/Runtime/CoreUObject/Public/UObject/Class.h
--- a/Engine/Source/Runtime/CoreUObject/Public/UObject/Class.h
+++ b/Engine/Source/Runtime/CoreUObject/Public/UObject/Class.h
@@ -552,14 +552,11 @@ enum EStructFlags
 	/** If set, this struct can share net serialization state across connections */
 	STRUCT_NetSharedSerialization = 0x00400000,
 
-	/**  */
-	STRUCT_SerializeNativeStructured = 0x00800000,
-
 	/** Struct flags that are automatically inherited */
 	STRUCT_Inherit				= STRUCT_HasInstancedReference|STRUCT_Atomic,
 
 	/** Flags that are always computed, never loaded or done with code generation */
-	STRUCT_ComputedFlags		= STRUCT_NetDeltaSerializeNative | STRUCT_NetSerializeNative | STRUCT_SerializeNative | STRUCT_SerializeNativeStructured | STRUCT_PostSerializeNative | STRUCT_CopyNative | STRUCT_IsPlainOldData | STRUCT_NoDestructor | STRUCT_ZeroConstructor | STRUCT_IdenticalNative | STRUCT_AddStructReferencedObjects | STRUCT_ExportTextItemNative | STRUCT_ImportTextItemNative | STRUCT_SerializeFromMismatchedTag | STRUCT_PostScriptConstruct | STRUCT_NetSharedSerialization
+	STRUCT_ComputedFlags		= STRUCT_NetDeltaSerializeNative | STRUCT_NetSerializeNative | STRUCT_SerializeNative | STRUCT_PostSerializeNative | STRUCT_CopyNative | STRUCT_IsPlainOldData | STRUCT_NoDestructor | STRUCT_ZeroConstructor | STRUCT_IdenticalNative | STRUCT_AddStructReferencedObjects | STRUCT_ExportTextItemNative | STRUCT_ImportTextItemNative | STRUCT_SerializeFromMismatchedTag | STRUCT_PostScriptConstruct | STRUCT_NetSharedSerialization
 };
 
 
@@ -1284,7 +1281,7 @@ public:
 	/** Returns true if this struct has a native serialize function */
 	bool UseNativeSerialization() const
 	{
-		if ((StructFlags&(STRUCT_SerializeNative | STRUCT_SerializeNativeStructured)) != 0)
+		if ((StructFlags&(STRUCT_SerializeNative)) != 0)
 		{
 			return true;
 		}
