THasGetTypeHash deprecated in favor of concepts checks.  THasGetTypeHash doesn't work correctly for derived types, e.g. FPlane inherits FVector's GetTypeHash(), but THasGetTypeHash<FPlane>::Value == false.

diff --git a/Engine/Source/Runtime/Core/Public/Concepts/GetTypeHashable.h b/Engine/Source/Runtime/Core/Public/Concepts/GetTypeHashable.h
new file mode 100644
--- /dev/null
+++ b/Engine/Source/Runtime/Core/Public/Concepts/GetTypeHashable.h
@@ -0,0 +1,15 @@
+// Copyright 1998-2019 Epic Games, Inc. All Rights Reserved.
+
+#pragma once
+
+#include "CoreTypes.h"
+
+/**
+ * Describes a type with a GetTypeHash overload.
+ */
+struct CGetTypeHashable {
+	template <typename T>
+	auto Requires(uint32& Result, const T& Val) -> decltype(
+		Result = GetTypeHash(Val)
+	);
+};
diff --git a/Engine/Source/Runtime/Core/Public/Containers/Map.h b/Engine/Source/Runtime/Core/Public/Containers/Map.h
--- a/Engine/Source/Runtime/Core/Public/Containers/Map.h
+++ b/Engine/Source/Runtime/Core/Public/Containers/Map.h
@@ -12,7 +12,8 @@
 #include "Containers/Set.h"
 #include "Algo/Reverse.h"
 #include "Templates/Tuple.h"
-#include "Templates/HasGetTypeHash.h"
+#include "Templates/Models.h"
+#include "Concepts/GetTypeHashable.h"
 #include "Containers/UnrealString.h"
 
 #define ExchangeB(A,B) {bool T=A; A=B; B=T;}
@@ -94,7 +95,7 @@ struct TDefaultMapKeyFuncs : BaseKeyFuncs<TPair<KeyType,ValueType>,KeyType,bInAl
 template<typename KeyType, typename ValueType, bool bInAllowDuplicateKeys>
 struct TDefaultMapHashableKeyFuncs : TDefaultMapKeyFuncs<KeyType, ValueType, bInAllowDuplicateKeys>
 {
-	static_assert(THasGetTypeHash<KeyType>::Value, "TMap must have a hashable KeyType unless a custom key func is provided.");
+	static_assert(TModels<CGetTypeHashable, KeyType>::Value, "TMap must have a hashable KeyType unless a custom key func is provided.");
 };
 
 /** 
diff --git a/Engine/Source/Runtime/Core/Public/Templates/HasGetTypeHash.h b/Engine/Source/Runtime/Core/Public/Templates/HasGetTypeHash.h
--- a/Engine/Source/Runtime/Core/Public/Templates/HasGetTypeHash.h
+++ b/Engine/Source/Runtime/Core/Public/Templates/HasGetTypeHash.h
@@ -48,7 +48,7 @@ namespace UE4GetTypeHashExists_Private
 * Traits class which tests if a type has a GetTypeHash overload.
 */
 template <typename T>
-struct THasGetTypeHash
+struct UE_DEPRECATED(4.23, "THasGetTypeHash has been deprecated, please use TModels<CGetTypeHashable, T> instead.") THasGetTypeHash
 {
 	enum { Value = UE4GetTypeHashExists_Private::GetTypeHashQuery<T>::Value };
 };
diff --git a/Engine/Source/Runtime/CoreUObject/Public/UObject/Class.h b/Engine/Source/Runtime/CoreUObject/Public/UObject/Class.h
--- a/Engine/Source/Runtime/CoreUObject/Public/UObject/Class.h
+++ b/Engine/Source/Runtime/CoreUObject/Public/UObject/Class.h
@@ -17,7 +17,8 @@
 #include "UObject/GarbageCollection.h"
 #include "UObject/CoreNative.h"
 #include "UObject/ReflectedTypeAccessors.h"
-#include "Templates/HasGetTypeHash.h"
+#include "Templates/Models.h"
+#include "Concepts/GetTypeHashable.h"
 #include "Templates/IsAbstract.h"
 #include "Templates/IsEnum.h"
 #include "Misc/Optional.h"
@@ -833,13 +834,13 @@ FORCEINLINE typename TEnableIf<TStructOpsTypeTraits<CPPSTRUCT>::WithStructuredSe
  * Selection of GetTypeHash call.
  */
 template<class CPPSTRUCT>
-FORCEINLINE typename TEnableIf<!THasGetTypeHash<CPPSTRUCT>::Value, uint32>::Type GetTypeHashOrNot(const CPPSTRUCT *Data)
+FORCEINLINE typename TEnableIf<!TModels<CGetTypeHashable, CPPSTRUCT>::Value, uint32>::Type GetTypeHashOrNot(const CPPSTRUCT *Data)
 {
 	return 0;
 }
 
 template<class CPPSTRUCT>
-FORCEINLINE typename TEnableIf<THasGetTypeHash<CPPSTRUCT>::Value, uint32>::Type GetTypeHashOrNot(const CPPSTRUCT *Data)
+FORCEINLINE typename TEnableIf<TModels<CGetTypeHashable, CPPSTRUCT>::Value, uint32>::Type GetTypeHashOrNot(const CPPSTRUCT *Data)
 {
 	return GetTypeHash(*Data);
 }
@@ -1161,7 +1162,7 @@ public:
 
 		virtual bool HasGetTypeHash() override
 		{
-			return THasGetTypeHash<CPPSTRUCT>::Value;
+			return TModels<CGetTypeHashable, CPPSTRUCT>::Value;
 		}
 		uint32 GetTypeHash(const void* Src) override
 		{
@@ -1174,7 +1175,7 @@ public:
 				  (TIsPODType<CPPSTRUCT>::Value ? CPF_IsPlainOldData : CPF_None)
 				| (TIsTriviallyDestructible<CPPSTRUCT>::Value ? CPF_NoDestructor : CPF_None)
 				| (TIsZeroConstructType<CPPSTRUCT>::Value ? CPF_ZeroConstructor : CPF_None)
-				| (THasGetTypeHash<CPPSTRUCT>::Value ? CPF_HasGetValueTypeHash : CPF_None);
+				| (TModels<CGetTypeHashable, CPPSTRUCT>::Value ? CPF_HasGetValueTypeHash : CPF_None);
 		}
 		bool IsAbstract() const override
 		{
diff --git a/Engine/Source/Runtime/CoreUObject/Public/UObject/LazyObjectPtr.h b/Engine/Source/Runtime/CoreUObject/Public/UObject/LazyObjectPtr.h
--- a/Engine/Source/Runtime/CoreUObject/Public/UObject/LazyObjectPtr.h
+++ b/Engine/Source/Runtime/CoreUObject/Public/UObject/LazyObjectPtr.h
@@ -171,7 +171,6 @@ public:
 
 template <> struct TIsPODType<FLazyObjectPtr> { enum { Value = TIsPODType<TPersistentObjectPtr<FUniqueObjectGuid> >::Value }; };
 template <> struct TIsWeakPointerType<FLazyObjectPtr> { enum { Value = TIsWeakPointerType<TPersistentObjectPtr<FUniqueObjectGuid> >::Value }; };
-template <> struct THasGetTypeHash<FLazyObjectPtr> { enum { Value = THasGetTypeHash<TPersistentObjectPtr<FUniqueObjectGuid> >::Value }; };
 
 /**
  * TLazyObjectPtr is templatized version of the generic FLazyObjectPtr
diff --git a/Engine/Source/Runtime/CoreUObject/Public/UObject/SoftObjectPtr.h b/Engine/Source/Runtime/CoreUObject/Public/UObject/SoftObjectPtr.h
--- a/Engine/Source/Runtime/CoreUObject/Public/UObject/SoftObjectPtr.h
+++ b/Engine/Source/Runtime/CoreUObject/Public/UObject/SoftObjectPtr.h
@@ -107,7 +107,6 @@ public:
 
 template <> struct TIsPODType<FSoftObjectPtr> { enum { Value = TIsPODType<TPersistentObjectPtr<FSoftObjectPath> >::Value }; };
 template <> struct TIsWeakPointerType<FSoftObjectPtr> { enum { Value = TIsWeakPointerType<TPersistentObjectPtr<FSoftObjectPath> >::Value }; };
-template <> struct THasGetTypeHash<FSoftObjectPtr> { enum { Value = THasGetTypeHash<TPersistentObjectPtr<FSoftObjectPath> >::Value }; };
 
 /**
  * TSoftObjectPtr is templatized wrapper of the generic FSoftObjectPtr, it can be used in UProperties
diff --git a/Engine/Source/Runtime/CoreUObject/Public/UObject/UnrealType.h b/Engine/Source/Runtime/CoreUObject/Public/UObject/UnrealType.h
--- a/Engine/Source/Runtime/CoreUObject/Public/UObject/UnrealType.h
+++ b/Engine/Source/Runtime/CoreUObject/Public/UObject/UnrealType.h
@@ -25,6 +25,8 @@
 #include "UObject/SoftObjectPtr.h"
 #include "UObject/PropertyTag.h"
 #include "Serialization/SerializedPropertyScope.h"
+#include "Templates/Models.h"
+#include "Concepts/GetTypeHashable.h"
 
 COREUOBJECT_API DECLARE_LOG_CATEGORY_EXTERN(LogType, Log, All);
 
@@ -967,7 +969,7 @@ protected:
 			(TIsPODType<TCppType>::Value ? CPF_IsPlainOldData : CPF_None) 
 			| (TIsTriviallyDestructible<TCppType>::Value ? CPF_NoDestructor : CPF_None) 
 			| (TIsZeroConstructType<TCppType>::Value ? CPF_ZeroConstructor : CPF_None)
-			| (THasGetTypeHash<TCppType>::Value ? CPF_HasGetValueTypeHash : CPF_None);
+			| (TModels<CGetTypeHashable, TCppType>::Value ? CPF_HasGetValueTypeHash : CPF_None);
 
 	}
 };
