Reordered includes alphabetically.

diff --git a/Engine/Source/Runtime/Core/Public/Containers/Map.h b/Engine/Source/Runtime/Core/Public/Containers/Map.h
--- a/Engine/Source/Runtime/Core/Public/Containers/Map.h
+++ b/Engine/Source/Runtime/Core/Public/Containers/Map.h
@@ -3,18 +3,19 @@
 #pragma once
 
 #include "CoreTypes.h"
+
+#include "Algo/Reverse.h"
+#include "Concepts/GetTypeHashable.h"
+#include "Containers/Set.h"
+#include "Containers/UnrealString.h"
 #include "Misc/AssertionMacros.h"
-#include "Templates/UnrealTypeTraits.h"
-#include "Templates/UnrealTemplate.h"
-#include "Templates/Sorting.h"
 #include "Misc/StructBuilder.h"
 #include "Templates/Function.h"
-#include "Containers/Set.h"
-#include "Algo/Reverse.h"
-#include "Templates/Tuple.h"
 #include "Templates/Models.h"
-#include "Concepts/GetTypeHashable.h"
-#include "Containers/UnrealString.h"
+#include "Templates/Sorting.h"
+#include "Templates/Tuple.h"
+#include "Templates/UnrealTemplate.h"
+#include "Templates/UnrealTypeTraits.h"
 
 #define ExchangeB(A,B) {bool T=A; A=B; B=T;}
 
diff --git a/Engine/Source/Runtime/CoreUObject/Public/UObject/Class.h b/Engine/Source/Runtime/CoreUObject/Public/UObject/Class.h
--- a/Engine/Source/Runtime/CoreUObject/Public/UObject/Class.h
+++ b/Engine/Source/Runtime/CoreUObject/Public/UObject/Class.h
@@ -7,23 +7,24 @@
 #pragma once
 
 #include "CoreMinimal.h"
-#include "UObject/Script.h"
-#include "UObject/ObjectMacros.h"
-#include "UObject/UObjectGlobals.h"
-#include "UObject/Object.h"
+
+#include "Concepts/GetTypeHashable.h"
+#include "Math/RandomStream.h"
+#include "Misc/EnumClassFlags.h"
 #include "Misc/FallbackStruct.h"
 #include "Misc/Guid.h"
-#include "Math/RandomStream.h"
-#include "UObject/GarbageCollection.h"
-#include "UObject/CoreNative.h"
-#include "UObject/ReflectedTypeAccessors.h"
-#include "Templates/Models.h"
-#include "Concepts/GetTypeHashable.h"
-#include "Templates/IsAbstract.h"
-#include "Templates/IsEnum.h"
 #include "Misc/Optional.h"
-#include "Misc/EnumClassFlags.h"
 #include "Misc/ScopeRWLock.h"
+#include "Templates/IsAbstract.h"
+#include "Templates/IsEnum.h"
+#include "Templates/Models.h"
+#include "UObject/CoreNative.h"
+#include "UObject/GarbageCollection.h"
+#include "UObject/Object.h"
+#include "UObject/ObjectMacros.h"
+#include "UObject/ReflectedTypeAccessors.h"
+#include "UObject/Script.h"
+#include "UObject/UObjectGlobals.h"
 
 struct FCustomPropertyListNode;
 struct FFrame;
diff --git a/Engine/Source/Runtime/CoreUObject/Public/UObject/UnrealType.h b/Engine/Source/Runtime/CoreUObject/Public/UObject/UnrealType.h
--- a/Engine/Source/Runtime/CoreUObject/Public/UObject/UnrealType.h
+++ b/Engine/Source/Runtime/CoreUObject/Public/UObject/UnrealType.h
@@ -7,26 +7,27 @@
 #pragma once
 
 #include "CoreMinimal.h"
-#include "UObject/ObjectMacros.h"
-#include "UObject/UObjectGlobals.h"
-#include "UObject/Object.h"
-#include "UObject/Class.h"
-#include "UObject/WeakObjectPtr.h"
-#include "UObject/CoreNetTypes.h"
-#include "UObject/ScriptInterface.h"
-#include "UObject/SparseDelegate.h"
+
+#include "Concepts/GetTypeHashable.h"
+#include "Containers/List.h"
+#include "Serialization/SerializedPropertyScope.h"
 #include "Templates/Casts.h"
+#include "Templates/Greater.h"
 #include "Templates/IsFloatingPoint.h"
 #include "Templates/IsIntegral.h"
 #include "Templates/IsSigned.h"
-#include "Templates/Greater.h"
-#include "Containers/List.h"
+#include "Templates/Models.h"
+#include "UObject/Class.h"
+#include "UObject/CoreNetTypes.h"
 #include "UObject/LazyObjectPtr.h"
-#include "UObject/SoftObjectPtr.h"
+#include "UObject/Object.h"
+#include "UObject/ObjectMacros.h"
 #include "UObject/PropertyTag.h"
-#include "Serialization/SerializedPropertyScope.h"
-#include "Templates/Models.h"
-#include "Concepts/GetTypeHashable.h"
+#include "UObject/ScriptInterface.h"
+#include "UObject/SoftObjectPtr.h"
+#include "UObject/SparseDelegate.h"
+#include "UObject/UObjectGlobals.h"
+#include "UObject/WeakObjectPtr.h"
 
 COREUOBJECT_API DECLARE_LOG_CATEGORY_EXTERN(LogType, Log, All);
 
