UPackages will no longer be hashed with their outer as their outer is null and the resulting has is the same as their name hash, nor will they be added to object -> outer map.

This saves up to 3MB of memory in our internal projects.

diff --git a/Engine/Source/Runtime/CoreUObject/Private/UObject/UObjectHash.cpp b/Engine/Source/Runtime/CoreUObject/Private/UObject/UObjectHash.cpp
--- a/Engine/Source/Runtime/CoreUObject/Private/UObject/UObjectHash.cpp
+++ b/Engine/Source/Runtime/CoreUObject/Private/UObject/UObjectHash.cpp
@@ -711,6 +711,8 @@ static FAutoConsoleCommand ShrinkUObjectHashTablesCmd(
 
 void GetObjectsWithOuter(const class UObjectBase* Outer, TArray<UObject *>& Results, bool bIncludeNestedObjects, EObjectFlags ExclusionFlags, EInternalObjectFlags ExclusionInternalFlags)
 {
+	check(Outer != nullptr);
+
 	// We don't want to return any objects that are currently being background loaded unless we're using the object iterator during async loading.
 	ExclusionInternalFlags |= EInternalObjectFlags::Unreachable;
 	if (!IsInAsyncLoadingThread())
@@ -759,6 +761,8 @@ void GetObjectsWithOuter(const class UObjectBase* Outer, TArray<UObject *>& Resu
 
 void ForEachObjectWithOuter(const class UObjectBase* Outer, TFunctionRef<void(UObject*)> Operation, bool bIncludeNestedObjects, EObjectFlags ExclusionFlags, EInternalObjectFlags ExclusionInternalFlags)
 {
+	check(Outer != nullptr);
+
 	// We don't want to return any objects that are currently being background loaded unless we're using the object iterator during async loading.
 	ExclusionInternalFlags |= EInternalObjectFlags::Unreachable;
 	if (!IsInAsyncLoadingThread())
@@ -996,11 +1000,15 @@ void HashObject(UObjectBase* Object)
 		checkSlow(!ThreadHash.PairExistsInHash(Hash, Object));  // if it already exists, something is wrong with the external code
 		ThreadHash.AddToHash(Hash, Object);
 
-		Hash = GetObjectOuterHash( Name, (PTRINT)Object->GetOuter() );
-		checkSlow( !ThreadHash.HashOuter.FindPair( Hash, Object ) );  // if it already exists, something is wrong with the external code
-		ThreadHash.HashOuter.Add( Hash, Object );
+		if (PTRINT Outer = (PTRINT)Object->GetOuter())
+		{
+			Hash = GetObjectOuterHash(Name, Outer);
+			checkSlow(!ThreadHash.HashOuter.FindPair(Hash, Object));  // if it already exists, something is wrong with the external code
+			ThreadHash.HashOuter.Add(Hash, Object);
+
+			AddToOuterMap(ThreadHash, Object);
+		}
 
-		AddToOuterMap( ThreadHash, Object );
 		AddToClassMap( ThreadHash, Object );
 	}
 }
@@ -1027,11 +1035,15 @@ void UnhashObject(UObjectBase* Object)
 		NumRemoved = ThreadHash.RemoveFromHash(Hash, Object);
 		check(NumRemoved == 1); // must have existed, else something is wrong with the external code
 
-		Hash = GetObjectOuterHash( Name, (PTRINT)Object->GetOuter() );
-		NumRemoved = ThreadHash.HashOuter.RemoveSingle( Hash, Object );
-		check( NumRemoved == 1 ); // must have existed, else something is wrong with the external code
+		if (PTRINT Outer = (PTRINT)Object->GetOuter())
+		{
+			Hash = GetObjectOuterHash(Name, Outer);
+			NumRemoved = ThreadHash.HashOuter.RemoveSingle(Hash, Object);
+			check(NumRemoved == 1); // must have existed, else something is wrong with the external code
+
+			RemoveFromOuterMap(ThreadHash, Object);
+		}
 
-		RemoveFromOuterMap( ThreadHash, Object );
 		RemoveFromClassMap( ThreadHash, Object );
 	}
 }
@@ -1214,5 +1226,13 @@ void LogHashOuterStatistics(FOutputDevice& Ar, const bool bShowHashBucketCollisi
 	FHashTableLock HashLock(FUObjectHashTables::Get());
 	LogHashStatisticsInternal(FUObjectHashTables::Get().HashOuter, Ar, bShowHashBucketCollisionInfo);
 	Ar.Logf(TEXT(""));
+
+	uint32 HashOuterMapSize = 0;
+	for (TPair<UObjectBase*, FHashBucket>& OuterMapEntry : FUObjectHashTables::Get().ObjectOuterMap)
+	{
+		HashOuterMapSize += OuterMapEntry.Value.GetItemsSize();
+	}
+	Ar.Logf(TEXT("Total memory allocated for Object Outer Map: %u bytes."), HashOuterMapSize);
+	Ar.Logf(TEXT(""));
 }
 
