Changed a WITH_EDITOR_DATA to WITH_TEXT_ARCHIVE_SUPPORT in SavePacakge.cpp so that editor builds compile without text archive support enabled

diff --git a/Engine/Source/Runtime/CoreUObject/Private/UObject/SavePackage.cpp b/Engine/Source/Runtime/CoreUObject/Private/UObject/SavePackage.cpp
--- a/Engine/Source/Runtime/CoreUObject/Private/UObject/SavePackage.cpp
+++ b/Engine/Source/Runtime/CoreUObject/Private/UObject/SavePackage.cpp
@@ -3975,7 +3975,7 @@ FSavePackageResultStruct UPackage::Save(UPackage* InOuter, UObject* Base, EObjec
 					Linker = TUniquePtr<FLinkerSave>(new FLinkerSave(InOuter, *TempFilename, bForceByteSwapping, bSaveUnversioned));
 				}
 
-#if WITH_EDITOR
+#if WITH_TEXT_ARCHIVE_SUPPORT
 				if (bTextFormat)
 				{
 					TextFormatArchive = IFileManager::Get().CreateFileWriter(*TextFormatTempFilename);
