When loading from json formatters, track the number of outstanding elements and pop them from the stack if we leave that scope before they've all been consumed

diff --git a/Engine/Source/Runtime/CoreUObject/Private/Serialization/Formatters/JsonArchiveInputFormatter.cpp b/Engine/Source/Runtime/CoreUObject/Private/Serialization/Formatters/JsonArchiveInputFormatter.cpp
--- a/Engine/Source/Runtime/CoreUObject/Private/Serialization/Formatters/JsonArchiveInputFormatter.cpp
+++ b/Engine/Source/Runtime/CoreUObject/Private/Serialization/Formatters/JsonArchiveInputFormatter.cpp
@@ -24,6 +24,9 @@ FJsonArchiveInputFormatter::FJsonArchiveInputFormatter(FArchive& InInner, TFunct
 	TSharedRef< TJsonReader<char> > Reader = TJsonReaderFactory<char>::Create(&InInner);
 	ensure(FJsonSerializer::Deserialize(Reader, RootObject));
 	ValueStack.Add(MakeShared<FJsonValueObject>(RootObject));
+
+	ValueStack.Reserve(64);
+	ArrayValuesRemainingStack.Reserve(64);
 }
 
 FJsonArchiveInputFormatter::~FJsonArchiveInputFormatter()
@@ -41,6 +44,7 @@ FStructuredArchiveFormatter* FJsonArchiveInputFormatter::CreateSubtreeReader()
 	Cloned->ObjectStack.Empty();
 	Cloned->ValueStack.Empty();
 	Cloned->MapIteratorStack.Empty();
+	Cloned->ArrayValuesRemainingStack.Empty();
 	Cloned->ValueStack.Push(ValueStack.Top());
 
 	return Cloned;
@@ -115,14 +119,26 @@ void FJsonArchiveInputFormatter::EnterArray(int32& NumElements)
 	}
 
 	NumElements = ArrayValue.Num();
+	ArrayValuesRemainingStack.Add(NumElements);
 }
 
 void FJsonArchiveInputFormatter::LeaveArray()
 {
+	check(ArrayValuesRemainingStack.Num() > 0);
+	int32 Remaining = ArrayValuesRemainingStack.Top();
+	ArrayValuesRemainingStack.Pop();
+	check(Remaining >= 0);
+	check(Remaining <= ValueStack.Num());
+	while (Remaining-- > 0)
+	{
+		ValueStack.Pop();
+	}
 }
 
 void FJsonArchiveInputFormatter::EnterArrayElement()
 {
+	check(ArrayValuesRemainingStack.Num() > 0);
+	check(ArrayValuesRemainingStack.Top() > 0);
 }
 
 void FJsonArchiveInputFormatter::EnterArrayElement_TextOnly(EArchiveValueType& OutType)
@@ -133,6 +149,7 @@ void FJsonArchiveInputFormatter::EnterArrayElement_TextOnly(EArchiveValueType& O
 void FJsonArchiveInputFormatter::LeaveArrayElement()
 {
 	ValueStack.Pop();
+	ArrayValuesRemainingStack.Top()--;
 }
 
 void FJsonArchiveInputFormatter::EnterStream()
@@ -148,6 +165,7 @@ void FJsonArchiveInputFormatter::EnterStream_TextOnly(int32& NumElements)
 
 void FJsonArchiveInputFormatter::LeaveStream()
 {
+	LeaveArray();
 }
 
 void FJsonArchiveInputFormatter::EnterStreamElement()
diff --git a/Engine/Source/Runtime/CoreUObject/Public/Serialization/Formatters/JsonArchiveInputFormatter.h b/Engine/Source/Runtime/CoreUObject/Public/Serialization/Formatters/JsonArchiveInputFormatter.h
--- a/Engine/Source/Runtime/CoreUObject/Public/Serialization/Formatters/JsonArchiveInputFormatter.h
+++ b/Engine/Source/Runtime/CoreUObject/Public/Serialization/Formatters/JsonArchiveInputFormatter.h
@@ -90,6 +90,7 @@ private:
 	TArray<FObjectRecord> ObjectStack;
 	TArray<TSharedPtr<FJsonValue>> ValueStack;
 	TArray<TMap<FString, TSharedPtr<FJsonValue>>::TIterator> MapIteratorStack;
+	TArray<int32> ArrayValuesRemainingStack;
 
 	static FString EscapeFieldName(const TCHAR* Name);
 	static FString UnescapeFieldName(const TCHAR* Name);
