Fix a crash when opening a text asset

diff --git a/Engine/Source/Runtime/CoreUObject/Private/UObject/LinkerLoad.cpp b/Engine/Source/Runtime/CoreUObject/Private/UObject/LinkerLoad.cpp
--- a/Engine/Source/Runtime/CoreUObject/Private/UObject/LinkerLoad.cpp
+++ b/Engine/Source/Runtime/CoreUObject/Private/UObject/LinkerLoad.cpp
@@ -468,10 +468,13 @@ void FLinkerLoad::SetLoader(FArchive* InLoader)
 
 	check(StructuredArchive == nullptr);
 	check(!StructuredArchiveRootRecord.IsSet());
-	check(StructuredArchiveFormatter == nullptr);
-	
-	// Create structured archive wrapper
-	StructuredArchiveFormatter = new FBinaryArchiveFormatter(*this);
+
+	if (StructuredArchiveFormatter == nullptr)
+	{
+		// Create structured archive wrapper
+		StructuredArchiveFormatter = new FBinaryArchiveFormatter(*this);
+	}
+
 	StructuredArchive = new FStructuredArchive(*StructuredArchiveFormatter);
 	StructuredArchiveRootRecord.Emplace(StructuredArchive->Open().EnterRecord());
 }
@@ -4756,10 +4759,10 @@ void FLinkerLoad::Detach()
 		CurrentLoadContext->RemoveDelayedLinkerClosePackage(this);
 	}
 
-		delete StructuredArchive;
-		StructuredArchive = nullptr;
-		delete StructuredArchiveFormatter;
-		StructuredArchiveFormatter = nullptr;
+	delete StructuredArchive;
+	StructuredArchive = nullptr;
+	delete StructuredArchiveFormatter;
+	StructuredArchiveFormatter = nullptr;
 
 	if (Loader)
 	{
