Remove unneeded check for HasMetaData before getting the display name.

diff --git a/Engine/Source/Runtime/CoreUObject/Private/UObject/Class.cpp b/Engine/Source/Runtime/CoreUObject/Private/UObject/Class.cpp
--- a/Engine/Source/Runtime/CoreUObject/Private/UObject/Class.cpp
+++ b/Engine/Source/Runtime/CoreUObject/Private/UObject/Class.cpp
@@ -242,12 +242,8 @@ FText UField::GetDisplayNameText() const
 
 	const FString Key = GetFullGroupName(false);
 
-	FString NativeDisplayName;
-	if( HasMetaData(NAME_DisplayName) )
-	{
-		NativeDisplayName = GetMetaData(NAME_DisplayName);
-	}
-	else
+	FString NativeDisplayName = GetMetaData(NAME_DisplayName);
+	if (NativeDisplayName.IsEmpty())
 	{
 		NativeDisplayName = FName::NameToDisplayString(FDisplayNameHelper::Get(*this), IsA<UBoolProperty>());
 	}
