Set up ReferenceNames for most runtime FGCObject classes, as these now properly show up in obj refs

diff --git a/Engine/Source/Runtime/CoreUObject/Public/UObject/ConstructorHelpers.h b/Engine/Source/Runtime/CoreUObject/Public/UObject/ConstructorHelpers.h
--- a/Engine/Source/Runtime/CoreUObject/Public/UObject/ConstructorHelpers.h
+++ b/Engine/Source/Runtime/CoreUObject/Public/UObject/ConstructorHelpers.h
@@ -114,10 +114,15 @@ public:
 			return !!Object;
 		}
 
-		virtual void AddReferencedObjects( FReferenceCollector& Collector )
+		virtual void AddReferencedObjects( FReferenceCollector& Collector ) override
 		{
 			Collector.AddReferencedObject(Object);
 		}
+
+		virtual FString GetReferencerName() const override
+		{
+			return TEXT("FObjectFinder");
+		}
 	};
 
 	template<class T>
@@ -152,10 +157,15 @@ public:
 			return !!Get();
 		}
 
-		virtual void AddReferencedObjects( FReferenceCollector& Collector )
+		virtual void AddReferencedObjects( FReferenceCollector& Collector ) override
 		{
 			Collector.AddReferencedObject(Object);
 		}
+
+		virtual FString GetReferencerName() const override
+		{
+			return TEXT("FObjectFinderOptional");
+		}
 	};
 
 	template<class T>
@@ -175,12 +185,17 @@ public:
 			return !!*Class;
 		}
 		
-		virtual void AddReferencedObjects( FReferenceCollector& Collector )
+		virtual void AddReferencedObjects( FReferenceCollector& Collector ) override
 		{
 			UClass* ReferencedClass = Class.Get();
 			Collector.AddReferencedObject(ReferencedClass);
 			Class = ReferencedClass;
 		}
+
+		virtual FString GetReferencerName() const override
+		{
+			return TEXT("FClassFinder");
+		}
 	};
 
 public:
diff --git a/Engine/Source/Runtime/CoreUObject/Public/UObject/GCObjectScopeGuard.h b/Engine/Source/Runtime/CoreUObject/Public/UObject/GCObjectScopeGuard.h
--- a/Engine/Source/Runtime/CoreUObject/Public/UObject/GCObjectScopeGuard.h
+++ b/Engine/Source/Runtime/CoreUObject/Public/UObject/GCObjectScopeGuard.h
@@ -30,6 +30,11 @@ public:
 		Collector.AddReferencedObject(Object);
 	}
 
+	virtual FString GetReferencerName() const override
+	{
+		return TEXT("FGCObjectScopeGuard");
+	}
+
 private:
 	const UObject* Object;
 };
@@ -62,6 +67,11 @@ public:
 		Collector.AddReferencedObjects(Objects);
 	}
 
+	virtual FString GetReferencerName() const override
+	{
+		return TEXT("TGCObjectsScopeGuard");
+	}
+
 private:
 	TArray<ObjectType*> Objects;
 };
diff --git a/Engine/Source/Runtime/CoreUObject/Public/UObject/StrongObjectPtr.h b/Engine/Source/Runtime/CoreUObject/Public/UObject/StrongObjectPtr.h
--- a/Engine/Source/Runtime/CoreUObject/Public/UObject/StrongObjectPtr.h
+++ b/Engine/Source/Runtime/CoreUObject/Public/UObject/StrongObjectPtr.h
@@ -141,6 +141,11 @@ private:
 			Collector.AddReferencedObject(Object);
 		}
 
+		virtual FString GetReferencerName() const override
+		{
+			return TEXT("TStrongObjectPtr");
+		}
+
 	private:
 		ObjectType* Object;
 	};
diff --git a/Engine/Source/Runtime/Engine/Classes/Components/ForceFeedbackComponent.h b/Engine/Source/Runtime/Engine/Classes/Components/ForceFeedbackComponent.h
--- a/Engine/Source/Runtime/Engine/Classes/Components/ForceFeedbackComponent.h
+++ b/Engine/Source/Runtime/Engine/Classes/Components/ForceFeedbackComponent.h
@@ -34,6 +34,7 @@ public:
 
 private:
 	virtual void AddReferencedObjects( FReferenceCollector& Collector ) override;
+	virtual FString GetReferencerName() const override;
 	virtual UWorld* GetTickableGameObjectWorld() const override;
 	virtual void Tick( float DeltaTime ) override;
 	virtual bool IsTickable() const override;
diff --git a/Engine/Source/Runtime/Engine/Classes/Particles/ParticleSystemManager.h b/Engine/Source/Runtime/Engine/Classes/Particles/ParticleSystemManager.h
--- a/Engine/Source/Runtime/Engine/Classes/Particles/ParticleSystemManager.h
+++ b/Engine/Source/Runtime/Engine/Classes/Particles/ParticleSystemManager.h
@@ -137,6 +137,7 @@ public:
 
 	// FGCObject interface
 	virtual void AddReferencedObjects(FReferenceCollector& Collector) override;
+	virtual FString GetReferencerName() const override;
 
 	void Dump();
 
diff --git a/Engine/Source/Runtime/Engine/Classes/Sound/SampleBuffer.h b/Engine/Source/Runtime/Engine/Classes/Sound/SampleBuffer.h
--- a/Engine/Source/Runtime/Engine/Classes/Sound/SampleBuffer.h
+++ b/Engine/Source/Runtime/Engine/Classes/Sound/SampleBuffer.h
@@ -441,7 +441,8 @@ namespace Audio
 		void Update();
 
 		//~ GCObject Interface
-		void AddReferencedObjects(FReferenceCollector& Collector) override;
+		virtual void AddReferencedObjects(FReferenceCollector& Collector) override;
+		virtual FString GetReferencerName() const override;
 		//~ GCObject Interface
 
 	private:
diff --git a/Engine/Source/Runtime/Engine/Private/Components/ForceFeedbackComponent.cpp b/Engine/Source/Runtime/Engine/Private/Components/ForceFeedbackComponent.cpp
--- a/Engine/Source/Runtime/Engine/Private/Components/ForceFeedbackComponent.cpp
+++ b/Engine/Source/Runtime/Engine/Private/Components/ForceFeedbackComponent.cpp
@@ -67,6 +67,11 @@ void FForceFeedbackManager::AddReferencedObjects(FReferenceCollector& Collector)
 	Collector.AddReferencedObjects(ActiveForceFeedbackComponents);
 }
 
+FString FForceFeedbackManager::GetReferencerName() const
+{
+	return TEXT("FForceFeedbackManager");
+}
+
 UWorld* FForceFeedbackManager::GetTickableGameObjectWorld() const
 {
 	return World;
diff --git a/Engine/Source/Runtime/Engine/Private/Components/TextRenderComponent.cpp b/Engine/Source/Runtime/Engine/Private/Components/TextRenderComponent.cpp
--- a/Engine/Source/Runtime/Engine/Private/Components/TextRenderComponent.cpp
+++ b/Engine/Source/Runtime/Engine/Private/Components/TextRenderComponent.cpp
@@ -465,6 +465,11 @@ public:
 		}
 	}
 
+	virtual FString GetReferencerName() const override
+	{
+		return TEXT("FTextRenderComponentMIDCache");
+	}
+
 private:
 	/** Key identifying an array of MIDs */
 	struct FKey
diff --git a/Engine/Source/Runtime/Engine/Private/DistanceFieldAtlas.cpp b/Engine/Source/Runtime/Engine/Private/DistanceFieldAtlas.cpp
--- a/Engine/Source/Runtime/Engine/Private/DistanceFieldAtlas.cpp
+++ b/Engine/Source/Runtime/Engine/Private/DistanceFieldAtlas.cpp
@@ -899,6 +899,11 @@ void FDistanceFieldAsyncQueue::AddReferencedObjects(FReferenceCollector& Collect
 	}
 }
 
+FString FDistanceFieldAsyncQueue::GetReferencerName() const
+{
+	return TEXT("FDistanceFieldAsyncQueue");
+}
+
 void FDistanceFieldAsyncQueue::ProcessAsyncTasks()
 {
 #if WITH_EDITOR
diff --git a/Engine/Source/Runtime/Engine/Private/Internationalization/StringTable.cpp b/Engine/Source/Runtime/Engine/Private/Internationalization/StringTable.cpp
--- a/Engine/Source/Runtime/Engine/Private/Internationalization/StringTable.cpp
+++ b/Engine/Source/Runtime/Engine/Private/Internationalization/StringTable.cpp
@@ -288,6 +288,11 @@ private:
 		Collector.AddReferencedObjects(KeepAliveStringTables);
 	}
 
+	virtual FString GetReferencerName() const override
+	{
+		return TEXT("FStringTableEngineBridge");
+	}
+
 private:
 	struct FAsyncLoadingStringTable
 	{
diff --git a/Engine/Source/Runtime/Engine/Private/Particles/ParticleSystemManager.cpp b/Engine/Source/Runtime/Engine/Private/Particles/ParticleSystemManager.cpp
--- a/Engine/Source/Runtime/Engine/Private/Particles/ParticleSystemManager.cpp
+++ b/Engine/Source/Runtime/Engine/Private/Particles/ParticleSystemManager.cpp
@@ -256,6 +256,11 @@ void FParticleSystemWorldManager::AddReferencedObjects(FReferenceCollector& Coll
 	}
 }
 
+FString FParticleSystemWorldManager::GetReferencerName() const
+{
+	return TEXT("FParticleSystemWorldManager");
+}
+
 void FParticleSystemWorldManager::HandlePostGarbageCollect()
 {
 	//UE_LOG(LogParticles, Warning, TEXT("| HandlePostGarbageCollect |"));
diff --git a/Engine/Source/Runtime/Engine/Private/PreviewScene.cpp b/Engine/Source/Runtime/Engine/Private/PreviewScene.cpp
--- a/Engine/Source/Runtime/Engine/Private/PreviewScene.cpp
+++ b/Engine/Source/Runtime/Engine/Private/PreviewScene.cpp
@@ -148,6 +148,11 @@ void FPreviewScene::AddReferencedObjects( FReferenceCollector& Collector )
 	Collector.AddReferencedObject( PreviewWorld );
 }
 
+FString FPreviewScene::GetReferencerName() const
+{
+	return TEXT("FPreviewScene");
+}
+
 void FPreviewScene::UpdateCaptureContents()
 {
 	USkyLightComponent::UpdateSkyCaptureContents(PreviewWorld);
diff --git a/Engine/Source/Runtime/Engine/Private/RepLayout.cpp b/Engine/Source/Runtime/Engine/Private/RepLayout.cpp
--- a/Engine/Source/Runtime/Engine/Private/RepLayout.cpp
+++ b/Engine/Source/Runtime/Engine/Private/RepLayout.cpp
@@ -5191,6 +5191,11 @@ void FRepLayout::AddReferencedObjects(FReferenceCollector& Collector)
 	}
 }
 
+FString FRepLayout::GetReferencerName() const
+{
+	return TEXT("FRepLayout");
+}
+
 void FRepLayout::CountBytes(FArchive& Ar) const
 {
 	GRANULAR_NETWORK_MEMORY_TRACKING_INIT(Ar, "FRepLayout::CountBytes");
diff --git a/Engine/Source/Runtime/Engine/Private/SampleBuffer.cpp b/Engine/Source/Runtime/Engine/Private/SampleBuffer.cpp
--- a/Engine/Source/Runtime/Engine/Private/SampleBuffer.cpp
+++ b/Engine/Source/Runtime/Engine/Private/SampleBuffer.cpp
@@ -81,6 +81,11 @@ namespace Audio
 		}
 	}
 
+	FString FSoundWavePCMLoader::GetReferencerName() const
+	{
+		return TEXT("FSoundWavePCMLoader");
+	}
+
 	FSoundWavePCMWriter::FSoundWavePCMWriter(int32 InChunkSize)
 		: CurrentSoundWave(0)
 		, CurrentState(ESoundWavePCMWriterState::Idle)
diff --git a/Engine/Source/Runtime/Engine/Private/Slate/DeferredCleanupSlateBrush.cpp b/Engine/Source/Runtime/Engine/Private/Slate/DeferredCleanupSlateBrush.cpp
--- a/Engine/Source/Runtime/Engine/Private/Slate/DeferredCleanupSlateBrush.cpp
+++ b/Engine/Source/Runtime/Engine/Private/Slate/DeferredCleanupSlateBrush.cpp
@@ -48,5 +48,5 @@ void FDeferredCleanupSlateBrush::AddReferencedObjects(FReferenceCollector& Colle
 
 FString FDeferredCleanupSlateBrush::GetReferencerName() const
 {
-	return "FDeferredCleanupSlateBrush";
+	return TEXT("FDeferredCleanupSlateBrush");
 }
\ No newline at end of file
diff --git a/Engine/Source/Runtime/Engine/Private/Slate/SlateGameResources.cpp b/Engine/Source/Runtime/Engine/Private/Slate/SlateGameResources.cpp
--- a/Engine/Source/Runtime/Engine/Private/Slate/SlateGameResources.cpp
+++ b/Engine/Source/Runtime/Engine/Private/Slate/SlateGameResources.cpp
@@ -306,3 +306,8 @@ void FSlateGameResources::AddReferencedObjects( FReferenceCollector& Collector )
 {
 	Collector.AddReferencedObjects( UIResources );
 }
+
+FString FSlateGameResources::GetReferencerName() const
+{
+	return TEXT("FSlateGameResources");
+}
diff --git a/Engine/Source/Runtime/Engine/Private/Subsystems/SubsystemCollection.cpp b/Engine/Source/Runtime/Engine/Private/Subsystems/SubsystemCollection.cpp
--- a/Engine/Source/Runtime/Engine/Private/Subsystems/SubsystemCollection.cpp
+++ b/Engine/Source/Runtime/Engine/Private/Subsystems/SubsystemCollection.cpp
@@ -173,6 +173,11 @@ void FSubsystemCollectionBase::AddReferencedObjects(FReferenceCollector& Collect
 	Collector.AddReferencedObjects(SubsystemMap);
 }
 
+FString FSubsystemCollectionBase::GetReferencerName() const
+{
+	return TEXT("FSubsystemCollectionBase");
+}
+
 bool FSubsystemCollectionBase::AddAndInitializeSubsystem(UClass* SubsystemClass)
 {
 	if (!SubsystemMap.Contains(SubsystemClass))
diff --git a/Engine/Source/Runtime/Engine/Private/UnrealClient.cpp b/Engine/Source/Runtime/Engine/Private/UnrealClient.cpp
--- a/Engine/Source/Runtime/Engine/Private/UnrealClient.cpp
+++ b/Engine/Source/Runtime/Engine/Private/UnrealClient.cpp
@@ -2041,6 +2041,11 @@ void FViewport::FHitProxyMap::AddReferencedObjects( FReferenceCollector& Collect
 	}
 }
 
+FString FViewport::FHitProxyMap::GetReferencerName() const
+{
+	return TEXT("FViewport::FHitProxyMap");
+}
+
 /**
  * Globally enables/disables rendering
  *
diff --git a/Engine/Source/Runtime/Engine/Public/DistanceFieldAtlas.h b/Engine/Source/Runtime/Engine/Public/DistanceFieldAtlas.h
--- a/Engine/Source/Runtime/Engine/Public/DistanceFieldAtlas.h
+++ b/Engine/Source/Runtime/Engine/Public/DistanceFieldAtlas.h
@@ -226,7 +226,10 @@ public:
 	ENGINE_API void ProcessAsyncTasks();
 
 	/** Exposes UObject references used by the async build. */
-	ENGINE_API void AddReferencedObjects(FReferenceCollector& Collector);
+	ENGINE_API virtual void AddReferencedObjects(FReferenceCollector& Collector) override;
+
+	/** Returns name of class for reference tracking */
+	ENGINE_API virtual FString GetReferencerName() const override;
 
 	/** Blocks until it is safe to shut down (worker threads are idle). */
 	ENGINE_API void Shutdown();
diff --git a/Engine/Source/Runtime/Engine/Public/Net/RepLayout.h b/Engine/Source/Runtime/Engine/Public/Net/RepLayout.h
--- a/Engine/Source/Runtime/Engine/Public/Net/RepLayout.h
+++ b/Engine/Source/Runtime/Engine/Public/Net/RepLayout.h
@@ -1564,6 +1564,7 @@ public:
 
 	//~ Begin FGCObject Interface
 	ENGINE_API virtual void AddReferencedObjects(FReferenceCollector& Collector) override;
+	ENGINE_API virtual FString GetReferencerName() const override;
 	//~ End FGCObject Interface
 
 	/**
diff --git a/Engine/Source/Runtime/Engine/Public/PreviewScene.h b/Engine/Source/Runtime/Engine/Public/PreviewScene.h
--- a/Engine/Source/Runtime/Engine/Public/PreviewScene.h
+++ b/Engine/Source/Runtime/Engine/Public/PreviewScene.h
@@ -77,7 +77,8 @@ public:
 	virtual void RemoveComponent(class UActorComponent* Component);
 
 	// Serializer.
-	void AddReferencedObjects( FReferenceCollector& Collector ) override;
+	virtual void AddReferencedObjects( FReferenceCollector& Collector ) override;
+	virtual FString GetReferencerName() const override;
 
 	// Accessors.
 	UWorld* GetWorld() const { return PreviewWorld; }
diff --git a/Engine/Source/Runtime/Engine/Public/Slate/SlateGameResources.h b/Engine/Source/Runtime/Engine/Public/Slate/SlateGameResources.h
--- a/Engine/Source/Runtime/Engine/Public/Slate/SlateGameResources.h
+++ b/Engine/Source/Runtime/Engine/Public/Slate/SlateGameResources.h
@@ -86,6 +86,7 @@ protected:
 	FName GetCleanName(const FName& AssetName) const;
 
 	virtual void AddReferencedObjects( FReferenceCollector& Collector ) override;
+	virtual FString GetReferencerName() const override;
 
 private:
 
diff --git a/Engine/Source/Runtime/Engine/Public/Subsystems/SubsystemCollection.h b/Engine/Source/Runtime/Engine/Public/Subsystems/SubsystemCollection.h
--- a/Engine/Source/Runtime/Engine/Public/Subsystems/SubsystemCollection.h
+++ b/Engine/Source/Runtime/Engine/Public/Subsystems/SubsystemCollection.h
@@ -26,6 +26,7 @@ public:
 
 	/* FGCObject Interface */
 	virtual void AddReferencedObjects(FReferenceCollector& Collector) override;
+	virtual FString GetReferencerName() const override;
 
 protected:
 	/** protected constructor - for use by the template only(FSubsystemCollection<TBaseType>) */
diff --git a/Engine/Source/Runtime/Engine/Public/UnrealClient.h b/Engine/Source/Runtime/Engine/Public/UnrealClient.h
--- a/Engine/Source/Runtime/Engine/Public/UnrealClient.h
+++ b/Engine/Source/Runtime/Engine/Public/UnrealClient.h
@@ -632,6 +632,7 @@ protected:
 
 		/** FGCObject interface */
 		virtual void AddReferencedObjects( FReferenceCollector& Collector ) override;
+		virtual FString GetReferencerName() const override;
 
 		const FTexture2DRHIRef& GetHitProxyTexture(void) const		{ return HitProxyTexture; }
 		const FTexture2DRHIRef& GetHitProxyCPUTexture(void) const		{ return HitProxyCPUTexture; }
diff --git a/Engine/Source/Runtime/MovieSceneCapture/Private/ActiveMovieSceneCaptures.cpp b/Engine/Source/Runtime/MovieSceneCapture/Private/ActiveMovieSceneCaptures.cpp
--- a/Engine/Source/Runtime/MovieSceneCapture/Private/ActiveMovieSceneCaptures.cpp
+++ b/Engine/Source/Runtime/MovieSceneCapture/Private/ActiveMovieSceneCaptures.cpp
@@ -57,3 +57,8 @@ void FActiveMovieSceneCaptures::AddReferencedObjects(FReferenceCollector& Collec
 {
 	Collector.AddReferencedObjects(ActiveCaptures);
 }
+
+FString FActiveMovieSceneCaptures::GetReferencerName() const
+{
+	return TEXT("FActiveMovieSceneCaptures");
+}
diff --git a/Engine/Source/Runtime/MovieSceneCapture/Private/ActiveMovieSceneCaptures.h b/Engine/Source/Runtime/MovieSceneCapture/Private/ActiveMovieSceneCaptures.h
--- a/Engine/Source/Runtime/MovieSceneCapture/Private/ActiveMovieSceneCaptures.h
+++ b/Engine/Source/Runtime/MovieSceneCapture/Private/ActiveMovieSceneCaptures.h
@@ -30,6 +30,7 @@ private:
 
 	/** FGCObject interface */
 	virtual void AddReferencedObjects(FReferenceCollector& Collector) override;
+	virtual FString GetReferencerName() const override;
 
 
 	/** FTickableGameObject interface */
diff --git a/Engine/Source/Runtime/NavigationSystem/Private/NavMesh/RecastNavMeshGenerator.cpp b/Engine/Source/Runtime/NavigationSystem/Private/NavMesh/RecastNavMeshGenerator.cpp
--- a/Engine/Source/Runtime/NavigationSystem/Private/NavMesh/RecastNavMeshGenerator.cpp
+++ b/Engine/Source/Runtime/NavigationSystem/Private/NavMesh/RecastNavMeshGenerator.cpp
@@ -3187,6 +3187,11 @@ void FRecastTileGenerator::AddReferencedObjects(FReferenceCollector& Collector)
 	}
 }
 
+FString FRecastTileGenerator::GetReferencerName() const
+{
+	return TEXT("FRecastTileGenerator");
+}
+
 static int32 CaclulateMaxTilesCount(const TNavStatArray<FBox>& NavigableAreas, float TileSizeinWorldUnits, float AvgLayersPerGridCell)
 {
 	int32 GridCellsCount = 0;
diff --git a/Engine/Source/Runtime/NavigationSystem/Public/NavMesh/RecastNavMeshGenerator.h b/Engine/Source/Runtime/NavigationSystem/Public/NavMesh/RecastNavMeshGenerator.h
--- a/Engine/Source/Runtime/NavigationSystem/Public/NavMesh/RecastNavMeshGenerator.h
+++ b/Engine/Source/Runtime/NavigationSystem/Public/NavMesh/RecastNavMeshGenerator.h
@@ -197,7 +197,8 @@ public:
 	uint32 UsedMemoryOnStartup;
 
 	// FGCObject begin
-	virtual void AddReferencedObjects(FReferenceCollector& Collector);
+	virtual void AddReferencedObjects(FReferenceCollector& Collector) override;
+	virtual FString GetReferencerName() const override;
 	// FGCObject end
 		
 protected:
diff --git a/Engine/Source/Runtime/Slate/Private/Widgets/SInvalidationPanel.cpp b/Engine/Source/Runtime/Slate/Private/Widgets/SInvalidationPanel.cpp
--- a/Engine/Source/Runtime/Slate/Private/Widgets/SInvalidationPanel.cpp
+++ b/Engine/Source/Runtime/Slate/Private/Widgets/SInvalidationPanel.cpp
@@ -237,6 +237,11 @@ void SInvalidationPanel::AddReferencedObjects(FReferenceCollector& Collector)
 #endif
 }
 
+FString SInvalidationPanel::GetReferencerName() const
+{
+	return TEXT("SInvalidationPanel");
+}
+
 void SInvalidationPanel::SetCanCache(bool InCanCache)
 {
 	bCanCache = InCanCache;
diff --git a/Engine/Source/Runtime/Slate/Public/Widgets/SInvalidationPanel.h b/Engine/Source/Runtime/Slate/Public/Widgets/SInvalidationPanel.h
--- a/Engine/Source/Runtime/Slate/Public/Widgets/SInvalidationPanel.h
+++ b/Engine/Source/Runtime/Slate/Public/Widgets/SInvalidationPanel.h
@@ -43,6 +43,7 @@ public:
 #endif
 
 	virtual void AddReferencedObjects(FReferenceCollector& Collector) override;
+	virtual FString GetReferencerName() const override;
 
 	void SetCanCache(bool InCanCache);
 
diff --git a/Engine/Source/Runtime/SlateCore/Private/Fonts/CompositeFont.cpp b/Engine/Source/Runtime/SlateCore/Private/Fonts/CompositeFont.cpp
--- a/Engine/Source/Runtime/SlateCore/Private/Fonts/CompositeFont.cpp
+++ b/Engine/Source/Runtime/SlateCore/Private/Fonts/CompositeFont.cpp
@@ -268,3 +268,8 @@ void FStandaloneCompositeFont::AddReferencedObjects(FReferenceCollector& Collect
 		}
 	}
 }
+
+FString FStandaloneCompositeFont::GetReferencerName() const
+{
+	return TEXT("FStandaloneCompositeFont");
+}
diff --git a/Engine/Source/Runtime/SlateCore/Private/Fonts/LegacySlateFontInfoCache.cpp b/Engine/Source/Runtime/SlateCore/Private/Fonts/LegacySlateFontInfoCache.cpp
--- a/Engine/Source/Runtime/SlateCore/Private/Fonts/LegacySlateFontInfoCache.cpp
+++ b/Engine/Source/Runtime/SlateCore/Private/Fonts/LegacySlateFontInfoCache.cpp
@@ -332,6 +332,11 @@ void FLegacySlateFontInfoCache::AddReferencedObjects(FReferenceCollector& Collec
 	}
 }
 
+FString FLegacySlateFontInfoCache::GetReferencerName() const
+{
+	return TEXT("FLegacySlateFontInfoCache");
+}
+
 void FLegacySlateFontInfoCache::HandleCultureChanged()
 {
 	// We set this to the current frame count, as this will prevent the fallback font being updated for the remainder of this frame (as the culture change may have affected the fallback font used)
diff --git a/Engine/Source/Runtime/SlateCore/Private/Fonts/LegacySlateFontInfoCache.h b/Engine/Source/Runtime/SlateCore/Private/Fonts/LegacySlateFontInfoCache.h
--- a/Engine/Source/Runtime/SlateCore/Private/Fonts/LegacySlateFontInfoCache.h
+++ b/Engine/Source/Runtime/SlateCore/Private/Fonts/LegacySlateFontInfoCache.h
@@ -85,6 +85,7 @@ public:
 
 	// FGCObject interface
 	virtual void AddReferencedObjects(FReferenceCollector& Collector) override;
+	virtual FString GetReferencerName() const override;
 
 private:
 	FLegacySlateFontInfoCache();
diff --git a/Engine/Source/Runtime/SlateCore/Private/Rendering/DrawElements.cpp b/Engine/Source/Runtime/SlateCore/Private/Rendering/DrawElements.cpp
--- a/Engine/Source/Runtime/SlateCore/Private/Rendering/DrawElements.cpp
+++ b/Engine/Source/Runtime/SlateCore/Private/Rendering/DrawElements.cpp
@@ -1288,6 +1288,11 @@ void FSlateWindowElementList::FWindowElementGCObject::AddReferencedObjects(FRefe
 	}
 }
 
+FString FSlateWindowElementList::FWindowElementGCObject::GetReferencerName() const
+{
+	return TEXT("FSlateWindowElementList");
+}
+
 void FSlateWindowElementList::AddReferencedObjects(FReferenceCollector& Collector)
 {
 	Collector.AddReferencedObjects(ResourcesToReport);
diff --git a/Engine/Source/Runtime/SlateCore/Public/Fonts/CompositeFont.h b/Engine/Source/Runtime/SlateCore/Public/Fonts/CompositeFont.h
--- a/Engine/Source/Runtime/SlateCore/Public/Fonts/CompositeFont.h
+++ b/Engine/Source/Runtime/SlateCore/Public/Fonts/CompositeFont.h
@@ -464,4 +464,5 @@ struct SLATECORE_API FStandaloneCompositeFont : public FCompositeFont, public FG
 
 	// FGCObject interface
 	virtual void AddReferencedObjects(FReferenceCollector& Collector) override;
+	virtual FString GetReferencerName() const override;
 };
diff --git a/Engine/Source/Runtime/SlateCore/Public/Rendering/DrawElements.h b/Engine/Source/Runtime/SlateCore/Public/Rendering/DrawElements.h
--- a/Engine/Source/Runtime/SlateCore/Public/Rendering/DrawElements.h
+++ b/Engine/Source/Runtime/SlateCore/Public/Rendering/DrawElements.h
@@ -1696,7 +1696,8 @@ private:
 
 		void ClearOwner();
 
-		SLATECORE_API virtual void AddReferencedObjects(FReferenceCollector& Collector) override;
+		virtual void AddReferencedObjects(FReferenceCollector& Collector) override;
+		virtual FString GetReferencerName() const override;
 
 	private:
 		FSlateWindowElementList* Owner;
diff --git a/Engine/Source/Runtime/UMG/Private/Slate/SMeshWidget.cpp b/Engine/Source/Runtime/UMG/Private/Slate/SMeshWidget.cpp
--- a/Engine/Source/Runtime/UMG/Private/Slate/SMeshWidget.cpp
+++ b/Engine/Source/Runtime/UMG/Private/Slate/SMeshWidget.cpp
@@ -244,6 +244,11 @@ void SMeshWidget::AddReferencedObjects(FReferenceCollector& Collector)
 }
 
 
+FString SMeshWidget::GetReferencerName() const
+{
+	return TEXT("SMeshWidget");
+}
+
 void SMeshWidget::PushUpdate(uint32 VectorArtId, const SMeshWidget& Widget, const FVector2D& Position, float Scale, uint32 BaseAddress)
 {
 	PushUpdate(VectorArtId, Widget, Position, Scale, static_cast<float>(BaseAddress));
diff --git a/Engine/Source/Runtime/UMG/Private/Slate/SRetainerWidget.cpp b/Engine/Source/Runtime/UMG/Private/Slate/SRetainerWidget.cpp
--- a/Engine/Source/Runtime/UMG/Private/Slate/SRetainerWidget.cpp
+++ b/Engine/Source/Runtime/UMG/Private/Slate/SRetainerWidget.cpp
@@ -68,6 +68,12 @@ public:
 		Collector.AddReferencedObject(RenderTarget);
 		Collector.AddReferencedObject(DynamicEffect);
 	}
+
+	virtual FString GetReferencerName() const override
+	{
+		return TEXT("FRetainerWidgetRenderingResources");
+	}
+	
 public:
 	FWidgetRenderer* WidgetRenderer;
 	UTextureRenderTarget2D* RenderTarget;
diff --git a/Engine/Source/Runtime/UMG/Private/Slate/UMGDragDropOp.cpp b/Engine/Source/Runtime/UMG/Private/Slate/UMGDragDropOp.cpp
--- a/Engine/Source/Runtime/UMG/Private/Slate/UMGDragDropOp.cpp
+++ b/Engine/Source/Runtime/UMG/Private/Slate/UMGDragDropOp.cpp
@@ -27,6 +27,11 @@ void FUMGDragDropOp::AddReferencedObjects(FReferenceCollector& Collector)
 	Collector.AddReferencedObject(GameViewport);
 }
 
+FString FUMGDragDropOp::GetReferencerName() const
+{
+	return TEXT("FUMGDragDropOp");
+}
+
 void FUMGDragDropOp::Construct()
 {
 
diff --git a/Engine/Source/Runtime/UMG/Public/Slate/SMeshWidget.h b/Engine/Source/Runtime/UMG/Public/Slate/SMeshWidget.h
--- a/Engine/Source/Runtime/UMG/Public/Slate/SMeshWidget.h
+++ b/Engine/Source/Runtime/UMG/Public/Slate/SMeshWidget.h
@@ -82,6 +82,7 @@ protected:
 
 	// ~ FGCObject
 	virtual void AddReferencedObjects(FReferenceCollector& Collector) override;
+	virtual FString GetReferencerName() const override;
 	// ~ FGCObject
 
 protected:
diff --git a/Engine/Source/Runtime/UMG/Public/Slate/UMGDragDropOp.h b/Engine/Source/Runtime/UMG/Public/Slate/UMGDragDropOp.h
--- a/Engine/Source/Runtime/UMG/Public/Slate/UMGDragDropOp.h
+++ b/Engine/Source/Runtime/UMG/Public/Slate/UMGDragDropOp.h
@@ -27,6 +27,7 @@ public:
 
 	// Begin FGCObject
 	virtual void AddReferencedObjects(FReferenceCollector& Collector) override;
+	virtual FString GetReferencerName() const override;
 	// End FGCObject
 
 	virtual bool AffectedByPointerEvent(const FPointerEvent& PointerEvent) override;
