[NO NON-PII COMMIT MESSAGE]

diff --git a/Engine/Source/Runtime/CoreUObject/Private/Misc/GCObjectReferencer.cpp b/Engine/Source/Runtime/CoreUObject/Private/Misc/GCObjectReferencer.cpp
--- a/Engine/Source/Runtime/CoreUObject/Private/Misc/GCObjectReferencer.cpp
+++ b/Engine/Source/Runtime/CoreUObject/Private/Misc/GCObjectReferencer.cpp
@@ -27,8 +27,10 @@ void UGCObjectReferencer::AddReferencedObjects(UObject* InThis, FReferenceCollec
 	for (FGCObject* Object : This->ReferencedObjects)
 	{
 		check(Object);
+		This->CurrentlySerializingObject = Object;
 		Object->AddReferencedObjects(Collector);
 	}
+	This->CurrentlySerializingObject = nullptr;
 	Super::AddReferencedObjects( This, Collector );
 	This->bIsAddingReferencedObjects = false;
 }
@@ -52,9 +54,14 @@ void UGCObjectReferencer::RemoveObject(FGCObject* Object)
 
 bool UGCObjectReferencer::GetReferencerName(UObject* Object, FString& OutName, bool bOnlyIfAddingReferenced) const
 {
-	if (bOnlyIfAddingReferenced && !bIsAddingReferencedObjects)
+	if (bOnlyIfAddingReferenced)
 	{
-		return false;
+		if (!bIsAddingReferencedObjects || !CurrentlySerializingObject)
+		{
+			return false;
+		}
+		OutName = CurrentlySerializingObject->GetReferencerName();
+		return true;
 	}
 
 	// Let each registered object handle its AddReferencedObjects call
diff --git a/Engine/Source/Runtime/CoreUObject/Public/UObject/GCObject.h b/Engine/Source/Runtime/CoreUObject/Public/UObject/GCObject.h
--- a/Engine/Source/Runtime/CoreUObject/Public/UObject/GCObject.h
+++ b/Engine/Source/Runtime/CoreUObject/Public/UObject/GCObject.h
@@ -31,6 +31,8 @@ class COREUOBJECT_API UGCObjectReferencer : public UObject
 	FCriticalSection ReferencedObjectsCritical;
 	/** True if we are currently inside AddReferencedObjects */
 	bool bIsAddingReferencedObjects = false;
+	/** Currently serializing FGCObject*, only valid if bIsAddingReferencedObjects */
+	FGCObject* CurrentlySerializingObject = nullptr;
 
 public:
 	DECLARE_CASTED_CLASS_INTRINSIC_WITH_API(UGCObjectReferencer, UObject, CLASS_Transient, TEXT("/Script/CoreUObject"), CASTCLASS_None, NO_API);
