[NO NON-PII COMMIT MESSAGE]

diff --git a/Engine/Source/Runtime/CoreUObject/Private/UObject/GarbageCollection.cpp b/Engine/Source/Runtime/CoreUObject/Private/UObject/GarbageCollection.cpp
--- a/Engine/Source/Runtime/CoreUObject/Private/UObject/GarbageCollection.cpp
+++ b/Engine/Source/Runtime/CoreUObject/Private/UObject/GarbageCollection.cpp
@@ -501,11 +501,15 @@ public:
 			{
 #if ENABLE_GC_DEBUG_OUTPUT
 				// this message is to help track down culprits behind "Object in PIE world still referenced" errors
-				if (GIsEditor && !GIsPlayInEditorWorld && ReferencingObject != NULL && !ReferencingObject->RootPackageHasAnyFlags(PKG_PlayInEditor) && Object->RootPackageHasAnyFlags(PKG_PlayInEditor))
+				if (GIsEditor && !GIsPlayInEditorWorld && ReferencingObject != nullptr && !ReferencingObject->HasAnyFlags(RF_Transient) && Object->RootPackageHasAnyFlags(PKG_PlayInEditor))
 				{
-					UE_LOG(LogGarbage, Warning, TEXT("GC detected illegal reference to PIE object from content [possibly via [todo]]:"));
-					UE_LOG(LogGarbage, Warning, TEXT("      PIE object: %s"), *Object->GetFullName());
-					UE_LOG(LogGarbage, Warning, TEXT("  NON-PIE object: %s"), *ReferencingObject->GetFullName());
+					UPackage* ReferencingPackage = ReferencingObject->GetOutermost();
+					if (!ReferencingPackage->HasAnyPackageFlags(PKG_PlayInEditor) && !ReferencingPackage->HasAnyFlags(RF_Transient))
+					{
+						UE_LOG(LogGarbage, Warning, TEXT("GC detected illegal reference to PIE object from content [possibly via [todo]]:"));
+						UE_LOG(LogGarbage, Warning, TEXT("      PIE object: %s"), *Object->GetFullName());
+						UE_LOG(LogGarbage, Warning, TEXT("  NON-PIE object: %s"), *ReferencingObject->GetFullName());
+					}
 				}
 #endif
 
diff --git a/Engine/Source/Runtime/CoreUObject/Private/UObject/UObjectGlobals.cpp b/Engine/Source/Runtime/CoreUObject/Private/UObject/UObjectGlobals.cpp
--- a/Engine/Source/Runtime/CoreUObject/Private/UObject/UObjectGlobals.cpp
+++ b/Engine/Source/Runtime/CoreUObject/Private/UObject/UObjectGlobals.cpp
@@ -3478,11 +3478,15 @@ private:
 	{
 #if ENABLE_GC_DEBUG_OUTPUT
 		// this message is to help track down culprits behind "Object in PIE world still referenced" errors
-		if ( GIsEditor && !GIsPlayInEditorWorld && !CurrentObject->RootPackageHasAnyFlags(PKG_PlayInEditor) && Object->RootPackageHasAnyFlags(PKG_PlayInEditor) )
+		if ( GIsEditor && !GIsPlayInEditorWorld && !CurrentObject->HasAnyFlags(RF_Transient) && Object->RootPackageHasAnyFlags(PKG_PlayInEditor) )
 		{
-			UE_LOG(LogGarbage, Warning, TEXT("GC detected illegal reference to PIE object from content [possibly via %s]:"), *ReferencingProperty->GetFullName());
-			UE_LOG(LogGarbage, Warning, TEXT("      PIE object: %s"), *Object->GetFullName());
-			UE_LOG(LogGarbage, Warning, TEXT("  NON-PIE object: %s"), *CurrentObject->GetFullName());
+			UPackage* ReferencingPackage = CurrentObject->GetOutermost();
+			if (!ReferencingPackage->HasAnyPackageFlags(PKG_PlayInEditor) && !ReferencingPackage->HasAnyFlags(RF_Transient))
+			{
+				UE_LOG(LogGarbage, Warning, TEXT("GC detected illegal reference to PIE object from content [possibly via %s]:"), *ReferencingProperty->GetFullName());
+				UE_LOG(LogGarbage, Warning, TEXT("      PIE object: %s"), *Object->GetFullName());
+				UE_LOG(LogGarbage, Warning, TEXT("  NON-PIE object: %s"), *CurrentObject->GetFullName());
+			}
 		}
 #endif
 
