CIS fix

diff --git a/Engine/Source/Runtime/CoreUObject/Private/UObject/GarbageCollection.cpp b/Engine/Source/Runtime/CoreUObject/Private/UObject/GarbageCollection.cpp
--- a/Engine/Source/Runtime/CoreUObject/Private/UObject/GarbageCollection.cpp
+++ b/Engine/Source/Runtime/CoreUObject/Private/UObject/GarbageCollection.cpp
@@ -23,6 +23,7 @@
 #include "UObject/UObjectClusters.h"
 #include "HAL/LowLevelMemTracker.h"
 #include "UObject/GarbageCollectionVerification.h"
+#include "UObject/Package.h"
 #include "Async/ParallelFor.h"
 #include "ProfilingDebugging/CsvProfiler.h"
 
