Fix FSoftObjectPath::TryLoad("/game/mapname.PersistentLevel:ActorName") to not call LoadObject on the full subobject path as that can fail in cooked builds. Instead, it will LoadObject the top level object then resolve off of that, which works in all cases
Also, correctly handle PIE fixup for TryLoad

diff --git a/Engine/Source/Runtime/CoreUObject/Private/UObject/SoftObjectPath.cpp b/Engine/Source/Runtime/CoreUObject/Private/UObject/SoftObjectPath.cpp
--- a/Engine/Source/Runtime/CoreUObject/Private/UObject/SoftObjectPath.cpp
+++ b/Engine/Source/Runtime/CoreUObject/Private/UObject/SoftObjectPath.cpp
@@ -350,7 +350,30 @@ UObject* FSoftObjectPath::TryLoad(FUObjectSerializeContext* InLoadContext) const
 
 	if (!IsNull())
 	{
-		LoadedObject = StaticLoadObject(UObject::StaticClass(), nullptr, *ToString(), nullptr, LOAD_None, nullptr, true, InLoadContext);
+		if (IsSubobject())
+		{
+			// For subobjects, it's not safe to call LoadObject directly, so we want to load the parent object and then resolve again
+			FSoftObjectPath TopLevelPath = FSoftObjectPath(AssetPathName, FString());
+			UObject* TopLevelObject = TopLevelPath.TryLoad(InLoadContext);
+
+			// This probably loaded the top-level object, so re-resolve ourselves
+			return ResolveObject();
+		}
+
+		FString PathString = ToString();
+#if WITH_EDITOR
+		if (GPlayInEditorID != INDEX_NONE)
+		{
+			// If we are in PIE and this hasn't already been fixed up, we need to fixup at resolution time. We cannot modify the path as it may be somewhere like a blueprint CDO
+			FSoftObjectPath FixupObjectPath = *this;
+			if (FixupObjectPath.FixupForPIE())
+			{
+				PathString = FixupObjectPath.ToString();
+			}
+		}
+#endif
+
+		LoadedObject = StaticLoadObject(UObject::StaticClass(), nullptr, *PathString, nullptr, LOAD_None, nullptr, true, InLoadContext);
 
 #if WITH_EDITOR
 		// Look at core redirects if we didn't find the object
