Improve comments about SaveGame property flag to make it clear it's not used by the GameplayStatics functions

diff --git a/Engine/Source/Runtime/Core/Public/Serialization/Archive.h b/Engine/Source/Runtime/Core/Public/Serialization/Archive.h
--- a/Engine/Source/Runtime/Core/Public/Serialization/Archive.h
+++ b/Engine/Source/Runtime/Core/Public/Serialization/Archive.h
@@ -1140,6 +1140,7 @@ PRAGMA_ENABLE_DEPRECATION_WARNINGS
 	/**
 	 * Indicates whether this archive is saving or loading game state
 	 *
+	 * @note This is intended for game-specific archives and is not true for any of the build in save methods
 	 * @return true if the archive is dealing with save games, false otherwise.
 	 */
 	bool IsSaveGame()
diff --git a/Engine/Source/Runtime/CoreUObject/Public/UObject/ObjectMacros.h b/Engine/Source/Runtime/CoreUObject/Public/UObject/ObjectMacros.h
--- a/Engine/Source/Runtime/CoreUObject/Public/UObject/ObjectMacros.h
+++ b/Engine/Source/Runtime/CoreUObject/Public/UObject/ObjectMacros.h
@@ -385,7 +385,7 @@ enum EPropertyFlags : uint64
 	CPF_DuplicateTransient				= 0x0000000000200000,	///< Property should always be reset to the default value during any type of duplication (copy/paste, binary duplication, etc.)
 	CPF_SubobjectReference				= 0x0000000000400000,	///< Property contains subobject references (TSubobjectPtr)
 	//CPF_    							= 0x0000000000800000,	///< 
-	CPF_SaveGame						= 0x0000000001000000,	///< Property should be serialized for save games
+	CPF_SaveGame						= 0x0000000001000000,	///< Property should be serialized for save games, this is only checked for game-specific archives with ArIsSaveGame
 	CPF_NoClear							= 0x0000000002000000,	///< Hide clear (and browse) button.
 	//CPF_  							= 0x0000000004000000,	///<
 	CPF_ReferenceParm					= 0x0000000008000000,	///< Value is passed by reference; CPF_OutParam and CPF_Param should also be set.
@@ -936,7 +936,8 @@ namespace UP
 		/// to use on struct properties or parameters.
 		AssetRegistrySearchable,
 
-		/// Property should be serialized for save game.
+		/// Property should be serialized for save games.
+		/// This is only checked for game-specific archives with ArIsSaveGame set
 		SaveGame,
 
 		/// MC Delegates only.  Property should be exposed for calling in blueprint code
diff --git a/Engine/Source/Runtime/Engine/Classes/Kismet/GameplayStatics.h b/Engine/Source/Runtime/Engine/Classes/Kismet/GameplayStatics.h
--- a/Engine/Source/Runtime/Engine/Classes/Kismet/GameplayStatics.h
+++ b/Engine/Source/Runtime/Engine/Classes/Kismet/GameplayStatics.h
@@ -842,6 +842,8 @@ public:
 
 	/** 
 	 * Serialize our USaveGame object into a given array of bytes
+	 * @note This will write out all non-transient properties, the SaveGame property flag is not checked
+	 * 
 	 * @param SaveGameObject	Object that contains data about the save game that we want to write out
 	 * @param OutSaveData		Byte array that is written into
 	 * @return					Whether we successfully wrote data
@@ -871,6 +873,8 @@ public:
 
 	/** 
 	 * Save the contents of the SaveGameObject to a platform-specific save slot/file.
+	 * @note This will write out all non-transient properties, the SaveGame property flag is not checked
+	 *
 	 * @param SaveGameObject	Object that contains data about the save game that we want to write out
 	 * @param SlotName			Name of save game slot to save to.
 	 * @param UserIndex			For some platforms, master user index to identify the user doing the saving.
