Merging //UE4/Dev-Main to Dev-Framework (//UE4/Dev-Framework) @ 6547008
#rnx

diff --git a/Engine/Source/Runtime/CoreUObject/Private/UObject/ReferenceChainSearch.cpp b/Engine/Source/Runtime/CoreUObject/Private/UObject/ReferenceChainSearch.cpp
--- a/Engine/Source/Runtime/CoreUObject/Private/UObject/ReferenceChainSearch.cpp
+++ b/Engine/Source/Runtime/CoreUObject/Private/UObject/ReferenceChainSearch.cpp
@@ -269,7 +269,7 @@ void FReferenceChainSearch::DumpChain(FReferenceChain* Chain)
 			const FNodeReferenceInfo& ReferenceInfo = Chain->GetReferenceInfo(NodeIndex);
 
 			UE_LOG(LogReferenceChain, Log, TEXT("%s%s%s%s"),
-				FCString::Spc(Chain->Num() - NodeIndex - 1),
+				FCString::Spc(FMath::Min<int32>(TCStringSpcHelper<TCHAR>::MAX_SPACES, Chain->Num() - NodeIndex - 1)),
 				*GetObjectFlags(Object),
 				*Object->GetFullName(),
 				*ReferenceInfo.ToString()
@@ -373,14 +373,17 @@ public:
 			else
 			{
 				RefInfo.Type = FReferenceChainSearch::EReferenceType::AddReferencedObjects;
-
 				if (FGCObject::GGCObjectReferencer && (!ReferencingObject || ReferencingObject == FGCObject::GGCObjectReferencer))
 				{
-					// If we have no referencing object, check the global referencer
-					FString ReferencerString;
-					if (FGCObject::GGCObjectReferencer->GetReferencerName(Object, ReferencerString, true))
+					FString RefName;
+					if (FGCObject::GGCObjectReferencer->GetReferencerName(Object, RefName, true))
+					{
+						RefInfo.Type = FReferenceChainSearch::EReferenceType::Property;
+						RefInfo.ReferencerName = FName(*RefName);
+					}
+					else if (ReferencingObject)
 					{
-						RefInfo.ReferencerName = FName(*ReferencerString);
+						RefInfo.ReferencerName = ReferencingObject->GetFName();
 					}
 				}
 				else if (ReferencingObject)
