Fix merge fallout from CL 6542056 combined with 5365249. I'm going to redo this fix in the source branch for 6542056

diff --git a/Engine/Source/Runtime/CoreUObject/Private/Misc/GCObjectReferencer.cpp b/Engine/Source/Runtime/CoreUObject/Private/Misc/GCObjectReferencer.cpp
--- a/Engine/Source/Runtime/CoreUObject/Private/Misc/GCObjectReferencer.cpp
+++ b/Engine/Source/Runtime/CoreUObject/Private/Misc/GCObjectReferencer.cpp
@@ -61,6 +61,11 @@ bool UGCObjectReferencer::GetReferencerName(UObject* Object, FString& OutName, b
 			return false;
 		}
 		OutName = CurrentlySerializingObject->GetReferencerName();
+		FString ReferencerProperty;
+		if (CurrentlySerializingObject->GetReferencerPropertyName(Object, ReferencerProperty))
+		{
+			OutName += TEXT(":") + ReferencerProperty;
+		}
 		return true;
 	}
 
diff --git a/Engine/Source/Runtime/CoreUObject/Private/Serialization/AsyncLoading.cpp b/Engine/Source/Runtime/CoreUObject/Private/Serialization/AsyncLoading.cpp
--- a/Engine/Source/Runtime/CoreUObject/Private/Serialization/AsyncLoading.cpp
+++ b/Engine/Source/Runtime/CoreUObject/Private/Serialization/AsyncLoading.cpp
@@ -5531,9 +5531,9 @@ FAsyncPackage::~FAsyncPackage()
 
 void FAsyncPackage::AddReferencedObjects(FReferenceCollector& Collector)
 {
-	Collector.AddReferencedObjects(ReferencedObjects, FGCObject::GGCObjectReferencer);
-	Collector.AddReferencedObjects(DeferredFinalizeObjects, FGCObject::GGCObjectReferencer);
-	Collector.AddReferencedObjects(PackageObjLoaded, FGCObject::GGCObjectReferencer);
+	Collector.AddReferencedObjects(ReferencedObjects);
+	Collector.AddReferencedObjects(DeferredFinalizeObjects);
+	Collector.AddReferencedObjects(PackageObjLoaded);
 }
 
 void FAsyncPackage::AddObjectReference(UObject* InObject)
diff --git a/Engine/Source/Runtime/CoreUObject/Private/UObject/ReferenceChainSearch.cpp b/Engine/Source/Runtime/CoreUObject/Private/UObject/ReferenceChainSearch.cpp
--- a/Engine/Source/Runtime/CoreUObject/Private/UObject/ReferenceChainSearch.cpp
+++ b/Engine/Source/Runtime/CoreUObject/Private/UObject/ReferenceChainSearch.cpp
@@ -378,7 +378,6 @@ public:
 					FString RefName;
 					if (FGCObject::GGCObjectReferencer->GetReferencerName(Object, RefName, true))
 					{
-						RefInfo.Type = FReferenceChainSearch::EReferenceType::Property;
 						RefInfo.ReferencerName = FName(*RefName);
 					}
 					else if (ReferencingObject)
diff --git a/Engine/Source/Runtime/Engine/Private/Slate/DeferredCleanupSlateBrush.cpp b/Engine/Source/Runtime/Engine/Private/Slate/DeferredCleanupSlateBrush.cpp
--- a/Engine/Source/Runtime/Engine/Private/Slate/DeferredCleanupSlateBrush.cpp
+++ b/Engine/Source/Runtime/Engine/Private/Slate/DeferredCleanupSlateBrush.cpp
@@ -43,7 +43,7 @@ TSharedRef<FDeferredCleanupSlateBrush> FDeferredCleanupSlateBrush::CreateBrush(
 
 void FDeferredCleanupSlateBrush::AddReferencedObjects(FReferenceCollector& Collector)
 {
-	InternalBrush.AddReferencedObjects(Collector, FGCObject::GGCObjectReferencer);
+	InternalBrush.AddReferencedObjects(Collector);
 }
 
 FString FDeferredCleanupSlateBrush::GetReferencerName() const
diff --git a/Engine/Source/Runtime/Engine/Private/StreamableManager.cpp b/Engine/Source/Runtime/Engine/Private/StreamableManager.cpp
--- a/Engine/Source/Runtime/Engine/Private/StreamableManager.cpp
+++ b/Engine/Source/Runtime/Engine/Private/StreamableManager.cpp
@@ -783,7 +783,7 @@ void FStreamableManager::AddReferencedObjects(FReferenceCollector& Collector)
 		FStreamable* Existing = It.Value();
 		if (Existing->Target)
 		{
-			Collector.AddReferencedObject(Existing->Target, FGCObject::GGCObjectReferencer);
+			Collector.AddReferencedObject(Existing->Target);
 		}
 	}
 
@@ -792,7 +792,7 @@ void FStreamableManager::AddReferencedObjects(FReferenceCollector& Collector)
 		FRedirectedPath& Existing = It.Value();
 		if (Existing.LoadedRedirector)
 		{
-			Collector.AddReferencedObject(Existing.LoadedRedirector, FGCObject::GGCObjectReferencer);
+			Collector.AddReferencedObject(Existing.LoadedRedirector);
 		}
 	}
 }
diff --git a/Engine/Source/Runtime/UMG/Private/Slate/SObjectWidget.cpp b/Engine/Source/Runtime/UMG/Private/Slate/SObjectWidget.cpp
--- a/Engine/Source/Runtime/UMG/Private/Slate/SObjectWidget.cpp
+++ b/Engine/Source/Runtime/UMG/Private/Slate/SObjectWidget.cpp
@@ -86,7 +86,7 @@ FString SObjectWidget::GetReferencerName() const
 
 void SObjectWidget::AddReferencedObjects(FReferenceCollector& Collector)
 {
-	Collector.AddReferencedObject(WidgetObject, FGCObject::GGCObjectReferencer);
+	Collector.AddReferencedObject(WidgetObject);
 }
 
 void SObjectWidget::SetPadding(const TAttribute<FMargin>& InMargin)
