Deprecate CreateAbstractDefaultSubobject as it did not work as advertised. CreateDefaultSubobject should now be used and it has been improved to provide a warning when a non-abstract class tries to create an abstract default subobject.

diff --git a/Engine/Source/Runtime/CoreUObject/Private/UObject/Obj.cpp b/Engine/Source/Runtime/CoreUObject/Private/UObject/Obj.cpp
--- a/Engine/Source/Runtime/CoreUObject/Private/UObject/Obj.cpp
+++ b/Engine/Source/Runtime/CoreUObject/Private/UObject/Obj.cpp
@@ -100,12 +100,12 @@ void UObject::EnsureNotRetrievingVTablePtr() const
 	UE_CLOG(GIsRetrievingVTablePtr, LogCore, Fatal, TEXT("We are currently retrieving VTable ptr. Please use FVTableHelper constructor instead."));
 }
 
-UObject* UObject::CreateDefaultSubobject(FName SubobjectFName, UClass* ReturnType, UClass* ClassToCreateByDefault, bool bIsRequired, bool bAbstract, bool bIsTransient)
+UObject* UObject::CreateDefaultSubobject(FName SubobjectFName, UClass* ReturnType, UClass* ClassToCreateByDefault, bool bIsRequired, bool bIsTransient)
 {
 	FObjectInitializer* CurrentInitializer = FUObjectThreadContext::Get().TopInitializer();
 	UE_CLOG(!CurrentInitializer, LogObj, Fatal, TEXT("No object initializer found during construction."));
 	UE_CLOG(CurrentInitializer->Obj != this, LogObj, Fatal, TEXT("Using incorrect object initializer."));
-	return CurrentInitializer->CreateDefaultSubobject(this, SubobjectFName, ReturnType, ClassToCreateByDefault, bIsRequired, bAbstract, bIsTransient);
+	return CurrentInitializer->CreateDefaultSubobject(this, SubobjectFName, ReturnType, ClassToCreateByDefault, bIsRequired, bIsTransient);
 }
 
 UObject* UObject::CreateEditorOnlyDefaultSubobjectImpl(FName SubobjectName, UClass* ReturnType, bool bTransient)
diff --git a/Engine/Source/Runtime/CoreUObject/Private/UObject/UObjectGlobals.cpp b/Engine/Source/Runtime/CoreUObject/Private/UObject/UObjectGlobals.cpp
--- a/Engine/Source/Runtime/CoreUObject/Private/UObject/UObjectGlobals.cpp
+++ b/Engine/Source/Runtime/CoreUObject/Private/UObject/UObjectGlobals.cpp
@@ -3639,7 +3639,7 @@ UScriptStruct* GetFallbackStruct()
 	return TBaseStructure<FFallbackStruct>::Get();
 }
 
-UObject* FObjectInitializer::CreateDefaultSubobject(UObject* Outer, FName SubobjectFName, UClass* ReturnType, UClass* ClassToCreateByDefault, bool bIsRequired, bool bAbstract, bool bIsTransient) const
+UObject* FObjectInitializer::CreateDefaultSubobject(UObject* Outer, FName SubobjectFName, UClass* ReturnType, UClass* ClassToCreateByDefault, bool bIsRequired, bool bIsTransient) const
 {
 	UE_CLOG(!FUObjectThreadContext::Get().IsInConstructor, LogClass, Fatal, TEXT("Subobjects cannot be created outside of UObject constructors. UObject constructing subobjects cannot be created using new or placement new operator."));
 	if (SubobjectFName == NAME_None)
@@ -3658,8 +3658,15 @@ UObject* FObjectInitializer::CreateDefaultSubobject(UObject* Outer, FName Subobj
 	{
 		check(OverrideClass->IsChildOf(ReturnType));
 
-		// Abstract sub-objects are only allowed when explicitly created with CreateAbstractDefaultSubobject.
-		if (!OverrideClass->HasAnyClassFlags(CLASS_Abstract) || !bAbstract)
+		if (OverrideClass->HasAnyClassFlags(CLASS_Abstract))
+		{
+			// Attempts to create an abstract class will return null. If it is not optional or the owning class is not also abstract report a warning.
+			if (!bIsRequired && !Outer->GetClass()->HasAnyClassFlags(CLASS_Abstract))
+			{
+				UE_LOG(LogClass, Warning, TEXT("Required default subobject %s not created as requested class %s is abstract. Returning null."), *SubobjectFName.ToString(), *OverrideClass->GetName());
+			}
+		}
+		else
 		{
 			UObject* Template = OverrideClass->GetDefaultObject(); // force the CDO to be created if it hasn't already
 			EObjectFlags SubobjectFlags = Outer->GetMaskedFlags(RF_PropagateToSubObjects);
@@ -3725,7 +3732,7 @@ UObject* FObjectInitializer::CreateEditorOnlyDefaultSubobject(UObject* Outer, FN
 #if WITH_EDITOR
 	if (GIsEditor)
 	{
-		UObject* EditorSubobject = CreateDefaultSubobject(Outer, SubobjectName, ReturnType, ReturnType, /*bIsRequired =*/ false, /*bIsAbstract =*/ false, bTransient);
+		UObject* EditorSubobject = CreateDefaultSubobject(Outer, SubobjectName, ReturnType, ReturnType, /*bIsRequired =*/ false, bTransient);
 		if (EditorSubobject)
 		{
 			EditorSubobject->MarkAsEditorOnlySubobject();
diff --git a/Engine/Source/Runtime/CoreUObject/Public/UObject/Object.h b/Engine/Source/Runtime/CoreUObject/Public/UObject/Object.h
--- a/Engine/Source/Runtime/CoreUObject/Public/UObject/Object.h
+++ b/Engine/Source/Runtime/CoreUObject/Public/UObject/Object.h
@@ -72,8 +72,14 @@ class COREUOBJECT_API UObject : public UObjectBaseUtility
 	/** DO NOT USE. This constructor is for internal usage only for hot-reload purposes. */
 	UObject(FVTableHelper& Helper);
 
+	UE_DEPRECATED(4.23, "CreateDefaultSubobject no longer takes bAbstract as a parameter.")
+	UObject* CreateDefaultSubobject(FName SubobjectFName, UClass* ReturnType, UClass* ClassToCreateByDefault, bool bIsRequired, bool bAbstract, bool bIsTransient)
+	{
+		return CreateDefaultSubobject(SubobjectFName, ReturnType, ClassToCreateByDefault, bIsRequired, bIsTransient);
+	}
+
 	/** Utility function for templates below */
-	UObject* CreateDefaultSubobject(FName SubobjectFName, UClass* ReturnType, UClass* ClassToCreateByDefault, bool bIsRequired, bool bAbstract, bool bIsTransient);
+	UObject* CreateDefaultSubobject(FName SubobjectFName, UClass* ReturnType, UClass* ClassToCreateByDefault, bool bIsRequired, bool bIsTransient);
 
 	/**
 	 * Create a component or subobject only to be used with the editor. They will be stripped out in packaged builds.
@@ -98,7 +104,7 @@ class COREUOBJECT_API UObject : public UObjectBaseUtility
 	TReturnType* CreateDefaultSubobject(FName SubobjectName, bool bTransient = false)
 	{
 		UClass* ReturnType = TReturnType::StaticClass();
-		return static_cast<TReturnType*>(CreateDefaultSubobject(SubobjectName, ReturnType, ReturnType, /*bIsRequired =*/ true, /*bIsAbstract =*/ false, bTransient));
+		return static_cast<TReturnType*>(CreateDefaultSubobject(SubobjectName, ReturnType, ReturnType, /*bIsRequired =*/ true, bTransient));
 	}
 	
 	/**
@@ -111,7 +117,7 @@ class COREUOBJECT_API UObject : public UObjectBaseUtility
 	template<class TReturnType, class TClassToConstructByDefault>
 	TReturnType* CreateDefaultSubobject(FName SubobjectName, bool bTransient = false)
 	{
-		return static_cast<TReturnType*>(CreateDefaultSubobject(SubobjectName, TReturnType::StaticClass(), TClassToConstructByDefault::StaticClass(), /*bIsRequired =*/ true, /*bIsAbstract =*/ false, bTransient));
+		return static_cast<TReturnType*>(CreateDefaultSubobject(SubobjectName, TReturnType::StaticClass(), TClassToConstructByDefault::StaticClass(), /*bIsRequired =*/ true, bTransient));
 	}
 	
 	/**
@@ -125,20 +131,21 @@ class COREUOBJECT_API UObject : public UObjectBaseUtility
 	TReturnType* CreateOptionalDefaultSubobject(FName SubobjectName, bool bTransient = false)
 	{
 		UClass* ReturnType = TReturnType::StaticClass();
-		return static_cast<TReturnType*>(CreateDefaultSubobject(SubobjectName, ReturnType, ReturnType, /*bIsRequired =*/ false, /*bIsAbstract =*/ false, bTransient));
+		return static_cast<TReturnType*>(CreateDefaultSubobject(SubobjectName, ReturnType, ReturnType, /*bIsRequired =*/ false, bTransient));
 	}
 	
 	/**
-	 * Create a subobject that has the Abstract class flag, child classes are expected to override this by calling CreateDEfaultObject with the same name and a non-abstract class.
+	 * Create a subobject that has the Abstract class flag, child classes are expected to override this by calling SetDefaultSubobjectClass with the same name and a non-abstract class.
 	 * @param	TReturnType					Class of return type, all overrides must be of this type
 	 * @param	SubobjectName				Name of the new component
 	 * @param	bTransient					True if the component is being assigned to a transient property. This does not make the component itself transient, but does stop it from inheriting parent defaults
 	 */
 	template<class TReturnType>
+	UE_DEPRECATED(4.23, "CreateAbstract did not work as intended and has been deprecated in favor of CreateDefaultObject")
 	TReturnType* CreateAbstractDefaultSubobject(FName SubobjectName, bool bTransient = false)
 	{
 		UClass* ReturnType = TReturnType::StaticClass();
-		return static_cast<TReturnType*>(CreateDefaultSubobject(SubobjectName, ReturnType, ReturnType, /*bIsRequired =*/ true, /*bIsAbstract =*/ true, bTransient));
+		return static_cast<TReturnType*>(CreateDefaultSubobject(SubobjectName, ReturnType, ReturnType, /*bIsRequired =*/ true, bTransient));
 	}
 
 	/**
diff --git a/Engine/Source/Runtime/CoreUObject/Public/UObject/UObjectGlobals.h b/Engine/Source/Runtime/CoreUObject/Public/UObject/UObjectGlobals.h
--- a/Engine/Source/Runtime/CoreUObject/Public/UObject/UObjectGlobals.h
+++ b/Engine/Source/Runtime/CoreUObject/Public/UObject/UObjectGlobals.h
@@ -814,7 +814,7 @@ public:
 	TReturnType* CreateDefaultSubobject(UObject* Outer, FName SubobjectName, bool bTransient = false) const
 	{
 		UClass* ReturnType = TReturnType::StaticClass();
-		return static_cast<TReturnType*>(CreateDefaultSubobject(Outer, SubobjectName, ReturnType, ReturnType, /*bIsRequired =*/ true, /*bIsAbstract =*/ false, bTransient));
+		return static_cast<TReturnType*>(CreateDefaultSubobject(Outer, SubobjectName, ReturnType, ReturnType, /*bIsRequired =*/ true, bTransient));
 	}
 
 	/**
@@ -829,7 +829,7 @@ public:
 	TReturnType* CreateOptionalDefaultSubobject(UObject* Outer, FName SubobjectName, bool bTransient = false) const
 	{
 		UClass* ReturnType = TReturnType::StaticClass();
-		return static_cast<TReturnType*>(CreateDefaultSubobject(Outer, SubobjectName, ReturnType, ReturnType, /*bIsRequired =*/ false, /*bIsAbstract =*/ false, bTransient));
+		return static_cast<TReturnType*>(CreateDefaultSubobject(Outer, SubobjectName, ReturnType, ReturnType, /*bIsRequired =*/ false, bTransient));
 	}
 
 	/**
@@ -840,11 +840,19 @@ public:
 	 * @param	SubobjectName				name of the new component
 	 * @param bTransient		true if the component is being assigned to a transient property
 	 */
+
+	/**
+	 * Create a subobject that has the Abstract class flag, child classes are expected to override this by calling SetDefaultSubobjectClass with the same name and a non-abstract class.
+	 * @param	TReturnType					Class of return type, all overrides must be of this type
+	 * @param	SubobjectName				Name of the new component
+	 * @param	bTransient					True if the component is being assigned to a transient property. This does not make the component itself transient, but does stop it from inheriting parent defaults
+	 */
 	template<class TReturnType>
+	UE_DEPRECATED(4.23, "CreateAbstract did not work as intended and has been deprecated in favor of CreateDefaultObject")
 	TReturnType* CreateAbstractDefaultSubobject(UObject* Outer, FName SubobjectName, bool bTransient = false) const
 	{
 		UClass* ReturnType = TReturnType::StaticClass();
-		return static_cast<TReturnType*>(CreateDefaultSubobject(Outer, SubobjectName, ReturnType, ReturnType, /*bIsRequired =*/ true, /*bIsAbstract =*/ true, bTransient));
+		return static_cast<TReturnType*>(CreateDefaultSubobject(Outer, SubobjectName, ReturnType, ReturnType, /*bIsRequired =*/ true, bTransient));
 	}
 
 	/** 
@@ -858,7 +866,7 @@ public:
 	template<class TReturnType, class TClassToConstructByDefault> 
 	TReturnType* CreateDefaultSubobject(UObject* Outer, FName SubobjectName, bool bTransient = false) const 
 	{ 
-		return static_cast<TReturnType*>(CreateDefaultSubobject(Outer, SubobjectName, TReturnType::StaticClass(), TClassToConstructByDefault::StaticClass(), /*bIsRequired =*/ true, /*bIsAbstract =*/ false, bTransient));
+		return static_cast<TReturnType*>(CreateDefaultSubobject(Outer, SubobjectName, TReturnType::StaticClass(), TClassToConstructByDefault::StaticClass(), /*bIsRequired =*/ true, bTransient));
 	}
 
 	/**
@@ -891,10 +899,17 @@ public:
 	 * @param	TClassToConstructByDefault	if the derived class has not overridden, create a component of this type (default is TReturnType)
 	 * @param	Outer						outer to construct the subobject in
 	 * @param	SubobjectName				name of the new component
-	 * @param bIsRequired			true if the component is required and will always be created even if DoNotCreateDefaultSubobject was sepcified.
+	 * @param bIsRequired			true if the component is required and will always be created even if DoNotCreateDefaultSubobject was specified.
 	 * @param bIsTransient		true if the component is being assigned to a transient property
 	 */
-	UObject* CreateDefaultSubobject(UObject* Outer, FName SubobjectFName, UClass* ReturnType, UClass* ClassToCreateByDefault, bool bIsRequired, bool bAbstract, bool bIsTransient) const;
+	UObject* CreateDefaultSubobject(UObject* Outer, FName SubobjectFName, UClass* ReturnType, UClass* ClassToCreateByDefault, bool bIsRequired, bool bIsTransient) const;
+
+	UE_DEPRECATED(4.23, "CreateDefaultSubobject no longer takes bAbstract as a parameter.")
+	UObject* CreateDefaultSubobject(UObject* Outer, FName SubobjectFName, UClass* ReturnType, UClass* ClassToCreateByDefault, bool bIsRequired, bool bAbstract, bool bIsTransient) const
+	{
+		return CreateDefaultSubobject(Outer, SubobjectFName, ReturnType, ClassToCreateByDefault, bIsRequired, bIsTransient);
+	}
+
 
 	/**
 	 * Sets the class of a subobject for a base class
diff --git a/Engine/Source/Runtime/Engine/Private/Components/ReflectionCaptureComponent.cpp b/Engine/Source/Runtime/Engine/Private/Components/ReflectionCaptureComponent.cpp
--- a/Engine/Source/Runtime/Engine/Private/Components/ReflectionCaptureComponent.cpp
+++ b/Engine/Source/Runtime/Engine/Private/Components/ReflectionCaptureComponent.cpp
@@ -116,7 +116,7 @@ void UReflectionCaptureComponent::PropagateLightingScenarioChange()
 AReflectionCapture::AReflectionCapture(const FObjectInitializer& ObjectInitializer)
 	: Super(ObjectInitializer)
 {
-	CaptureComponent = CreateAbstractDefaultSubobject<UReflectionCaptureComponent>(TEXT("NewReflectionComponent"));
+	CaptureComponent = CreateDefaultSubobject<UReflectionCaptureComponent>(TEXT("NewReflectionComponent"));
 
 	bCanBeInCluster = true;
 
diff --git a/Engine/Source/Runtime/Engine/Private/Light.cpp b/Engine/Source/Runtime/Engine/Private/Light.cpp
--- a/Engine/Source/Runtime/Engine/Private/Light.cpp
+++ b/Engine/Source/Runtime/Engine/Private/Light.cpp
@@ -16,7 +16,7 @@
 ALight::ALight(const FObjectInitializer& ObjectInitializer)
 	: Super(ObjectInitializer)
 {
-	LightComponent = CreateAbstractDefaultSubobject<ULightComponent>(TEXT("LightComponent0"));
+	LightComponent = CreateDefaultSubobject<ULightComponent>(TEXT("LightComponent0"));
 
 	bCollideWhenPlacing = true;
 	SpawnCollisionHandlingMethod = ESpawnActorCollisionHandlingMethod::AdjustIfPossibleButDontSpawnIfColliding;
diff --git a/Engine/Source/Runtime/Engine/Private/TriggerBase.cpp b/Engine/Source/Runtime/Engine/Private/TriggerBase.cpp
--- a/Engine/Source/Runtime/Engine/Private/TriggerBase.cpp
+++ b/Engine/Source/Runtime/Engine/Private/TriggerBase.cpp
@@ -14,7 +14,7 @@ ATriggerBase::ATriggerBase(const FObjectInitializer& ObjectInitializer)
 
 	// ATriggerBase is requesting UShapeComponent which is abstract, however it is responsibility
 	// of a derived class to override this type with ObjectInitializer.SetDefaultSubobjectClass.
-	CollisionComponent = CreateAbstractDefaultSubobject<UShapeComponent>(TEXT("CollisionComp"));
+	CollisionComponent = CreateDefaultSubobject<UShapeComponent>(TEXT("CollisionComp"));
 	if (CollisionComponent)
 	{
 		RootComponent = CollisionComponent;
