Add editor only commands to audit bytecode contents

diff --git a/Engine/Config/BaseInput.ini b/Engine/Config/BaseInput.ini
--- a/Engine/Config/BaseInput.ini
+++ b/Engine/Config/BaseInput.ini
@@ -204,6 +204,10 @@ MaxScrollbackSize=1024
 +ManualAutoCompleteList=(Command="Stat SOUNDWAVES",Desc="Shows active SoundWaves")
 +ManualAutoCompleteList=(Command="Stat SOUNDCUES",Desc="Shows active SoundCues")
 +ManualAutoCompleteList=(Command="Stat SOUNDS",Desc="<?> <sort=distance|class|name|waves|default> <-debug> <off> Shows active SoundCues and SoundWaves")
++ManualAutoCompleteList=(Command="ScriptAudit LongestFunctions",Desc="List functions that contain the most bytecode - optionally include # of entries to list")
++ManualAutoCompleteList=(Command="ScriptAudit FrequentFunctionsCalled",Desc="List functions that are most frequently called from bytecode - optionally include # of entries to list")
++ManualAutoCompleteList=(Command="ScriptAudit FrequentInstructions",Desc="List most frequently used instructions - optionally include # of entries to list")
++ManualAutoCompleteList=(Command="ScriptAudit TotalBytecodeSize",Desc="Gather total size of bytecode in bytes of currently loaded functions")
 +ManualAutoCompleteList=(Command="Audio3dVisualize",Desc="Shows locations of sound sources playing (white text) and their left and right channel locations respectively (red and green).  Virtualized loops (if enabled) display in blue.")
 +ManualAutoCompleteList=(Command="StartMovieCapture",Desc=)
 +ManualAutoCompleteList=(Command="StopMovieCapture",Desc=)
diff --git a/Engine/Source/Runtime/Core/Public/Misc/CoreMisc.h b/Engine/Source/Runtime/Core/Public/Misc/CoreMisc.h
--- a/Engine/Source/Runtime/Core/Public/Misc/CoreMisc.h
+++ b/Engine/Source/Runtime/Core/Public/Misc/CoreMisc.h
@@ -265,6 +265,14 @@ struct CORE_API FScopedScriptExceptionHandler
 	#endif
 #endif
 
+#ifndef SCRIPT_AUDIT_ROUTINES
+	#if !(UE_BUILD_SHIPPING || UE_BUILD_TEST)
+		#define SCRIPT_AUDIT_ROUTINES 1
+	#else
+		#define SCRIPT_AUDIT_ROUTINES 0
+	#endif
+#endif
+
 #if DO_BLUEPRINT_GUARD
 struct FFrame;
 
diff --git a/Engine/Source/Runtime/CoreUObject/Private/UObject/ScriptCore.cpp b/Engine/Source/Runtime/CoreUObject/Private/UObject/ScriptCore.cpp
--- a/Engine/Source/Runtime/CoreUObject/Private/UObject/ScriptCore.cpp
+++ b/Engine/Source/Runtime/CoreUObject/Private/UObject/ScriptCore.cpp
@@ -82,6 +82,23 @@ COREUOBJECT_API void GInitRunaway()
 COREUOBJECT_API void GInitRunaway() {}
 #endif
 
+#define STORE_INSTRUCTION_NAMES SCRIPT_AUDIT_ROUTINES
+
+#if STORE_INSTRUCTION_NAMES
+const char* GNativeFuncNames[EX_Max];
+
+#define STORE_INSTRUCTION_NAME(inst) \
+static struct F##inst##Registrar \
+{ \
+	F##inst##Registrar() \
+	{ \
+		GNativeFuncNames[inst] = #inst; \
+	} \
+} ##inst##RegistrarInst;
+#else
+#define STORE_INSTRUCTION_NAME(inst)
+#endif//STORE_INSTRUCTION_NAMES
+
 #define IMPLEMENT_FUNCTION(func) \
 	static FNativeFunctionRegistrar UObject##func##Registar(UObject::StaticClass(),#func,&UObject::func);
 
@@ -90,6 +107,7 @@ COREUOBJECT_API void GInitRunaway() {}
 	static uint8 UObject##func##CastTemp = GRegisterCast( CastIndex, &UObject::func );
 
 #define IMPLEMENT_VM_FUNCTION(BytecodeIndex, func) \
+	STORE_INSTRUCTION_NAME(BytecodeIndex) \
 	IMPLEMENT_FUNCTION(func) \
 	static uint8 UObject##func##BytecodeTemp = GRegisterNative( BytecodeIndex, &UObject::func );
 
@@ -1289,6 +1307,431 @@ FBlueprintEventTimer::FScopedNativeTimer::~FScopedNativeTimer()
 
 #endif
 
+#if SCRIPT_AUDIT_ROUTINES
+
+// heap would be more time efficient:
+template<typename T>
+void NBest(TArray<T>& OutBest, const T& NewEntry, TFunctionRef<bool(const T&, const T&)> IsBetter)
+{
+	if(IsBetter(NewEntry, OutBest.Last()))
+	{
+		// find insertion point:
+		int32 InsertIdx = INDEX_NONE;
+		// O(n):
+		for(int32 I = 0; I < OutBest.Num(); ++I)
+		{
+			if(IsBetter( NewEntry, OutBest[I] ))
+			{
+				InsertIdx = I;
+				break;
+			}
+		}
+
+		// O(n):
+		OutBest.Insert(NewEntry, InsertIdx);
+		OutBest.Pop();
+	}
+}
+
+static void OutputLongestFunctions(FOutputDevice& Ar, int32 Num)
+{
+	// max heap would be more efficient
+	TArray<UFunction*> LongestFunctions;
+	LongestFunctions.AddDefaulted(Num);
+
+	for( TObjectIterator<UClass> It; It; ++It)
+	{
+		UClass* BPGC = *It;
+		for(TFieldIterator<UFunction> FuncIt(BPGC, EFieldIteratorFlags::ExcludeSuper); FuncIt; ++FuncIt)
+		{
+			UFunction* Fn = *FuncIt;
+			int32 LenScript = Fn->Script.Num();
+
+			NBest<UFunction*>(LongestFunctions, Fn, 
+				[LenScript](UFunction* A, UFunction* B)
+				{
+					return B == nullptr || LenScript > B->Script.Num();
+				}
+			);
+		}
+	}
+
+	if(LongestFunctions.Num() == 0)
+	{
+		Ar.Log(TEXT("No script functions found when looking for longest functions."));
+	}
+	else
+	{
+		for(UFunction* Fn : LongestFunctions)
+		{
+			if(!Fn)
+			{
+				break;
+			}
+
+			Ar.Logf(TEXT("%s %s %d"), *Fn->GetName(), *Fn->GetOuter()->GetName(), Fn->Script.Num());
+		}
+	}
+}
+
+static void OutputMostFrequentlyCalledFunctions(FOutputDevice& Ar, int32 Num)
+{
+	// Script serialization is recursive and requires certain symbols (e.g. Script, a reference
+	// to the bytecode), so we declare a type so that we have some scope:
+	struct FCallFrequencyCounter
+	{
+		FCallFrequencyCounter(TArray<uint8>& InScript)
+			: Script(InScript)
+		{
+		}
+
+		TArray<uint8>& Script;
+		TMap<UFunction*, int32>* FunctionCallCounts = nullptr;
+		// Could try and get more context on vcalls, but for
+		// this macro auditing tool name should be enough:
+		TMap<FName, int32>* VirtualFunctionCallCounts = nullptr;
+
+		void* GetLinker() { return nullptr; }
+
+		EExprToken SerializeExpr(int32& iCode, FArchive& Ar)
+		{
+			#define SERIALIZEEXPR_INC
+			#define SERIALIZEEXPR_AUTO_UNDEF_XFER_MACROS
+			
+				if(iCode < Script.Num())
+				{
+					switch((EExprToken)Script[iCode])
+					{
+					case EX_CallMath:
+					case EX_LocalFinalFunction:
+					case EX_FinalFunction:
+						{
+							// peak UFunction*:
+							if(FunctionCallCounts)
+							{
+								UFunction* Fn = nullptr;
+								FMemory::Memcpy( &Fn, &Script[iCode+1], sizeof(UFunction*) );
+								if(ensure(Fn))
+								{
+									check(Fn->IsValidLowLevel());
+									FunctionCallCounts->FindOrAdd(Fn)++;
+								}
+							}
+						}
+						break;
+					case EX_VirtualFunction:
+					case EX_LocalVirtualFunction:
+						{
+							// peak function name:
+							if(VirtualFunctionCallCounts)
+							{
+								FScriptName ScriptName;
+								FMemory::Memcpy( &ScriptName, &Script[iCode+1], sizeof(FScriptName) );
+								VirtualFunctionCallCounts->FindOrAdd(ScriptNameToName(ScriptName))++;
+							}
+						}
+						break;
+
+					}
+				}
+
+			#include "UObject/ScriptSerialization.h"
+				return Expr;
+			#undef SERIALIZEEXPR_INC
+			#undef SERIALIZEEXPR_AUTO_UNDEF_XFER_MACROS
+		}
+
+		void CountCalls(TMap<UFunction*, int32>* InFunctionCallCounts, TMap<FName, int32>* InVirtualFunctionCallCounts)
+		{
+			FunctionCallCounts = InFunctionCallCounts;
+			VirtualFunctionCallCounts = InVirtualFunctionCallCounts;
+
+			int32 iCode = 0;
+			const int32 ScriptSizeBytes = Script.Num();
+			FArchive DummyArchive;
+
+			while (iCode < ScriptSizeBytes)
+			{
+				SerializeExpr(iCode, DummyArchive);
+			}
+		}
+	};
+
+	TMap<UFunction*, int32> FunctionCallCounts;
+	TMap<FName, int32> VirtualFunctionCallCounts;
+
+	for( TObjectIterator<UClass> It; It; ++It)
+	{
+		UClass* BPGC = *It;
+		for(TFieldIterator<UFunction> FuncIt(BPGC, EFieldIteratorFlags::ExcludeSuper); FuncIt; ++FuncIt)
+		{
+			UFunction* Fn = *FuncIt;
+
+			// disassem and log function calls:
+			FCallFrequencyCounter Counter(Fn->Script);
+			Counter.CountCalls(&FunctionCallCounts, &VirtualFunctionCallCounts);
+		}
+	}
+
+	// sort by # calls:
+	{
+		TArray< TPair<UFunction*, int32> > FunctionCallsSorted;
+		FunctionCallsSorted.AddDefaulted(Num);
+		for(const TPair<UFunction*, int32>& Calls : FunctionCallCounts )
+		{
+			NBest<TPair<UFunction*, int32>>(FunctionCallsSorted, Calls, 
+				[](const TPair<UFunction*, int32>& A, const TPair<UFunction*, int32>& B) -> bool
+				{
+					return B.Key == nullptr || A.Value > B.Value;
+				}
+			);
+		}
+
+		if(FunctionCallsSorted.Num())
+		{
+			Ar.Logf(TEXT("Top %d function call targets"), FunctionCallsSorted.Num());
+			for(TPair<UFunction*, int32>& Calls : FunctionCallsSorted )
+			{
+				if(Calls.Key == nullptr)
+				{
+					break;
+				}
+
+				Ar.Logf(TEXT("%s %s %d"), *Calls.Key->GetName(), *Calls.Key->GetOuter()->GetName(), Calls.Value);
+			}
+		}
+		else
+		{
+			Ar.Log(TEXT("No function call instructions found in memory"));
+		}
+	}
+
+	{
+		TArray< TPair<FName, int32>> VirtualFunctionCallsSorted;
+		VirtualFunctionCallsSorted.AddDefaulted(Num);
+		for(const TPair<FName, int32>& Calls : VirtualFunctionCallCounts )
+		{
+			NBest<TPair<FName, int32>>(VirtualFunctionCallsSorted, Calls, 
+				[](const TPair<FName, int32>& A, const TPair<FName, int32>& B) -> bool
+				{
+					return B.Key == FName() || A.Value > B.Value;
+				}
+			);
+		}
+
+		if(VirtualFunctionCallsSorted.Num())
+		{
+			Ar.Logf(TEXT("Top %d virtual function call targets"), VirtualFunctionCallsSorted.Num());
+			for(TPair<FName, int32>& Calls : VirtualFunctionCallsSorted )
+			{
+				if(Calls.Key == FName())
+				{
+					break;
+				}
+
+				Ar.Logf(TEXT("%s %d"), *(Calls.Key.ToString()), Calls.Value);
+			}
+		}
+		else
+		{
+			Ar.Log(TEXT("No virtual function call instructions in memory"));
+		}
+	}
+}
+
+static void OutputMostFrequentlyUsedInstructions(FOutputDevice& Ar, int32 Num)
+{
+	// Script serialization is recursive and requires certain symbols (e.g. Script, a reference
+	// to the bytecode), so we declare a type so that we have some scope:
+	struct FInstructionFrequencyCounter
+	{
+		FInstructionFrequencyCounter(TArray<uint8>& InScript)
+			: Script(InScript)
+		{
+		}
+		
+		TArray<uint8>& Script;
+		TMap<EExprToken, int32>* InstructionCallCounts;
+
+		void* GetLinker() { return nullptr; }
+
+		
+		EExprToken SerializeExpr(int32& iCode, FArchive& Ar)
+		{
+			#define SERIALIZEEXPR_INC
+			#define SERIALIZEEXPR_AUTO_UNDEF_XFER_MACROS
+			
+				if(iCode < Script.Num())
+				{
+					if(InstructionCallCounts)
+					{
+						InstructionCallCounts->FindOrAdd((EExprToken)Script[iCode])++;
+					}
+				}
+
+			#include "UObject/ScriptSerialization.h"
+				return Expr;
+			#undef SERIALIZEEXPR_INC
+			#undef SERIALIZEEXPR_AUTO_UNDEF_XFER_MACROS
+		}
+
+		void CountInstructions(TMap<EExprToken, int32>* InInstructionCallCounts)
+		{
+			InstructionCallCounts = InInstructionCallCounts;
+
+			int32 iCode = 0;
+			const int32 ScriptSizeBytes = Script.Num();
+			FArchive DummyArchive;
+
+			while (iCode < ScriptSizeBytes)
+			{
+				SerializeExpr(iCode, DummyArchive);
+			}
+		}
+	};
+
+	
+	TMap<EExprToken, int32> InstructionCallCounts;
+
+	for( TObjectIterator<UClass> It; It; ++It)
+	{
+		UClass* BPGC = *It;
+		for(TFieldIterator<UFunction> FuncIt(BPGC, EFieldIteratorFlags::ExcludeSuper); FuncIt; ++FuncIt)
+		{
+			UFunction* Fn = *FuncIt;
+
+			// disassem and log function calls:
+			FInstructionFrequencyCounter Counter(Fn->Script);
+			Counter.CountInstructions(&InstructionCallCounts);
+		}
+	}
+
+	// sort by #:
+	{
+		TArray<TPair<EExprToken, int32>> InstructionCountsSorted;
+		InstructionCountsSorted.AddDefaulted(Num);
+		
+		for(const TPair<EExprToken, int32>& Instruction : InstructionCallCounts )
+		{
+			NBest<TPair<EExprToken, int32>>(InstructionCountsSorted, Instruction, 
+				[](const TPair<EExprToken, int32>& A, const TPair<EExprToken, int32>& B) -> bool
+				{
+					return A.Value > B.Value;
+				}
+			);
+		}
+
+		if(InstructionCountsSorted.Num())
+		{
+			Ar.Logf(TEXT("Top %d bytecode instructions"), InstructionCountsSorted.Num());
+			for(TPair<EExprToken, int32>& Instruction : InstructionCountsSorted )
+			{
+				if(Instruction.Value == 0)
+				{
+					break;
+				}
+
+#if STORE_INSTRUCTION_NAMES
+				if(GNativeFuncNames[Instruction.Key])
+				{
+					FString AsString = GNativeFuncNames[Instruction.Key];
+					Ar.Logf(TEXT("%s %d"), *AsString, Instruction.Value);
+				}
+				else
+				{
+					Ar.Logf(TEXT("0x%x %d"), Instruction.Key, Instruction.Value);
+				}
+#else
+				Ar.Logf(TEXT("0x%x %d"), Instruction.Key, Instruction.Value);
+#endif
+			}
+		}
+		else
+		{
+			Ar.Log(TEXT("No instructions found in memory"));
+		}
+	}
+}
+
+static void OutputTotalBytecodeSize(FOutputDevice& Ar)
+{
+	uint32 TotalSize = 0;
+	
+	for( TObjectIterator<UClass> It; It; ++It)
+	{
+		UClass* BPGC = *It;
+		for(TFieldIterator<UFunction> FuncIt(BPGC, EFieldIteratorFlags::ExcludeSuper); FuncIt; ++FuncIt)
+		{
+			UFunction* Fn = *FuncIt;
+
+			TotalSize += Fn->Script.Num();
+		}
+	}
+
+	Ar.Logf(TEXT("Total bytecode size: %d"), TotalSize);
+}
+
+struct FScriptAuditExec 
+	: public FSelfRegisteringExec
+{
+	// FSelfRegisteringExec:
+	virtual bool Exec(UWorld* InWorld, const TCHAR* Cmd, FOutputDevice& Ar) override;
+} ScriptAudit;
+
+bool FScriptAuditExec::Exec(UWorld* InWorld, const TCHAR* Cmd, FOutputDevice& Ar)
+{
+	if (FParse::Command(&Cmd, TEXT("ScriptAudit")))
+	{
+		FString ParsedCommand = FParse::Token(Cmd, 0);
+
+		if (ParsedCommand.Equals(TEXT("LongestFunctions"),ESearchCase::IgnoreCase))
+		{
+			int32 NumToOutput = 20;
+
+			FString Num = FParse::Token(Cmd, 0);
+			if(!Num.IsEmpty())
+			{
+				NumToOutput = FCString::Atoi(*Num);
+			}
+			OutputLongestFunctions(Ar, NumToOutput);
+			return true;
+		}
+		else if (ParsedCommand.Equals(TEXT("FrequentFunctionsCalled"),ESearchCase::IgnoreCase))
+		{
+			int32 NumToOutput = 20;
+
+			FString Num = FParse::Token(Cmd, 0);
+			if(!Num.IsEmpty())
+			{
+				NumToOutput = FCString::Atoi(*Num);
+			}
+			OutputMostFrequentlyCalledFunctions(Ar, NumToOutput);
+			return true;
+		}
+		else if (ParsedCommand.Equals(TEXT("FrequentInstructions"),ESearchCase::IgnoreCase))
+		{
+			int32 NumToOutput = 20;
+
+			FString Num = FParse::Token(Cmd, 0);
+			if(!Num.IsEmpty())
+			{
+				NumToOutput = FCString::Atoi(*Num);
+			}
+			OutputMostFrequentlyUsedInstructions(Ar, NumToOutput);
+			return true;
+		}
+		else if (ParsedCommand.Equals(TEXT("TotalBytecodeSize"),ESearchCase::IgnoreCase))
+		{
+			OutputTotalBytecodeSize(Ar);
+			return true;
+		}
+	}
+
+	return false;
+}
+
+#endif //SCRIPT_AUDIT_ROUTINES
+
 // Switch for a lightweight process event counter, useful when disabling the blueprint guard
 // which can taint profiling results:
 #define LIGHTWEIGHT_PROCESS_EVENT_COUNTER 0 && !DO_BLUEPRINT_GUARD
