Non-unity CIS fix.

#rnx

diff --git a/Engine/Source/Runtime/CoreUObject/Private/UObject/ScriptCore.cpp b/Engine/Source/Runtime/CoreUObject/Private/UObject/ScriptCore.cpp
--- a/Engine/Source/Runtime/CoreUObject/Private/UObject/ScriptCore.cpp
+++ b/Engine/Source/Runtime/CoreUObject/Private/UObject/ScriptCore.cpp
@@ -16,6 +16,7 @@
 #include "UObject/Script.h"
 #include "UObject/ObjectMacros.h"
 #include "UObject/UObjectBaseUtility.h"
+#include "UObject/UObjectIterator.h"
 #include "UObject/Object.h"
 #include "UObject/CoreNative.h"
 #include "UObject/Class.h"
