CIS fix (Linux).

#rnx

diff --git a/Engine/Source/Runtime/CoreUObject/Private/UObject/ScriptCore.cpp b/Engine/Source/Runtime/CoreUObject/Private/UObject/ScriptCore.cpp
--- a/Engine/Source/Runtime/CoreUObject/Private/UObject/ScriptCore.cpp
+++ b/Engine/Source/Runtime/CoreUObject/Private/UObject/ScriptCore.cpp
@@ -95,7 +95,7 @@ static struct F##inst##Registrar \
 	{ \
 		GNativeFuncNames[inst] = #inst; \
 	} \
-} ##inst##RegistrarInst;
+} inst##RegistrarInst;
 #else
 #define STORE_INSTRUCTION_NAME(inst)
 #endif//STORE_INSTRUCTION_NAMES
