Fix shadowing warnings from CL# 6623723
#rnx

diff --git a/Engine/Source/Runtime/CoreUObject/Private/UObject/ScriptCore.cpp b/Engine/Source/Runtime/CoreUObject/Private/UObject/ScriptCore.cpp
--- a/Engine/Source/Runtime/CoreUObject/Private/UObject/ScriptCore.cpp
+++ b/Engine/Source/Runtime/CoreUObject/Private/UObject/ScriptCore.cpp
@@ -1375,7 +1375,7 @@ static void OutputLongestFunctions(FOutputDevice& Ar, int32 Num)
 	}
 }
 
-static void OutputMostFrequentlyCalledFunctions(FOutputDevice& Ar, int32 Num)
+static void OutputMostFrequentlyCalledFunctions(FOutputDevice& OutputAr, int32 Num)
 {
 	// Script serialization is recursive and requires certain symbols (e.g. Script, a reference
 	// to the bytecode), so we declare a type so that we have some scope:
@@ -1490,7 +1490,7 @@ static void OutputMostFrequentlyCalledFunctions(FOutputDevice& Ar, int32 Num)
 
 		if(FunctionCallsSorted.Num())
 		{
-			Ar.Logf(TEXT("Top %d function call targets"), FunctionCallsSorted.Num());
+			OutputAr.Logf(TEXT("Top %d function call targets"), FunctionCallsSorted.Num());
 			for(TPair<UFunction*, int32>& Calls : FunctionCallsSorted )
 			{
 				if(Calls.Key == nullptr)
@@ -1498,12 +1498,12 @@ static void OutputMostFrequentlyCalledFunctions(FOutputDevice& Ar, int32 Num)
 					break;
 				}
 
-				Ar.Logf(TEXT("%s %s %d"), *Calls.Key->GetName(), *Calls.Key->GetOuter()->GetName(), Calls.Value);
+				OutputAr.Logf(TEXT("%s %s %d"), *Calls.Key->GetName(), *Calls.Key->GetOuter()->GetName(), Calls.Value);
 			}
 		}
 		else
 		{
-			Ar.Log(TEXT("No function call instructions found in memory"));
+			OutputAr.Log(TEXT("No function call instructions found in memory"));
 		}
 	}
 
@@ -1522,7 +1522,7 @@ static void OutputMostFrequentlyCalledFunctions(FOutputDevice& Ar, int32 Num)
 
 		if(VirtualFunctionCallsSorted.Num())
 		{
-			Ar.Logf(TEXT("Top %d virtual function call targets"), VirtualFunctionCallsSorted.Num());
+			OutputAr.Logf(TEXT("Top %d virtual function call targets"), VirtualFunctionCallsSorted.Num());
 			for(TPair<FName, int32>& Calls : VirtualFunctionCallsSorted )
 			{
 				if(Calls.Key == FName())
@@ -1530,17 +1530,17 @@ static void OutputMostFrequentlyCalledFunctions(FOutputDevice& Ar, int32 Num)
 					break;
 				}
 
-				Ar.Logf(TEXT("%s %d"), *(Calls.Key.ToString()), Calls.Value);
+				OutputAr.Logf(TEXT("%s %d"), *(Calls.Key.ToString()), Calls.Value);
 			}
 		}
 		else
 		{
-			Ar.Log(TEXT("No virtual function call instructions in memory"));
+			OutputAr.Log(TEXT("No virtual function call instructions in memory"));
 		}
 	}
 }
 
-static void OutputMostFrequentlyUsedInstructions(FOutputDevice& Ar, int32 Num)
+static void OutputMostFrequentlyUsedInstructions(FOutputDevice& OutputAr, int32 Num)
 {
 	// Script serialization is recursive and requires certain symbols (e.g. Script, a reference
 	// to the bytecode), so we declare a type so that we have some scope:
@@ -1624,7 +1624,7 @@ static void OutputMostFrequentlyUsedInstructions(FOutputDevice& Ar, int32 Num)
 
 		if(InstructionCountsSorted.Num())
 		{
-			Ar.Logf(TEXT("Top %d bytecode instructions"), InstructionCountsSorted.Num());
+			OutputAr.Logf(TEXT("Top %d bytecode instructions"), InstructionCountsSorted.Num());
 			for(TPair<EExprToken, int32>& Instruction : InstructionCountsSorted )
 			{
 				if(Instruction.Value == 0)
@@ -1636,20 +1636,20 @@ static void OutputMostFrequentlyUsedInstructions(FOutputDevice& Ar, int32 Num)
 				if(GNativeFuncNames[Instruction.Key])
 				{
 					FString AsString = GNativeFuncNames[Instruction.Key];
-					Ar.Logf(TEXT("%s %d"), *AsString, Instruction.Value);
+					OutputAr.Logf(TEXT("%s %d"), *AsString, Instruction.Value);
 				}
 				else
 				{
-					Ar.Logf(TEXT("0x%x %d"), Instruction.Key, Instruction.Value);
+					OutputAr.Logf(TEXT("0x%x %d"), Instruction.Key, Instruction.Value);
 				}
 #else
-				Ar.Logf(TEXT("0x%x %d"), Instruction.Key, Instruction.Value);
+				OutputAr.Logf(TEXT("0x%x %d"), Instruction.Key, Instruction.Value);
 #endif
 			}
 		}
 		else
 		{
-			Ar.Log(TEXT("No instructions found in memory"));
+			OutputAr.Log(TEXT("No instructions found in memory"));
 		}
 	}
 }
