Documentating "DocumentationPolicy" meta tag for inside ObjectMacros.h

diff --git a/Engine/Source/Runtime/CoreUObject/Public/UObject/ObjectMacros.h b/Engine/Source/Runtime/CoreUObject/Public/UObject/ObjectMacros.h
--- a/Engine/Source/Runtime/CoreUObject/Public/UObject/ObjectMacros.h
+++ b/Engine/Source/Runtime/CoreUObject/Public/UObject/ObjectMacros.h
@@ -983,6 +983,9 @@ namespace UM
 
 		/// A short tooltip that is used in some contexts where the full tooltip might be overwhelming (such as the parent class picker dialog)
 		ShortTooltip,
+
+		/// A setting to determine validation of tooltips and comments. Needs to be set to "Strict"
+		DocumentationPolicy,
 	};
 
 	// Metadata usable in UCLASS
