Added a check for the presence of an interface class to fix a strange crash in the wild

diff --git a/Engine/Source/Runtime/CoreUObject/Private/UObject/Class.cpp b/Engine/Source/Runtime/CoreUObject/Private/UObject/Class.cpp
--- a/Engine/Source/Runtime/CoreUObject/Private/UObject/Class.cpp
+++ b/Engine/Source/Runtime/CoreUObject/Private/UObject/Class.cpp
@@ -3871,7 +3871,7 @@ bool UClass::ImplementsInterface( const class UClass* SomeInterface ) const
 			for (TArray<FImplementedInterface>::TConstIterator It(CurrentClass->Interfaces); It; ++It)
 			{
 				const UClass* InterfaceClass = It->Class;
-				if (InterfaceClass->IsChildOf(SomeInterface))
+				if (InterfaceClass && InterfaceClass->IsChildOf(SomeInterface))
 				{
 					return true;
 				}
