Fixed log warning to reflect the current way of allowing cooked content in the editor

diff --git a/Engine/Source/Runtime/CoreUObject/Private/UObject/LinkerLoad.cpp b/Engine/Source/Runtime/CoreUObject/Private/UObject/LinkerLoad.cpp
--- a/Engine/Source/Runtime/CoreUObject/Private/UObject/LinkerLoad.cpp
+++ b/Engine/Source/Runtime/CoreUObject/Private/UObject/LinkerLoad.cpp
@@ -1224,7 +1224,7 @@ FLinkerLoad::ELinkerStatus FLinkerLoad::SerializePackageFileSummary()
 			if (!GAllowCookedDataInEditorBuilds)
 			{
 				UE_LOG(LogLinker, Warning, 
-					TEXT("Unable to load package (%s). Package contains cooked data which is not supported by the current build. Set [Core.System] AllowCookedDataInEditorBuilds to true in Engine.ini to allow it."), 
+					TEXT("Unable to load package (%s). Package contains cooked data which is not supported by the current build. Enable 'Allow Cooked Content In The Editor' in Project Settings under 'Engine - Cooker' section to load it."), 
 					*Filename);
 				return LINKER_Failed;
 			}
