Fixed crashes related to multithreaded destruction of widget components

diff --git a/Engine/Source/Runtime/CoreUObject/Private/UObject/GarbageCollection.cpp b/Engine/Source/Runtime/CoreUObject/Private/UObject/GarbageCollection.cpp
--- a/Engine/Source/Runtime/CoreUObject/Private/UObject/GarbageCollection.cpp
+++ b/Engine/Source/Runtime/CoreUObject/Private/UObject/GarbageCollection.cpp
@@ -151,7 +151,7 @@ static FAutoConsoleVariableRef CIncrementalBeginDestroyEnabled(
 	ECVF_Default
 );
 
-int32 GMultithreadedDestructionEnabled = 1;
+int32 GMultithreadedDestructionEnabled = 0;
 static FAutoConsoleVariableRef CMultithreadedDestructionEnabled(
 	TEXT("gc.MultithreadedDestructionEnabled"),
 	GMultithreadedDestructionEnabled,
@@ -2003,6 +2003,11 @@ void UObject::AddReferencedObjects(UObject* This, FReferenceCollector& Collector
 #endif
 }
 
+bool UObject::IsDestructionThreadSafe() const
+{
+	return true;
+}
+
 /*-----------------------------------------------------------------------------
 	Implementation of realtime garbage collection helper functions in 
 	UProperty, UClass, ...
diff --git a/Engine/Source/Runtime/CoreUObject/Public/UObject/Object.h b/Engine/Source/Runtime/CoreUObject/Public/UObject/Object.h
--- a/Engine/Source/Runtime/CoreUObject/Public/UObject/Object.h
+++ b/Engine/Source/Runtime/CoreUObject/Public/UObject/Object.h
@@ -451,10 +451,7 @@ public:
 	*
 	* @return	true if this object's destructor is thread safe
 	*/
-	virtual bool IsDestructionThreadSafe() const
-	{
-		return true;
-	}
+	virtual bool IsDestructionThreadSafe() const;
 
 	/**
 	* Called during cooking. Must return all objects that will be Preload()ed when this is serialized at load time. Only used by the EDL.
diff --git a/Engine/Source/Runtime/UMG/Public/Components/Widget.h b/Engine/Source/Runtime/UMG/Public/Components/Widget.h
--- a/Engine/Source/Runtime/UMG/Public/Components/Widget.h
+++ b/Engine/Source/Runtime/UMG/Public/Components/Widget.h
@@ -852,6 +852,7 @@ public:
 	// Begin UObject
 	virtual UWorld* GetWorld() const override;
 	virtual void FinishDestroy() override;
+	virtual bool IsDestructionThreadSafe() const override { return false; }
 	// End UObject
 
 	FORCEINLINE bool CanSafelyRouteEvent()
diff --git a/Engine/Source/Runtime/UMG/Public/Components/WidgetComponent.h b/Engine/Source/Runtime/UMG/Public/Components/WidgetComponent.h
--- a/Engine/Source/Runtime/UMG/Public/Components/WidgetComponent.h
+++ b/Engine/Source/Runtime/UMG/Public/Components/WidgetComponent.h
@@ -85,6 +85,7 @@ class UMG_API UWidgetComponent : public UMeshComponent
 public:
 	//UObject interface
 	virtual void Serialize(FArchive& Ar) override;
+	virtual bool IsDestructionThreadSafe() const override { return false; }
 	//~ End UObject Interface
 
 	/** UActorComponent Interface */
diff --git a/Engine/Source/Runtime/UMG/Public/Components/WidgetInteractionComponent.h b/Engine/Source/Runtime/UMG/Public/Components/WidgetInteractionComponent.h
--- a/Engine/Source/Runtime/UMG/Public/Components/WidgetInteractionComponent.h
+++ b/Engine/Source/Runtime/UMG/Public/Components/WidgetInteractionComponent.h
@@ -75,6 +75,10 @@ public:
 	virtual void TickComponent(float DeltaTime, enum ELevelTick TickType, FActorComponentTickFunction* ThisTickFunction) override;
 	// End UActorComponent
 	
+	// Begin UObject interface
+	virtual bool IsDestructionThreadSafe() const override { return false; }
+	// End UObject
+
 	/**
 	 * Presses a key as if the mouse/pointer were the source of it.  Normally you would just use
 	 * Left/Right mouse button for the Key.  However - advanced uses could also be imagined where you
