1st change to improve the IO perf
At the moment precaching is limited to 2 packages, even if there is memory/IO available for more packages to be preloaded.

diff --git a/Engine/Source/Runtime/CoreUObject/Private/Serialization/AsyncLoading.cpp b/Engine/Source/Runtime/CoreUObject/Private/Serialization/AsyncLoading.cpp
--- a/Engine/Source/Runtime/CoreUObject/Private/Serialization/AsyncLoading.cpp
+++ b/Engine/Source/Runtime/CoreUObject/Private/Serialization/AsyncLoading.cpp
@@ -182,6 +182,14 @@ static FAutoConsoleVariableRef CVar_MaxReadyRequestsToStallMB(
 	TEXT("Controls the maximum amount memory for unhandled IO requests before we stall the pak precacher to let the CPU catch up (in megabytes).")
 );
 
+int32 GMaxPrecacheRequestsInFlight = 2;
+static FAutoConsoleVariableRef CVar_MaxPrecacheRequestsInFlight(
+	TEXT("s.MaxPrecacheRequestsInFlight"),
+	GMaxPrecacheRequestsInFlight,
+	TEXT("Controls the maximum amount of precache requests to have in flight.")
+);
+
+
 int32 GMaxIncomingRequestsToStall = 100;
 static FAutoConsoleVariableRef CVar_MaxIncomingRequestsToStall(
 	TEXT("s.MaxIncomingRequestsToStall"),
@@ -3897,7 +3905,7 @@ EAsyncPackageState::Type FAsyncPackage::ProcessImportsAndExports_Event()
 			break; // requeue this to give other packages a chance to start IO
 		}
 		bDidSomething = false;
-		if (PrecacheRequests.Num() < 2 && ExportsThatCanHaveIOStarted.Num() && CanAddWaitingPackages())
+		if (PrecacheRequests.Num() < GMaxPrecacheRequestsInFlight && ExportsThatCanHaveIOStarted.Num() && CanAddWaitingPackages())
 		{
 			bDidSomething = true;
 			StartPrecacheRequest();
