Prevent async loading textures from stalling the gamethread in Postload(), when using async loading packages.

diff --git a/Engine/Source/Runtime/CoreUObject/Private/Serialization/AsyncLoading.cpp b/Engine/Source/Runtime/CoreUObject/Private/Serialization/AsyncLoading.cpp
--- a/Engine/Source/Runtime/CoreUObject/Private/Serialization/AsyncLoading.cpp
+++ b/Engine/Source/Runtime/CoreUObject/Private/Serialization/AsyncLoading.cpp
@@ -6734,7 +6734,12 @@ EAsyncPackageState::Type FAsyncPackage::PostLoadObjects()
 		UObject* Object = PackageObjLoaded[PostLoadIndex++];
 		if (Object)
 		{
-			if (!bIsMultithreaded || (bAsyncPostLoadEnabled && CanPostLoadOnAsyncLoadingThread(Object)))
+			if (!Object->IsReadyForAsyncPostLoad())
+			{
+				--PostLoadIndex;
+				break;
+			}
+			else if (!bIsMultithreaded || (bAsyncPostLoadEnabled && CanPostLoadOnAsyncLoadingThread(Object)))
 			{
 				SCOPED_ACCUM_LOADTIME(PostLoad, StaticGetNativeClassName(Object->GetClass()));
 
diff --git a/Engine/Source/Runtime/CoreUObject/Private/Serialization/BulkData.cpp b/Engine/Source/Runtime/CoreUObject/Private/Serialization/BulkData.cpp
--- a/Engine/Source/Runtime/CoreUObject/Private/Serialization/BulkData.cpp
+++ b/Engine/Source/Runtime/CoreUObject/Private/Serialization/BulkData.cpp
@@ -432,7 +432,7 @@ bool FUntypedBulkData::IsBulkDataLoaded() const
 	return !!BulkData;
 }
 
-bool FUntypedBulkData::IsAsyncLoadingComplete()
+bool FUntypedBulkData::IsAsyncLoadingComplete() const 
 {
 	return SerializeFuture.IsValid() == false || SerializeFuture.WaitFor(FTimespan::Zero());
 }
diff --git a/Engine/Source/Runtime/CoreUObject/Public/Serialization/BulkData.h b/Engine/Source/Runtime/CoreUObject/Public/Serialization/BulkData.h
--- a/Engine/Source/Runtime/CoreUObject/Public/Serialization/BulkData.h
+++ b/Engine/Source/Runtime/CoreUObject/Public/Serialization/BulkData.h
@@ -366,7 +366,7 @@ public:
 	*
 	* @return true if bulk data has been loaded or async loading was not used to load this data, false otherwise
 	*/
-	bool IsAsyncLoadingComplete();
+	bool IsAsyncLoadingComplete() const;
 
 	/**
 	* Returns whether this bulk data is used
diff --git a/Engine/Source/Runtime/CoreUObject/Public/UObject/Object.h b/Engine/Source/Runtime/CoreUObject/Public/UObject/Object.h
--- a/Engine/Source/Runtime/CoreUObject/Public/UObject/Object.h
+++ b/Engine/Source/Runtime/CoreUObject/Public/UObject/Object.h
@@ -235,6 +235,11 @@ public:
 	virtual void LoadedFromAnotherClass(const FName& OldClassName) {}
 #endif
 
+	/**
+	 * Called before calling PostLoad() in FAsyncPackage::PostLoadObjects(). This is the safeguard to prevent PostLoad() from stalling the main thread.
+	 */
+	virtual bool IsReadyForAsyncPostLoad() const { return true; }
+
 	/** 
 	 * Do any object-specific cleanup required immediately after loading an object.
 	 * This is not called for newly-created objects, and by default will always execute on the game thread.
diff --git a/Engine/Source/Runtime/Engine/Classes/Engine/Texture.h b/Engine/Source/Runtime/Engine/Classes/Engine/Texture.h
--- a/Engine/Source/Runtime/Engine/Classes/Engine/Texture.h
+++ b/Engine/Source/Runtime/Engine/Classes/Engine/Texture.h
@@ -450,6 +450,9 @@ struct FTexturePlatformData
 	/** Destructor. */
 	ENGINE_API ~FTexturePlatformData();
 
+	/** Return whether TryLoadMips() would stall because async loaded mips are not yet available. */
+	bool IsReadyForAsyncPostLoad() const;
+
 	/**
 	 * Try to load mips from the derived data cache.
 	 * @param FirstMipToLoad - The first mip index to load.
diff --git a/Engine/Source/Runtime/Engine/Classes/Engine/Texture2D.h b/Engine/Source/Runtime/Engine/Classes/Engine/Texture2D.h
--- a/Engine/Source/Runtime/Engine/Classes/Engine/Texture2D.h
+++ b/Engine/Source/Runtime/Engine/Classes/Engine/Texture2D.h
@@ -103,6 +103,7 @@ public:
 	virtual void PostEditChangeProperty(FPropertyChangedEvent& PropertyChangedEvent) override;
 #endif // WITH_EDITOR
 	virtual void BeginDestroy() override;
+	virtual bool IsReadyForAsyncPostLoad() const override;
 	virtual void PostLoad() override;
 	virtual void PreSave(const class ITargetPlatform* TargetPlatform) override;
 	virtual void GetAssetRegistryTags(TArray<FAssetRegistryTag>& OutTags) const override;
diff --git a/Engine/Source/Runtime/Engine/Private/Texture2D.cpp b/Engine/Source/Runtime/Engine/Private/Texture2D.cpp
--- a/Engine/Source/Runtime/Engine/Private/Texture2D.cpp
+++ b/Engine/Source/Runtime/Engine/Private/Texture2D.cpp
@@ -674,6 +674,12 @@ void UTexture2D::CancelPendingTextureStreaming()
 	// No need to call FlushResourceStreaming(), since calling CancelPendingMipChangeRequest has an immediate effect.
 }
 
+bool UTexture2D::IsReadyForAsyncPostLoad() const
+{
+	return !PlatformData || PlatformData->IsReadyForAsyncPostLoad();
+
+}
+
 void UTexture2D::PostLoad()
 {
 #if WITH_EDITOR
diff --git a/Engine/Source/Runtime/Engine/Private/TextureDerivedData.cpp b/Engine/Source/Runtime/Engine/Private/TextureDerivedData.cpp
--- a/Engine/Source/Runtime/Engine/Private/TextureDerivedData.cpp
+++ b/Engine/Source/Runtime/Engine/Private/TextureDerivedData.cpp
@@ -1038,6 +1038,19 @@ FTexturePlatformData::~FTexturePlatformData()
 	if (VTData) delete VTData;
 }
 
+bool FTexturePlatformData::IsReadyForAsyncPostLoad() const
+{
+	for (int32 MipIndex = 0; MipIndex < Mips.Num(); ++MipIndex)
+	{
+		const FTexture2DMipMap& Mip = Mips[MipIndex];
+		if (!Mip.BulkData.IsAsyncLoadingComplete())
+		{
+			return false;
+		}
+	}
+	return true;
+}
+
 bool FTexturePlatformData::TryLoadMips(int32 FirstMipToLoad, void** OutMipData)
 {
 	int32 NumMipsCached = 0;
