#UE4 A few fixes for the SaveConcurrent save path. Avoiding improper thread usage of a few systems.

diff --git a/Engine/Source/Editor/UnrealEd/Private/CookOnTheFlyServer.cpp b/Engine/Source/Editor/UnrealEd/Private/CookOnTheFlyServer.cpp
--- a/Engine/Source/Editor/UnrealEd/Private/CookOnTheFlyServer.cpp
+++ b/Engine/Source/Editor/UnrealEd/Private/CookOnTheFlyServer.cpp
@@ -7667,7 +7667,7 @@ uint32 UCookOnTheFlyServer::FullLoadAndSave(uint32& CookedPackageCount)
 									GIsCookerLoadingPackage = false;
 								}
 
-								if (!bIsTexture)
+								if (!bIsTexture || bSaveConcurrent)
 								{
 									SCOPE_TIMER(FullLoadAndSave_BeginCache);
 									Obj->BeginCacheForCookedPlatformData(TargetPlatform);
@@ -7748,6 +7748,21 @@ uint32 UCookOnTheFlyServer::FullLoadAndSave(uint32& CookedPackageCount)
 
 	ProcessedPackages.Empty();
 
+	// When saving concurrently, flush async loading since that is normally done internally in SavePackage
+	if (bSaveConcurrent)
+	{
+		UE_LOG(LogCook, Display, TEXT("Flushing async loading..."));
+		SCOPE_TIMER(FullLoadAndSave_FlushAsyncLoading);
+		FlushAsyncLoading();
+	}
+
+	if (bSaveConcurrent)
+	{
+		UE_LOG(LogCook, Display, TEXT("Waiting for async tasks..."));
+		SCOPE_TIMER(FullLoadAndSave_ProcessThreadUntilIdle);
+		FTaskGraphInterface::Get().ProcessThreadUntilIdle(ENamedThreads::GameThread);
+	}
+
 	// Wait for all shaders to finish compiling
 	if (GShaderCompilingManager)
 	{
@@ -7911,6 +7926,14 @@ uint32 UCookOnTheFlyServer::FullLoadAndSave(uint32& CookedPackageCount)
 						FSavePackageResultStruct SaveResult = GEditor->Save(Package, World, FlagsToCook, *PlatFilename, GError, NULL, bSwap, false, SaveFlags, Target, FDateTime::MinValue(), false);
 						GIsCookerLoadingPackage = false;
 
+						if (SaveResult == ESavePackageResult::Success && UAssetManager::IsValid())
+						{
+							if (!UAssetManager::Get().VerifyCanCookPackage(Package->GetFName()))
+							{
+								SaveResult = ESavePackageResult::Error;
+							}
+						}
+
 						const bool bSucceededSavePackage = (SaveResult == ESavePackageResult::Success || SaveResult == ESavePackageResult::GenerateStub || SaveResult == ESavePackageResult::ReplaceCompletely);
 						if (bSucceededSavePackage)
 						{
diff --git a/Engine/Source/Runtime/CoreUObject/Private/UObject/SavePackage.cpp b/Engine/Source/Runtime/CoreUObject/Private/UObject/SavePackage.cpp
--- a/Engine/Source/Runtime/CoreUObject/Private/UObject/SavePackage.cpp
+++ b/Engine/Source/Runtime/CoreUObject/Private/UObject/SavePackage.cpp
@@ -3659,11 +3659,14 @@ FSavePackageResultStruct UPackage::Save(UPackage* InOuter, UObject* Base, EObjec
 		// point to the new version of the path
 		Filename = *NewPath;
 
-		// We need to fulfill all pending streaming and async loading requests to then allow us to lock the global IO manager. 
-		// The latter implies flushing all file handles which is a pre-requisite of saving a package. The code basically needs 
-		// to be sure that we are not reading from a file that is about to be overwritten and that there is no way we might 
-		// start reading from the file till we are done overwriting it.
-		FlushAsyncLoading();
+		if (!bSavingConcurrent)
+		{
+			// We need to fulfill all pending streaming and async loading requests to then allow us to lock the global IO manager. 
+			// The latter implies flushing all file handles which is a pre-requisite of saving a package. The code basically needs 
+			// to be sure that we are not reading from a file that is about to be overwritten and that there is no way we might 
+			// start reading from the file till we are done overwriting it.
+			FlushAsyncLoading();
+		}
 
 		(*GFlushStreamingFunc)();
 
@@ -3845,7 +3848,14 @@ FSavePackageResultStruct UPackage::Save(UPackage* InOuter, UObject* Base, EObjec
 						: bSavingConcurrent(InSavingConcurrent)
 					{
 						// We need the same lock as GC so that no StaticFindObject can happen in parallel to saveing a package
-						FGCCSyncObject::Get().GCLock();
+						if (IsInGameThread())
+						{
+							FGCCSyncObject::Get().GCLock();
+						}
+						else
+						{
+							FGCCSyncObject::Get().LockAsync();
+						}
 
 						// Do not change GIsSavingPackage while saving concurrently. It should have been set before and after all packages are saved
 						if (!bSavingConcurrent)
@@ -3859,7 +3869,14 @@ FSavePackageResultStruct UPackage::Save(UPackage* InOuter, UObject* Base, EObjec
 						{
 							GIsSavingPackage = false;
 						}
-						FGCCSyncObject::Get().GCUnlock();
+						if (IsInGameThread())
+						{
+							FGCCSyncObject::Get().GCUnlock();
+						}
+						else
+						{
+							FGCCSyncObject::Get().UnlockAsync();
+						}
 					}
 
 					bool bSavingConcurrent;
