Making cooked Animation Blueprints work in the editor

diff --git a/Engine/Source/Runtime/CoreUObject/Private/UObject/StructScriptLoader.cpp b/Engine/Source/Runtime/CoreUObject/Private/UObject/StructScriptLoader.cpp
--- a/Engine/Source/Runtime/CoreUObject/Private/UObject/StructScriptLoader.cpp
+++ b/Engine/Source/Runtime/CoreUObject/Private/UObject/StructScriptLoader.cpp
@@ -205,10 +205,17 @@ bool FStructScriptLoader::LoadStructWithScript(UStruct* DestScriptContainer, FAr
 	// script code), we define a minimum script version
 	bool bSkipScriptSerialization = (Ar.UE4Ver() < VER_MIN_SCRIPTVM_UE4) || (Ar.LicenseeUE4Ver() < VER_MIN_SCRIPTVM_LICENSEEUE4);
 #if WITH_EDITOR
-	static const FBoolConfigValueHelper SkipByteCodeHelper(TEXT("StructSerialization"), TEXT("SkipByteCodeSerialization"));
-	// in editor builds, we're going to regenerate the bytecode anyways, so it
-	// is a waste of cycles to try and serialize it in
-	bSkipScriptSerialization |= (bool)SkipByteCodeHelper;
+	if (!bSkipScriptSerialization)
+	{
+		// Always serialize script for cooked packages loaded in the editor
+		if (!(GAllowCookedDataInEditorBuilds && Ar.GetLinker() && Ar.GetLinker()->LinkerRoot && Ar.GetLinker()->LinkerRoot->bIsCookedForEditor))
+		{
+			static const FBoolConfigValueHelper SkipByteCodeHelper(TEXT("StructSerialization"), TEXT("SkipByteCodeSerialization"));
+			// in editor builds, we're going to regenerate the bytecode anyways, so it
+			// is a waste of cycles to try and serialize it in
+			bSkipScriptSerialization |= (bool)SkipByteCodeHelper;
+		}
+	}
 #endif // WITH_EDITOR
 	bSkipScriptSerialization &= bIsLinkerLoader; // to keep consistent with old UStruct::Serialize() functionality
 
