Fixing CIS

diff --git a/Engine/Source/Runtime/CoreUObject/Private/UObject/StructScriptLoader.cpp b/Engine/Source/Runtime/CoreUObject/Private/UObject/StructScriptLoader.cpp
--- a/Engine/Source/Runtime/CoreUObject/Private/UObject/StructScriptLoader.cpp
+++ b/Engine/Source/Runtime/CoreUObject/Private/UObject/StructScriptLoader.cpp
@@ -5,6 +5,7 @@
 #include "Misc/CoreMisc.h"
 #include "UObject/LinkerLoad.h"
 #include "Serialization/ArchiveScriptReferenceCollector.h"
+#include "UObject/Package.h"
 
 /*******************************************************************************
  * FDeferredScriptTracker
