Log use of additional GC time for Android/iOS.
If mobile takes > 10s for GC to complete then increase threshold to 30s, log usage of additional time and trigger report from hitch detector.
Add delegate which is called when GC extended time was used.
Increase Android GC timeout to match iOS.
#rnx

diff --git a/Engine/Source/Runtime/Core/Private/Misc/CoreDelegates.cpp b/Engine/Source/Runtime/Core/Private/Misc/CoreDelegates.cpp
--- a/Engine/Source/Runtime/Core/Private/Misc/CoreDelegates.cpp
+++ b/Engine/Source/Runtime/Core/Private/Misc/CoreDelegates.cpp
@@ -171,6 +171,8 @@ FCoreDelegates::FIsLoadingMovieCurrentlyPlaying FCoreDelegates::IsLoadingMovieCu
 
 FCoreDelegates::FShouldLaunchUrl FCoreDelegates::ShouldLaunchUrl;
 
+FCoreDelegates::FOnGCFinishDestroyTimeExtended FCoreDelegates::OnGCFinishDestroyTimeExtended;
+
 PRAGMA_DISABLE_DEPRECATION_WARNINGS
 FCrashOverrideParameters::~FCrashOverrideParameters()
 {
diff --git a/Engine/Source/Runtime/Core/Public/Misc/CoreDelegates.h b/Engine/Source/Runtime/Core/Public/Misc/CoreDelegates.h
--- a/Engine/Source/Runtime/Core/Public/Misc/CoreDelegates.h
+++ b/Engine/Source/Runtime/Core/Public/Misc/CoreDelegates.h
@@ -549,6 +549,10 @@ public:
 	DECLARE_DELEGATE_RetVal_OneParam(bool, FShouldLaunchUrl, const TCHAR* /* URL */);
 	static FShouldLaunchUrl ShouldLaunchUrl;
 
+	/** Sent when GC finish destroy takes more time than expected */
+	DECLARE_MULTICAST_DELEGATE_OneParam(FOnGCFinishDestroyTimeExtended, const FString&);
+	static FOnGCFinishDestroyTimeExtended OnGCFinishDestroyTimeExtended;
+
 private:
 
 	// Callbacks for hotfixes
diff --git a/Engine/Source/Runtime/CoreUObject/Private/UObject/GarbageCollection.cpp b/Engine/Source/Runtime/CoreUObject/Private/UObject/GarbageCollection.cpp
--- a/Engine/Source/Runtime/CoreUObject/Private/UObject/GarbageCollection.cpp
+++ b/Engine/Source/Runtime/CoreUObject/Private/UObject/GarbageCollection.cpp
@@ -1449,6 +1449,11 @@ bool IncrementalDestroyGarbage(bool bUseTimeLimit, float TimeLimit)
 		// Have we finished the first round of attempting to call FinishDestroy on unreachable objects?
 		if (GObjCurrentPurgeObjectIndex >= GUnreachableObjects.Num())
 		{
+			double MaxTimeForFinishDestroy = 10.00;
+			bool bFinishDestroyTimeExtended = false;
+			FString FirstObjectNotReadyWhenTimeExtended;
+			int32 StartObjectsPendingDestructionCount = GGCObjectsPendingDestructionCount;
+
 			// We've finished iterating over all unreachable objects, but we need still need to handle
 			// objects that were deferred.
 			int32 LastLoopObjectsPendingDestructionCount = GGCObjectsPendingDestructionCount;
@@ -1515,10 +1520,17 @@ bool IncrementalDestroyGarbage(bool bUseTimeLimit, float TimeLimit)
 					if (FPlatformProperties::RequiresCookedData())
 					{
 						const bool bPollTimeLimit = ((FinishDestroyTimePollCounter++) % TimeLimitEnforcementGranularityForDestroy == 0);
-#if PLATFORM_IOS
-                        const double MaxTimeForFinishDestroy = 30.0;
-#else
-						const double MaxTimeForFinishDestroy = 10.0;
+#if PLATFORM_IOS || PLATFORM_ANDROID
+						if( !bFinishDestroyTimeExtended && (FPlatformTime::Seconds() - GCStartTime) > MaxTimeForFinishDestroy)
+						{
+							MaxTimeForFinishDestroy = 30.0;
+							bFinishDestroyTimeExtended = true;
+#if USE_HITCH_DETECTION
+							GHitchDetected = true;
+#endif
+							FirstObjectNotReadyWhenTimeExtended = GetFullNameSafe(GGCObjectsPendingDestruction[0]);
+						}
+						else
 #endif
 						// Check if we spent too much time on waiting for FinishDestroy without making any progress
 						if (LastLoopObjectsPendingDestructionCount == GGCObjectsPendingDestructionCount && bPollTimeLimit &&
@@ -1573,6 +1585,13 @@ bool IncrementalDestroyGarbage(bool bUseTimeLimit, float TimeLimit)
 			// Have all objects been destroyed now?
 			if( GGCObjectsPendingDestructionCount == 0 )
 			{
+				if (bFinishDestroyTimeExtended)
+				{
+					FString Msg = FString::Printf(TEXT("Additional time was required to finish routing FinishDestroy, spent %.2fs on routing FinishDestroy to %d objects. 1st obj not ready: '%s'."), (FPlatformTime::Seconds() - GCStartTime), StartObjectsPendingDestructionCount, *FirstObjectNotReadyWhenTimeExtended);
+					UE_LOG(LogGarbage, Warning, TEXT("%s"), *Msg );
+					FCoreDelegates::OnGCFinishDestroyTimeExtended.Broadcast(Msg);
+				}
+
 				// Release memory we used for objects pending destruction, leaving some slack space
 				GGCObjectsPendingDestruction.Empty( 256 );
 
diff --git a/Engine/Source/Runtime/OpenGLDrv/Private/Android/AndroidEGL.cpp b/Engine/Source/Runtime/OpenGLDrv/Private/Android/AndroidEGL.cpp
--- a/Engine/Source/Runtime/OpenGLDrv/Private/Android/AndroidEGL.cpp
+++ b/Engine/Source/Runtime/OpenGLDrv/Private/Android/AndroidEGL.cpp
@@ -1715,6 +1715,7 @@ extern FCriticalSection GAndroidWindowLock;
 
 void BlockOnLostWindowRenderCommand(TSharedPtr<FEvent, ESPMode::ThreadSafe> RTBlockedTrigger)
 {
+	QUICK_SCOPE_CYCLE_COUNTER(STAT_BlockOnLostWindowRenderCommand);
 	check(IsInRenderingThread());
 
 	// Hold GC scope guard, as GC will timeout if anything waits for RT fences.
