#UE4 Nonunity fix

diff --git a/Engine/Source/Runtime/CoreUObject/Private/UObject/GarbageCollection.cpp b/Engine/Source/Runtime/CoreUObject/Private/UObject/GarbageCollection.cpp
--- a/Engine/Source/Runtime/CoreUObject/Private/UObject/GarbageCollection.cpp
+++ b/Engine/Source/Runtime/CoreUObject/Private/UObject/GarbageCollection.cpp
@@ -9,6 +9,7 @@
 #include "Misc/TimeGuard.h"
 #include "HAL/IConsoleManager.h"
 #include "Misc/App.h"
+#include "Misc/CoreDelegates.h"
 #include "UObject/ScriptInterface.h"
 #include "UObject/UObjectAllocator.h"
 #include "UObject/UObjectBase.h"
