Ensure/handle missing PropertyClass in USoftObjectProperty when calling GetCPPType

diff --git a/Engine/Source/Runtime/CoreUObject/Public/UObject/UnrealType.h b/Engine/Source/Runtime/CoreUObject/Public/UObject/UnrealType.h
--- a/Engine/Source/Runtime/CoreUObject/Public/UObject/UnrealType.h
+++ b/Engine/Source/Runtime/CoreUObject/Public/UObject/UnrealType.h
@@ -2195,6 +2195,18 @@ class COREUOBJECT_API USoftObjectProperty : public TUObjectPropertyBase<FSoftObj
 	virtual bool AllowCrossLevel() const override;
 	virtual FString GetCPPTypeCustom(FString* ExtendedTypeText, uint32 CPPExportFlags, const FString& InnerNativeTypeName)  const override;
 
+	virtual FString GetCPPType(FString* ExtendedTypeText, uint32 CPPExportFlags) const override
+	{
+		if (ensureMsgf(PropertyClass, TEXT("Soft object property missing PropertyClass: %s"), *GetFullNameSafe(this)))
+		{
+			return Super::GetCPPType(ExtendedTypeText, CPPExportFlags);
+		}
+		else
+		{
+			return TEXT("TSoftObjectPtr<UObject>");
+		}
+	}
+
 private:
 	virtual uint32 GetValueTypeHashInternal(const void* Src) const override;
 public:
