Fix invalid file read crashes in AssetRegistry by properly handling IO errors when saving source files in the editor

diff --git a/Engine/Source/Runtime/CoreUObject/Private/UObject/LinkerSave.cpp b/Engine/Source/Runtime/CoreUObject/Private/UObject/LinkerSave.cpp
--- a/Engine/Source/Runtime/CoreUObject/Private/UObject/LinkerSave.cpp
+++ b/Engine/Source/Runtime/CoreUObject/Private/UObject/LinkerSave.cpp
@@ -171,22 +171,24 @@ FLinkerSave::FLinkerSave(UPackage* InParent, bool bForceByteSwapping, bool bInSa
 #endif // USE_STABLE_LOCALIZATION_KEYS
 	}
 }
-/**
- * Detaches file saver and hence file handle.
- */
-void FLinkerSave::Detach()
+
+bool FLinkerSave::CloseAndDestroySaver()
 {
+	bool bSuccess = true;
 	if (Saver)
 	{
+		// first, do an explicit close to check for archive errors
+		bSuccess = Saver->Close();
+		// then, destroy it
 		delete Saver;
 	}
 	Saver = nullptr;
+	return bSuccess;
 }
 
 FLinkerSave::~FLinkerSave()
 {
-	// Detach file saver/ handle.
-	Detach();
+	CloseAndDestroySaver();
 }
 
 int32 FLinkerSave::MapName(FNameEntryId Id) const
diff --git a/Engine/Source/Runtime/CoreUObject/Private/UObject/SavePackage.cpp b/Engine/Source/Runtime/CoreUObject/Private/UObject/SavePackage.cpp
--- a/Engine/Source/Runtime/CoreUObject/Private/UObject/SavePackage.cpp
+++ b/Engine/Source/Runtime/CoreUObject/Private/UObject/SavePackage.cpp
@@ -431,7 +431,7 @@ static bool EndSavingIfCancelled( FLinkerSave* Linker, const FString& TempFilena
 	if ( GWarn->ReceivedUserCancel() )
 	{
 		// free the file handle and delete the temporary file
-		Linker->Detach();
+		Linker->CloseAndDestroySaver();
 		IFileManager::Get().Delete( *TempFilename );
 		return true;
 	}
@@ -4439,7 +4439,7 @@ FSavePackageResultStruct UPackage::Save(UPackage* InOuter, UObject* Base, EObjec
 					}
 
 					// Free the file handle and delete the temporary file
-					Linker->Detach();
+					Linker->CloseAndDestroySaver();
 					IFileManager::Get().Delete( *TempFilename );
 					if (!(SaveFlags & SAVE_NoError))
 					{
@@ -4491,7 +4491,7 @@ FSavePackageResultStruct UPackage::Save(UPackage* InOuter, UObject* Base, EObjec
 					}
 
 					// free the file handle and delete the temporary file
-					Linker->Detach();
+					Linker->CloseAndDestroySaver();
 					IFileManager::Get().Delete(*TempFilename);
 					if (!(SaveFlags & SAVE_NoError))
 					{
@@ -5995,14 +5995,22 @@ FSavePackageResultStruct UPackage::Save(UPackage* InOuter, UObject* Base, EObjec
 				}
 				SlowTask.EnterProgressFrame();
 
-				// Detach archive used for saving, closing file handle.
+				// Destroy archives used for saving, closing file handle.
 				if (!bSaveAsync)
 				{
-					Linker->Detach();
+					const bool bFileWriterSuccess = Linker->CloseAndDestroySaver();
 
 					delete StructuredArchive;
 					delete Formatter;
 					delete TextFormatArchive;
+
+					if (!bFileWriterSuccess)
+					{
+						IFileManager::Get().Delete(*TempFilename);
+						UE_LOG(LogSavePackage, Error, TEXT("Error writing temp file '%s' for '%s'"),
+							*TempFilename, Filename);
+						return ESavePackageResult::Error;
+					}
 				}
 				UNCLOCK_CYCLES(Time);
 				UE_CLOG(!bDiffing, LogSavePackage, Verbose,  TEXT("Save=%.2fms"), FPlatformTime::ToMilliseconds(Time) );
@@ -6087,7 +6095,7 @@ FSavePackageResultStruct UPackage::Save(UPackage* InOuter, UObject* Base, EObjec
 								AsyncWriteFile(MoveTemp(DataPtr), DataSize, *NewPathToSave, FinalTimeStamp);
 							}
 						}
-						Linker->Detach();
+						Linker->CloseAndDestroySaver();
 
 						delete StructuredArchive;
 						delete Formatter;
diff --git a/Engine/Source/Runtime/CoreUObject/Public/UObject/LinkerSave.h b/Engine/Source/Runtime/CoreUObject/Public/UObject/LinkerSave.h
--- a/Engine/Source/Runtime/CoreUObject/Public/UObject/LinkerSave.h
+++ b/Engine/Source/Runtime/CoreUObject/Public/UObject/LinkerSave.h
@@ -67,11 +67,11 @@ public:
 	/** A mapping of package name to generated script SHA keys */
 	COREUOBJECT_API static TMap<FString, TArray<uint8> > PackagesToScriptSHAMap;
 
-	/** Constructor */
+	/** Constructor for file writer */
 	FLinkerSave(UPackage* InParent, const TCHAR* InFilename, bool bForceByteSwapping, bool bInSaveUnversioned = false );
 	/** Constructor for memory writer */
 	FLinkerSave(UPackage* InParent, bool bForceByteSwapping, bool bInSaveUnversioned = false );
-	/** Constructor for custom savers */
+	/** Constructor for custom savers. The linker assumes ownership of the custom saver. */
 	FLinkerSave(UPackage* InParent, FArchive *InSaver, bool bForceByteSwapping, bool bInSaveUnversioned = false);
 
 	/** Returns the appropriate name index for the source name, or 0 if not found in NameIndices */
@@ -110,9 +110,10 @@ public:
 	void Serialize( void* V, int64 Length );
 
 	/**
-	 * Detaches file saver and hence file handle.
+	 * Closes and deletes the Saver (file, memory or custom writer) which will close any associated file handle.
+	 * Returns false if the owned saver contains errors after closing it, true otherwise.
 	 */
-	void Detach();
+	bool CloseAndDestroySaver();
 
 	/**
 	 * Sets a flag indicating that this archive contains data required to be gathered for localization.
