Fix UEnum::NetSerializeProperty not serializing the correct number of bits.
This would only cause problems if the Max Value for an enum was a power of 2, and we tried to serialize the max value.

diff --git a/Engine/Source/Runtime/Core/Public/Misc/NetworkVersion.h b/Engine/Source/Runtime/Core/Public/Misc/NetworkVersion.h
--- a/Engine/Source/Runtime/Core/Public/Misc/NetworkVersion.h
+++ b/Engine/Source/Runtime/Core/Public/Misc/NetworkVersion.h
@@ -43,7 +43,7 @@ enum EEngineNetworkVersionHistory
 	HISTORY_NETEXPORT_SERIALIZATION = 9,			// Bump version due to serialization change to FNetFieldExport
 	HISTORY_NETEXPORT_SERIALIZE_FIX = 10,			// Bump version to fix net field export name serialization 
 	HISTORY_FAST_ARRAY_DELTA_STRUCT = 11,			// Bump version to allow fast array serialization, delta struct serialization.
-
+	HISTORY_FIX_ENUM_SERIALIZATION = 12,			// Bump version to fix enum net serialization issues.
 	// New history items go above here.
 
 	HISTORY_ENGINENETVERSION_PLUS_ONE,
diff --git a/Engine/Source/Runtime/CoreUObject/Private/UObject/EnumProperty.cpp b/Engine/Source/Runtime/CoreUObject/Private/UObject/EnumProperty.cpp
--- a/Engine/Source/Runtime/CoreUObject/Private/UObject/EnumProperty.cpp
+++ b/Engine/Source/Runtime/CoreUObject/Private/UObject/EnumProperty.cpp
@@ -8,6 +8,7 @@
 #include "Templates/IsSigned.h"
 #include "Algo/Find.h"
 #include "UObject/LinkerLoad.h"
+#include "Misc/NetworkVersion.h"
 
 namespace UE4EnumProperty_Private
 {
@@ -141,7 +142,14 @@ void UEnumProperty::SerializeItem(FStructuredArchive::FSlot Slot, void* Value, v
 
 bool UEnumProperty::NetSerializeItem(FArchive& Ar, UPackageMap* Map, void* Data, TArray<uint8>* MetaData) const
 {
-	Ar.SerializeBits(Data, FMath::CeilLogTwo64(Enum->GetMaxEnumValue()));
+	if (Ar.EngineNetVer() < HISTORY_FIX_ENUM_SERIALIZATION)
+	{
+		Ar.SerializeBits(Data, FMath::CeilLogTwo64(Enum->GetMaxEnumValue()));
+	}
+	else
+	{
+		Ar.SerializeBits(Data, FMath::CeilLogTwo64(Enum->GetMaxEnumValue() + 1));
+	}
 	return 1;
 }
 
