Fix sub object reinstancing when serialized sub object type does not match CDO object type

diff --git a/Engine/Source/Runtime/CoreUObject/Private/UObject/CoreNative.cpp b/Engine/Source/Runtime/CoreUObject/Private/UObject/CoreNative.cpp
--- a/Engine/Source/Runtime/CoreUObject/Private/UObject/CoreNative.cpp
+++ b/Engine/Source/Runtime/CoreUObject/Private/UObject/CoreNative.cpp
@@ -133,7 +133,7 @@ UObject* FObjectInstancingGraph::GetInstancedSubobject( UObject* SourceSubobject
 	{
 		bool bAllowedSelfReference = bAllowSelfReference && SourceSubobject == SourceRoot;
 
-		bool bShouldInstance = bAllowedSelfReference || SourceSubobject->IsIn(SourceRoot);
+		bool bShouldInstance = bAllowedSelfReference || CurrentValue->IsIn(SourceRoot);
 		if ( !bShouldInstance && CurrentValue->GetOuter() == CurrentObject->GetArchetype() )
 		{
 			// this code is intended to catch cases where SourceRoot contains subobjects assigned to instanced object properties, where the subobject's class
