When checking if additional GC time is required on ios/android dont check platform time with each iteration.
#rnx

diff --git a/Engine/Source/Runtime/CoreUObject/Private/UObject/GarbageCollection.cpp b/Engine/Source/Runtime/CoreUObject/Private/UObject/GarbageCollection.cpp
--- a/Engine/Source/Runtime/CoreUObject/Private/UObject/GarbageCollection.cpp
+++ b/Engine/Source/Runtime/CoreUObject/Private/UObject/GarbageCollection.cpp
@@ -1522,7 +1522,7 @@ bool IncrementalDestroyGarbage(bool bUseTimeLimit, float TimeLimit)
 					{
 						const bool bPollTimeLimit = ((FinishDestroyTimePollCounter++) % TimeLimitEnforcementGranularityForDestroy == 0);
 #if PLATFORM_IOS || PLATFORM_ANDROID
-						if( !bFinishDestroyTimeExtended && (FPlatformTime::Seconds() - GCStartTime) > MaxTimeForFinishDestroy)
+						if(bPollTimeLimit && !bFinishDestroyTimeExtended && (FPlatformTime::Seconds() - GCStartTime) > MaxTimeForFinishDestroy )
 						{
 							MaxTimeForFinishDestroy = 30.0;
 							bFinishDestroyTimeExtended = true;
