Part 2 of async loading stall workaround when post-loading async textures.

diff --git a/Engine/Source/Runtime/CoreUObject/Private/Serialization/AsyncLoading.cpp b/Engine/Source/Runtime/CoreUObject/Private/Serialization/AsyncLoading.cpp
--- a/Engine/Source/Runtime/CoreUObject/Private/Serialization/AsyncLoading.cpp
+++ b/Engine/Source/Runtime/CoreUObject/Private/Serialization/AsyncLoading.cpp
@@ -6811,6 +6811,13 @@ EAsyncPackageState::Type FAsyncPackage::PostLoadDeferredObjects(double InTickSta
 	{
 		UObject* Object = DeferredPostLoadObjects[DeferredPostLoadIndex++];
 		check(Object);
+
+		if (!Object->IsReadyForAsyncPostLoad())
+		{
+			--DeferredPostLoadIndex;
+			break;
+		}
+
 		LastObjectWorkWasPerformedOn = Object;
 		LastTypeOfWorkPerformed = TEXT("postloading_gamethread");
 
