Fix some warning spam when async loading packages that were redirected with CoreRedirects

diff --git a/Engine/Source/Runtime/CoreUObject/Private/Serialization/AsyncLoading.cpp b/Engine/Source/Runtime/CoreUObject/Private/Serialization/AsyncLoading.cpp
--- a/Engine/Source/Runtime/CoreUObject/Private/Serialization/AsyncLoading.cpp
+++ b/Engine/Source/Runtime/CoreUObject/Private/Serialization/AsyncLoading.cpp
@@ -6358,7 +6358,12 @@ EAsyncPackageState::Type FAsyncPackage::LoadImports(FFlushTree* FlushTree)
 		{
 			continue;
 		}
-			
+
+		// This may be an import left behind from a core redirects fixup
+		if (Import->ObjectName.IsNone())
+		{
+			continue;
+		}
 
 		// Don't try to import a package that is in an import table that we know is an invalid entry
 		if (FLinkerLoad::IsKnownMissingPackage(Import->ObjectName))
