Remove out of date component check code.
#rnx

diff --git a/Engine/Source/Runtime/CoreUObject/Private/UObject/Obj.cpp b/Engine/Source/Runtime/CoreUObject/Private/UObject/Obj.cpp
--- a/Engine/Source/Runtime/CoreUObject/Private/UObject/Obj.cpp
+++ b/Engine/Source/Runtime/CoreUObject/Private/UObject/Obj.cpp
@@ -1541,7 +1541,6 @@ protected:
 };
 
 
-#define ALLOW_SUB_SUB_OBJECTS (1)
 // if this is set to fatal, then we don't run any testing since it is time consuming.
 DEFINE_LOG_CATEGORY_STATIC(LogCheckSubobjects, Fatal, All);
 
@@ -1622,16 +1621,6 @@ bool UObject::CheckDefaultSubobjectsInternal() const
 	}
 
 
-	TArray<UObject *> AllCollectedComponents;
-	CollectDefaultSubobjects(AllCollectedComponents, true);
-	TArray<UObject *> DirectCollectedComponents;
-	CollectDefaultSubobjects(DirectCollectedComponents, false);
-		
-	AllCollectedComponents.Sort();
-	DirectCollectedComponents.Sort();
-
-	CompCheck(ALLOW_SUB_SUB_OBJECTS || AllCollectedComponents == DirectCollectedComponents); // just say no to subobjects of subobjects
-
 	return Result;
 }
 
