Change GetToolTipText to cutoff after 2 empty or punctuation separator lines, allowing large comments for the API docs but smaller ones for blueprint
Expose FormatNativeToolTip so it can be used by things other than GetToolTipText
Fix Actor's comment to not be enormous in the blueprint editor

diff --git a/Engine/Source/Runtime/CoreUObject/Private/UObject/Class.cpp b/Engine/Source/Runtime/CoreUObject/Private/UObject/Class.cpp
--- a/Engine/Source/Runtime/CoreUObject/Private/UObject/Class.cpp
+++ b/Engine/Source/Runtime/CoreUObject/Private/UObject/Class.cpp
@@ -305,14 +305,9 @@ FText UField::GetToolTipText(bool bShortTooltip) const
 		{
 			NativeToolTip = FName::NameToDisplayString(FDisplayNameHelper::Get(*this), IsA<UBoolProperty>());
 		}
-		else
+		else if (!bShortTooltip && IsNative())
 		{
-			static const FString DoxygenSee(TEXT("@see"));
-			static const FString TooltipSee(TEXT("See:"));
-			if (NativeToolTip.ReplaceInline(*DoxygenSee, *TooltipSee) > 0)
-			{
-				NativeToolTip.TrimEndInline();
-			}
+			FormatNativeToolTip(NativeToolTip, true);
 		}
 		LocalizedToolTip = FText::FromString(NativeToolTip);
 	}
@@ -320,6 +315,67 @@ FText UField::GetToolTipText(bool bShortTooltip) const
 	return LocalizedToolTip;
 }
 
+void UField::FormatNativeToolTip(FString& ToolTipString, bool bRemoveExtraSections)
+{
+	// First do doxygen replace
+	static const FString DoxygenSee(TEXT("@see"));
+	static const FString TooltipSee(TEXT("See:"));
+	ToolTipString.ReplaceInline(*DoxygenSee, *TooltipSee);
+
+	bool bCurrentLineIsEmpty = true;
+	int32 EmptyLineCount = 0;
+	int32 LastContentIndex = INDEX_NONE;
+	const int32 ToolTipLength = ToolTipString.Len();
+		
+	// Start looking for empty lines and whitespace to strip
+	for (int32 StrIndex = 0; StrIndex < ToolTipLength; StrIndex++)
+	{
+		TCHAR CurrentChar = ToolTipString[StrIndex];
+
+		if (!FChar::IsWhitespace(CurrentChar))
+		{
+			if (FChar::IsPunct(CurrentChar))
+			{
+				// Punctuation is considered content if it's on a line with alphanumeric text
+				if (!bCurrentLineIsEmpty)
+				{
+					LastContentIndex = StrIndex;
+				}
+			}
+			else
+			{
+				// This is something alphanumeric, this is always content and mark line as not empty
+				bCurrentLineIsEmpty = false;
+				LastContentIndex = StrIndex;
+			}
+		}
+		else if (CurrentChar == TEXT('\n'))
+		{
+			if (bCurrentLineIsEmpty)
+			{
+				EmptyLineCount++;
+				if (bRemoveExtraSections && EmptyLineCount >= 2)
+				{
+					// If we get two empty or punctuation/separator lines in a row, cut off the string if requested
+					break;
+				}
+			}
+			else
+			{
+				EmptyLineCount = 0;
+			}
+
+			bCurrentLineIsEmpty = true;
+		}
+	}
+
+	// Trim string to last content character, this strips trailing whitespace as well as extra sections if needed
+	if (LastContentIndex >= 0 && LastContentIndex != ToolTipLength - 1)
+	{
+		ToolTipString.RemoveAt(LastContentIndex + 1, ToolTipLength - (LastContentIndex + 1));
+	}
+}
+
 /**
  * Determines if the property has any metadata associated with the key
  * 
diff --git a/Engine/Source/Runtime/CoreUObject/Public/UObject/Class.h b/Engine/Source/Runtime/CoreUObject/Public/UObject/Class.h
--- a/Engine/Source/Runtime/CoreUObject/Public/UObject/Class.h
+++ b/Engine/Source/Runtime/CoreUObject/Public/UObject/Class.h
@@ -108,6 +108,14 @@ class COREUOBJECT_API UField : public UObject
 	 */
 	FText GetToolTipText(bool bShortTooltip = false) const;
 
+	/** 
+	 * Formats a source comment into the form we want to show in the editor, is used by GetToolTipText and anything else that will get a native tooltip 
+	 * 
+	 * @param ToolTipString			String parsed out of C++ headers that is modified in place
+	 * @param bRemoveExtraSections	If true, cut off the comment on first line separator or 2 empty lines in a row
+	 */
+	static void FormatNativeToolTip(FString& ToolTipString, bool bRemoveExtraSections = true);
+
 	/**
 	 * Determines if the property has any metadata associated with the key
 	 * 
diff --git a/Engine/Source/Runtime/Engine/Classes/GameFramework/Actor.h b/Engine/Source/Runtime/Engine/Classes/GameFramework/Actor.h
--- a/Engine/Source/Runtime/Engine/Classes/GameFramework/Actor.h
+++ b/Engine/Source/Runtime/Engine/Classes/GameFramework/Actor.h
@@ -73,6 +73,7 @@ extern ENGINE_API FUObjectAnnotationSparseBool GSelectedActorAnnotation;
  * Actors may contain a collection of ActorComponents, which can be used to control how actors move, how they are rendered, etc.
  * The other main function of an Actor is the replication of properties and function calls across the network during play.
  * 
+ * 
  * Actor initialization has multiple steps, here's the order of important virtual functions that get called:
  * - UObject::PostLoad: For actors statically placed in a level, the normal UObject PostLoad gets called both in the editor and during gameplay.
  *                      This is not called for newly spawned actors.
