Part of PR #5146

diff --git a/Engine/Source/Runtime/CoreUObject/Public/UObject/ObjectMacros.h b/Engine/Source/Runtime/CoreUObject/Public/UObject/ObjectMacros.h
--- a/Engine/Source/Runtime/CoreUObject/Public/UObject/ObjectMacros.h
+++ b/Engine/Source/Runtime/CoreUObject/Public/UObject/ObjectMacros.h
@@ -1267,7 +1267,8 @@ namespace UM
 		/// [FunctionMetadata] Used when ArrayParm has been specified to indicate other function parameters that should be treated as wild card properties linked to the type of the array parameter.
 		ArrayTypeDependentParams,
 
-		/// [FunctionMetadata]
+		/// [FunctionMetadata] For reference parameters, indicates that a value should be created to be used for the input if none is linked via BP.
+		/// This also allows for inline editing of the default value on some types (take FRotator for instance). Only valid for inputs.
 		AutoCreateRefTerm,
 
 		/// [FunctionMetadata] This function is an internal implementation detail, used to implement another function or node.  It is never directly exposed in a graph.
@@ -1285,7 +1286,7 @@ namespace UM
 		/// [FunctionMetadata] Indicates that a BlueprintCallable function should display in the compact display mode and the name to use in that mode.
 		CompactNodeTitle,
 
-		/// [FunctionMetadata]
+		/// [FunctionMetadata] Used with CustomThunk to declare that a parameter is actually polymorphic
 		CustomStructureParam,
 
 		/// [FunctionMetadata] For BlueprintCallable functions indicates that the object property named's default value should be the self context of the node
@@ -1297,7 +1298,7 @@ namespace UM
 		/// [ClassMetadata] [FunctionMetadata] Used in conjunction with DeprecatedNode or DeprecatedFunction to customize the warning message displayed to the user.
 		// DeprecationMessage, (Commented out so as to avoid duplicate name with version in the Class section, but still show in the function section)
 
-		/// [FunctionMetadata] For BlueprintCallable functions indicates that an input exec pin should be created for each entry in the enum specified.
+		/// [FunctionMetadata] For BlueprintCallable functions indicates that an input/output (determined by whether it is an input/output enum) exec pin should be created for each entry in the enum specified.
 		ExpandEnumAsExecs,
 
 		/// [ClassMetadata] [PropertyMetadata] [FunctionMetadata] The name to display for this class, property, or function instead of auto-generating it from the name.
@@ -1342,7 +1343,7 @@ namespace UM
 		/// [FunctionMetadata] For BlueprintCallable functions indicates that the parameter pin should be hidden from the user's view.
 		HidePin,
 
-		/// [FunctionMetadata]
+		/// [FunctionMetadata] For some functions used by async task nodes, specify this parameter should be skipped when exposing pins
 		HideSpawnParms,
 
 		/// [FunctionMetadata] For BlueprintCallable functions provides additional keywords to be associated with the function for search purposes.
