Added version of SetDefaultSubobjectClass that takes UClass explicitly rather than relying on templated type.
Cleaned up unnecessary function overrides and conversions of FName to string
Fixed capitalization on IsLegalOverride

diff --git a/Engine/Source/Runtime/CoreUObject/Private/UObject/UObjectGlobals.cpp b/Engine/Source/Runtime/CoreUObject/Private/UObject/UObjectGlobals.cpp
--- a/Engine/Source/Runtime/CoreUObject/Private/UObject/UObjectGlobals.cpp
+++ b/Engine/Source/Runtime/CoreUObject/Private/UObject/UObjectGlobals.cpp
@@ -3019,7 +3019,7 @@ void FObjectInitializer::InitProperties(UObject* Obj, UClass* DefaultsClass, UOb
 	}
 }
 
-bool FObjectInitializer::IslegalOverride(FName InComponentName, class UClass *DerivedComponentClass, class UClass *BaseComponentClass) const
+bool FObjectInitializer::IsLegalOverride(FName InComponentName, class UClass *DerivedComponentClass, class UClass *BaseComponentClass) const
 {
 	if (DerivedComponentClass && BaseComponentClass && !DerivedComponentClass->IsChildOf(BaseComponentClass) )
 	{
@@ -3030,10 +3030,10 @@ bool FObjectInitializer::IslegalOverride(FName InComponentName, class UClass *De
 	return true;
 }
 
-void FObjectInitializer::AssertIfSubobjectSetupIsNotAllowed(const TCHAR* SubobjectName) const
+void FObjectInitializer::AssertIfSubobjectSetupIsNotAllowed(const FName SubobjectName) const
 {
 	UE_CLOG(!bSubobjectClassInitializationAllowed, LogUObjectGlobals, Fatal,
-		TEXT("%s.%s: Subobject class setup is only allowed in base class constructor call (in the initialization list)"), Obj ? *Obj->GetFullName() : TEXT("NULL"), SubobjectName);
+		TEXT("%s.%s: Subobject class setup is only allowed in base class constructor call (in the initialization list)"), Obj ? *Obj->GetFullName() : TEXT("NULL"), *SubobjectName.GetPlainNameString());
 }
 
 #if DO_CHECK
diff --git a/Engine/Source/Runtime/CoreUObject/Public/UObject/UObjectGlobals.h b/Engine/Source/Runtime/CoreUObject/Public/UObject/UObjectGlobals.h
--- a/Engine/Source/Runtime/CoreUObject/Public/UObject/UObjectGlobals.h
+++ b/Engine/Source/Runtime/CoreUObject/Public/UObject/UObjectGlobals.h
@@ -926,26 +926,25 @@ public:
 
 
 	/**
-	 * Sets the class of a subobject for a base class
+	 * Sets the class to use for a subobject defined in a base class, the class must be a subclass of the class used by the base class.
 	 * @param	SubobjectName	name of the new component or subobject
+	 * @param	Class			The class to use for the specified subobject or component.
 	 */
-	template<class T>
-	FObjectInitializer const& SetDefaultSubobjectClass(FName SubobjectName) const
+	FObjectInitializer const& SetDefaultSubobjectClass(FName SubobjectName, UClass* Class) const
 	{
-		AssertIfSubobjectSetupIsNotAllowed(*SubobjectName.GetPlainNameString());
-		ComponentOverrides.Add(SubobjectName, T::StaticClass(), *this);
+		AssertIfSubobjectSetupIsNotAllowed(SubobjectName);
+		ComponentOverrides.Add(SubobjectName, Class, *this);
 		return *this;
 	}
+
 	/**
-	 * Sets the class of a subobject for a base class
+	 * Sets the class to use for a subobject defined in a base class, the class must be a subclass of the class used by the base class.
 	 * @param	SubobjectName	name of the new component or subobject
 	 */
 	template<class T>
-	FORCEINLINE FObjectInitializer const& SetDefaultSubobjectClass(TCHAR const*SubobjectName) const
+	FObjectInitializer const& SetDefaultSubobjectClass(FName SubobjectName) const
 	{
-		AssertIfSubobjectSetupIsNotAllowed(SubobjectName);
-		ComponentOverrides.Add(SubobjectName, T::StaticClass(), *this);
-		return *this;
+		return SetDefaultSubobjectClass(SubobjectName, T::StaticClass());
 	}
 
 	/**
@@ -953,17 +952,6 @@ public:
 	 * @param	SubobjectName	name of the new component or subobject to not create
 	 */
 	FObjectInitializer const& DoNotCreateDefaultSubobject(FName SubobjectName) const
-	{
-		AssertIfSubobjectSetupIsNotAllowed(*SubobjectName.GetPlainNameString());
-		ComponentOverrides.Add(SubobjectName, nullptr, *this);
-		return *this;
-	}
-
-	/**
-	 * Indicates that a base class should not create a component
-	 * @param	ComponentName	name of the new component or subobject to not create
-	 */
-	FORCEINLINE FObjectInitializer const& DoNotCreateDefaultSubobject(TCHAR const*SubobjectName) const
 	{
 		AssertIfSubobjectSetupIsNotAllowed(SubobjectName);
 		ComponentOverrides.Add(SubobjectName, nullptr, *this);
@@ -973,7 +961,7 @@ public:
 	/** 
 	 * Internal use only, checks if the override is legal and if not deal with error messages
 	**/
-	bool IslegalOverride(FName InComponentName, class UClass *DerivedComponentClass, class UClass *BaseComponentClass) const;
+	bool IsLegalOverride(FName InComponentName, class UClass* DerivedComponentClass, class UClass* BaseComponentClass) const;
 
 	/**
 	 * Asserts with the specified message if code is executed inside UObject constructor
@@ -1044,7 +1032,7 @@ private:
 	struct FOverrides
 	{
 		/**  Add an override, make sure it is legal **/
-		void Add(FName InComponentName, UClass *InComponentClass, FObjectInitializer const& ObjectInitializer)
+		void Add(FName InComponentName, UClass* InComponentClass, FObjectInitializer const& ObjectInitializer)
 		{
 			int32 Index = Find(InComponentName);
 			if (Index == INDEX_NONE)
@@ -1053,7 +1041,7 @@ private:
 			}
 			else if (InComponentClass && Overrides[Index].ComponentClass)
 			{
-				ObjectInitializer.IslegalOverride(InComponentName, Overrides[Index].ComponentClass, InComponentClass); // if a base class is asking for an override, the existing override (which we are going to use) had better be derived
+				ObjectInitializer.IsLegalOverride(InComponentName, Overrides[Index].ComponentClass, InComponentClass); // if a base class is asking for an override, the existing override (which we are going to use) had better be derived
 			}
 		}
 		/**  Retrieve an override, or TClassToConstructByDefault::StaticClass or nullptr if this was removed by a derived class **/
@@ -1067,7 +1055,7 @@ private:
 			}
 			else if (Overrides[Index].ComponentClass)
 			{
-				if (ObjectInitializer.IslegalOverride(InComponentName, Overrides[Index].ComponentClass, ReturnType)) // if THE base class is asking for a T, the existing override (which we are going to use) had better be derived
+				if (ObjectInitializer.IsLegalOverride(InComponentName, Overrides[Index].ComponentClass, ReturnType)) // if THE base class is asking for a T, the existing override (which we are going to use) had better be derived
 				{
 					return Overrides[Index].ComponentClass; // the override is of an acceptable class, so use it
 				}
@@ -1130,7 +1118,7 @@ private:
 	};
 
 	/** Asserts if SetDefaultSubobjectClass or DoNotCreateOptionalDefaultSuobject are called inside of the constructor body */
-	void AssertIfSubobjectSetupIsNotAllowed(const TCHAR* SubobjectName) const;
+	void AssertIfSubobjectSetupIsNotAllowed(const FName SubobjectName) const;
 
 	/**  object to initialize, from static allocate object, after construction **/
 	UObject* Obj;
