Minor clean up
#rnx

diff --git a/Engine/Source/Runtime/CoreUObject/Public/UObject/UObjectGlobals.h b/Engine/Source/Runtime/CoreUObject/Public/UObject/UObjectGlobals.h
--- a/Engine/Source/Runtime/CoreUObject/Public/UObject/UObjectGlobals.h
+++ b/Engine/Source/Runtime/CoreUObject/Public/UObject/UObjectGlobals.h
@@ -1028,7 +1028,7 @@ private:
 
 private:
 
-	/**  Littel helper struct to manage overrides from dervied classes **/
+	/**  Little helper struct to manage overrides from derived classes **/
 	struct FOverrides
 	{
 		/**  Add an override, make sure it is legal **/
@@ -1048,7 +1048,7 @@ private:
 		UClass* Get(FName InComponentName, UClass* ReturnType, UClass* ClassToConstructByDefault, FObjectInitializer const& ObjectInitializer)
 		{
 			int32 Index = Find(InComponentName);
-			UClass *BaseComponentClass = ClassToConstructByDefault;
+			UClass* BaseComponentClass = ClassToConstructByDefault;
 			if (Index == INDEX_NONE)
 			{
 				return BaseComponentClass; // no override so just do what the base class wanted
@@ -1063,6 +1063,7 @@ private:
 			}
 			return nullptr;  // the override is of nullptr, which means "don't create this component"
 		}
+
 private:
 		/**  Search for an override **/
 		int32 Find(FName InComponentName)
@@ -1080,7 +1081,7 @@ private:
 		struct FOverride
 		{
 			FName	ComponentName;
-			UClass *ComponentClass;
+			UClass* ComponentClass;
 			FOverride(FName InComponentName, UClass *InComponentClass)
 				: ComponentName(InComponentName)
 				, ComponentClass(InComponentClass)
@@ -1090,7 +1091,7 @@ private:
 		/**  The override array **/
 		TArray<FOverride, TInlineAllocator<8> > Overrides;
 	};
-	/**  Littel helper struct to manage overrides from dervied classes **/
+	/**  Little helper struct to manage overrides from derived classes **/
 	struct FSubobjectsToInit
 	{
 		/**  Add a subobject **/
