const correct Outers and Classes for object hash functions

diff --git a/Engine/Source/Runtime/CoreUObject/Private/UObject/UObjectHash.cpp b/Engine/Source/Runtime/CoreUObject/Private/UObject/UObjectHash.cpp
--- a/Engine/Source/Runtime/CoreUObject/Private/UObject/UObjectHash.cpp
+++ b/Engine/Source/Runtime/CoreUObject/Private/UObject/UObjectHash.cpp
@@ -427,7 +427,7 @@ static FORCEINLINE int32 GetObjectOuterHash(FName ObjName,PTRINT Outer)
 	return GetTypeHash(ObjName) + (Outer >> 6);
 }
 
-UObject* StaticFindObjectFastExplicitThreadSafe(FUObjectHashTables& ThreadHash, UClass* ObjectClass, FName ObjectName, const FString& ObjectPathName, bool bExactClass, EObjectFlags ExcludeFlags/*=0*/)
+UObject* StaticFindObjectFastExplicitThreadSafe(FUObjectHashTables& ThreadHash, const UClass* ObjectClass, FName ObjectName, const FString& ObjectPathName, bool bExactClass, EObjectFlags ExcludeFlags/*=0*/)
 {
 	const EInternalObjectFlags ExclusiveInternalFlags = EInternalObjectFlags::Unreachable;
 
@@ -477,7 +477,7 @@ UObject* StaticFindObjectFastExplicitThreadSafe(FUObjectHashTables& ThreadHash,
  * @param	ExclusiveFlags	Ignores objects that contain any of the specified exclusive flags
  * @return	Returns a pointer to the found object or nullptr if none could be found
  */
-UObject* StaticFindObjectFastExplicit( UClass* ObjectClass, FName ObjectName, const FString& ObjectPathName, bool bExactClass, EObjectFlags ExcludeFlags/*=0*/ )
+UObject* StaticFindObjectFastExplicit( const UClass* ObjectClass, FName ObjectName, const FString& ObjectPathName, bool bExactClass, EObjectFlags ExcludeFlags/*=0*/ )
 {
 	checkSlow(FPackageName::IsShortPackageName(ObjectName)); //@Package name transition, we aren't checking the name here because we know this is only used for texture
 
@@ -489,7 +489,7 @@ UObject* StaticFindObjectFastExplicit( UClass* ObjectClass, FName ObjectName, co
 	return Result;
 }
 
-UObject* StaticFindObjectFastInternalThreadSafe(FUObjectHashTables& ThreadHash, UClass* ObjectClass, UObject* ObjectPackage, FName ObjectName, bool bExactClass, bool bAnyPackage, EObjectFlags ExcludeFlags, EInternalObjectFlags ExclusiveInternalFlags)
+UObject* StaticFindObjectFastInternalThreadSafe(FUObjectHashTables& ThreadHash, const UClass* ObjectClass, const UObject* ObjectPackage, FName ObjectName, bool bExactClass, bool bAnyPackage, EObjectFlags ExcludeFlags, EInternalObjectFlags ExclusiveInternalFlags)
 {
 	ExclusiveInternalFlags |= EInternalObjectFlags::Unreachable;
 
@@ -600,7 +600,7 @@ UObject* StaticFindObjectFastInternalThreadSafe(FUObjectHashTables& ThreadHash,
 	return Result;
 }
 
-UObject* StaticFindObjectFastInternal(UClass* ObjectClass, UObject* ObjectPackage, FName ObjectName, bool bExactClass, bool bAnyPackage, EObjectFlags ExcludeFlags, EInternalObjectFlags ExclusiveInternalFlags)
+UObject* StaticFindObjectFastInternal(const UClass* ObjectClass, const UObject* ObjectPackage, FName ObjectName, bool bExactClass, bool bAnyPackage, EObjectFlags ExcludeFlags, EInternalObjectFlags ExclusiveInternalFlags)
 {
 	INC_DWORD_STAT(STAT_FindObjectFast);
 
@@ -812,7 +812,7 @@ void ForEachObjectWithOuter(const class UObjectBase* Outer, TFunctionRef<void(UO
 	}
 }
 
-UObjectBase* FindObjectWithOuter(class UObjectBase* Outer, class UClass* ClassToLookFor, FName NameToLookFor)
+UObjectBase* FindObjectWithOuter(const class UObjectBase* Outer, const class UClass* ClassToLookFor, FName NameToLookFor)
 {
 	UObject* Result = nullptr;
 	check( Outer );
@@ -825,7 +825,7 @@ UObjectBase* FindObjectWithOuter(class UObjectBase* Outer, class UClass* ClassTo
 
 	if( NameToLookFor != NAME_None )
 	{
-		Result = StaticFindObjectFastInternal(ClassToLookFor, static_cast<UObject*>(Outer), NameToLookFor, false, false, RF_NoFlags, ExclusionInternalFlags);
+		Result = StaticFindObjectFastInternal(ClassToLookFor, static_cast<const UObject*>(Outer), NameToLookFor, false, false, RF_NoFlags, ExclusionInternalFlags);
 	}
 	else
 	{
@@ -854,21 +854,20 @@ UObjectBase* FindObjectWithOuter(class UObjectBase* Outer, class UClass* ClassTo
 }
 
 /** Helper function that returns all the children of the specified class recursively */
-template<typename ArrayAllocator>
-static void RecursivelyPopulateDerivedClasses(FUObjectHashTables& ThreadHash, UClass* ParentClass, TArray<UClass*, ArrayAllocator>& OutAllDerivedClass)
+template<typename ClassType, typename ArrayAllocator>
+static void RecursivelyPopulateDerivedClasses(FUObjectHashTables& ThreadHash, const UClass* ParentClass, TArray<ClassType, ArrayAllocator>& OutAllDerivedClass)
 {
 	// Start search with the parent class at virtual index Num-1, then continue searching from index Num as things are added
 	int32 SearchIndex = OutAllDerivedClass.Num() - 1;
-	UClass* SearchClass = ParentClass;
+	const UClass* SearchClass = ParentClass;
 
 	while (1)
 	{
 		TSet<UClass*>* ChildSet = ThreadHash.ClassToChildListMap.Find(SearchClass);
 		if (ChildSet)
 		{
-			for (auto ChildIt = ChildSet->CreateConstIterator(); ChildIt; ++ChildIt)
+			for (UClass* ChildClass : *ChildSet)
 			{
-				UClass* ChildClass = *ChildIt;
 				OutAllDerivedClass.Add(ChildClass);
 			}
 		}
@@ -887,7 +886,7 @@ static void RecursivelyPopulateDerivedClasses(FUObjectHashTables& ThreadHash, UC
 	}
 }
 
-void GetObjectsOfClass(UClass* ClassToLookFor, TArray<UObject *>& Results, bool bIncludeDerivedClasses, EObjectFlags ExclusionFlags, EInternalObjectFlags ExclusionInternalFlags)
+void GetObjectsOfClass(const UClass* ClassToLookFor, TArray<UObject *>& Results, bool bIncludeDerivedClasses, EObjectFlags ExclusionFlags, EInternalObjectFlags ExclusionInternalFlags)
 {
 	SCOPE_CYCLE_COUNTER(STAT_Hash_GetObjectsOfClass);
 
@@ -901,7 +900,7 @@ void GetObjectsOfClass(UClass* ClassToLookFor, TArray<UObject *>& Results, bool
 	check(Results.Num() <= GUObjectArray.GetObjectArrayNum()); // otherwise we have a cycle in the outer chain, which should not be possible
 }
 
-void ForEachObjectOfClass(UClass* ClassToLookFor, TFunctionRef<void(UObject*)> Operation, bool bIncludeDerivedClasses, EObjectFlags ExclusionFlags, EInternalObjectFlags ExclusionInternalFlags)
+void ForEachObjectOfClass(const UClass* ClassToLookFor, TFunctionRef<void(UObject*)> Operation, bool bIncludeDerivedClasses, EObjectFlags ExclusionFlags, EInternalObjectFlags ExclusionInternalFlags)
 {
 	// We don't want to return any objects that are currently being background loaded unless we're using the object iterator during async loading.
 	ExclusionInternalFlags |= EInternalObjectFlags::Unreachable;
@@ -911,7 +910,7 @@ void ForEachObjectOfClass(UClass* ClassToLookFor, TFunctionRef<void(UObject*)> O
 	}
 
 	// Most classes searched for have around 10 subclasses, some have hundreds
-	TArray<UClass*, TInlineAllocator<16>> ClassesToSearch;
+	TArray<const UClass*, TInlineAllocator<16>> ClassesToSearch;
 	ClassesToSearch.Add(ClassToLookFor);
 
 	FUObjectHashTables& ThreadHash = FUObjectHashTables::Get();
@@ -922,7 +921,7 @@ void ForEachObjectOfClass(UClass* ClassToLookFor, TFunctionRef<void(UObject*)> O
 		RecursivelyPopulateDerivedClasses(ThreadHash, ClassToLookFor, ClassesToSearch);
 	}
 
-	for (UClass* SearchClass : ClassesToSearch)
+	for (const UClass* SearchClass : ClassesToSearch)
 	{
 		FHashBucket* List = ThreadHash.ClassToObjectListMap.Find(SearchClass);
 		if (List)
@@ -939,7 +938,7 @@ void ForEachObjectOfClass(UClass* ClassToLookFor, TFunctionRef<void(UObject*)> O
 	}
 }
 
-void GetDerivedClasses(UClass* ClassToLookFor, TArray<UClass *>& Results, bool bRecursive)
+void GetDerivedClasses(const UClass* ClassToLookFor, TArray<UClass*>& Results, bool bRecursive)
 {
 	auto& ThreadHash = FUObjectHashTables::Get();
 	FHashTableLock HashLock(ThreadHash);
@@ -958,9 +957,9 @@ void GetDerivedClasses(UClass* ClassToLookFor, TArray<UClass *>& Results, bool b
 	}
 }
 
-bool ClassHasInstancesAsyncLoading(UClass* ClassToLookFor)
+bool ClassHasInstancesAsyncLoading(const UClass* ClassToLookFor)
 {
-	TArray<UClass*> ClassesToSearch;
+	TArray<const UClass*> ClassesToSearch;
 	ClassesToSearch.Add(ClassToLookFor);
 
 	auto& ThreadHash = FUObjectHashTables::Get();
@@ -968,7 +967,7 @@ bool ClassHasInstancesAsyncLoading(UClass* ClassToLookFor)
 
 	RecursivelyPopulateDerivedClasses(ThreadHash, ClassToLookFor, ClassesToSearch);
 
-	for (UClass* SearchClass : ClassesToSearch)
+	for (const UClass* SearchClass : ClassesToSearch)
 	{
 		FHashBucket* List = ThreadHash.ClassToObjectListMap.Find(SearchClass);
 		if (List)
diff --git a/Engine/Source/Runtime/CoreUObject/Public/UObject/UObjectHash.h b/Engine/Source/Runtime/CoreUObject/Public/UObject/UObjectHash.h
--- a/Engine/Source/Runtime/CoreUObject/Public/UObject/UObjectHash.h
+++ b/Engine/Source/Runtime/CoreUObject/Public/UObject/UObjectHash.h
@@ -29,7 +29,7 @@ DECLARE_DWORD_COUNTER_STAT_EXTERN(TEXT("NumObjects"), STAT_Hash_NumObjects, STAT
  * @param	ExclusiveInternalFlags	Ignores objects that contain any of the specified internal exclusive flags
  * @return	Returns a pointer to the found object or NULL if none could be found
  */
-UObject* StaticFindObjectFastInternal(UClass* Class, UObject* InOuter, FName InName, bool ExactClass = false, bool AnyPackage = false, EObjectFlags ExclusiveFlags = RF_NoFlags, EInternalObjectFlags ExclusiveInternalFlags = EInternalObjectFlags::None);
+UObject* StaticFindObjectFastInternal(const UClass* Class, const UObject* InOuter, FName InName, bool ExactClass = false, bool AnyPackage = false, EObjectFlags ExclusiveFlags = RF_NoFlags, EInternalObjectFlags ExclusiveInternalFlags = EInternalObjectFlags::None);
 
 /**
  * Variation of StaticFindObjectFast that uses explicit path.
@@ -41,7 +41,7 @@ UObject* StaticFindObjectFastInternal(UClass* Class, UObject* InOuter, FName InN
  * @param	ExclusiveFlags	Ignores objects that contain any of the specified exclusive flags
  * @return	Returns a pointer to the found object or NULL if none could be found
  */
-UObject* StaticFindObjectFastExplicit(UClass* ObjectClass, FName ObjectName, const FString& ObjectPathName, bool bExactClass, EObjectFlags ExcludeFlags = RF_NoFlags);
+UObject* StaticFindObjectFastExplicit(const UClass* ObjectClass, FName ObjectName, const FString& ObjectPathName, bool bExactClass, EObjectFlags ExcludeFlags = RF_NoFlags);
 
 /**
  * Return all objects with a given outer
@@ -72,7 +72,7 @@ COREUOBJECT_API void ForEachObjectWithOuter(const class UObjectBase* Outer, TFun
  * @param	ClassToLookFor				if NULL, ignore this parameter, otherwise require the returned object have this class
  * @param	NameToLookFor				if NAME_None, ignore this parameter, otherwise require the returned object have this name
  */
-COREUOBJECT_API class UObjectBase* FindObjectWithOuter(class UObjectBase* Outer, class UClass* ClassToLookFor = NULL, FName NameToLookFor = NAME_None);
+COREUOBJECT_API class UObjectBase* FindObjectWithOuter(const class UObjectBase* Outer, const class UClass* ClassToLookFor = nullptr, FName NameToLookFor = NAME_None);
 
 /**
  * Returns an array of objects of a specific class. Optionally, results can include objects of derived classes as well.
@@ -83,7 +83,7 @@ COREUOBJECT_API class UObjectBase* FindObjectWithOuter(class UObjectBase* Outer,
  * @param	AdditionalExcludeFlags		Objects with any of these flags will be excluded from the results.
  * @param	ExclusiveInternalFlags	Specifies internal flags to use as a filter for which objects to return
  */
-COREUOBJECT_API void GetObjectsOfClass(UClass* ClassToLookFor, TArray<UObject *>& Results, bool bIncludeDerivedClasses = true, EObjectFlags ExcludeFlags = RF_ClassDefaultObject, EInternalObjectFlags ExclusionInternalFlags = EInternalObjectFlags::None);
+COREUOBJECT_API void GetObjectsOfClass(const UClass* ClassToLookFor, TArray<UObject *>& Results, bool bIncludeDerivedClasses = true, EObjectFlags ExcludeFlags = RF_ClassDefaultObject, EInternalObjectFlags ExclusionInternalFlags = EInternalObjectFlags::None);
 
 /**
  * Performs an operation on all objects with a given outer
@@ -93,7 +93,7 @@ COREUOBJECT_API void GetObjectsOfClass(UClass* ClassToLookFor, TArray<UObject *>
  * @param	bIncludeDerivedClasses		If true, the results will include objects of child classes as well.
  * @param	AdditionalExcludeFlags		Objects with any of these flags will be excluded from the results.
  */
-COREUOBJECT_API void ForEachObjectOfClass(UClass* ClassToLookFor, TFunctionRef<void(UObject*)> Operation, bool bIncludeDerivedClasses = true, EObjectFlags ExcludeFlags = RF_ClassDefaultObject, EInternalObjectFlags ExclusionInternalFlags = EInternalObjectFlags::None);
+COREUOBJECT_API void ForEachObjectOfClass(const UClass* ClassToLookFor, TFunctionRef<void(UObject*)> Operation, bool bIncludeDerivedClasses = true, EObjectFlags ExcludeFlags = RF_ClassDefaultObject, EInternalObjectFlags ExclusionInternalFlags = EInternalObjectFlags::None);
 
 /**
  * Returns an array of classes that were derived from the specified class.
@@ -102,7 +102,7 @@ COREUOBJECT_API void ForEachObjectOfClass(UClass* ClassToLookFor, TFunctionRef<v
  * @param	Results						An output list of child classes of the specified parent class.
  * @param	bRecursive					If true, the results will include children of the children classes, recursively. Otherwise, only direct decedents will be included.
  */
-COREUOBJECT_API void GetDerivedClasses(UClass* ClassToLookFor, TArray<UClass *>& Results, bool bRecursive = true);
+COREUOBJECT_API void GetDerivedClasses(const UClass* ClassToLookFor, TArray<UClass *>& Results, bool bRecursive = true);
 
 /**
  * Returns true if any instances of the class in question are currently being async loaded.
@@ -110,7 +110,7 @@ COREUOBJECT_API void GetDerivedClasses(UClass* ClassToLookFor, TArray<UClass *>&
  * @param	ClassToLookFor				The class in question
  * @return	True if there are any instances of the class being async loaded - includes instances based on derived classes. Otherwise, false
  */
-COREUOBJECT_API bool ClassHasInstancesAsyncLoading(UClass* ClassToLookFor);
+COREUOBJECT_API bool ClassHasInstancesAsyncLoading(const UClass* ClassToLookFor);
 
 /**
  * Add an object to the name hash tables
