Move IsLegalOverride in to FOverrides subclass and make it non fatal. Give proper error messages when the override is not legal.

diff --git a/Engine/Source/Runtime/CoreUObject/Private/UObject/UObjectGlobals.cpp b/Engine/Source/Runtime/CoreUObject/Private/UObject/UObjectGlobals.cpp
--- a/Engine/Source/Runtime/CoreUObject/Private/UObject/UObjectGlobals.cpp
+++ b/Engine/Source/Runtime/CoreUObject/Private/UObject/UObjectGlobals.cpp
@@ -3019,12 +3019,56 @@ void FObjectInitializer::InitProperties(UObject* Obj, UClass* DefaultsClass, UOb
 	}
 }
 
-bool FObjectInitializer::IsLegalOverride(FName InComponentName, class UClass *DerivedComponentClass, class UClass *BaseComponentClass) const
+/**  Add an override, make sure it is legal **/
+void FObjectInitializer::FOverrides::Add(FName InComponentName, UClass* InComponentClass, FObjectInitializer const& ObjectInitializer)
 {
-	if (DerivedComponentClass && BaseComponentClass && !DerivedComponentClass->IsChildOf(BaseComponentClass) )
+	const int32 Index = Find(InComponentName);
+	if (Index == INDEX_NONE)
+	{
+		Overrides.Emplace(FOverride(InComponentName, InComponentClass));
+	}
+	else if (InComponentClass && Overrides[Index].ComponentClass)
+	{
+		// if a base class is asking for an override, the existing override (which we are going to use) had better be derived
+		if (!IsLegalOverride(Overrides[Index].ComponentClass, InComponentClass))
+		{
+			UE_LOG(LogUObjectGlobals, Error, TEXT("%s is not a legal override for component %s because it does not derive from %s. Will use %s when constructing component."),
+				*Overrides[Index].ComponentClass->GetFullName(), *InComponentName.ToString(), *InComponentClass->GetFullName(), *InComponentClass->GetFullName());
+
+			Overrides[Index].ComponentClass = InComponentClass;
+		}
+	}
+}
+
+/**  Retrieve an override, or TClassToConstructByDefault::StaticClass or nullptr if this was removed by a derived class **/
+UClass* FObjectInitializer::FOverrides::Get(FName InComponentName, UClass* ReturnType, UClass* ClassToConstructByDefault, FObjectInitializer const& ObjectInitializer) const
+{
+	const int32 Index = Find(InComponentName);
+	UClass* BaseComponentClass = ClassToConstructByDefault;
+	if (Index == INDEX_NONE)
+	{
+		return BaseComponentClass; // no override so just do what the base class wanted
+	}
+	else if (Overrides[Index].ComponentClass)
+	{
+		if (IsLegalOverride(Overrides[Index].ComponentClass, ReturnType)) // if THE base class is asking for a T, the existing override (which we are going to use) had better be derived
+		{
+			return Overrides[Index].ComponentClass; // the override is of an acceptable class, so use it
+		}
+		else
+		{
+			UE_LOG(LogUObjectGlobals, Error, TEXT("%s is not a legal override for component %s because it does not derive from %s. Using %s to construct component."),
+				*Overrides[Index].ComponentClass->GetFullName(), *InComponentName.ToString(), *ReturnType->GetFullName(), *ClassToConstructByDefault->GetFullName());
+
+			return ClassToConstructByDefault;
+		}
+	}
+	return nullptr;  // the override is of nullptr, which means "don't create this component"
+}
+bool FObjectInitializer::FOverrides::IsLegalOverride(const UClass* DerivedComponentClass, const UClass* BaseComponentClass)
+{
+	if (DerivedComponentClass && BaseComponentClass && !DerivedComponentClass->IsChildOf(BaseComponentClass))
 	{
-		UE_LOG(LogUObjectGlobals, Fatal, TEXT("%s is not a legal override for component %s because it does not derive from %s."), 
-			*DerivedComponentClass->GetFullName(), *InComponentName.ToString(), *BaseComponentClass->GetFullName());
 		return false;
 	}
 	return true;
diff --git a/Engine/Source/Runtime/CoreUObject/Public/UObject/UObjectGlobals.h b/Engine/Source/Runtime/CoreUObject/Public/UObject/UObjectGlobals.h
--- a/Engine/Source/Runtime/CoreUObject/Public/UObject/UObjectGlobals.h
+++ b/Engine/Source/Runtime/CoreUObject/Public/UObject/UObjectGlobals.h
@@ -958,11 +958,6 @@ public:
 		return *this;
 	}
 
-	/** 
-	 * Internal use only, checks if the override is legal and if not deal with error messages
-	**/
-	bool IsLegalOverride(FName InComponentName, class UClass* DerivedComponentClass, class UClass* BaseComponentClass) const;
-
 	/**
 	 * Asserts with the specified message if code is executed inside UObject constructor
 	 **/
@@ -1032,41 +1027,16 @@ private:
 	struct FOverrides
 	{
 		/**  Add an override, make sure it is legal **/
-		void Add(FName InComponentName, UClass* InComponentClass, FObjectInitializer const& ObjectInitializer)
-		{
-			int32 Index = Find(InComponentName);
-			if (Index == INDEX_NONE)
-			{
-				new (Overrides) FOverride(InComponentName, InComponentClass);
-			}
-			else if (InComponentClass && Overrides[Index].ComponentClass)
-			{
-				ObjectInitializer.IsLegalOverride(InComponentName, Overrides[Index].ComponentClass, InComponentClass); // if a base class is asking for an override, the existing override (which we are going to use) had better be derived
-			}
-		}
+		void Add(FName InComponentName, UClass* InComponentClass, FObjectInitializer const& ObjectInitializer);
+
 		/**  Retrieve an override, or TClassToConstructByDefault::StaticClass or nullptr if this was removed by a derived class **/
-		UClass* Get(FName InComponentName, UClass* ReturnType, UClass* ClassToConstructByDefault, FObjectInitializer const& ObjectInitializer)
-		{
-			int32 Index = Find(InComponentName);
-			UClass* BaseComponentClass = ClassToConstructByDefault;
-			if (Index == INDEX_NONE)
-			{
-				return BaseComponentClass; // no override so just do what the base class wanted
-			}
-			else if (Overrides[Index].ComponentClass)
-			{
-				if (ObjectInitializer.IsLegalOverride(InComponentName, Overrides[Index].ComponentClass, ReturnType)) // if THE base class is asking for a T, the existing override (which we are going to use) had better be derived
-				{
-					return Overrides[Index].ComponentClass; // the override is of an acceptable class, so use it
-				}
-				// else return nullptr; this is a unacceptable override
-			}
-			return nullptr;  // the override is of nullptr, which means "don't create this component"
-		}
+		UClass* Get(FName InComponentName, UClass* ReturnType, UClass* ClassToConstructByDefault, FObjectInitializer const& ObjectInitializer) const;
+
+	private:
+		static bool IsLegalOverride(const UClass* DerivedComponentClass, const UClass* BaseComponentClass);
 
-private:
 		/**  Search for an override **/
-		int32 Find(FName InComponentName)
+		int32 Find(FName InComponentName) const
 		{
 			for (int32 Index = 0 ; Index < Overrides.Num(); Index++)
 			{
