Introduce owned packages:
 - Owned packages can reference private objects in their owner package.
 - Owned packages can reference private objects in other packages sharing the same owner.

Note: UE4 object version changed.

#rnx

diff --git a/Engine/Source/Editor/UnrealEd/Private/Commandlets/PackageUtilities.cpp b/Engine/Source/Editor/UnrealEd/Private/Commandlets/PackageUtilities.cpp
--- a/Engine/Source/Editor/UnrealEd/Private/Commandlets/PackageUtilities.cpp
+++ b/Engine/Source/Editor/UnrealEd/Private/Commandlets/PackageUtilities.cpp
@@ -960,8 +960,12 @@ void FPkgInfoReporter_Log::GeneratePackageReport( FLinkerLoad* InLinker /*=nullp
 	Out.Logf(ELogVerbosity::Display, TEXT("\t  Custom Versions:\n%s"), *Linker->Summary.GetCustomVersionContainer().ToString("\t\t"));
 	
 
-	FString szGUID = Linker->Summary.Guid.ToString();
-	Out.Logf(ELogVerbosity::Display, TEXT("\t             Guid: %s"), *szGUID );
+	Out.Logf(ELogVerbosity::Display, TEXT("\t             Guid: %s"), *Linker->Summary.Guid.ToString() );
+	Out.Logf(ELogVerbosity::Display, TEXT("\t   PersistentGuid: %s"), *Linker->Summary.PersistentGuid.ToString() );
+	if (Linker->Summary.OwnerPersistentGuid.IsValid())
+	{
+		Out.Logf(ELogVerbosity::Display, TEXT("\t    OwnerGuid: %s"), *Linker->Summary.OwnerPersistentGuid.ToString() );
+	}
 	Out.Logf(ELogVerbosity::Display, TEXT("\t      Generations:"));
 	for( int32 i = 0; i < Linker->Summary.Generations.Num(); ++i )
 	{
diff --git a/Engine/Source/Runtime/Core/Public/UObject/ObjectVersion.h b/Engine/Source/Runtime/Core/Public/UObject/ObjectVersion.h
--- a/Engine/Source/Runtime/Core/Public/UObject/ObjectVersion.h
+++ b/Engine/Source/Runtime/Core/Public/UObject/ObjectVersion.h
@@ -631,6 +631,8 @@ enum EUnrealEngineObjectUE4Version
 	VER_UE4_ADDED_PACKAGE_SUMMARY_LOCALIZATION_ID,
 	// Fixed case insensitive hashes of wide strings containing character values from 128-255
 	VER_UE4_FIX_WIDE_STRING_CRC,
+	// Added package owner to allow private references
+	VER_UE4_ADDED_PACKAGE_OWNER,
 
 	// -----<new versions can be added before this line>-------------------------------------------------
 	// - this needs to be the last line (see note below)
diff --git a/Engine/Source/Runtime/CoreUObject/Private/UObject/LinkerLoad.cpp b/Engine/Source/Runtime/CoreUObject/Private/UObject/LinkerLoad.cpp
--- a/Engine/Source/Runtime/CoreUObject/Private/UObject/LinkerLoad.cpp
+++ b/Engine/Source/Runtime/CoreUObject/Private/UObject/LinkerLoad.cpp
@@ -1346,9 +1346,14 @@ FLinkerLoad::ELinkerStatus FLinkerLoad::SerializePackageFileSummary()
 			// Propagate package file size
 			LinkerRootPackage->FileSize = TotalSize();
 
-			// Propagate package Guid
+			// Propagate package Guids
 			LinkerRootPackage->SetGuid( Summary.Guid );
 
+#if WITH_EDITORONLY_DATA
+			LinkerRootPackage->SetPersistentGuid( Summary.PersistentGuid );
+			LinkerRootPackage->SetOwnerPersistentGuid( Summary.OwnerPersistentGuid );
+#endif
+
 			// Remember the linker versions
 			LinkerRootPackage->LinkerPackageVersion = this->UE4Ver();
 			LinkerRootPackage->LinkerLicenseeVersion = this->LicenseeUE4Ver();
@@ -2884,7 +2889,16 @@ bool FLinkerLoad::VerifyImportInner(const int32 ImportIndex, FString& WarningSuf
 								}
 							}
 						}
-						if( !(SourceExport.ObjectFlags & RF_Public) )
+
+						const bool bIsImportPublic = !!(SourceExport.ObjectFlags & RF_Public);
+						const FPackageFileSummary& ImportSummary = Import.SourceLinker->Summary;
+#if WITH_EDITORONLY_DATA
+						const bool bIsImportOwned = (ImportSummary.OwnerPersistentGuid.IsValid() && ((ImportSummary.OwnerPersistentGuid == Summary.PersistentGuid) || (ImportSummary.OwnerPersistentGuid == Summary.OwnerPersistentGuid)) ||
+							                        (Summary.OwnerPersistentGuid.IsValid() && ((Summary.OwnerPersistentGuid == ImportSummary.PersistentGuid) || (Summary.OwnerPersistentGuid == ImportSummary.OwnerPersistentGuid))));
+#else
+						const bool bIsImportOwned = false;
+#endif
+						if( !bIsImportPublic && !bIsImportOwned)
 						{
 							SafeReplace = SafeReplace || (GIsEditor && !IsRunningCommandlet());
 
diff --git a/Engine/Source/Runtime/CoreUObject/Private/UObject/Package.cpp b/Engine/Source/Runtime/CoreUObject/Private/UObject/Package.cpp
--- a/Engine/Source/Runtime/CoreUObject/Private/UObject/Package.cpp
+++ b/Engine/Source/Runtime/CoreUObject/Private/UObject/Package.cpp
@@ -124,6 +124,23 @@ void UPackage::AddReferencedObjects(UObject* InThis, FReferenceCollector& Collec
 	Super::AddReferencedObjects(This, Collector);
 }
 
+#if WITH_EDITORONLY_DATA
+bool UPackage::IsOwned() const
+{
+	return OwnerPersistentGuid.IsValid();
+}
+
+bool UPackage::IsOwnedBy(const UPackage* Package) const
+{
+	return OwnerPersistentGuid.IsValid() && (OwnerPersistentGuid == Package->GetPersistentGuid());
+}
+
+bool UPackage::HasSameOwner(const UPackage* Package) const
+{
+	return OwnerPersistentGuid.IsValid() && (OwnerPersistentGuid == Package->OwnerPersistentGuid);
+}
+#endif
+
 /**
  * Gets (after possibly creating) a metadata object for this package
  *
diff --git a/Engine/Source/Runtime/CoreUObject/Private/UObject/PackageFileSummary.cpp b/Engine/Source/Runtime/CoreUObject/Private/UObject/PackageFileSummary.cpp
--- a/Engine/Source/Runtime/CoreUObject/Private/UObject/PackageFileSummary.cpp
+++ b/Engine/Source/Runtime/CoreUObject/Private/UObject/PackageFileSummary.cpp
@@ -211,8 +211,31 @@ void operator<<(FStructuredArchive::FSlot Slot, FPackageFileSummary& Sum)
 
 		Record << NAMED_ITEM("ThumbnailTableOffset", Sum.ThumbnailTableOffset);
 
+		Record << NAMED_ITEM("Guid", Sum.Guid);
+
+		if (BaseArchive.IsSaving() || Sum.FileVersionUE4 >= VER_UE4_ADDED_PACKAGE_OWNER)
+		{
+			if (!BaseArchive.IsFilterEditorOnly())
+			{
+#if WITH_EDITORONLY_DATA
+				Record << NAMED_ITEM("PersistentGuid", Sum.PersistentGuid) << NAMED_ITEM("OwnerPersistentGuid", Sum.OwnerPersistentGuid);
+#else
+				FGuid PersistentGuid;
+				FGuid OwnerPersistentGuid;
+				Record << NAMED_ITEM("PersistentGuid", PersistentGuid) << NAMED_ITEM("OwnerPersistentGuid", OwnerPersistentGuid);
+#endif
+			}
+		}
+#if WITH_EDITORONLY_DATA
+		else
+		{
+			// By assigning the current package guid, we maintain a stable persistent guid, so we can reference this package even if it wasn't resaved.
+			Sum.PersistentGuid = Sum.Guid;
+		}
+#endif
+
 		int32 GenerationCount = Sum.Generations.Num();
-		Record << NAMED_ITEM("Guid", Sum.Guid) << NAMED_ITEM("GenerationCount", GenerationCount);
+		Record << NAMED_ITEM("GenerationCount", GenerationCount);
 		if (BaseArchive.IsLoading() && GenerationCount > 0)
 		{
 			Sum.Generations.Reset(GenerationCount);
diff --git a/Engine/Source/Runtime/CoreUObject/Private/UObject/SavePackage.cpp b/Engine/Source/Runtime/CoreUObject/Private/UObject/SavePackage.cpp
--- a/Engine/Source/Runtime/CoreUObject/Private/UObject/SavePackage.cpp
+++ b/Engine/Source/Runtime/CoreUObject/Private/UObject/SavePackage.cpp
@@ -4333,7 +4333,13 @@ FSavePackageResultStruct UPackage::Save(UPackage* InOuter, UObject* Base, EObjec
 								continue;
 							}
 
-
+#if WITH_EDITORONLY_DATA
+							// Allow referencing private objects into/from/between owned packages
+							if (ObjPackage->IsOwnedBy(InOuter) || InOuter->IsOwnedBy(ObjPackage) || InOuter->HasSameOwner(ObjPackage))
+							{
+								continue;
+							}
+#endif
 							if( !Obj->HasAnyFlags(RF_Public) && !Obj->HasAnyFlags(RF_Transient))
 							{
 								if (!IsEventDrivenLoaderEnabledInCookedBuilds() || !TargetPlatform || !ObjPackage->HasAnyPackageFlags(PKG_CompiledIn))
@@ -4364,7 +4370,6 @@ FSavePackageResultStruct UPackage::Save(UPackage* InOuter, UObject* Base, EObjec
 									LevelObjects.Add(Obj);
 								}
 							}
-
 						}
 					}
 				}
@@ -4505,19 +4510,31 @@ FSavePackageResultStruct UPackage::Save(UPackage* InOuter, UObject* Base, EObjec
 				{
 					// Conform to previous generation of file.
 					UE_LOG(LogSavePackage, Log,  TEXT("Conformal save, relative to: %s, Generation %i"), *Conform->Filename, Conform->Summary.Generations.Num()+1 );
-					Linker->Summary.Guid        = Conform->Summary.Guid;
+					Linker->Summary.Guid = Conform->Summary.Guid;
+#if WITH_EDITORONLY_DATA
+					Linker->Summary.PersistentGuid = Conform->Summary.PersistentGuid;
+					Linker->Summary.OwnerPersistentGuid	= Conform->Summary.OwnerPersistentGuid;
+#endif
 					Linker->Summary.Generations = Conform->Summary.Generations;
 				}
 				else if (SaveFlags & SAVE_KeepGUID)
 				{
 					// First generation file, keep existing GUID
 					Linker->Summary.Guid = InOuter->Guid;
+#if WITH_EDITORONLY_DATA
+					Linker->Summary.PersistentGuid = InOuter->PersistentGuid;
+					Linker->Summary.OwnerPersistentGuid = InOuter->OwnerPersistentGuid;
+#endif
 					Linker->Summary.Generations = TArray<FGenerationInfo>();
 				}
 				else
 				{
 					// First generation file.
 					Linker->Summary.Guid = FGuid::NewGuid();
+#if WITH_EDITORONLY_DATA
+					Linker->Summary.PersistentGuid = InOuter->PersistentGuid;
+					Linker->Summary.OwnerPersistentGuid = InOuter->OwnerPersistentGuid;
+#endif
 					Linker->Summary.Generations = TArray<FGenerationInfo>();
 
 					// make sure the UPackage's copy of the GUID is up to date
diff --git a/Engine/Source/Runtime/CoreUObject/Public/UObject/Package.h b/Engine/Source/Runtime/CoreUObject/Public/UObject/Package.h
--- a/Engine/Source/Runtime/CoreUObject/Public/UObject/Package.h
+++ b/Engine/Source/Runtime/CoreUObject/Public/UObject/Package.h
@@ -151,8 +151,17 @@ private:
 	FName	FolderName;
 #endif
 
-	/** GUID of package if it was loaded from disk; used by netcode to make sure packages match between client and server */
+	/** GUID of package if it was loaded from disk; used by netcode to make sure packages match between client and server. Changes at every save. */
 	FGuid Guid;
+
+#if WITH_EDITORONLY_DATA
+	/** Persistent GUID of package if it was loaded from disk. Persistent across saves. */
+	FGuid PersistentGuid;
+
+	/** Persistent GUID of owning package, to allow private references when saving */
+	FGuid OwnerPersistentGuid;
+#endif
+
 	/** Chunk IDs for the streaming install chunks this package will be placed in.  Empty for no chunk */
 	TArray<int32> ChunkIDs;
 
@@ -454,6 +463,34 @@ public:
 		Guid = NewGuid;
 	}
 
+#if WITH_EDITORONLY_DATA
+	/** returns our persistent Guid */
+	FORCEINLINE FGuid GetPersistentGuid() const
+	{
+		return PersistentGuid;
+	}
+	/** sets a specific persistent Guid */
+	FORCEINLINE void SetPersistentGuid(FGuid NewPersistentGuid)
+	{
+		PersistentGuid = NewPersistentGuid;
+	}
+
+	/** returns our owner persistent Guid */
+	FORCEINLINE FGuid GetOwnerPersistentGuid() const
+	{
+		return OwnerPersistentGuid;
+	}
+	/** sets a specific owner persistent Guid */
+	FORCEINLINE void SetOwnerPersistentGuid(FGuid NewOwnerPersistentGuid)
+	{
+		OwnerPersistentGuid = NewOwnerPersistentGuid;
+	}
+
+	bool IsOwned() const;
+	bool IsOwnedBy(const UPackage* Package) const;
+	bool HasSameOwner(const UPackage* Package) const;
+#endif
+
 	/** returns our FileSize */
 	FORCEINLINE int64 GetFileSize()
 	{
diff --git a/Engine/Source/Runtime/CoreUObject/Public/UObject/PackageFileSummary.h b/Engine/Source/Runtime/CoreUObject/Public/UObject/PackageFileSummary.h
--- a/Engine/Source/Runtime/CoreUObject/Public/UObject/PackageFileSummary.h
+++ b/Engine/Source/Runtime/CoreUObject/Public/UObject/PackageFileSummary.h
@@ -151,6 +151,18 @@ public:
 	*/
 	FGuid	Guid;
 
+#if WITH_EDITORONLY_DATA
+	/**
+	* Current persistent id for this package
+	*/
+	FGuid	PersistentGuid;
+
+	/**
+	* Package persistent owner for this package
+	*/
+	FGuid	OwnerPersistentGuid;
+#endif
+	
 	/**
 	* Data about previous versions of this package
 	*/
