Fix static analysis warning.
#rnx

diff --git a/Engine/Source/Runtime/CoreUObject/Private/UObject/SavePackage.cpp b/Engine/Source/Runtime/CoreUObject/Private/UObject/SavePackage.cpp
--- a/Engine/Source/Runtime/CoreUObject/Private/UObject/SavePackage.cpp
+++ b/Engine/Source/Runtime/CoreUObject/Private/UObject/SavePackage.cpp
@@ -4335,7 +4335,7 @@ FSavePackageResultStruct UPackage::Save(UPackage* InOuter, UObject* Base, EObjec
 
 #if WITH_EDITORONLY_DATA
 							// Allow referencing private objects into/from/between owned packages
-							if (ObjPackage->IsOwnedBy(InOuter) || InOuter->IsOwnedBy(ObjPackage) || InOuter->HasSameOwner(ObjPackage))
+							if (ObjPackage && (ObjPackage->IsOwnedBy(InOuter) || InOuter->IsOwnedBy(ObjPackage) || InOuter->HasSameOwner(ObjPackage)))
 							{
 								continue;
 							}
