Remove accidentally checked in diagnostic
#rnx

diff --git a/Engine/Source/Runtime/CoreUObject/Public/UObject/Object.h b/Engine/Source/Runtime/CoreUObject/Public/UObject/Object.h
--- a/Engine/Source/Runtime/CoreUObject/Public/UObject/Object.h
+++ b/Engine/Source/Runtime/CoreUObject/Public/UObject/Object.h
@@ -229,7 +229,6 @@ public:
 	/** Utility to allow overrides of Modify to avoid doing work if this object cannot be safely modified */
 	bool CanModify() const;
 #else
-	UE_DEPRECATED(4.24,"How many things call this?")
 	FORCEINLINE bool Modify(bool bAlwaysMarkDirty = true) { return false; }
 #endif
 
