Run GC tasks on hi-pri threads (big cores) on Android devices. Reduces GC time by 2.5x on high end devices

diff --git a/Engine/Source/Runtime/CoreUObject/Public/UObject/FastReferenceCollector.h b/Engine/Source/Runtime/CoreUObject/Public/UObject/FastReferenceCollector.h
--- a/Engine/Source/Runtime/CoreUObject/Public/UObject/FastReferenceCollector.h
+++ b/Engine/Source/Runtime/CoreUObject/Public/UObject/FastReferenceCollector.h
@@ -525,6 +525,9 @@ public:
 
 				int32 NumThreads = FTaskGraphInterface::Get().GetNumWorkerThreads();
 				int32 NumBackgroundThreads = ENamedThreads::bHasBackgroundThreads ? NumThreads : 0;
+				ENamedThreads::Type NormalThreadName = ENamedThreads::AnyNormalThreadNormalTask;
+				ENamedThreads::Type BackgroundThreadName = ENamedThreads::AnyBackgroundThreadNormalTask;
+
 #if ((PLATFORM_PS4 && USE_7TH_CORE) || PLATFORM_XBOXONE)
 				if (NumBackgroundThreads)
 				{
@@ -535,6 +538,10 @@ public:
 				{
 					NumBackgroundThreads = 6 - NumThreads;
 				}
+#elif PLATFORM_ANDROID
+				// On devices with overridden affinity only HiPri threads can run on big cores
+				NormalThreadName = ENamedThreads::AnyHiPriThreadHiPriTask; 
+				NumBackgroundThreads = 0; // run on single group
 #endif
 				int32 NumTasks = NumThreads + NumBackgroundThreads;
 
@@ -553,7 +560,7 @@ public:
 				}
 				for (int32 Chunk = 0; Chunk < NumTasks; Chunk++)
 				{
-					ChunkTasks.Add(TGraphTask< FCollectorTaskProcessorTask >::CreateTask().ConstructAndDispatchWhenReady(TaskQueue, Chunk >= NumThreads ? ENamedThreads::AnyBackgroundThreadNormalTask : ENamedThreads::AnyNormalThreadNormalTask));
+					ChunkTasks.Add(TGraphTask< FCollectorTaskProcessorTask >::CreateTask().ConstructAndDispatchWhenReady(TaskQueue, Chunk >= NumThreads ? BackgroundThreadName : NormalThreadName));
 				}
 
 				QUICK_SCOPE_CYCLE_COUNTER(STAT_GC_Subtask_Wait);
