Store a function pointer on UDynamicClass to the custom class initialization (in much the same way that AddReferencedObjects pointer is stored) to ensure the class is initialized before beginning to construct the CDO rather than code embedded in the constructor
#rnx

diff --git a/Engine/Source/Developer/BlueprintCompilerCppBackend/Private/BlueprintCompilerCppBackendValueHelper.cpp b/Engine/Source/Developer/BlueprintCompilerCppBackend/Private/BlueprintCompilerCppBackendValueHelper.cpp
--- a/Engine/Source/Developer/BlueprintCompilerCppBackend/Private/BlueprintCompilerCppBackendValueHelper.cpp
+++ b/Engine/Source/Developer/BlueprintCompilerCppBackend/Private/BlueprintCompilerCppBackendValueHelper.cpp
@@ -1813,14 +1813,6 @@ void FEmitDefaultValueHelper::GenerateConstructor(FEmitterLocalContext& Context)
 		Context.AddLine(TEXT("{"));
 		Context.IncreaseIndent();
 
-		// Call CustomDynamicClassInitialization
-		Context.AddLine(FString::Printf(TEXT("if(HasAnyFlags(RF_ClassDefaultObject) && (%s::StaticClass() == GetClass()))"), *CppClassName));
-		Context.AddLine(TEXT("{"));
-		Context.IncreaseIndent();
-		Context.AddLine(FString::Printf(TEXT("%s::__CustomDynamicClassInitialization(CastChecked<UDynamicClass>(GetClass()));"), *CppClassName));
-		Context.DecreaseIndent();
-		Context.AddLine(TEXT("}"));
-
 		// Subobjects that must be fixed after serialization
 		TArray<FDefaultSubobjectData> SubobjectsToInit;
 		TArray<FNonNativeComponentData> ComponentsToInit;
diff --git a/Engine/Source/Runtime/CoreUObject/Private/UObject/Class.cpp b/Engine/Source/Runtime/CoreUObject/Private/UObject/Class.cpp
--- a/Engine/Source/Runtime/CoreUObject/Private/UObject/Class.cpp
+++ b/Engine/Source/Runtime/CoreUObject/Private/UObject/Class.cpp
@@ -3137,6 +3137,13 @@ UObject* UClass::CreateDefaultObject()
 					USparseDelegateFunction* SparseDelegateFunction = CastChecked<USparseDelegateFunction>(SparseDelegateIt->SignatureFunction);
 					FSparseDelegateStorage::RegisterDelegateOffset(ClassDefaultObject, SparseDelegateFunction->DelegateName, (size_t)&SparseDelegate - (size_t)ClassDefaultObject);
 				}
+				if (HasAnyClassFlags(CLASS_CompiledFromBlueprint))
+				{
+					if (UDynamicClass* DynamicClass = Cast<UDynamicClass>(this))
+					{
+						(*(DynamicClass->DynamicClassInitializer))(DynamicClass);
+					}
+				}
 				(*ClassConstructor)(FObjectInitializer(ClassDefaultObject, ParentDefaultObject, false, bShouldInitializeProperties));
 				if (bDoNotify)
 				{
@@ -4612,7 +4619,8 @@ void GetPrivateStaticClassBody(
 	UClass::ClassAddReferencedObjectsType InClassAddReferencedObjects,
 	UClass::StaticClassFunctionType InSuperClassFn,
 	UClass::StaticClassFunctionType InWithinClassFn,
-	bool bIsDynamic /*= false*/
+	bool bIsDynamic /*= false*/,
+	UDynamicClass::DynamicClassInitializerType InDynamicClassInitializerFn /*= nullptr*/
 	)
 {
 #if WITH_HOT_RELOAD
@@ -4690,7 +4698,8 @@ void GetPrivateStaticClassBody(
 			EObjectFlags(RF_Public | RF_Standalone | RF_Transient | RF_Dynamic | (GIsInitialLoad ? RF_MarkAsRootSet : RF_NoFlags)),
 			InClassConstructor,
 			InClassVTableHelperCtorCaller,
-			InClassAddReferencedObjects
+			InClassAddReferencedObjects,
+			InDynamicClassInitializerFn
 			);
 		check(ReturnClass);
 	}
@@ -5277,7 +5286,8 @@ UDynamicClass::UDynamicClass(
 	EObjectFlags	InFlags,
 	ClassConstructorType InClassConstructor,
 	ClassVTableHelperCtorCallerType InClassVTableHelperCtorCaller,
-	ClassAddReferencedObjectsType InClassAddReferencedObjects)
+	ClassAddReferencedObjectsType InClassAddReferencedObjects,
+	DynamicClassInitializerType InDynamicClassInitializer)
 : UClass(
   EC_StaticConstructor
 , InName
@@ -5291,6 +5301,7 @@ UDynamicClass::UDynamicClass(
 , InClassVTableHelperCtorCaller
 , InClassAddReferencedObjects)
 , AnimClassImplementation(nullptr)
+, DynamicClassInitializer(InDynamicClassInitializer)
 {
 }
 
diff --git a/Engine/Source/Runtime/CoreUObject/Public/UObject/Class.h b/Engine/Source/Runtime/CoreUObject/Public/UObject/Class.h
--- a/Engine/Source/Runtime/CoreUObject/Public/UObject/Class.h
+++ b/Engine/Source/Runtime/CoreUObject/Public/UObject/Class.h
@@ -2871,12 +2871,15 @@ class COREUOBJECT_API UDynamicClass : public UClass
 
 public:
 
+	typedef void (*DynamicClassInitializerType)	(UDynamicClass*);
+
 	UDynamicClass(const FObjectInitializer& ObjectInitializer = FObjectInitializer::Get());
 	explicit UDynamicClass(const FObjectInitializer& ObjectInitializer, UClass* InSuperClass);
 	UDynamicClass(EStaticConstructor, FName InName, uint32 InSize, uint32 InAlignment, EClassFlags InClassFlags, EClassCastFlags InClassCastFlags,
 		const TCHAR* InClassConfigName, EObjectFlags InFlags, ClassConstructorType InClassConstructor,
 		ClassVTableHelperCtorCallerType InClassVTableHelperCtorCaller,
-		ClassAddReferencedObjectsType InClassAddReferencedObjects);
+		ClassAddReferencedObjectsType InClassAddReferencedObjects,
+		DynamicClassInitializerType InDynamicClassInitializer);
 
 	// UObject interface.
 	static void AddReferencedObjects(UObject* InThis, FReferenceCollector& Collector);
@@ -2905,6 +2908,8 @@ public:
 
 	/** IAnimClassInterface (UAnimClassData) or null */
 	UObject* AnimClassImplementation;
+
+	DynamicClassInitializerType DynamicClassInitializer;
 };
 
 /**
@@ -2968,7 +2973,8 @@ COREUOBJECT_API void GetPrivateStaticClassBody(
 	UClass::ClassAddReferencedObjectsType InClassAddReferencedObjects,
 	UClass::StaticClassFunctionType InSuperClassFn,
 	UClass::StaticClassFunctionType InWithinClassFn,
-	bool bIsDynamic = false);
+	bool bIsDynamic = false,
+	UDynamicClass::DynamicClassInitializerType InDynamicClassInitializer = nullptr);
 
 /*-----------------------------------------------------------------------------
 	FObjectInstancingGraph.
diff --git a/Engine/Source/Runtime/CoreUObject/Public/UObject/ObjectMacros.h b/Engine/Source/Runtime/CoreUObject/Public/UObject/ObjectMacros.h
--- a/Engine/Source/Runtime/CoreUObject/Public/UObject/ObjectMacros.h
+++ b/Engine/Source/Runtime/CoreUObject/Public/UObject/ObjectMacros.h
@@ -1665,7 +1665,8 @@ public: \
 			&TClass::AddReferencedObjects, \
 			&TClass::Super::StaticClass, \
 			&TClass::WithinClass::StaticClass, \
-			true \
+			true, \
+			&TClass::__CustomDynamicClassInitialization \
 			); \
 		} \
 		return PrivateStaticClass; \
