const parameters to ultimately make NewObject able to take a const UClass*
#rnx

diff --git a/Engine/Source/Runtime/CoreUObject/Private/UObject/Class.cpp b/Engine/Source/Runtime/CoreUObject/Private/UObject/Class.cpp
--- a/Engine/Source/Runtime/CoreUObject/Private/UObject/Class.cpp
+++ b/Engine/Source/Runtime/CoreUObject/Private/UObject/Class.cpp
@@ -3228,7 +3228,7 @@ FFeedbackContext& UClass::GetDefaultPropertiesFeedbackContext()
 * Get the name of the CDO for the this class
 * @return The name of the CDO
 */
-FName UClass::GetDefaultObjectName()
+FName UClass::GetDefaultObjectName() const
 {
 	FString DefaultName;
 	DefaultName.Reserve(NAME_SIZE);
@@ -5353,10 +5353,9 @@ void UDynamicClass::PurgeClass(bool bRecompilingOnLoad)
 	AnimClassImplementation = nullptr;
 }
 
-UObject* UDynamicClass::FindArchetype(UClass* ArchetypeClass, const FName ArchetypeName) const
+UObject* UDynamicClass::FindArchetype(const UClass* ArchetypeClass, const FName ArchetypeName) const
 {
-	UDynamicClass* ThisClass = const_cast<UDynamicClass*>(this);
-	UObject* Archetype = static_cast<UObject*>(FindObjectWithOuter(ThisClass, ArchetypeClass, ArchetypeName));
+	UObject* Archetype = static_cast<UObject*>(FindObjectWithOuter(this, ArchetypeClass, ArchetypeName));
 	if (!Archetype)
 	{
 		// See UBlueprintGeneratedClass::FindArchetype, UE-35259, UE-37480
diff --git a/Engine/Source/Runtime/CoreUObject/Private/UObject/UObjectArchetype.cpp b/Engine/Source/Runtime/CoreUObject/Private/UObject/UObjectArchetype.cpp
--- a/Engine/Source/Runtime/CoreUObject/Private/UObject/UObjectArchetype.cpp
+++ b/Engine/Source/Runtime/CoreUObject/Private/UObject/UObjectArchetype.cpp
@@ -12,7 +12,7 @@
 #include "UObject/UObjectAnnotation.h"
 #include "Stats/StatsMisc.h"
 
-UObject* GetArchetypeFromRequiredInfoImpl(UClass* Class, UObject* Outer, FName Name, EObjectFlags ObjectFlags, bool bUseUpToDateClass)
+UObject* GetArchetypeFromRequiredInfoImpl(const UClass* Class, const UObject* Outer, FName Name, EObjectFlags ObjectFlags, bool bUseUpToDateClass)
 {
 	UObject* Result = NULL;
 	const bool bIsCDO = !!(ObjectFlags&RF_ClassDefaultObject);
@@ -37,8 +37,8 @@ UObject* GetArchetypeFromRequiredInfoImpl(UClass* Class, UObject* Outer, FName N
 			}
 			else if (!!(ObjectFlags&RF_InheritableComponentTemplate) && Outer->IsA<UClass>())
 			{
-				UClass* OuterSuperClass = static_cast<UClass*>(Outer)->GetSuperClass();
-				for (UClass* SuperClassArchetype = bUseUpToDateClass && OuterSuperClass ? OuterSuperClass->GetAuthoritativeClass() : OuterSuperClass;
+				const UClass* OuterSuperClass = static_cast<const UClass*>(Outer)->GetSuperClass();
+				for (const UClass* SuperClassArchetype = bUseUpToDateClass && OuterSuperClass ? OuterSuperClass->GetAuthoritativeClass() : OuterSuperClass;
 					SuperClassArchetype && SuperClassArchetype->HasAllClassFlags(CLASS_CompiledFromBlueprint);
 					SuperClassArchetype = SuperClassArchetype->GetSuperClass())
 				{
@@ -92,7 +92,7 @@ UObject* GetArchetypeFromRequiredInfoImpl(UClass* Class, UObject* Outer, FName N
 	return Result;
 }
 
-UObject* UObject::GetArchetypeFromRequiredInfo(UClass* Class, UObject* Outer, FName Name, EObjectFlags ObjectFlags)
+UObject* UObject::GetArchetypeFromRequiredInfo(const UClass* Class, const UObject* Outer, FName Name, EObjectFlags ObjectFlags)
 {
 	bool bUseUpToDateClass = false;
 #if WITH_EDITOR
diff --git a/Engine/Source/Runtime/CoreUObject/Private/UObject/UObjectGlobals.cpp b/Engine/Source/Runtime/CoreUObject/Private/UObject/UObjectGlobals.cpp
--- a/Engine/Source/Runtime/CoreUObject/Private/UObject/UObjectGlobals.cpp
+++ b/Engine/Source/Runtime/CoreUObject/Private/UObject/UObjectGlobals.cpp
@@ -251,7 +251,7 @@ struct FPerClassNumberSuffixAnnotation
 		return false;
 	}
 
-	TMap<UClass*, int32> Suffixes;
+	TMap<const UClass*, int32> Suffixes;
 };
 
 /**
@@ -259,7 +259,7 @@ struct FPerClassNumberSuffixAnnotation
  *
  * Updating is done via a callback because a lock needs to be maintained while this happens.
  */
-int32 UpdateSuffixForNextNewObject(UObject* Parent, UClass* Class, TFunctionRef<void(int32&)> IndexMutator)
+int32 UpdateSuffixForNextNewObject(UObject* Parent, const UClass* Class, TFunctionRef<void(int32&)> IndexMutator)
 {
 	static FCriticalSection PerClassNumberSuffixAnnotationMutex;
 	static FUObjectAnnotationDense<FPerClassNumberSuffixAnnotation, true> PerClassNumberSuffixAnnotation;
@@ -1747,7 +1747,7 @@ void EndLoad(FUObjectSerializeContext* LoadContext)
 	Object name functions.
 -----------------------------------------------------------------------------*/
 
-FName MakeUniqueObjectName( UObject* Parent, UClass* Class, FName InBaseName/*=NAME_None*/ )
+FName MakeUniqueObjectName( UObject* Parent, const UClass* Class, FName InBaseName/*=NAME_None*/ )
 {
 	CSV_SCOPED_TIMING_STAT(UObject, MakeUniqueObjectName);
 	check(Class);
@@ -2171,7 +2171,7 @@ void SnapshotTransactionBuffer(UObject* Object)
 	}
 }
 
-bool StaticAllocateObjectErrorTests( UClass* InClass, UObject* InOuter, FName InName, EObjectFlags InFlags)
+bool StaticAllocateObjectErrorTests( const UClass* InClass, UObject* InOuter, FName InName, EObjectFlags InFlags)
 {
 	// Validation checks.
 	if( !InClass )
@@ -2239,7 +2239,7 @@ COREUOBJECT_API bool GOutputCookingWarnings = false;
 
 UObject* StaticAllocateObject
 (
-	UClass*			InClass,
+	const UClass*	InClass,
 	UObject*		InOuter,
 	FName			InName,
 	EObjectFlags	InFlags,
@@ -2435,7 +2435,7 @@ UObject* StaticAllocateObject
 	if (!bSubObject)
 	{
 		FMemory::Memzero((void *)Obj, TotalSize);
-		new ((void *)Obj) UObjectBase(InClass, InFlags|RF_NeedInitialization, InternalSetFlags, InOuter, InName);
+		new ((void *)Obj) UObjectBase(const_cast<UClass*>(InClass), InFlags|RF_NeedInitialization, InternalSetFlags, InOuter, InName);
 	}
 	else
 	{
@@ -3086,7 +3086,7 @@ void FObjectInitializer::AssertIfSubobjectSetupIsNotAllowed(const FName Subobjec
 }
 
 #if DO_CHECK
-void CheckIsClassChildOf_Internal(UClass* Parent, UClass* Child)
+void CheckIsClassChildOf_Internal(const UClass* Parent, const UClass* Child)
 {
 	// This is a function to avoid platform compilation issues
 	checkf(Child, TEXT("NewObject called with a nullptr class object"));
@@ -3096,7 +3096,7 @@ void CheckIsClassChildOf_Internal(UClass* Parent, UClass* Child)
 
 UObject* StaticConstructObject_Internal
 (
-	UClass*			InClass,
+	const UClass*	InClass,
 	UObject*		InOuter								/*=GetTransientPackage()*/,
 	FName			InName								/*=NAME_None*/,
 	EObjectFlags	InFlags								/*=0*/,
diff --git a/Engine/Source/Runtime/CoreUObject/Public/UObject/Class.h b/Engine/Source/Runtime/CoreUObject/Public/UObject/Class.h
--- a/Engine/Source/Runtime/CoreUObject/Public/UObject/Class.h
+++ b/Engine/Source/Runtime/CoreUObject/Public/UObject/Class.h
@@ -2245,7 +2245,7 @@ public:
 	ClassAddReferencedObjectsType ClassAddReferencedObjects;
 
 	/** Class pseudo-unique counter; used to accelerate unique instance name generation */
-	uint32 ClassUnique:31;
+	mutable uint32 ClassUnique:31;
 
 	/** Used to check if the class was cooked or not */
 	uint32 bCooked:1;
@@ -2525,11 +2525,11 @@ public:
 	 * @param	bCreateIfNeeded if true (default) then the CDO is created if it is null
 	 * @return		the CDO for this class
 	 */
-	UObject* GetDefaultObject(bool bCreateIfNeeded = true)
+	UObject* GetDefaultObject(bool bCreateIfNeeded = true) const
 	{
 		if (ClassDefaultObject == nullptr && bCreateIfNeeded)
 		{
-			CreateDefaultObject();
+			const_cast<UClass*>(this)->CreateDefaultObject();
 		}
 
 		return ClassDefaultObject;
@@ -2558,7 +2558,7 @@ public:
 	 * Get the name of the CDO for the this class
 	 * @return The name of the CDO
 	 */
-	FName GetDefaultObjectName();
+	FName GetDefaultObjectName() const;
 
 	/** Returns memory used to store temporary data on an instance, used by blueprints */
 	virtual uint8* GetPersistentUberGraphFrame(UObject* Obj, UFunction* FuncToCheck) const
@@ -2581,7 +2581,7 @@ public:
 	 * @return		the CDO for this class
 	 */
 	template<class T>
-	T* GetDefaultObject()
+	T* GetDefaultObject() const
 	{
 		UObject *Ret = GetDefaultObject();
 		check(Ret->IsA(T::StaticClass()));
@@ -2793,7 +2793,7 @@ public:
 	virtual bool HasProperty(UProperty* InProperty) const;
 
 	/** Finds the object that is used as the parent object when serializing properties, overridden for blueprints */
-	virtual UObject* FindArchetype(UClass* ArchetypeClass, const FName ArchetypeName) const { return nullptr; }
+	virtual UObject* FindArchetype(const UClass* ArchetypeClass, const FName ArchetypeName) const { return nullptr; }
 
 	/** Returns archetype object for CDO */
 	virtual UObject* GetArchetypeForCDO() const;
@@ -2892,7 +2892,7 @@ public:
 	// UClass interface
 	virtual UObject* CreateDefaultObject();
 	virtual void PurgeClass(bool bRecompilingOnLoad) override;
-	virtual UObject* FindArchetype(UClass* ArchetypeClass, const FName ArchetypeName) const override;
+	virtual UObject* FindArchetype(const UClass* ArchetypeClass, const FName ArchetypeName) const override;
 	virtual void SetupObjectInitializer(FObjectInitializer& ObjectInitializer) const override;
 
 	/** Find a struct property, called from generated code */
diff --git a/Engine/Source/Runtime/CoreUObject/Public/UObject/Object.h b/Engine/Source/Runtime/CoreUObject/Public/UObject/Object.h
--- a/Engine/Source/Runtime/CoreUObject/Public/UObject/Object.h
+++ b/Engine/Source/Runtime/CoreUObject/Public/UObject/Object.h
@@ -1116,7 +1116,7 @@ public:
 	 * 
 	 * @return the archetype for this object
 	 */
-	static UObject* GetArchetypeFromRequiredInfo(UClass* Class, UObject* Outer, FName Name, EObjectFlags ObjectFlags);
+	static UObject* GetArchetypeFromRequiredInfo(const UClass* Class, const UObject* Outer, FName Name, EObjectFlags ObjectFlags);
 
 	/**
 	 * Return the template this object is based on. 
diff --git a/Engine/Source/Runtime/CoreUObject/Public/UObject/UObjectGlobals.h b/Engine/Source/Runtime/CoreUObject/Public/UObject/UObjectGlobals.h
--- a/Engine/Source/Runtime/CoreUObject/Public/UObject/UObjectGlobals.h
+++ b/Engine/Source/Runtime/CoreUObject/Public/UObject/UObjectGlobals.h
@@ -278,7 +278,7 @@ COREUOBJECT_API UClass* StaticLoadClass(UClass* BaseClass, UObject* InOuter, con
  *
  * @return	A pointer to a fully initialized object of the specified class.
  */
-COREUOBJECT_API UObject* StaticConstructObject_Internal(UClass* Class, UObject* InOuter = (UObject*)GetTransientPackage(), FName Name = NAME_None, EObjectFlags SetFlags = RF_NoFlags, EInternalObjectFlags InternalSetFlags = EInternalObjectFlags::None, UObject* Template = nullptr, bool bCopyTransientsFromClassDefaults = false, struct FObjectInstancingGraph* InstanceGraph = nullptr, bool bAssumeTemplateIsArchetype = false);
+COREUOBJECT_API UObject* StaticConstructObject_Internal(const UClass* Class, UObject* InOuter = (UObject*)GetTransientPackage(), FName Name = NAME_None, EObjectFlags SetFlags = RF_NoFlags, EInternalObjectFlags InternalSetFlags = EInternalObjectFlags::None, UObject* Template = nullptr, bool bCopyTransientsFromClassDefaults = false, struct FObjectInstancingGraph* InstanceGraph = nullptr, bool bAssumeTemplateIsArchetype = false);
 
 /**
  * Creates a copy of SourceObject using the Outer and Name specified, as well as copies of all objects contained by SourceObject.  
@@ -500,7 +500,7 @@ COREUOBJECT_API void IncrementalPurgeGarbage( bool bUseTimeLimit, float TimeLimi
  * @return	name is the form BaseName_##, where ## is the number of objects of this
  *			type that have been created since the last time the class was garbage collected.
  */
-COREUOBJECT_API FName MakeUniqueObjectName( UObject* Outer, UClass* Class, FName BaseName=NAME_None );
+COREUOBJECT_API FName MakeUniqueObjectName( UObject* Outer, const UClass* Class, FName BaseName=NAME_None );
 
 /**
  * Given a display label string, generates an FName slug that is a valid FName for that label.
@@ -640,7 +640,7 @@ COREUOBJECT_API void SnapshotTransactionBuffer(UObject* Object);
  * @param	SetFlags	the ObjectFlags to assign to the new object. some flags can affect the behavior of constructing the object.
  * @return	true if nullptr should be returned; there was a problem reported 
  */
-bool StaticAllocateObjectErrorTests( UClass* Class, UObject* InOuter, FName Name, EObjectFlags SetFlags);
+bool StaticAllocateObjectErrorTests( const UClass* Class, UObject* InOuter, FName Name, EObjectFlags SetFlags);
 
 /**
  * Create a new instance of an object or replace an existing object.  If both an Outer and Name are specified, and there is an object already in memory with the same Class, Outer, and Name, the
@@ -655,7 +655,7 @@ bool StaticAllocateObjectErrorTests( UClass* Class, UObject* InOuter, FName Name
  * @param bOutReusedSubobject	flag indicating if the object is a subobject that has already been created (in which case further initialization is not necessary).
  * @return	a pointer to a fully initialized object of the specified class.
  */
-COREUOBJECT_API UObject* StaticAllocateObject(UClass* Class, UObject* InOuter, FName Name, EObjectFlags SetFlags, EInternalObjectFlags InternalSetFlags = EInternalObjectFlags::None, bool bCanReuseSubobjects = false, bool* bOutReusedSubobject = nullptr);
+COREUOBJECT_API UObject* StaticAllocateObject(const UClass* Class, UObject* InOuter, FName Name, EObjectFlags SetFlags, EInternalObjectFlags InternalSetFlags = EInternalObjectFlags::None, bool bCanReuseSubobjects = false, bool* bOutReusedSubobject = nullptr);
 
 /** @deprecated Use raw pointers or TWeakObjectPtr instead */
 class COREUOBJECT_API FSubobjectPtr
@@ -1176,7 +1176,7 @@ public:
 
 #if DO_CHECK
 /** Called by NewObject to make sure Child is actually a child of Parent */
-COREUOBJECT_API void CheckIsClassChildOf_Internal(UClass* Parent, UClass* Child);
+COREUOBJECT_API void CheckIsClassChildOf_Internal(const UClass* Parent, const UClass* Child);
 #endif
 
 /**
@@ -1194,7 +1194,7 @@ COREUOBJECT_API void CheckIsClassChildOf_Internal(UClass* Parent, UClass* Child)
  */
 template< class T >
 FUNCTION_NON_NULL_RETURN_START
-	T* NewObject(UObject* Outer, UClass* Class, FName Name = NAME_None, EObjectFlags Flags = RF_NoFlags, UObject* Template = nullptr, bool bCopyTransientsFromClassDefaults = false, FObjectInstancingGraph* InInstanceGraph = nullptr)
+	T* NewObject(UObject* Outer, const UClass* Class, FName Name = NAME_None, EObjectFlags Flags = RF_NoFlags, UObject* Template = nullptr, bool bCopyTransientsFromClassDefaults = false, FObjectInstancingGraph* InInstanceGraph = nullptr)
 FUNCTION_NON_NULL_RETURN_END
 {
 	if (Name == NAME_None)
diff --git a/Engine/Source/Runtime/Engine/Classes/Engine/BlueprintGeneratedClass.h b/Engine/Source/Runtime/Engine/Classes/Engine/BlueprintGeneratedClass.h
--- a/Engine/Source/Runtime/Engine/Classes/Engine/BlueprintGeneratedClass.h
+++ b/Engine/Source/Runtime/Engine/Classes/Engine/BlueprintGeneratedClass.h
@@ -747,7 +747,7 @@ public:
 	virtual void PurgeClass(bool bRecompilingOnLoad) override;
 	virtual void Bind() override;
 	virtual void GetDefaultObjectPreloadDependencies(TArray<UObject*>& OutDeps) override;
-	virtual UObject* FindArchetype(UClass* ArchetypeClass, const FName ArchetypeName) const override;
+	virtual UObject* FindArchetype(const UClass* ArchetypeClass, const FName ArchetypeName) const override;
 
 	virtual void InitPropertiesFromCustomList(uint8* DataPtr, const uint8* DefaultDataPtr) override;
 	virtual void SetupObjectInitializer(FObjectInitializer& ObjectInitializer) const override;
diff --git a/Engine/Source/Runtime/Engine/Private/BlueprintGeneratedClass.cpp b/Engine/Source/Runtime/Engine/Private/BlueprintGeneratedClass.cpp
--- a/Engine/Source/Runtime/Engine/Private/BlueprintGeneratedClass.cpp
+++ b/Engine/Source/Runtime/Engine/Private/BlueprintGeneratedClass.cpp
@@ -745,7 +745,7 @@ UInheritableComponentHandler* UBlueprintGeneratedClass::GetInheritableComponentH
 }
 
 
-UObject* UBlueprintGeneratedClass::FindArchetype(UClass* ArchetypeClass, const FName ArchetypeName) const
+UObject* UBlueprintGeneratedClass::FindArchetype(const UClass* ArchetypeClass, const FName ArchetypeName) const
 {
 	UObject* Archetype = nullptr;
 
