Fix UpdateSuffixForNextNewObject definition signature
#rnx

diff --git a/Engine/Source/Runtime/CoreUObject/Public/UObject/UObjectGlobals.h b/Engine/Source/Runtime/CoreUObject/Public/UObject/UObjectGlobals.h
--- a/Engine/Source/Runtime/CoreUObject/Public/UObject/UObjectGlobals.h
+++ b/Engine/Source/Runtime/CoreUObject/Public/UObject/UObjectGlobals.h
@@ -188,7 +188,7 @@ COREUOBJECT_API bool ResolveName(UObject*& Outer, FString& ObjectsReferenceStrin
 COREUOBJECT_API bool SafeLoadError( UObject* Outer, uint32 LoadFlags, const TCHAR* ErrorMessage);
 
 /** Internal function used to update the suffix to be given to the next newly-created unnamed object. */
-COREUOBJECT_API int32 UpdateSuffixForNextNewObject(UObject* Parent, UClass* Class, TFunctionRef<void(int32&)> IndexMutator);
+COREUOBJECT_API int32 UpdateSuffixForNextNewObject(UObject* Parent, const UClass* Class, TFunctionRef<void(int32&)> IndexMutator);
 
 
 /**
