Fix a compile issue (missing semicolon).

#rnx

diff --git a/Engine/Source/Runtime/CoreUObject/Private/UObject/SavePackage.cpp b/Engine/Source/Runtime/CoreUObject/Private/UObject/SavePackage.cpp
--- a/Engine/Source/Runtime/CoreUObject/Private/UObject/SavePackage.cpp
+++ b/Engine/Source/Runtime/CoreUObject/Private/UObject/SavePackage.cpp
@@ -344,7 +344,7 @@ private:
 
 class COREUOBJECT_API FArchiveObjectCrc32NonEditorProperties : public FArchiveObjectCrc32
 {
-}
+};
 
 #endif
 
