End the async loading thread loop on core uobject shutdown. This causes games to not crash on exit after changes from CL 8321295

diff --git a/Engine/Source/Runtime/CoreUObject/Private/Serialization/AsyncLoading.cpp b/Engine/Source/Runtime/CoreUObject/Private/Serialization/AsyncLoading.cpp
--- a/Engine/Source/Runtime/CoreUObject/Private/Serialization/AsyncLoading.cpp
+++ b/Engine/Source/Runtime/CoreUObject/Private/Serialization/AsyncLoading.cpp
@@ -5396,6 +5396,18 @@ void FAsyncLoadingThread::Stop()
 	StopTaskCounter.Increment();
 }
 
+void FAsyncLoadingThread::Kill()
+{
+	check(IsInGameThread());
+
+	Stop();
+
+	if (IsMultithreaded() && Thread)
+	{
+		Thread->Kill(true); 
+	}
+}
+
 void FAsyncLoadingThread::CancelAsyncLoading()
 {
 	check(IsInGameThread());	
diff --git a/Engine/Source/Runtime/CoreUObject/Private/Serialization/AsyncLoadingThread.h b/Engine/Source/Runtime/CoreUObject/Private/Serialization/AsyncLoadingThread.h
--- a/Engine/Source/Runtime/CoreUObject/Private/Serialization/AsyncLoadingThread.h
+++ b/Engine/Source/Runtime/CoreUObject/Private/Serialization/AsyncLoadingThread.h
@@ -413,6 +413,11 @@ public:
 	*/
 	void CancelAsyncLoading();
 
+	/**
+	* [GAME THREAD] Stops the async loading thread and blocks until the thread has exited.
+	*/
+	void Kill();
+
 	/**
 	* [GAME THREAD] Suspends async loading thread
 	*/
diff --git a/Engine/Source/Runtime/CoreUObject/Private/UObject/CoreNative.cpp b/Engine/Source/Runtime/CoreUObject/Private/UObject/CoreNative.cpp
--- a/Engine/Source/Runtime/CoreUObject/Private/UObject/CoreNative.cpp
+++ b/Engine/Source/Runtime/CoreUObject/Private/UObject/CoreNative.cpp
@@ -8,6 +8,7 @@
 #include "Misc/RuntimeErrors.h"
 #include "UObject/Stack.h"
 #include "Serialization/AsyncLoading.h"
+#include "Serialization/AsyncLoadingThread.h"
 
 void UClassRegisterAllCompiledInClasses();
 bool IsInAsyncLoadingThreadCoreUObjectInternal();
@@ -59,6 +60,10 @@ void FCoreUObjectModule::StartupModule()
 
 void FCoreUObjectModule::ShutdownModule()
 {
+	// Before deleting our handler, we need to make sure any async loading threads are stopped.
+	// Otherwise the threads will be stuck in an invalid state.
+	FAsyncLoadingThread::Get().Kill();
+
 	DeleteGlobalPrecacheHandler(GlobalPrecacheHandler);
 	GlobalPrecacheHandler = nullptr;
 }
