Fix missing linker reset in deferred blueprint objects before replacing placeholder object.

#rnx

diff --git a/Engine/Source/Runtime/CoreUObject/Private/Blueprint/BlueprintSupport.cpp b/Engine/Source/Runtime/CoreUObject/Private/Blueprint/BlueprintSupport.cpp
--- a/Engine/Source/Runtime/CoreUObject/Private/Blueprint/BlueprintSupport.cpp
+++ b/Engine/Source/Runtime/CoreUObject/Private/Blueprint/BlueprintSupport.cpp
@@ -2068,6 +2068,7 @@ void FLinkerLoad::ResolveDeferredExports(UClass* LoadClass)
 				if (ensure(PlaceholderExport))
 				{
 					// replace the placeholder with the proper object instance
+					PlaceholderExport->SetLinker(nullptr, INDEX_NONE);
 					Export.ResetObject();
 					UObject* ExportObj = CreateExport(ExportIndex);
 
