MediaPlayer: Added meta data description for HideInDetailPanel. #rb none

diff --git a/Engine/Source/Runtime/CoreUObject/Public/UObject/ObjectMacros.h b/Engine/Source/Runtime/CoreUObject/Public/UObject/ObjectMacros.h
--- a/Engine/Source/Runtime/CoreUObject/Public/UObject/ObjectMacros.h
+++ b/Engine/Source/Runtime/CoreUObject/Public/UObject/ObjectMacros.h
@@ -1142,6 +1142,9 @@ namespace UM
 		/// [PropertyMetadata] Used for FColor and FLinearColor properties. Indicates that the Alpha property should be hidden when displaying the property widget in the details.
 		HideAlphaChannel,
 
+		/// [PropertyMetadata] Indicates that the property should be hidden in the details panel. Currently only used by events.
+		HideInDetailPanel,
+
 		/// [PropertyMetadata] Used for Subclass and SoftClass properties. Specifies to hide the ability to change view options in the class picker
 		HideViewOptions,
 
