Allow calling code to ImportText to avoid doing work to create the struct name used for error messages until it is actually needed.

diff --git a/Engine/Source/Runtime/AssetRegistry/Private/AssetBundleData.cpp b/Engine/Source/Runtime/AssetRegistry/Private/AssetBundleData.cpp
--- a/Engine/Source/Runtime/AssetRegistry/Private/AssetBundleData.cpp
+++ b/Engine/Source/Runtime/AssetRegistry/Private/AssetBundleData.cpp
@@ -13,7 +13,7 @@ bool FAssetBundleData::SetFromAssetData(const FAssetData& AssetData)
 
 	if (AssetData.GetTagValue(FAssetBundleData::StaticStruct()->GetFName(), TagValue))
 	{
-		if (FAssetBundleData::StaticStruct()->ImportText(*TagValue, this, nullptr, PPF_None, (FOutputDevice*)GWarn, AssetData.AssetName.ToString()))
+		if (FAssetBundleData::StaticStruct()->ImportText(*TagValue, this, nullptr, PPF_None, (FOutputDevice*)GWarn, [&AssetData]() { return AssetData.AssetName.ToString(); }))
 		{
 			FPrimaryAssetId FoundId = AssetData.GetPrimaryAssetId();
 
diff --git a/Engine/Source/Runtime/CoreUObject/Private/UObject/Class.cpp b/Engine/Source/Runtime/CoreUObject/Private/UObject/Class.cpp
--- a/Engine/Source/Runtime/CoreUObject/Private/UObject/Class.cpp
+++ b/Engine/Source/Runtime/CoreUObject/Private/UObject/Class.cpp
@@ -2262,6 +2262,11 @@ void UScriptStruct::SerializeItem(FStructuredArchive::FSlot Slot, void* Value, v
 }
 
 const TCHAR* UScriptStruct::ImportText(const TCHAR* InBuffer, void* Value, UObject* OwnerObject, int32 PortFlags, FOutputDevice* ErrorText, const FString& StructName, bool bAllowNativeOverride)
+{
+	return ImportText(InBuffer, Value, OwnerObject, PortFlags, ErrorText, [&StructName](){return StructName;}, bAllowNativeOverride);
+}
+
+const TCHAR* UScriptStruct::ImportText(const TCHAR* InBuffer, void* Value, UObject* OwnerObject, int32 PortFlags, FOutputDevice* ErrorText, const TFunctionRef<FString()>& StructNameGetter, bool bAllowNativeOverride)
 {
 	if (bAllowNativeOverride && StructFlags & STRUCT_ImportTextItemNative)
 	{
@@ -2301,7 +2306,7 @@ const TCHAR* UScriptStruct::ImportText(const TCHAR* InBuffer, void* Value, UObje
 
 					if (*Buffer != TCHAR('\"'))
 					{
-						ErrorText->Logf(TEXT("%sImportText (%s): Bad quoted string at: %s"), ErrorCount++ > 0 ? LINE_TERMINATOR : TEXT(""), *StructName, Buffer);
+						ErrorText->Logf(TEXT("%sImportText (%s): Bad quoted string at: %s"), ErrorCount++ > 0 ? LINE_TERMINATOR : TEXT(""), *StructNameGetter(), Buffer);
 						return nullptr;
 					}
 				}
@@ -2314,7 +2319,7 @@ const TCHAR* UScriptStruct::ImportText(const TCHAR* InBuffer, void* Value, UObje
 					SubCount--;
 					if (SubCount < 0)
 					{
-						ErrorText->Logf(TEXT("%sImportText (%s): Too many closing parenthesis in: %s"), ErrorCount++ > 0 ? LINE_TERMINATOR : TEXT(""), *StructName, InBuffer);
+						ErrorText->Logf(TEXT("%sImportText (%s): Too many closing parenthesis in: %s"), ErrorCount++ > 0 ? LINE_TERMINATOR : TEXT(""), *StructNameGetter(), InBuffer);
 						return nullptr;
 					}
 				}
@@ -2322,7 +2327,7 @@ const TCHAR* UScriptStruct::ImportText(const TCHAR* InBuffer, void* Value, UObje
 			}
 			if (SubCount > 0)
 			{
-				ErrorText->Logf(TEXT("%sImportText(%s): Not enough closing parenthesis in: %s"), ErrorCount++ > 0 ? LINE_TERMINATOR : TEXT(""), *StructName, InBuffer);
+				ErrorText->Logf(TEXT("%sImportText(%s): Not enough closing parenthesis in: %s"), ErrorCount++ > 0 ? LINE_TERMINATOR : TEXT(""), *StructNameGetter(), InBuffer);
 				return nullptr;
 			}
 
@@ -2334,7 +2339,7 @@ const TCHAR* UScriptStruct::ImportText(const TCHAR* InBuffer, void* Value, UObje
 			}
 			else if (*Buffer != TCHAR(')'))
 			{
-				ErrorText->Logf(TEXT("%sImportText (%s): Missing closing parenthesis: %s"), ErrorCount++ > 0 ? LINE_TERMINATOR : TEXT(""), *StructName, InBuffer);
+				ErrorText->Logf(TEXT("%sImportText (%s): Missing closing parenthesis: %s"), ErrorCount++ > 0 ? LINE_TERMINATOR : TEXT(""), *StructNameGetter(), InBuffer);
 				return nullptr;
 			}
 
@@ -2346,7 +2351,7 @@ const TCHAR* UScriptStruct::ImportText(const TCHAR* InBuffer, void* Value, UObje
 	}
 	else
 	{
-		ErrorText->Logf(TEXT("%sImportText (%s): Missing opening parenthesis: %s"), ErrorCount++ > 0 ? LINE_TERMINATOR : TEXT(""), *StructName, InBuffer); //-V547
+		ErrorText->Logf(TEXT("%sImportText (%s): Missing opening parenthesis: %s"), ErrorCount++ > 0 ? LINE_TERMINATOR : TEXT(""), *StructNameGetter(), InBuffer); //-V547
 		return nullptr;
 	}
 	return Buffer;
diff --git a/Engine/Source/Runtime/CoreUObject/Public/UObject/Class.h b/Engine/Source/Runtime/CoreUObject/Public/UObject/Class.h
--- a/Engine/Source/Runtime/CoreUObject/Public/UObject/Class.h
+++ b/Engine/Source/Runtime/CoreUObject/Public/UObject/Class.h
@@ -1344,6 +1344,20 @@ public:
 	 */
 	COREUOBJECT_API const TCHAR* ImportText(const TCHAR* Buffer, void* Value, UObject* OwnerObject, int32 PortFlags, FOutputDevice* ErrorText, const FString& StructName, bool bAllowNativeOverride = true);
 
+	/**
+	 * Sets value of script struct based on imported string
+	 *
+	 * @param	Buffer			String to read text data out of
+	 * @param	Value			Struct that will be modified
+	 * @param	OwnerObject		UObject that contains this struct
+	 * @param	PortFlags		EPropertyPortFlags controlling import behavior
+	 * @param	ErrorText		What to print import errors to
+	 * @param	StructNameGetter Function to return the struct name to avoid doing work if no error message is forthcoming
+	 * @param	bAllowNativeOverride If true, will try to run native version of export text on the struct
+	 * @return Buffer after parsing has succeeded, or NULL on failure
+	 */
+	COREUOBJECT_API const TCHAR* ImportText(const TCHAR* Buffer, void* Value, UObject* OwnerObject, int32 PortFlags, FOutputDevice* ErrorText, const TFunctionRef<FString()>& StructNameGetter, bool bAllowNativeOverride = true);
+
 	/**
 	 * Compare two script structs
 	 *
