MoveTemp some async delegates
#rnx

diff --git a/Engine/Source/Runtime/CoreUObject/Private/Serialization/AsyncLoading.cpp b/Engine/Source/Runtime/CoreUObject/Private/Serialization/AsyncLoading.cpp
--- a/Engine/Source/Runtime/CoreUObject/Private/Serialization/AsyncLoading.cpp
+++ b/Engine/Source/Runtime/CoreUObject/Private/Serialization/AsyncLoading.cpp
@@ -7295,7 +7295,7 @@ int32 LoadPackageAsync(const FString& InName, const FGuid* InGuid /*= nullptr*/,
 		TUniquePtr<FLoadPackageAsyncDelegate> CompletionDelegatePtr;
 		if (InCompletionDelegate.IsBound())
 		{
-			CompletionDelegatePtr.Reset(new FLoadPackageAsyncDelegate(InCompletionDelegate));
+			CompletionDelegatePtr.Reset(new FLoadPackageAsyncDelegate(MoveTemp(InCompletionDelegate)));
 		}
 
 		// Add new package request
@@ -7314,7 +7314,7 @@ int32 LoadPackageAsync(const FString& PackageName, FLoadPackageAsyncDelegate Com
 {
 	const FGuid* Guid = nullptr;
 	const TCHAR* PackageToLoadFrom = nullptr;
-	return LoadPackageAsync(PackageName, Guid, PackageToLoadFrom, CompletionDelegate, InPackageFlags, InPIEInstanceID, InPackagePriority );
+	return LoadPackageAsync(PackageName, Guid, PackageToLoadFrom, MoveTemp(CompletionDelegate), InPackageFlags, InPIEInstanceID, InPackagePriority );
 }
 
 void CancelAsyncLoading()
