[AsyncLoading] Explicit shutdown of async loading

- reverted changes from CL 8231774
- added ShutdownAsyncThread in UObjectBaseShutdown

diff --git a/Engine/Source/Runtime/CoreUObject/Private/Serialization/AsyncLoading.cpp b/Engine/Source/Runtime/CoreUObject/Private/Serialization/AsyncLoading.cpp
--- a/Engine/Source/Runtime/CoreUObject/Private/Serialization/AsyncLoading.cpp
+++ b/Engine/Source/Runtime/CoreUObject/Private/Serialization/AsyncLoading.cpp
@@ -1706,16 +1706,7 @@ struct FPrecacheCallbackHandler
 		FPlatformFileManager::Get().GetPlatformFile().SetAsyncMinimumPriority(NewMinPriority);
 	}
 };
-
-FPrecacheCallbackHandler* InitGlobalPrecacheHandler()
-{
-	return new FPrecacheCallbackHandler();
-}
-
-void DeleteGlobalPrecacheHandler(FPrecacheCallbackHandler* InHandler)
-{
-	delete InHandler;
-}
+TUniquePtr<FPrecacheCallbackHandler> GPrecacheCallbackHandler;
 
 int32 GRandomizeLoadOrder = 0;
 static FAutoConsoleVariableRef CVarRandomizeLoadOrder(
@@ -1797,7 +1788,7 @@ FORCEINLINE static bool CanAddWaitingPackages()
 	//marked by pathological load time performance, and this does not cover the "excessive load times when loading deployed, uncompressed data" case.
 	//applying a sane cap in all circumstances would not be a terrible idea.
 	const int32 MaxWaitingPackageCount = 1024;
-	return !FileOpenLogActive() || FCoreUObjectModule::FCoreUObjectModule::GetGlobalPrecacheHandler()->WaitingPackages.Num() < MaxWaitingPackageCount;
+	return !FileOpenLogActive() || GPrecacheCallbackHandler->WaitingPackages.Num() < MaxWaitingPackageCount;
 }
 
 void FAsyncLoadingThread::QueueEvent_CreateLinker(FAsyncPackage* Package, int32 EventSystemPriority)
@@ -3829,7 +3820,7 @@ void FAsyncPackage::StartPrecacheRequest()
 	FAsyncArchive* AsyncLoader = Linker->GetAsyncLoader();
 	check(AsyncLoader);
 
-	IAsyncReadRequest* Precache = AsyncLoader->MakeEventDrivenPrecacheRequest(NewReq.Offset, NewReq.BytesToRead, FCoreUObjectModule::GetGlobalPrecacheHandler()->GetCompletionCallback());
+	IAsyncReadRequest* Precache = AsyncLoader->MakeEventDrivenPrecacheRequest(NewReq.Offset, NewReq.BytesToRead, GPrecacheCallbackHandler->GetCompletionCallback());
 
 	NewReq.FirstExportCovered = LocalExportIndex;
 	NewReq.LastExportCovered = LastExportIndex;
@@ -3841,7 +3832,7 @@ void FAsyncPackage::StartPrecacheRequest()
 	check(!PrecacheRequests.Contains(Precache));
 	FExportIORequest& RequestInPlace = PrecacheRequests.Add(Precache);
 	Swap(RequestInPlace, NewReq);
-	FCoreUObjectModule::GetGlobalPrecacheHandler()->RegisterNewPrecacheRequest(Precache, this);
+	GPrecacheCallbackHandler->RegisterNewPrecacheRequest(Precache, this);
 }
 
 int64 FAsyncPackage::PrecacheRequestReady(IAsyncReadRequest * Read)
@@ -3864,7 +3855,7 @@ void FAsyncPackage::MakeNextPrecacheRequestCurrent()
 	CurrentBlockBytes = Req.BytesToRead;
 	ExportsInThisBlock.Reset();
 
-	FCoreUObjectModule::GetGlobalPrecacheHandler()->FinishRequest(Req.BytesToRead);
+	GPrecacheCallbackHandler->FinishRequest(Req.BytesToRead);
 
 	FAsyncArchive* AsyncLoader = Linker->GetAsyncLoader();
 	check(AsyncLoader);
@@ -4521,7 +4512,7 @@ EAsyncPackageState::Type FAsyncLoadingThread::ProcessAsyncLoading(int32& OutPack
 
 			bool bDidSomething = false;
 			{
-				bDidSomething = FCoreUObjectModule::GetGlobalPrecacheHandler()->ProcessIncoming();
+				bDidSomething = GPrecacheCallbackHandler->ProcessIncoming();
 				OutPackagesProcessed += (bDidSomething ? 1 : 0);
 
 				if (IsTimeLimitExceeded(TickStartTime, bUseTimeLimit, TimeLimit, TEXT("ProcessIncoming"), nullptr))
@@ -4623,7 +4614,7 @@ EAsyncPackageState::Type FAsyncLoadingThread::ProcessAsyncLoading(int32& OutPack
 			{
 				continue;
 			}
-			bool bAnyIOOutstanding = FCoreUObjectModule::GetGlobalPrecacheHandler()->AnyIOOutstanding();
+			bool bAnyIOOutstanding = GPrecacheCallbackHandler->AnyIOOutstanding();
 			if (bAnyIOOutstanding)
 			{
 				SCOPED_LOADTIMER(Package_EventIOWait);
@@ -4635,7 +4626,7 @@ EAsyncPackageState::Type FAsyncLoadingThread::ProcessAsyncLoading(int32& OutPack
 						const float RemainingTimeLimit = FMath::Max(0.0f, TimeLimit - (float)(FPlatformTime::Seconds() - TickStartTime));
 						if (RemainingTimeLimit > 0.0f)
 						{
-							bool bGotIO = FCoreUObjectModule::GetGlobalPrecacheHandler()->WaitForIO(RemainingTimeLimit);
+							bool bGotIO = GPrecacheCallbackHandler->WaitForIO(RemainingTimeLimit);
 							if (bGotIO)
 							{
 								OutPackagesProcessed++;
@@ -4651,7 +4642,7 @@ EAsyncPackageState::Type FAsyncLoadingThread::ProcessAsyncLoading(int32& OutPack
 				}
 				else
 				{
-					bool bGotIO = FCoreUObjectModule::GetGlobalPrecacheHandler()->WaitForIO(10.0f); // wait "forever"
+					bool bGotIO = GPrecacheCallbackHandler->WaitForIO(10.0f); // wait "forever"
 					if (!bGotIO)
 					{
 						//UE_LOG(LogStreaming, Error, TEXT("Waited for 10 seconds on IO...."));
@@ -5292,7 +5283,7 @@ uint32 FAsyncLoadingThread::Run()
 
 void FAsyncLoadingThread::CheckForCycles()
 {
-	if (FCoreUObjectModule::GetGlobalPrecacheHandler()->AnyIOOutstanding() || EventQueue.EventQueue.Num())
+	if (GPrecacheCallbackHandler->AnyIOOutstanding() || EventQueue.EventQueue.Num())
 	{
 		// we can't check for cycles if there is stuff in flight.
 		return;
@@ -6075,8 +6066,8 @@ EAsyncPackageState::Type FAsyncPackage::CreateLinker()
 				UE_LOG(LogStreaming, Fatal, TEXT("Package %s was reloaded before it even closed the linker from a previous load. Seems like a waste of time eh?"), *Desc.Name.ToString());
 				check(Package);
 				FWeakAsyncPackagePtr WeakPtr(this);
-				FCoreUObjectModule::GetGlobalPrecacheHandler()->RegisterNewSummaryRequest(this);
-				FCoreUObjectModule::GetGlobalPrecacheHandler()->SummaryComplete(WeakPtr);
+				GPrecacheCallbackHandler->RegisterNewSummaryRequest(this);
+				GPrecacheCallbackHandler->SummaryComplete(WeakPtr);
 			}
 		}
 
@@ -6171,17 +6162,17 @@ EAsyncPackageState::Type FAsyncPackage::CreateLinker()
 					, TFunction<void()>(
 						[WeakPtr]()
 				{
-					FCoreUObjectModule::GetGlobalPrecacheHandler()->SummaryComplete(WeakPtr);
+					GPrecacheCallbackHandler->SummaryComplete(WeakPtr);
 				}
 				));
 				if (Linker)
 				{
-					FCoreUObjectModule::GetGlobalPrecacheHandler()->RegisterNewSummaryRequest(this);
+					GPrecacheCallbackHandler->RegisterNewSummaryRequest(this);
 					if (Linker->bDynamicClassLinker)
 					{
 						//native blueprint 
 						check(!Linker->GetAsyncLoader());
-						FCoreUObjectModule::GetGlobalPrecacheHandler()->SummaryComplete(WeakPtr);
+						GPrecacheCallbackHandler->SummaryComplete(WeakPtr);
 					}
 				}
 			}
@@ -7355,9 +7346,16 @@ void InitAsyncThread()
 {
 	FCoreDelegates::OnSyncLoadPackage.AddStatic([](const FString&) { GSyncLoadCount++; });
 
+	GPrecacheCallbackHandler = MakeUnique<FPrecacheCallbackHandler>();
 	FAsyncLoadingThread::Get().InitializeAsyncThread();
 }
 
+void ShutdownAsyncThread()
+{
+	FAsyncLoadingThread::Get().Kill();
+	GPrecacheCallbackHandler.Reset();
+}
+
 bool IsInAsyncLoadingThreadCoreUObjectInternal()
 {
 	return FAsyncLoadingThread::Get().IsInAsyncLoadThread();
diff --git a/Engine/Source/Runtime/CoreUObject/Private/UObject/CoreNative.cpp b/Engine/Source/Runtime/CoreUObject/Private/UObject/CoreNative.cpp
--- a/Engine/Source/Runtime/CoreUObject/Private/UObject/CoreNative.cpp
+++ b/Engine/Source/Runtime/CoreUObject/Private/UObject/CoreNative.cpp
@@ -7,8 +7,6 @@
 #include "Misc/PackageName.h"
 #include "Misc/RuntimeErrors.h"
 #include "UObject/Stack.h"
-#include "Serialization/AsyncLoading.h"
-#include "Serialization/AsyncLoadingThread.h"
 
 void UClassRegisterAllCompiledInClasses();
 bool IsInAsyncLoadingThreadCoreUObjectInternal();
@@ -18,57 +16,48 @@ void ResumeAsyncLoadingInternal();
 bool IsAsyncLoadingSuspendedInternal();
 bool IsAsyncLoadingMultithreadedCoreUObjectInternal();
 
-void FCoreUObjectModule::RouteRuntimeMessageToBP(ELogVerbosity::Type Verbosity, const ANSICHAR* FileName, int32 LineNumber, const FText& Message)
+// CoreUObject module. Handles UObject system pre-init (registers init function with Core callbacks).
+class FCoreUObjectModule : public FDefaultModuleImpl
 {
+public:
+	static void RouteRuntimeMessageToBP(ELogVerbosity::Type Verbosity, const ANSICHAR* FileName, int32 LineNumber, const FText& Message)
+	{
 #if UE_RAISE_RUNTIME_ERRORS && !NO_LOGGING
-	check((Verbosity == ELogVerbosity::Error) || (Verbosity == ELogVerbosity::Warning));
-	FMsg::Logf_Internal(FileName, LineNumber, LogScript.GetCategoryName(), Verbosity, TEXT("%s(%d): Runtime %s: \"%s\""), ANSI_TO_TCHAR(FileName), LineNumber, (Verbosity == ELogVerbosity::Error) ? TEXT("Error") : TEXT("Warning"), *Message.ToString());
+		check((Verbosity == ELogVerbosity::Error) || (Verbosity == ELogVerbosity::Warning));
+		FMsg::Logf_Internal(FileName, LineNumber, LogScript.GetCategoryName(), Verbosity, TEXT("%s(%d): Runtime %s: \"%s\""), ANSI_TO_TCHAR(FileName), LineNumber, (Verbosity == ELogVerbosity::Error) ? TEXT("Error") : TEXT("Warning"), *Message.ToString());
 #endif
-	FFrame::KismetExecutionMessage(*Message.ToString(), Verbosity);
-}
-
-void FCoreUObjectModule::StartupModule()
-{
-	GlobalPrecacheHandler = InitGlobalPrecacheHandler();
+		FFrame::KismetExecutionMessage(*Message.ToString(), Verbosity);
+	}
 
-	// Register all classes that have been loaded so far. This is required for CVars to work.		
-	UClassRegisterAllCompiledInClasses();
+	virtual void StartupModule() override
+	{
+		// Register all classes that have been loaded so far. This is required for CVars to work.		
+		UClassRegisterAllCompiledInClasses();
 
-	void InitUObject();
-	FCoreDelegates::OnInit.AddStatic(InitUObject);
+		void InitUObject();
+		FCoreDelegates::OnInit.AddStatic(InitUObject);
 
-	// Substitute Core version of async loading functions with CoreUObject ones.
-	IsInAsyncLoadingThread = &IsInAsyncLoadingThreadCoreUObjectInternal;
-	IsAsyncLoading = &IsAsyncLoadingCoreUObjectInternal;
-	SuspendAsyncLoading = &SuspendAsyncLoadingInternal;
-	ResumeAsyncLoading = &ResumeAsyncLoadingInternal;
-	IsAsyncLoadingSuspended = &IsAsyncLoadingSuspendedInternal;
-	IsAsyncLoadingMultithreaded = &IsAsyncLoadingMultithreadedCoreUObjectInternal;
+		// Substitute Core version of async loading functions with CoreUObject ones.
+		IsInAsyncLoadingThread = &IsInAsyncLoadingThreadCoreUObjectInternal;
+		IsAsyncLoading = &IsAsyncLoadingCoreUObjectInternal;
+		SuspendAsyncLoading = &SuspendAsyncLoadingInternal;
+		ResumeAsyncLoading = &ResumeAsyncLoadingInternal;
+		IsAsyncLoadingSuspended = &IsAsyncLoadingSuspendedInternal;
+		IsAsyncLoadingMultithreaded = &IsAsyncLoadingMultithreadedCoreUObjectInternal;
 
-	// Register the script callstack callback to the runtime error logging
+		// Register the script callstack callback to the runtime error logging
 #if UE_RAISE_RUNTIME_ERRORS
-	FRuntimeErrors::OnRuntimeIssueLogged.BindStatic(&FCoreUObjectModule::RouteRuntimeMessageToBP);
+		FRuntimeErrors::OnRuntimeIssueLogged.BindStatic(&FCoreUObjectModule::RouteRuntimeMessageToBP);
 #endif
 
-	// Make sure that additional content mount points can be registered after CoreUObject loads
-	FPackageName::EnsureContentPathsAreRegistered();
+		// Make sure that additional content mount points can be registered after CoreUObject loads
+		FPackageName::EnsureContentPathsAreRegistered();
 
 #if DO_BLUEPRINT_GUARD
-	FFrame::InitPrintScriptCallstack();
+		FFrame::InitPrintScriptCallstack();
 #endif
-}
-
-void FCoreUObjectModule::ShutdownModule()
-{
-	// Before deleting our handler, we need to make sure any async loading threads are stopped.
-	// Otherwise the threads will be stuck in an invalid state.
-	FAsyncLoadingThread::Get().Kill();
-
-	DeleteGlobalPrecacheHandler(GlobalPrecacheHandler);
-	GlobalPrecacheHandler = nullptr;
-}
-
-struct FPrecacheCallbackHandler* FCoreUObjectModule::GlobalPrecacheHandler = nullptr;
+	}
+};
 IMPLEMENT_MODULE( FCoreUObjectModule, CoreUObject );
 
 // if we are not using compiled in natives, we still need this as a base class for intrinsics
diff --git a/Engine/Source/Runtime/CoreUObject/Private/UObject/UObjectBase.cpp b/Engine/Source/Runtime/CoreUObject/Private/UObject/UObjectBase.cpp
--- a/Engine/Source/Runtime/CoreUObject/Private/UObject/UObjectBase.cpp
+++ b/Engine/Source/Runtime/CoreUObject/Private/UObject/UObjectBase.cpp
@@ -1136,6 +1136,9 @@ void UObjectBaseInit()
  */
 void UObjectBaseShutdown()
 {
+	void ShutdownAsyncThread();
+	ShutdownAsyncThread();
+
 	GUObjectArray.ShutdownUObjectArray();
 	Internal::GetUObjectSubsystemInitialised() = false;
 }
diff --git a/Engine/Source/Runtime/CoreUObject/Public/Serialization/AsyncLoading.h b/Engine/Source/Runtime/CoreUObject/Public/Serialization/AsyncLoading.h
--- a/Engine/Source/Runtime/CoreUObject/Public/Serialization/AsyncLoading.h
+++ b/Engine/Source/Runtime/CoreUObject/Public/Serialization/AsyncLoading.h
@@ -959,10 +959,6 @@ struct FScopedAsyncPackageEvent
 	~FScopedAsyncPackageEvent();
 };
 
-// Access to precache handler
-extern COREUOBJECT_API struct FPrecacheCallbackHandler* InitGlobalPrecacheHandler();
-extern COREUOBJECT_API void DeleteGlobalPrecacheHandler(struct FPrecacheCallbackHandler* InHandler);
-
 // Stats for ChartCreation.cpp
 extern COREUOBJECT_API double GFlushAsyncLoadingTime;
 extern COREUOBJECT_API uint32 GFlushAsyncLoadingCount;
diff --git a/Engine/Source/Runtime/CoreUObject/Public/UObject/CoreNative.h b/Engine/Source/Runtime/CoreUObject/Public/UObject/CoreNative.h
--- a/Engine/Source/Runtime/CoreUObject/Public/UObject/CoreNative.h
+++ b/Engine/Source/Runtime/CoreUObject/Public/UObject/CoreNative.h
@@ -7,7 +7,6 @@
 #pragma once
 
 #include "CoreMinimal.h"
-#include "Modules/ModuleManager.h"
 #include "UObject/Script.h"
 #include "UObject/Object.h"
 
@@ -40,20 +39,3 @@ struct FNativeFunctionRegistrar
 
 	static COREUOBJECT_API void RegisterFunctions(class UClass* Class, const FNameNativePtrPair* InArray, int32 NumFunctions);
 };
-
-// CoreUObject module. Handles UObject system pre-init (registers init function with Core callbacks).
-class FCoreUObjectModule : public FDefaultModuleImpl
-{
-public:
-	static void RouteRuntimeMessageToBP(ELogVerbosity::Type Verbosity, const ANSICHAR* FileName, int32 LineNumber, const FText& Message);
-
-	virtual void StartupModule() override;
-
-	virtual void ShutdownModule() override;
-
-	static struct FPrecacheCallbackHandler* GetGlobalPrecacheHandler() { return GlobalPrecacheHandler; }
-
-private:
-
-	static struct FPrecacheCallbackHandler* GlobalPrecacheHandler;
-};
