Fix CIS issue UE4Editor Linux

diff --git a/Engine/Source/Runtime/CoreUObject/Private/UObject/CoreNative.cpp b/Engine/Source/Runtime/CoreUObject/Private/UObject/CoreNative.cpp
--- a/Engine/Source/Runtime/CoreUObject/Private/UObject/CoreNative.cpp
+++ b/Engine/Source/Runtime/CoreUObject/Private/UObject/CoreNative.cpp
@@ -2,6 +2,7 @@
 
 #include "UObject/CoreNative.h"
 #include "Misc/CoreDelegates.h"
+#include "Modules/ModuleManager.h"
 #include "UObject/Class.h"
 #include "UObject/UnrealType.h"
 #include "Misc/PackageName.h"
