Make StaticFindObject() handle searching for any package with a package-specifying string

The change in behavior was introduced by an optimization in UE 4.23

(ushell-p4-cherrypick of 8633853 by [REDACTED_PERSON])

diff --git a/Engine/Source/Runtime/CoreUObject/Private/UObject/UObjectHash.cpp b/Engine/Source/Runtime/CoreUObject/Private/UObject/UObjectHash.cpp
--- a/Engine/Source/Runtime/CoreUObject/Private/UObject/UObjectHash.cpp
+++ b/Engine/Source/Runtime/CoreUObject/Private/UObject/UObjectHash.cpp
@@ -489,6 +489,42 @@ UObject* StaticFindObjectFastExplicit( const UClass* ObjectClass, FName ObjectNa
 	return Result;
 }
 
+static FName SplitInnerAndOuter(const TCHAR* OuterAndInner, const TCHAR* Delimiter, FName& OutOuter, uint32 InnerNumber)
+{
+	check(*Delimiter == ':');
+
+	int32 OuterLen = static_cast<int32>(Delimiter - OuterAndInner);
+	OutOuter = FName(OuterLen, OuterAndInner);
+	return FName(Delimiter + 1, InnerNumber);
+}
+
+static FName ExtractInnerAndOuterFromPath(FName ObjectPath, FName& OutOuter)
+{
+	TCHAR PathBuffer[NAME_SIZE];
+	ObjectPath.GetPlainNameString(PathBuffer);
+
+	// Find package separator . or subobject separator :
+	constexpr FAsciiSet DotColon(".:");
+	const TCHAR* Path = PathBuffer;
+	while (true)
+	{
+		const TCHAR* DelimiterOrEnd = FAsciiSet::FindFirstOrEnd(Path, DotColon);
+
+		if (*DelimiterOrEnd == '\0')
+		{
+			return ObjectPath;
+		}
+		else if (*DelimiterOrEnd == ':')
+		{
+			return SplitInnerAndOuter(Path, DelimiterOrEnd, OutOuter, ObjectPath.GetNumber());
+		}
+
+		// We have a package prefix, drop it
+		check(*DelimiterOrEnd == '.');
+		Path = DelimiterOrEnd + 1;
+	}
+}
+
 UObject* StaticFindObjectFastInternalThreadSafe(FUObjectHashTables& ThreadHash, const UClass* ObjectClass, const UObject* ObjectPackage, FName ObjectName, bool bExactClass, bool bAnyPackage, EObjectFlags ExcludeFlags, EInternalObjectFlags ExclusiveInternalFlags)
 {
 	ExclusiveInternalFlags |= EInternalObjectFlags::Unreachable;
@@ -536,20 +572,8 @@ UObject* StaticFindObjectFastInternalThreadSafe(FUObjectHashTables& ThreadHash,
 	else
 	{
 		// Find an object with the specified name and (optional) class, in any package; if bAnyPackage is false, only matches top-level packages
-		FName ActualObjectName = ObjectName;
 		FName VerifyOuterName;
-		TCHAR PlainObjectName[NAME_SIZE];
-		int32 PlainObjectNameLen = ObjectName.GetPlainNameString(PlainObjectName);
-		
-		// Drop part prefixed by '.' or ':' and keep the suffix part
-		constexpr FAsciiSet DotColon(".:");
-		const TCHAR* DelimiterOrEnd = FAsciiSet::FindLastOrEnd(PlainObjectName, DotColon);
-		if (*DelimiterOrEnd)
-		{
-			ActualObjectName = FName(DelimiterOrEnd + 1, ObjectName.GetNumber());
-			int32 OuterLen = static_cast<int32>(DelimiterOrEnd - PlainObjectName);
-			VerifyOuterName = FName(OuterLen, PlainObjectName);
-		}
+		FName ActualObjectName = ExtractInnerAndOuterFromPath(ObjectName, /* out*/ VerifyOuterName);
 
 		const int32 Hash = GetObjectHash(ActualObjectName);
 		FHashTableLock HashLock(ThreadHash);
