Fixed StaticFindObject bug

diff --git a/Engine/Source/Runtime/CoreUObject/Private/UObject/UObjectHash.cpp b/Engine/Source/Runtime/CoreUObject/Private/UObject/UObjectHash.cpp
--- a/Engine/Source/Runtime/CoreUObject/Private/UObject/UObjectHash.cpp
+++ b/Engine/Source/Runtime/CoreUObject/Private/UObject/UObjectHash.cpp
@@ -512,7 +512,7 @@ static FName ExtractInnerAndOuterFromPath(FName ObjectPath, FName& OutOuter)
 
 		if (*DelimiterOrEnd == '\0')
 		{
-			return ObjectPath;
+			return Path == PathBuffer ? ObjectPath : FName(Path, ObjectPath.GetNumber());
 		}
 		else if (*DelimiterOrEnd == ':')
 		{
