Fixed a possible hang in GC when the number of objects to be destroyed was dividable by TimeLimitEnforcementGranularityForDeletion (100)

diff --git a/Engine/Source/Runtime/CoreUObject/Private/UObject/GarbageCollection.cpp b/Engine/Source/Runtime/CoreUObject/Private/UObject/GarbageCollection.cpp
--- a/Engine/Source/Runtime/CoreUObject/Private/UObject/GarbageCollection.cpp
+++ b/Engine/Source/Runtime/CoreUObject/Private/UObject/GarbageCollection.cpp
@@ -304,7 +304,7 @@ class FAsyncPurge : public FRunnable
 			++ObjCurrentPurgeObjectIndex;
 
 			// Time slicing when running on the game thread
-			if (bUseTimeLimit && (ProcessedObjectsCount == TimeLimitEnforcementGranularityForDeletion))
+			if (bUseTimeLimit && (ProcessedObjectsCount == TimeLimitEnforcementGranularityForDeletion) && (ObjCurrentPurgeObjectIndex < GUnreachableObjects.Num()))
 			{
 				ProcessedObjectsCount = 0;
 				if ((FPlatformTime::Seconds() - StartTime) > TimeLimit)
@@ -329,7 +329,7 @@ class FAsyncPurge : public FRunnable
 			Object->~UObject();
 			GUObjectAllocator.FreeUObject(Object);
 
-			if (bUseTimeLimit && (ProcessedObjectsCount == TimeLimitEnforcementGranularityForDeletion))
+			if (bUseTimeLimit && (ProcessedObjectsCount == TimeLimitEnforcementGranularityForDeletion) && !GameThreadObjects.IsEmpty())
 			{
 				ProcessedObjectsCount = 0;
 				if ((FPlatformTime::Seconds() - StartTime) > TimeLimit)
