Fix possible null dereference.

diff --git a/Engine/Source/Runtime/CoreUObject/Private/UObject/ScriptCore.cpp b/Engine/Source/Runtime/CoreUObject/Private/UObject/ScriptCore.cpp
--- a/Engine/Source/Runtime/CoreUObject/Private/UObject/ScriptCore.cpp
+++ b/Engine/Source/Runtime/CoreUObject/Private/UObject/ScriptCore.cpp
@@ -2001,7 +2001,7 @@ DEFINE_FUNCTION(UObject::execClassSparseDataVariable)
 
 	if (VarProperty == nullptr || P_THIS->GetSparseClassDataStruct() == nullptr)
 	{
-		FBlueprintExceptionInfo ExceptionInfo(EBlueprintExceptionType::AccessViolation, FText::Format(LOCTEXT("MissingProperty", "Attempted to access missing sparse property '{0}' {1}, {2}. If this is a packaged/cooked build, are you attempting to use an editor-only property?"), FText::FromString(GetNameSafe(VarProperty)), FText::FromString(GetNameSafe(P_THIS->GetSparseClassDataStruct())), FText::FromString(GetNameSafe((UClass*)VarProperty->GetOuter()))));
+		FBlueprintExceptionInfo ExceptionInfo(EBlueprintExceptionType::AccessViolation, FText::Format(LOCTEXT("MissingProperty", "Attempted to access missing sparse property '{0}' {1}, {2}. If this is a packaged/cooked build, are you attempting to use an editor-only property?"), FText::FromString(GetNameSafe(VarProperty)), FText::FromString(GetNameSafe(P_THIS->GetSparseClassDataStruct())), FText::FromString(GetNameSafe(VarProperty ? (UClass*)VarProperty->GetOuter() : nullptr))));
 		FBlueprintCoreDelegates::ThrowScriptException(P_THIS, Stack, ExceptionInfo);
 
 		Stack.MostRecentPropertyAddress = nullptr;
