PR #6238: Fix add/remove delegate comments (Contributed by [REDACTED_PERSON])
#rnx

diff --git a/Engine/Source/Runtime/CoreUObject/Private/UObject/PropertyMulticastDelegate.cpp b/Engine/Source/Runtime/CoreUObject/Private/UObject/PropertyMulticastDelegate.cpp
--- a/Engine/Source/Runtime/CoreUObject/Private/UObject/PropertyMulticastDelegate.cpp
+++ b/Engine/Source/Runtime/CoreUObject/Private/UObject/PropertyMulticastDelegate.cpp
@@ -284,7 +284,7 @@ const TCHAR* UMulticastDelegateProperty::ImportText_Remove( const TCHAR* Buffer,
 		return NULL;
 	}
 
-	// Remove this delegate to our multicast delegate's invocation list
+	// Remove this delegate from our multicast delegate's invocation list
 	RemoveDelegate(ImportedDelegate, Parent, PropertyValue);
 
 	SkipWhitespace(Buffer);
@@ -390,7 +390,7 @@ void UMulticastInlineDelegateProperty::AddDelegate(FScriptDelegate ScriptDelegat
 
 	FMulticastScriptDelegate& MulticastDelegate = (*(FMulticastScriptDelegate*)PropertyValue);
 
-	// Remove this delegate to our multicast delegate's invocation list
+	// Add this delegate to our multicast delegate's invocation list
 	MulticastDelegate.AddUnique(MoveTemp(ScriptDelegate));
 }
 
@@ -400,7 +400,7 @@ void UMulticastInlineDelegateProperty::RemoveDelegate(const FScriptDelegate& Scr
 
 	FMulticastScriptDelegate& MulticastDelegate = (*(FMulticastScriptDelegate*)PropertyValue);
 
-	// Remove this delegate to our multicast delegate's invocation list
+	// Remove this delegate from our multicast delegate's invocation list
 	MulticastDelegate.Remove(ScriptDelegate);
 }
 
