TStrongObjectPtr can now be constructed with nullptr.
TStrongObjectPtr can now be constructed with TStrongObjectPtrs rvalues of other types.
Fix for crash when calling IsValid() on a moved-from TStrongObjectPtr.
Fewer template instantiations.
Unrelated TStrongObjectPtrs can now be used to overload functions, e.g. Func(StrongPtrToPawn) will now do the right thing if Func is overloaded to take TStrongObjectPtr<AActor> and TStrongObjectPtr<UActorComponent>.

diff --git a/Engine/Source/Runtime/CoreUObject/Public/UObject/StrongObjectPtr.h b/Engine/Source/Runtime/CoreUObject/Public/UObject/StrongObjectPtr.h
--- a/Engine/Source/Runtime/CoreUObject/Public/UObject/StrongObjectPtr.h
+++ b/Engine/Source/Runtime/CoreUObject/Public/UObject/StrongObjectPtr.h
@@ -3,7 +3,52 @@
 #pragma once
 
 #include "UObject/GCObject.h"
-#include "Templates/UniquePtr.h"
+#include "Templates/EnableIf.h"
+#include "Templates/PointerIsConvertibleFromTo.h"
+#include "Templates/UniqueObj.h"
+
+namespace UE4StrongObjectPtr_Private
+{
+	class FInternalReferenceCollector : public FGCObject
+	{
+	public:
+		FInternalReferenceCollector(const volatile UObject* InObject = nullptr)
+			: Object(InObject)
+		{
+			check(IsInGameThread());
+		}
+
+		virtual ~FInternalReferenceCollector()
+		{
+			check(IsInGameThread());
+		}
+
+		bool IsValid() const
+		{
+			return Object != nullptr;
+		}
+
+		template <typename UObjectType>
+		FORCEINLINE UObjectType* GetAs() const
+		{
+			return (UObjectType*)Object;
+		}
+
+		FORCEINLINE void Set(const volatile UObject* InObject)
+		{
+			Object = InObject;
+		}
+
+		//~ FGCObject interface
+		virtual void AddReferencedObjects(FReferenceCollector& Collector) override
+		{
+			Collector.AddReferencedObject(Object);
+		}
+
+	private:
+		const volatile UObject* Object;
+	};
+}
 
 /**
  * Specific implementation of FGCObject that prevents a single UObject-based pointer from being GC'd while this guard is in scope.
@@ -13,66 +58,48 @@ template <typename ObjectType>
 class TStrongObjectPtr
 {
 public:
-	FORCEINLINE_DEBUGGABLE TStrongObjectPtr()
-		: ReferenceCollector(MakeUnique<FInternalReferenceCollector>(nullptr))
+	TStrongObjectPtr(TStrongObjectPtr&& InOther) = default;
+	TStrongObjectPtr(const TStrongObjectPtr& InOther) = default;
+	TStrongObjectPtr& operator=(TStrongObjectPtr&& InOther) = default;
+	~TStrongObjectPtr() = default;
+
+	FORCEINLINE_DEBUGGABLE TStrongObjectPtr(TYPE_OF_NULLPTR = nullptr)
 	{
 		static_assert(TPointerIsConvertibleFromTo<ObjectType, const volatile UObject>::Value, "TStrongObjectPtr can only be constructed with UObject types");
 	}
 
 	FORCEINLINE_DEBUGGABLE explicit TStrongObjectPtr(ObjectType* InObject)
-		: ReferenceCollector(MakeUnique<FInternalReferenceCollector>(InObject))
+		: ReferenceCollector(InObject)
 	{
 		static_assert(TPointerIsConvertibleFromTo<ObjectType, const volatile UObject>::Value, "TStrongObjectPtr can only be constructed with UObject types");
 	}
 
-	FORCEINLINE_DEBUGGABLE TStrongObjectPtr(const TStrongObjectPtr& InOther)
-		: ReferenceCollector(MakeUnique<FInternalReferenceCollector>(InOther.Get()))
-	{
-	}
-
-	template <typename OtherObjectType>
+	template <
+		typename OtherObjectType,
+		typename = typename TEnableIf<TPointerIsConvertibleFromTo<OtherObjectType, ObjectType>::Value>::Type
+	>
 	FORCEINLINE_DEBUGGABLE TStrongObjectPtr(const TStrongObjectPtr<OtherObjectType>& InOther)
-		: ReferenceCollector(MakeUnique<FInternalReferenceCollector>(InOther.Get()))
+		: ReferenceCollector(InOther.Get())
 	{
 	}
 
 	FORCEINLINE_DEBUGGABLE TStrongObjectPtr& operator=(const TStrongObjectPtr& InOther)
 	{
+		// TUniqueObj is not assignable so we need to implement this instead of defaulting it.
 		ReferenceCollector->Set(InOther.Get());
 		return *this;
 	}
 
 	template <typename OtherObjectType>
-	FORCEINLINE_DEBUGGABLE TStrongObjectPtr& operator=(const TStrongObjectPtr<OtherObjectType>& InOther)
+	FORCEINLINE_DEBUGGABLE typename TEnableIf<
+		TPointerIsConvertibleFromTo<OtherObjectType, ObjectType>::Value,
+		TStrongObjectPtr&
+	>::Type operator=(const TStrongObjectPtr<OtherObjectType>& InOther)
 	{
 		ReferenceCollector->Set(InOther.Get());
 		return *this;
 	}
 
-	FORCEINLINE_DEBUGGABLE TStrongObjectPtr(TStrongObjectPtr&& InOther)
-	{
-		ReferenceCollector = MoveTemp(InOther.ReferenceCollector);
-	}
-
-	template <typename OtherObjectType>
-	FORCEINLINE_DEBUGGABLE TStrongObjectPtr(TStrongObjectPtr<OtherObjectType>&& InOther)
-	{
-		ReferenceCollector = MoveTemp(InOther.ReferenceCollector);
-	}
-
-	FORCEINLINE_DEBUGGABLE TStrongObjectPtr& operator=(TStrongObjectPtr&& InOther)
-	{
-		ReferenceCollector = MoveTemp(InOther.ReferenceCollector);
-		return *this;
-	}
-
-	template <typename OtherObjectType>
-	FORCEINLINE_DEBUGGABLE TStrongObjectPtr<OtherObjectType>& operator=(TStrongObjectPtr<OtherObjectType>&& InOther)
-	{
-		ReferenceCollector = MoveTemp(InOther.ReferenceCollector);
-		return *this;
-	}
-
 	FORCEINLINE_DEBUGGABLE ObjectType& operator*() const
 	{
 		check(IsValid());
@@ -87,17 +114,17 @@ public:
 
 	FORCEINLINE_DEBUGGABLE bool IsValid() const
 	{
-		return Get() != nullptr;
+		return ReferenceCollector->IsValid();
 	}
 
 	FORCEINLINE_DEBUGGABLE explicit operator bool() const
 	{
-		return Get() != nullptr;
+		return ReferenceCollector->IsValid();
 	}
 
 	FORCEINLINE_DEBUGGABLE ObjectType* Get() const
 	{
-		return ReferenceCollector->Get();
+		return ReferenceCollector->GetAs<ObjectType>();
 	}
 
 	FORCEINLINE_DEBUGGABLE void Reset(ObjectType* InNewObject = nullptr)
@@ -111,41 +138,7 @@ public:
 	}
 
 private:
-	class FInternalReferenceCollector : public FGCObject
-	{
-	public:
-		FInternalReferenceCollector(ObjectType* InObject)
-			: Object(InObject)
-		{
-			check(IsInGameThread());
-		}
-
-		virtual ~FInternalReferenceCollector()
-		{
-			check(IsInGameThread());
-		}
-
-		FORCEINLINE ObjectType* Get() const
-		{
-			return Object;
-		}
-
-		FORCEINLINE void Set(ObjectType* InObject)
-		{
-			Object = InObject;
-		}
-
-		//~ FGCObject interface
-		virtual void AddReferencedObjects(FReferenceCollector& Collector) override
-		{
-			Collector.AddReferencedObject(Object);
-		}
-
-	private:
-		ObjectType* Object;
-	};
-
-	TUniquePtr<FInternalReferenceCollector> ReferenceCollector;
+	TUniqueObj<UE4StrongObjectPtr_Private::FInternalReferenceCollector> ReferenceCollector;
 };
 
 template <typename LHSObjectType, typename RHSObjectType>
