Fixing crashes when running DDC or GenerateDistillFileSets commandlets.

diff --git a/Engine/Source/Editor/UnrealEd/Private/Commandlets/DerivedDataCacheCommandlet.cpp b/Engine/Source/Editor/UnrealEd/Private/Commandlets/DerivedDataCacheCommandlet.cpp
--- a/Engine/Source/Editor/UnrealEd/Private/Commandlets/DerivedDataCacheCommandlet.cpp
+++ b/Engine/Source/Editor/UnrealEd/Private/Commandlets/DerivedDataCacheCommandlet.cpp
@@ -244,12 +244,12 @@ int32 UDerivedDataCacheCommandlet::Main( const FString& Params )
 			// Keep track of which packages have already been processed along with the map.
 			{
 				const double FindProcessedPackagesStartTime = FPlatformTime::Seconds();
-				TArray<UObject *> ObjectsInOuter;
-				GetObjectsWithOuter(NULL, ObjectsInOuter, false);
-				for (int32 Index = 0; Index < ObjectsInOuter.Num(); Index++)
+				TArray<UObject *> AllPackages;
+				GetObjectsOfClass(UPackage::StaticClass(), AllPackages);
+				for (int32 Index = 0; Index < AllPackages.Num(); Index++)
 				{
-					UPackage* Pkg = Cast<UPackage>(ObjectsInOuter[Index]);
-					if (!Pkg)
+					UPackage* Pkg = Cast<UPackage>(AllPackages[Index]);
+					if (!Pkg || Pkg->GetOuter())
 					{
 						continue;
 					}
diff --git a/Engine/Source/Editor/UnrealEd/Private/Commandlets/GenerateDistillFileSetsCommandlet.cpp b/Engine/Source/Editor/UnrealEd/Private/Commandlets/GenerateDistillFileSetsCommandlet.cpp
--- a/Engine/Source/Editor/UnrealEd/Private/Commandlets/GenerateDistillFileSetsCommandlet.cpp
+++ b/Engine/Source/Editor/UnrealEd/Private/Commandlets/GenerateDistillFileSetsCommandlet.cpp
@@ -261,11 +261,11 @@ int32 UGenerateDistillFileSetsCommandlet::Main( const FString& InParams )
 				AllPackageNames.Add(Package->GetName());
 
 				UE_LOG(LogGenerateDistillFileSetsCommandlet, Display, TEXT( "Finding content referenced by %s..." ), *MapPackage );
-				TArray<UObject *> ObjectsInOuter;
-				GetObjectsWithOuter(NULL, ObjectsInOuter, false);
-				for (int32 Index = 0; Index < ObjectsInOuter.Num(); Index++)
+				TArray<UObject *> AllPackages;
+				GetObjectsOfClass(UPackage::StaticClass(), AllPackages);
+				for (int32 Index = 0; Index < AllPackages.Num(); Index++)
 				{
-					FString OtherName = ObjectsInOuter[Index]->GetOutermost()->GetName();
+					FString OtherName = AllPackages[Index]->GetOutermost()->GetName();
 					if (!AllPackageNames.Contains(OtherName))
 					{
 						AllPackageNames.Add(OtherName);
diff --git a/Engine/Source/Runtime/CoreUObject/Private/UObject/UObjectHash.cpp b/Engine/Source/Runtime/CoreUObject/Private/UObject/UObjectHash.cpp
--- a/Engine/Source/Runtime/CoreUObject/Private/UObject/UObjectHash.cpp
+++ b/Engine/Source/Runtime/CoreUObject/Private/UObject/UObjectHash.cpp
@@ -711,7 +711,7 @@ static FAutoConsoleCommand ShrinkUObjectHashTablesCmd(
 
 void GetObjectsWithOuter(const class UObjectBase* Outer, TArray<UObject *>& Results, bool bIncludeNestedObjects, EObjectFlags ExclusionFlags, EInternalObjectFlags ExclusionInternalFlags)
 {
-	check(Outer != nullptr);
+	checkf(Outer != nullptr, TEXT("Getting objects with a null outer is no longer supported. If you want to get all packages you might consider using GetObjectsOfClass instead."));
 
 	// We don't want to return any objects that are currently being background loaded unless we're using the object iterator during async loading.
 	ExclusionInternalFlags |= EInternalObjectFlags::Unreachable;
@@ -761,7 +761,7 @@ void GetObjectsWithOuter(const class UObjectBase* Outer, TArray<UObject *>& Resu
 
 void ForEachObjectWithOuter(const class UObjectBase* Outer, TFunctionRef<void(UObject*)> Operation, bool bIncludeNestedObjects, EObjectFlags ExclusionFlags, EInternalObjectFlags ExclusionInternalFlags)
 {
-	check(Outer != nullptr);
+	checkf(Outer != nullptr, TEXT("Getting objects with a null outer is no longer supported. If you want to get all packages you might consider using GetObjectsOfClass instead."));
 
 	// We don't want to return any objects that are currently being background loaded unless we're using the object iterator during async loading.
 	ExclusionInternalFlags |= EInternalObjectFlags::Unreachable;
