TBA: Detect static property arrays during saving and write through a structured stream

diff --git a/Engine/Source/Runtime/CoreUObject/Private/UObject/Class.cpp b/Engine/Source/Runtime/CoreUObject/Private/UObject/Class.cpp
--- a/Engine/Source/Runtime/CoreUObject/Private/UObject/Class.cpp
+++ b/Engine/Source/Runtime/CoreUObject/Private/UObject/Class.cpp
@@ -1166,6 +1166,13 @@ void UStruct::SerializeTaggedProperties(FStructuredArchive::FSlot Slot, uint8* D
 			{
 				const int32 LoopMin = CustomPropertyNode ? CustomPropertyNode->ArrayIndex : 0;
 				const int32 LoopMax = CustomPropertyNode ? LoopMin + 1 : Property->ArrayDim;
+
+				TOptional<FStructuredArchive::FStream> ArrayPropertyRoot;
+				if (UnderlyingArchive.IsTextFormat() && (LoopMax - LoopMin > 1))
+				{
+					ArrayPropertyRoot.Emplace(PropertiesRecord.EnterStream(FIELD_NAME(*Property->GetName())));
+				}
+
 				for( int32 Idx = LoopMin; Idx < LoopMax; Idx++ )
 				{
 					uint8* DataPtr      = Property->ContainerPtrToValuePtr           <uint8>(Data, Idx);
@@ -1189,7 +1196,8 @@ void UStruct::SerializeTaggedProperties(FStructuredArchive::FSlot Slot, uint8* D
 							Tag.SetPropertyGuid(PropertyGuid);
 						}
 
-						FStructuredArchive::FRecord PropertyRecord = PropertiesRecord.EnterRecord(FIELD_NAME(*Tag.Name.ToString()));
+						FStructuredArchive::FSlot PropertySlot = ArrayPropertyRoot.IsSet() ? ArrayPropertyRoot->EnterElement() : PropertiesRecord.EnterField(FIELD_NAME(*Tag.Name.ToString()));
+						FStructuredArchive::FRecord PropertyRecord = PropertySlot.EnterRecord();
 
 						PropertyRecord << NAMED_FIELD(Tag);
 
