Copy up of Unreal Insights from //Tasks/UE4/Dev-Core-Trace.

diff --git a/Engine/Source/Runtime/Core/Private/ProfilingDebugging/PlatformFileTrace.cpp b/Engine/Source/Runtime/Core/Private/ProfilingDebugging/PlatformFileTrace.cpp
new file mode 100644
--- /dev/null
+++ b/Engine/Source/Runtime/Core/Private/ProfilingDebugging/PlatformFileTrace.cpp
@@ -0,0 +1,141 @@
+// Copyright 1998-2019 Epic Games, Inc. All Rights Reserved.
+#include "ProfilingDebugging/PlatformFileTrace.h"
+
+#if PLATFORMFILETRACE_ENABLED
+
+#include "Trace/Trace.h"
+#include "Misc/CString.h"
+#include "HAL/PlatformTime.h"
+#include "HAL/PlatformTLS.h"
+
+UE_TRACE_EVENT_BEGIN(PlatformFile, BeginOpen)
+	UE_TRACE_EVENT_FIELD(uint64, Cycle)
+	UE_TRACE_EVENT_FIELD(uint64, TempHandle)
+	UE_TRACE_EVENT_FIELD(uint32, ThreadId)
+UE_TRACE_EVENT_END()
+
+UE_TRACE_EVENT_BEGIN(PlatformFile, EndOpen)
+	UE_TRACE_EVENT_FIELD(uint64, Cycle)
+	UE_TRACE_EVENT_FIELD(uint64, TempHandle)
+	UE_TRACE_EVENT_FIELD(uint64, FileHandle)
+UE_TRACE_EVENT_END()
+
+UE_TRACE_EVENT_BEGIN(PlatformFile, BeginClose)
+	UE_TRACE_EVENT_FIELD(uint64, Cycle)
+	UE_TRACE_EVENT_FIELD(uint64, TempHandle)
+	UE_TRACE_EVENT_FIELD(uint64, FileHandle)
+	UE_TRACE_EVENT_FIELD(uint32, ThreadId)
+UE_TRACE_EVENT_END()
+
+UE_TRACE_EVENT_BEGIN(PlatformFile, EndClose)
+	UE_TRACE_EVENT_FIELD(uint64, Cycle)
+	UE_TRACE_EVENT_FIELD(uint64, TempHandle)
+UE_TRACE_EVENT_END()
+
+UE_TRACE_EVENT_BEGIN(PlatformFile, BeginRead)
+	UE_TRACE_EVENT_FIELD(uint64, Cycle)
+	UE_TRACE_EVENT_FIELD(uint64, ReadHandle)
+	UE_TRACE_EVENT_FIELD(uint64, FileHandle)
+	UE_TRACE_EVENT_FIELD(uint64, Offset)
+	UE_TRACE_EVENT_FIELD(uint64, Size)
+	UE_TRACE_EVENT_FIELD(uint32, ThreadId)
+UE_TRACE_EVENT_END()
+
+UE_TRACE_EVENT_BEGIN(PlatformFile, EndRead)
+	UE_TRACE_EVENT_FIELD(uint64, Cycle)
+	UE_TRACE_EVENT_FIELD(uint64, ReadHandle)
+	UE_TRACE_EVENT_FIELD(uint64, SizeRead)
+	UE_TRACE_EVENT_FIELD(uint32, ThreadId)
+UE_TRACE_EVENT_END()
+
+UE_TRACE_EVENT_BEGIN(PlatformFile, BeginWrite)
+	UE_TRACE_EVENT_FIELD(uint64, Cycle)
+	UE_TRACE_EVENT_FIELD(uint64, WriteHandle)
+	UE_TRACE_EVENT_FIELD(uint64, FileHandle)
+	UE_TRACE_EVENT_FIELD(uint64, Offset)
+	UE_TRACE_EVENT_FIELD(uint64, Size)
+	UE_TRACE_EVENT_FIELD(uint32, ThreadId)
+UE_TRACE_EVENT_END()
+
+UE_TRACE_EVENT_BEGIN(PlatformFile, EndWrite)
+	UE_TRACE_EVENT_FIELD(uint64, Cycle)
+	UE_TRACE_EVENT_FIELD(uint64, WriteHandle)
+	UE_TRACE_EVENT_FIELD(uint64, SizeWritten)
+	UE_TRACE_EVENT_FIELD(uint32, ThreadId)
+UE_TRACE_EVENT_END()
+
+void FPlatformFileTrace::BeginOpen(uint64 TempHandle, const TCHAR* Path)
+{
+	uint16 PathSize = (FCString::Strlen(Path) + 1) * sizeof(TCHAR);
+	UE_TRACE_LOG(PlatformFile, BeginOpen, PathSize)
+		<< BeginOpen.Cycle(FPlatformTime::Cycles64())
+		<< BeginOpen.TempHandle(TempHandle)
+		<< BeginOpen.Attachment(Path, PathSize)
+		<< BeginOpen.ThreadId(FPlatformTLS::GetCurrentThreadId());
+}
+
+void FPlatformFileTrace::EndOpen(uint64 TempHandle, uint64 FileHandle)
+{
+	UE_TRACE_LOG(PlatformFile, EndOpen)
+		<< EndOpen.Cycle(FPlatformTime::Cycles64())
+		<< EndOpen.TempHandle(TempHandle)
+		<< EndOpen.FileHandle(FileHandle);
+}
+
+void FPlatformFileTrace::BeginClose(uint64 TempHandle, uint64 FileHandle)
+{
+	UE_TRACE_LOG(PlatformFile, BeginClose)
+		<< BeginClose.Cycle(FPlatformTime::Cycles64())
+		<< BeginClose.TempHandle(TempHandle)
+		<< BeginClose.FileHandle(FileHandle)
+		<< BeginClose.ThreadId(FPlatformTLS::GetCurrentThreadId());
+}
+
+void FPlatformFileTrace::EndClose(uint64 TempHandle)
+{
+	UE_TRACE_LOG(PlatformFile, EndClose)
+		<< EndClose.Cycle(FPlatformTime::Cycles64())
+		<< EndClose.TempHandle(TempHandle);
+}
+
+void FPlatformFileTrace::BeginRead(uint64 ReadHandle, uint64 FileHandle, uint64 Offset, uint64 Size)
+{
+	UE_TRACE_LOG(PlatformFile, BeginRead)
+		<< BeginRead.Cycle(FPlatformTime::Cycles64())
+		<< BeginRead.ReadHandle(ReadHandle)
+		<< BeginRead.FileHandle(FileHandle)
+		<< BeginRead.Offset(Offset)
+		<< BeginRead.Size(Size)
+		<< BeginRead.ThreadId(FPlatformTLS::GetCurrentThreadId());
+}
+
+void FPlatformFileTrace::EndRead(uint64 ReadHandle, uint64 SizeRead)
+{
+	UE_TRACE_LOG(PlatformFile, EndRead)
+		<< EndRead.Cycle(FPlatformTime::Cycles64())
+		<< EndRead.ReadHandle(ReadHandle)
+		<< EndRead.SizeRead(SizeRead)
+		<< EndRead.ThreadId(FPlatformTLS::GetCurrentThreadId());
+}
+
+void FPlatformFileTrace::BeginWrite(uint64 WriteHandle, uint64 FileHandle, uint64 Offset, uint64 Size)
+{
+	UE_TRACE_LOG(PlatformFile, BeginWrite)
+		<< BeginWrite.Cycle(FPlatformTime::Cycles64())
+		<< BeginWrite.WriteHandle(WriteHandle)
+		<< BeginWrite.FileHandle(FileHandle)
+		<< BeginWrite.Offset(Offset)
+		<< BeginWrite.Size(Size)
+		<< BeginWrite.ThreadId(FPlatformTLS::GetCurrentThreadId());
+}
+
+void FPlatformFileTrace::EndWrite(uint64 WriteHandle, uint64 SizeWritten)
+{
+	UE_TRACE_LOG(PlatformFile, EndWrite)
+		<< EndWrite.Cycle(FPlatformTime::Cycles64())
+		<< EndWrite.WriteHandle(WriteHandle)
+		<< EndWrite.SizeWritten(SizeWritten)
+		<< EndWrite.ThreadId(FPlatformTLS::GetCurrentThreadId());
+}
+
+#endif
\ No newline at end of file
diff --git a/Engine/Source/Runtime/Core/Private/Stats/Stats2.cpp b/Engine/Source/Runtime/Core/Private/Stats/Stats2.cpp
--- a/Engine/Source/Runtime/Core/Private/Stats/Stats2.cpp
+++ b/Engine/Source/Runtime/Core/Private/Stats/Stats2.cpp
@@ -602,6 +602,21 @@ public:
 		{
 			Found->AlwaysEnabledNamesInThisGroup.Add( Stat, Result );
 		}
+
+#if CPUPROFILERTRACE_ENABLED
+		if (bCycleStat)
+		{
+			Result->TraceCpuProfilerSpecId = FCpuProfilerEventSpec::AssignId(*StatDescription, CpuProfilerGroup_Stats);
+		}
+#endif
+#if STATSTRACE_ENABLED
+		if (!bCycleStat && (InStatType == EStatDataType::ST_int64 || InStatType == EStatDataType::ST_double))
+		{
+			ANSICHAR NameBuffer[1024];
+			StatShortName.GetPlainANSIString(NameBuffer);
+			FStatsTrace::DeclareStat(Stat, NameBuffer, *StatDescription, InStatType == EStatDataType::ST_double, MemoryRegion != FPlatformMemory::MCR_Invalid, bShouldClearEveryFrame);
+		}
+#endif
 		return TStatId(Result);
 	}
 
diff --git a/Engine/Source/Runtime/Core/Public/HAL/PlatformMisc.h b/Engine/Source/Runtime/Core/Public/HAL/PlatformMisc.h
--- a/Engine/Source/Runtime/Core/Public/HAL/PlatformMisc.h
+++ b/Engine/Source/Runtime/Core/Public/HAL/PlatformMisc.h
@@ -51,6 +51,8 @@
 
 #if ENABLE_NAMED_EVENTS
 
+#include "ProfilingDebugging/CpuProfilerTrace.h"
+
 class CORE_API FScopedNamedEvent
 {
 public:
@@ -131,10 +133,23 @@ public:
 #define NAMED_EVENT_STR(x) L##x
 #endif
 
-#define SCOPED_NAMED_EVENT(Name, Color)          FScopedNamedEventStatic ANONYMOUS_VARIABLE(NamedEvent_##Name##_)(Color, NAMED_EVENT_STR(#Name));
-#define SCOPED_NAMED_EVENT_FSTRING(Text, Color)  FScopedNamedEvent       ANONYMOUS_VARIABLE(NamedEvent_)         (Color, *Text);
-#define SCOPED_NAMED_EVENT_TEXT(Text, Color)     FScopedNamedEventStatic ANONYMOUS_VARIABLE(NamedEvent_)         (Color, NAMED_EVENT_STR(Text));
-#define SCOPED_NAMED_EVENT_F(Format, Color, ...) FScopedNamedEvent       ANONYMOUS_VARIABLE(NamedEvent_)         (Color, *FString::Printf(Format, __VA_ARGS__));
+#define SCOPED_NAMED_EVENT(Name, Color) \
+	TRACE_CPUPROFILER_EVENT_SCOPE(Name) \
+	FScopedNamedEventStatic ANONYMOUS_VARIABLE(NamedEvent_##Name##_)(Color, NAMED_EVENT_STR(#Name));
+
+#define SCOPED_NAMED_EVENT_FSTRING(Text, Color) \
+	TRACE_CPUPROFILER_DYNAMIC_EVENT_SCOPE(*Text) \
+	FScopedNamedEvent       ANONYMOUS_VARIABLE(NamedEvent_)         (Color, *Text);
+
+#define SCOPED_NAMED_EVENT_TEXT(Text, Color) \
+	TRACE_CPUPROFILER_EVENT_SCOPE_TEXT(L##Text) \
+	FScopedNamedEventStatic ANONYMOUS_VARIABLE(NamedEvent_)         (Color, NAMED_EVENT_STR(Text));
+
+#define SCOPED_NAMED_EVENT_F(Format, Color, ...) \
+	FString __NamedEventFormattedName(FString::Printf(Format, __VA_ARGS__)); \
+	TRACE_CPUPROFILER_DYNAMIC_EVENT_SCOPE(*__NamedEventFormattedName) \
+	FScopedNamedEvent       ANONYMOUS_VARIABLE(NamedEvent_)         (Color, *__NamedEventFormattedName);
+
 #define SCOPED_PROFILER_COLOR(Color)			 FScopedProfilerColor    ANONYMOUS_VARIABLE(ProfilerColor_##Name##_)(Color);
 
 
diff --git a/Engine/Source/Runtime/Core/Public/Logging/LogTrace.h b/Engine/Source/Runtime/Core/Public/Logging/LogTrace.h
new file mode 100644
--- /dev/null
+++ b/Engine/Source/Runtime/Core/Public/Logging/LogTrace.h
@@ -0,0 +1,53 @@
+// Copyright 1998-2019 Epic Games, Inc. All Rights Reserved.
+
+#pragma once
+
+#include "CoreTypes.h"
+
+#if !IS_PROGRAM && !UE_BUILD_SHIPPING
+#define LOGTRACE_ENABLED 0
+#else
+#define LOGTRACE_ENABLED 0
+#endif
+
+#if LOGTRACE_ENABLED
+#include "ProfilingDebugging/FormatArgsTrace.h"
+
+struct FLogCategoryBase;
+
+struct FLogTrace
+{
+	CORE_API static void OutputLogCategory(const FLogCategoryBase* Category, const TCHAR* Name, ELogVerbosity::Type DefaultVerbosity);
+	CORE_API static void OutputLogMessageSpec(const void* LogPoint, const FLogCategoryBase* Category, ELogVerbosity::Type Verbosity, const ANSICHAR* File, int32 Line, const TCHAR* Format);
+
+	template <typename... Types>
+	static void OutputLogMessage(const void* LogPoint, Types... FormatArgs)
+	{
+		uint8 FormatArgsBuffer[4096];
+		uint16 FormatArgsSize = FFormatArgsTrace::EncodeArguments(FormatArgsBuffer, FormatArgs...);
+		if (FormatArgsSize)
+		{
+			OutputLogMessageInternal(LogPoint, FormatArgsSize, FormatArgsBuffer);
+		}
+	}
+
+private:
+	CORE_API static void OutputLogMessageInternal(const void* LogPoint, uint16 EncodedFormatArgsSize, uint8* EncodedFormatArgs);
+};
+
+#define TRACE_LOG_CATEGORY(Category, Name, DefaultVerbosity) \
+	FLogTrace::OutputLogCategory(Category, Name, DefaultVerbosity);
+
+#define TRACE_LOG_MESSAGE(Category, Verbosity, Format, ...) \
+	static bool PREPROCESSOR_JOIN(__LogPoint, __LINE__); \
+	if (!PREPROCESSOR_JOIN(__LogPoint, __LINE__)) \
+	{ \
+		FLogTrace::OutputLogMessageSpec(&PREPROCESSOR_JOIN(__LogPoint, __LINE__), &Category, ELogVerbosity::Verbosity, __FILE__, __LINE__, Format); \
+		PREPROCESSOR_JOIN(__LogPoint, __LINE__) = true; \
+	} \
+	FLogTrace::OutputLogMessage(&PREPROCESSOR_JOIN(__LogPoint, __LINE__), ##__VA_ARGS__);
+
+#else
+#define TRACE_LOG_CATEGORY(Category, Name, DefaultVerbosity)
+#define TRACE_LOG_MESSAGE(Category, Verbosity, Format, ...)
+#endif
\ No newline at end of file
diff --git a/Engine/Source/Runtime/Core/Public/ProfilingDebugging/CpuProfilerTrace.h b/Engine/Source/Runtime/Core/Public/ProfilingDebugging/CpuProfilerTrace.h
new file mode 100644
--- /dev/null
+++ b/Engine/Source/Runtime/Core/Public/ProfilingDebugging/CpuProfilerTrace.h
@@ -0,0 +1,93 @@
+// Copyright 1998-2019 Epic Games, Inc. All Rights Reserved.
+
+#pragma once
+
+#include "CoreTypes.h"
+
+#if !IS_PROGRAM && !UE_BUILD_SHIPPING
+#define CPUPROFILERTRACE_ENABLED 0
+#else
+#define CPUPROFILERTRACE_ENABLED 0
+#endif
+
+enum ECpuProfilerGroup
+{
+	CpuProfilerGroup_Default,
+	CpuProfilerGroup_Stats,
+	CpuProfilerGroup_LoadTime,
+	CpuProfilerGroup_CsvProfiler,
+};
+
+#if CPUPROFILERTRACE_ENABLED
+
+struct FCpuProfilerTrace
+{
+	CORE_API static void OutputBeginEvent(uint16 SpecId);
+	CORE_API static void OutputBeginDynamicEvent(const TCHAR* Name);
+	CORE_API static void OutputEndEvent();
+	CORE_API static uint64 GetCurrentCycle();
+};
+
+struct FCpuProfilerEventSpec
+{
+	CORE_API static uint16 AssignId(const TCHAR* Name, ECpuProfilerGroup Group);
+};
+
+struct FCpuProfilerEventScope
+{
+	FCpuProfilerEventScope(uint16 InSpecId)
+	{
+		FCpuProfilerTrace::OutputBeginEvent(InSpecId);
+	}
+
+	~FCpuProfilerEventScope()
+	{
+		FCpuProfilerTrace::OutputEndEvent();
+	}
+};
+
+struct FCpuProfilerDynamicEventScope
+{
+	FCpuProfilerDynamicEventScope(const TCHAR* InName)
+	{
+		FCpuProfilerTrace::OutputBeginDynamicEvent(InName);
+	}
+
+	~FCpuProfilerDynamicEventScope()
+	{
+		FCpuProfilerTrace::OutputEndEvent();
+	}
+};
+
+#define TRACE_CPUPROFILER_EVENT_SCOPE_GROUP(Name, Group) \
+	static uint16 PREPROCESSOR_JOIN(__CpuProfilerEventSpecId, __LINE__); \
+	if (PREPROCESSOR_JOIN(__CpuProfilerEventSpecId, __LINE__) == 0) { \
+		PREPROCESSOR_JOIN(__CpuProfilerEventSpecId, __LINE__) = FCpuProfilerEventSpec::AssignId(L###Name, Group); \
+	} \
+	FCpuProfilerEventScope PREPROCESSOR_JOIN(__CpuProfilerEventScope, __LINE__)(PREPROCESSOR_JOIN(__CpuProfilerEventSpecId, __LINE__));
+
+#define TRACE_CPUPROFILER_EVENT_SCOPE(Name) \
+	TRACE_CPUPROFILER_EVENT_SCOPE_GROUP(Name, CpuProfilerGroup_Default);
+
+#define TRACE_CPUPROFILER_EVENT_SCOPE_TEXT_GROUP(Name, Group) \
+	static uint16 PREPROCESSOR_JOIN(__CpuProfilerEventSpecId, __LINE__); \
+	if (PREPROCESSOR_JOIN(__CpuProfilerEventSpecId, __LINE__) == 0) { \
+		PREPROCESSOR_JOIN(__CpuProfilerEventSpecId, __LINE__) = FCpuProfilerEventSpec::AssignId(Name, Group); \
+	} \
+	FCpuProfilerEventScope PREPROCESSOR_JOIN(__CpuProfilerEventScope, __LINE__)(PREPROCESSOR_JOIN(__CpuProfilerEventSpecId, __LINE__));
+
+#define TRACE_CPUPROFILER_EVENT_SCOPE_TEXT(Name) \
+	TRACE_CPUPROFILER_EVENT_SCOPE_TEXT_GROUP(Name, CpuProfilerGroup_Default)
+
+#define TRACE_CPUPROFILER_DYNAMIC_EVENT_SCOPE(Name) \
+	FCpuProfilerDynamicEventScope ANONYMOUS_VARIABLE(__Scope)(Name);
+
+#else
+
+#define TRACE_CPUPROFILER_EVENT_SCOPE_GROUP(Name, Group)
+#define TRACE_CPUPROFILER_EVENT_SCOPE(Name)
+#define TRACE_CPUPROFILER_EVENT_SCOPE_TEXT_GROUP(Name, Group)
+#define TRACE_CPUPROFILER_EVENT_SCOPE_TEXT(Name)
+#define TRACE_CPUPROFILER_DYNAMIC_EVENT_SCOPE(Name)
+
+#endif
diff --git a/Engine/Source/Runtime/Core/Public/ProfilingDebugging/FormatArgsTrace.h b/Engine/Source/Runtime/Core/Public/ProfilingDebugging/FormatArgsTrace.h
new file mode 100644
--- /dev/null
+++ b/Engine/Source/Runtime/Core/Public/ProfilingDebugging/FormatArgsTrace.h
@@ -0,0 +1,112 @@
+// Copyright 1998-2019 Epic Games, Inc. All Rights Reserved.
+
+#pragma once
+
+#include "CoreTypes.h"
+#include "Misc/CString.h"
+#include "Templates/IsFloatingPoint.h"
+#include "Templates/UnrealTypeTraits.h"
+#include "Templates/UnrealTemplate.h"
+
+struct FFormatArgsTrace
+{
+	enum EFormatArgTypeCode
+	{
+		FormatArgTypeCode_CategoryBitShift = 6,
+		FormatArgTypeCode_SizeBitMask = (1 << FormatArgTypeCode_CategoryBitShift) - 1,
+		FormatArgTypeCode_CategoryBitMask = ~FormatArgTypeCode_SizeBitMask,
+		FormatArgTypeCode_CategoryInteger = 1 << FormatArgTypeCode_CategoryBitShift,
+		FormatArgTypeCode_CategoryFloatingPoint = 2 << FormatArgTypeCode_CategoryBitShift,
+		FormatArgTypeCode_CategoryString = 3 << FormatArgTypeCode_CategoryBitShift,
+	};
+
+	template <int BufferSize, typename... Types>
+	static uint16 EncodeArguments(uint8(&Buffer)[BufferSize], Types... FormatArgs)
+	{
+		static_assert(BufferSize < 65536, "Maximum buffer size of 16 bits exceeded");
+		uint64 FormatArgsCount = sizeof...(FormatArgs);
+		if (FormatArgsCount >= 256)
+		{
+			// Nope
+			return 0;
+		}
+		uint64 FormatArgsSize = 1 + FormatArgsCount + GetArgumentsEncodedSize(FormatArgs...);
+		if (FormatArgsSize > BufferSize)
+		{
+			// Nope
+			return 0;
+		}
+		uint8* TypeCodesBufferPtr = Buffer;
+		*TypeCodesBufferPtr++ = uint8(FormatArgsCount);
+		uint8* PayloadBufferPtr = TypeCodesBufferPtr + FormatArgsCount;
+		EncodeArgumentsInternal(TypeCodesBufferPtr, PayloadBufferPtr, FormatArgs...);
+		check(PayloadBufferPtr - Buffer == FormatArgsSize);
+		return FormatArgsSize;
+	}
+
+private:
+	template <typename T>
+	struct TIsStringArgument
+	{
+		enum { Value = TAnd<TIsPointer<T>, TIsCharType<typename TRemoveCV<typename TRemovePointer<T>::Type>::Type>>::Value };
+	};
+
+	template <typename T>
+	constexpr static typename TEnableIf<!TIsStringArgument<T>::Value, uint64>::Type GetArgumentEncodedSize(T Argument)
+	{
+		return sizeof(T);
+	}
+
+	template <typename T, typename CharType = typename TRemoveCV<typename TRemovePointer<T>::Type>::Type>
+	static typename TEnableIf<TIsStringArgument<T>::Value, uint64>::Type GetArgumentEncodedSize(T Argument)
+	{
+		return (TCString<CharType>::Strlen(Argument) + 1) * sizeof(CharType);
+	}
+
+	constexpr static uint64 GetArgumentsEncodedSize()
+	{
+		return 0;
+	}
+
+	template <typename T, typename... Types>
+	static uint64 GetArgumentsEncodedSize(T Head, Types... Tail)
+	{
+		return GetArgumentEncodedSize(Head) + GetArgumentsEncodedSize(Tail...);
+	}
+
+	template <typename T>
+	static typename TEnableIf<TAnd<TNot<TIsFloatingPoint<T>>, TNot<TIsStringArgument<T>>>::Value>::Type EncodeArgumentInternal(uint8*& TypeCodesPtr, uint8*& PayloadPtr, T Argument)
+	{
+		*TypeCodesPtr++ = FormatArgTypeCode_CategoryInteger | sizeof(T);
+		*reinterpret_cast<T*>(PayloadPtr) = Argument;
+		PayloadPtr += sizeof(T);
+	}
+
+	template <typename T>
+	static typename TEnableIf<TIsFloatingPoint<T>::Value>::Type EncodeArgumentInternal(uint8*& TypeCodesPtr, uint8*& PayloadPtr, T Argument)
+	{
+		*TypeCodesPtr++ = FormatArgTypeCode_CategoryFloatingPoint | sizeof(T);
+		*reinterpret_cast<T*>(PayloadPtr) = Argument;
+		PayloadPtr += sizeof(T);
+	}
+
+	template <typename T, typename CharType = typename TRemoveCV<typename TRemovePointer<T>::Type>::Type>
+	static typename TEnableIf<TIsStringArgument<T>::Value>::Type EncodeArgumentInternal(uint8*& TypeCodesPtr, uint8*& PayloadPtr, T Argument)
+	{
+		*TypeCodesPtr++ = FormatArgTypeCode_CategoryString | sizeof(CharType);
+		uint16 Length = (TCString<CharType>::Strlen(Argument) + 1) * sizeof(CharType);
+		memcpy(PayloadPtr, Argument, Length);
+		PayloadPtr += Length;
+	}
+
+	constexpr static void EncodeArgumentsInternal(uint8*& TypeCodesPtr, uint8*& PayloadPtr)
+	{
+	}
+
+	template <typename T, typename... Types>
+	static void EncodeArgumentsInternal(uint8*& ArgDescriptorsPtr, uint8*& ArgPayloadPtr, T Head, Types... Tail)
+	{
+		EncodeArgumentInternal(ArgDescriptorsPtr, ArgPayloadPtr, Head);
+		EncodeArgumentsInternal(ArgDescriptorsPtr, ArgPayloadPtr, Tail...);
+	}
+};
\ No newline at end of file
diff --git a/Engine/Source/Runtime/Core/Public/ProfilingDebugging/MiscTrace.h b/Engine/Source/Runtime/Core/Public/ProfilingDebugging/MiscTrace.h
new file mode 100644
--- /dev/null
+++ b/Engine/Source/Runtime/Core/Public/ProfilingDebugging/MiscTrace.h
@@ -0,0 +1,136 @@
+// Copyright 1998-2019 Epic Games, Inc. All Rights Reserved.
+
+#pragma once
+
+#include "CoreTypes.h"
+
+#if !IS_PROGRAM && !UE_BUILD_SHIPPING
+#define MISCTRACE_ENABLED 1
+#else
+#define MISCTRACE_ENABLED 0
+#endif
+
+enum ETraceThreadGroup
+{
+	TraceThreadGroup_Render,
+	TraceThreadGroup_AsyncLoading,
+	TraceThreadGroup_TaskGraphHigh,
+	TraceThreadGroup_TaskGraphNormal,
+	TraceThreadGroup_TaskGraphLow,
+	TraceThreadGroup_LargeThreadPool,
+	TraceThreadGroup_ThreadPool,
+	TraceThreadGroup_BackgroundThreadPool,
+	TraceThreadGroup_IOThreadPool,
+
+	TraceThreadGroup_Count
+};
+
+enum ETraceFrameType
+{
+	TraceFrameType_Game,
+	TraceFrameType_Rendering,
+
+	TraceFrameType_Count
+};
+
+struct FTraceUtils
+{
+	static void Encode7bit(uint64 Value, uint8*& BufferPtr)
+	{
+		do
+		{
+			uint8 HasMoreBytes = (Value > uint64(0x7F)) << 7;
+			*(BufferPtr++) = (uint8)(Value & 0x7F) | HasMoreBytes;
+			Value >>= 7;
+		} while (Value > 0);
+	}
+
+	static void EncodeZigZag(int64 Value, uint8*& BufferPtr)
+	{
+		Encode7bit((Value << 1) ^ (Value >> 63), BufferPtr);
+	}
+};
+
+#if MISCTRACE_ENABLED
+
+#include "ProfilingDebugging/FormatArgsTrace.h"
+
+class FName;
+
+struct FMiscTrace
+{
+	CORE_API static void OutputRegisterGameThread(uint32 Id);
+	CORE_API static void OutputCreateThread(uint32 Id, const TCHAR* Name, uint32 Priority);
+	CORE_API static void OutputSetThreadGroup(uint32 Id, ETraceThreadGroup Group);
+	CORE_API static void OutputBeginThreadGroupScope(ETraceThreadGroup Group);
+	CORE_API static void OutputEndThreadGroupScope();
+	CORE_API static void OutputBookmarkSpec(const void* BookmarkPoint, const ANSICHAR* File, int32 Line, const TCHAR* Format);
+	template <typename... Types>
+	static void OutputBookmark(const void* BookmarkPoint, Types... FormatArgs)
+	{
+		uint8 FormatArgsBuffer[4096];
+		uint16 FormatArgsSize = FFormatArgsTrace::EncodeArguments(FormatArgsBuffer, FormatArgs...);
+		if (FormatArgsSize)
+		{
+			OutputBookmarkInternal(BookmarkPoint, FormatArgsSize, FormatArgsBuffer);
+		}
+	}
+
+	CORE_API static void OutputBeginFrame(ETraceFrameType FrameType);
+	CORE_API static void OutputEndFrame(ETraceFrameType FrameType);
+
+	struct FThreadGroupScope
+	{
+		FThreadGroupScope(ETraceThreadGroup Group)
+		{
+			FMiscTrace::OutputBeginThreadGroupScope(Group);
+		}
+
+		~FThreadGroupScope()
+		{
+			FMiscTrace::OutputEndThreadGroupScope();
+		}
+	};
+
+private:
+	CORE_API static void OutputBookmarkInternal(const void* BookmarkPoint, uint16 EncodedFormatArgsSize, uint8* EncodedFormatArgs);
+};
+
+#define TRACE_REGISTER_GAME_THREAD(Id) \
+	FMiscTrace::OutputRegisterGameThread(Id);
+
+#define TRACE_CREATE_THREAD(Id, Name, Priority) \
+	FMiscTrace::OutputCreateThread(Id, Name, Priority);
+
+#define TRACE_SET_THREAD_GROUP(Id, Group) \
+	FMiscTrace::OutputSetThreadGroup(Id, Group);
+
+#define TRACE_THREAD_GROUP_SCOPE(Group) \
+	FMiscTrace::FThreadGroupScope ANONYMOUS_VARIABLE(ThreadGroupScope) (Group);
+
+#define TRACE_BOOKMARK(Format, ...) \
+	static bool PREPROCESSOR_JOIN(__BookmarkPoint, __LINE__); \
+	if (!PREPROCESSOR_JOIN(__BookmarkPoint, __LINE__)) \
+	{ \
+		FMiscTrace::OutputBookmarkSpec(&PREPROCESSOR_JOIN(__BookmarkPoint, __LINE__), __FILE__, __LINE__, Format); \
+		PREPROCESSOR_JOIN(__BookmarkPoint, __LINE__) = true; \
+	} \
+	FMiscTrace::OutputBookmark(&PREPROCESSOR_JOIN(__BookmarkPoint, __LINE__), ##__VA_ARGS__);
+
+#define TRACE_BEGIN_FRAME(FrameType) \
+	FMiscTrace::OutputBeginFrame(FrameType);
+
+#define TRACE_END_FRAME(FrameType) \
+	FMiscTrace::OutputEndFrame(FrameType);
+
+#else
+
+#define TRACE_REGISTER_GAME_THREAD(...)
+#define TRACE_CREATE_THREAD(...)
+#define TRACE_SET_THREAD_GROUP(...)
+#define TRACE_THREAD_GROUP_SCOPE(...)
+#define TRACE_BOOKMARK(...)
+#define TRACE_BEGIN_FRAME(...)
+#define TRACE_END_FRAME(...)
+
+#endif
\ No newline at end of file
diff --git a/Engine/Source/Runtime/Core/Public/ProfilingDebugging/PlatformFileTrace.h b/Engine/Source/Runtime/Core/Public/ProfilingDebugging/PlatformFileTrace.h
new file mode 100644
--- /dev/null
+++ b/Engine/Source/Runtime/Core/Public/ProfilingDebugging/PlatformFileTrace.h
@@ -0,0 +1,64 @@
+// Copyright 1998-2019 Epic Games, Inc. All Rights Reserved.
+
+#pragma once
+
+#include "CoreTypes.h"
+
+#if !IS_PROGRAM && !UE_BUILD_SHIPPING
+#define PLATFORMFILETRACE_ENABLED 0
+#else
+#define PLATFORMFILETRACE_ENABLED 0
+#endif
+
+#if PLATFORMFILETRACE_ENABLED
+
+struct FPlatformFileTrace
+{
+	static void BeginOpen(uint64 TempHandle, const TCHAR* Path);
+	static void EndOpen(uint64 TempHandle, uint64 FileHandle);
+	static void BeginClose(uint64 TempHandle, uint64 FileHandle);
+	static void EndClose(uint64 TempHandle);
+	static void BeginRead(uint64 ReadHandle, uint64 FileHandle, uint64 Offset, uint64 Size);
+	static void EndRead(uint64 ReadHandle, uint64 SizeRead);
+	static void BeginWrite(uint64 WriteHandle, uint64 FileHandle, uint64 Offset, uint64 Size);
+	static void EndWrite(uint64 WriteHandle, uint64 SizeWritten);
+};
+
+#define TRACE_PLATFORMFILE_BEGIN_OPEN(Path) \
+	uint64 __TempHandle; \
+	FPlatformFileTrace::BeginOpen(uint64(&__TempHandle), Path);
+
+#define TRACE_PLATFORMFILE_END_OPEN(FileHandle) \
+	FPlatformFileTrace::EndOpen(uint64(&__TempHandle), uint64(FileHandle));
+
+#define TRACE_PLATFORMFILE_BEGIN_CLOSE(FileHandle) \
+	uint64 __TempHandle; \
+	FPlatformFileTrace::BeginClose(uint64(&__TempHandle), uint64(FileHandle));
+
+#define TRACE_PLATFORMFILE_END_CLOSE() \
+	FPlatformFileTrace::EndClose(uint64(&__TempHandle));
+
+#define TRACE_PLATFORMFILE_BEGIN_READ(ReadHandle, FileHandle, Offset, Size) \
+	FPlatformFileTrace::BeginRead(uint64(ReadHandle), uint64(FileHandle), Offset, Size);
+
+#define TRACE_PLATFORMFILE_END_READ(ReadHandle, SizeRead) \
+	FPlatformFileTrace::EndRead(uint64(ReadHandle), SizeRead);
+
+#define TRACE_PLATFORMFILE_BEGIN_WRITE(WriteHandle, FileHandle, Offset, Size) \
+	FPlatformFileTrace::BeginWrite(uint64(WriteHandle), uint64(FileHandle), Offset, Size);
+
+#define TRACE_PLATFORMFILE_END_WRITE(WriteHandle, SizeWritten) \
+	FPlatformFileTrace::EndWrite(uint64(WriteHandle), SizeWritten);
+
+#else
+
+#define TRACE_PLATFORMFILE_BEGIN_OPEN(Path)
+#define TRACE_PLATFORMFILE_END_OPEN(FileHandle)
+#define TRACE_PLATFORMFILE_BEGIN_CLOSE(FileHandle)
+#define TRACE_PLATFORMFILE_END_CLOSE()
+#define TRACE_PLATFORMFILE_BEGIN_READ(ReadHandle, FileHandle, Offset, Size)
+#define TRACE_PLATFORMFILE_END_READ(ReadHandle, SizeRead)
+#define TRACE_PLATFORMFILE_BEGIN_WRITE(WriteHandle, FileHandle, Offset, Size)
+#define TRACE_PLATFORMFILE_END_WRITE(WriteHandle, SizeWritten)
+
+#endif
\ No newline at end of file
diff --git a/Engine/Source/Runtime/Core/Public/Stats/Stats2.h b/Engine/Source/Runtime/Core/Public/Stats/Stats2.h
--- a/Engine/Source/Runtime/Core/Public/Stats/Stats2.h
+++ b/Engine/Source/Runtime/Core/Public/Stats/Stats2.h
@@ -19,6 +19,9 @@
 #include "Math/Color.h"
 #include "StatsCommon.h"
 #include "Templates/UniquePtr.h"
+#include "ProfilingDebugging/CpuProfilerTrace.h"
+#include "ProfilingDebugging/MiscTrace.h"
+#include "StatsTrace.h"
 
 class FScopeCycleCounter;
 class FThreadStats;
@@ -120,6 +123,10 @@ struct TStatIdData
 
 	/** const WIDECHAR* pointer to a string describing the stat */
 	TUniquePtr<ANSICHAR[]> StatDescriptionAnsi;
+
+#if CPUPROFILERTRACE_ENABLED
+	uint32 TraceCpuProfilerSpecId = 0;
+#endif
 };
 
 struct TStatId
@@ -155,6 +162,13 @@ struct TStatId
 		return MinimalNameToName(StatIdPtr->Name);
 	}
 
+#if CPUPROFILERTRACE_ENABLED
+	FORCEINLINE uint16 GetTraceCpuProfilerSpecId() const
+	{
+		return StatIdPtr->TraceCpuProfilerSpecId;
+	}
+#endif
+
 	FORCEINLINE static const TStatIdData& GetStatNone()
 	{
 		return TStatId_NAME_None;
@@ -1545,6 +1559,9 @@ class FCycleCounter
 {
 	/** Name of the stat, usually a short name **/
 	FName StatId;
+#if CPUPROFILERTRACE_ENABLED
+	uint16 TraceCpuProfilerSpecId = 0;
+#endif
 
 public:
 
@@ -1559,6 +1576,13 @@ public:
 		{
 			return;
 		}
+#if CPUPROFILERTRACE_ENABLED
+		if (bAlways || GCycleStatsShouldEmitNamedEvents > 0)
+		{
+			TraceCpuProfilerSpecId = InStatId.GetTraceCpuProfilerSpecId();
+			FCpuProfilerTrace::OutputBeginEvent(TraceCpuProfilerSpecId);
+		}
+#endif
 
 		if( (bAlways && FThreadStats::WillEverCollectData()) || FThreadStats::IsCollectingData() )
 		{
@@ -1583,6 +1607,13 @@ public:
 	 */
 	FORCEINLINE_STATS void Stop()
 	{
+#if CPUPROFILERTRACE_ENABLED
+		if (TraceCpuProfilerSpecId)
+		{
+			FCpuProfilerTrace::OutputEndEvent();
+			TraceCpuProfilerSpecId = 0;
+		}
+#endif
 		if( !StatId.IsNone() )
 		{
 			FThreadStats::AddMessage(StatId, EStatOperation::CycleScopeEnd);
@@ -1946,74 +1977,126 @@ struct FStat_##StatName\
 #define INC_DWORD_STAT(Stat) \
 {\
 	if (FThreadStats::IsCollectingData() || !GET_STATISEVERYFRAME(Stat)) \
+	{ \
 		FThreadStats::AddMessage(GET_STATFNAME(Stat), EStatOperation::Add, int64(1));\
+		TRACE_STAT_INCREMENT(GET_STATFNAME(Stat)); \
+	} \
 }
 #define INC_FLOAT_STAT_BY(Stat, Amount) \
 {\
 	if (Amount != 0.0f) \
+	{ \
 		if (FThreadStats::IsCollectingData() || !GET_STATISEVERYFRAME(Stat)) \
-				FThreadStats::AddMessage(GET_STATFNAME(Stat), EStatOperation::Add, double(Amount));\
+		{ \
+			FThreadStats::AddMessage(GET_STATFNAME(Stat), EStatOperation::Add, double(Amount));\
+			TRACE_STAT_ADD(GET_STATFNAME(Stat), double(Amount)); \
+		} \
+	} \
 }
 #define INC_DWORD_STAT_BY(Stat, Amount) \
 {\
 	if (Amount != 0) \
+	{ \
 		if (FThreadStats::IsCollectingData() || !GET_STATISEVERYFRAME(Stat)) \
+		{ \
 			FThreadStats::AddMessage(GET_STATFNAME(Stat), EStatOperation::Add, int64(Amount));\
+			TRACE_STAT_ADD(GET_STATFNAME(Stat), int64(Amount)); \
+		} \
+	} \
 }
 #define INC_DWORD_STAT_FNAME_BY(StatFName, Amount) \
 {\
 	if (Amount != 0) \
-			FThreadStats::AddMessage(StatFName, EStatOperation::Add, int64(Amount));\
+	{ \
+		FThreadStats::AddMessage(StatFName, EStatOperation::Add, int64(Amount));\
+		TRACE_STAT_ADD(StatFName, int64(Amount)); \
+	} \
 }
 #define INC_MEMORY_STAT_BY(Stat, Amount) \
 {\
 	if (Amount != 0) \
+	{ \
 		if (FThreadStats::IsCollectingData() || !GET_STATISEVERYFRAME(Stat)) \
+		{ \
 			FThreadStats::AddMessage(GET_STATFNAME(Stat), EStatOperation::Add, int64(Amount));\
+			TRACE_STAT_ADD(GET_STATFNAME(Stat), int64(Amount)); \
+		} \
+	} \
 }
 #define DEC_DWORD_STAT(Stat) \
 {\
 	if (FThreadStats::IsCollectingData() || !GET_STATISEVERYFRAME(Stat)) \
+	{ \
 		FThreadStats::AddMessage(GET_STATFNAME(Stat), EStatOperation::Subtract, int64(1));\
+		TRACE_STAT_DECREMENT(GET_STATFNAME(Stat)); \
+	} \
 }
 #define DEC_FLOAT_STAT_BY(Stat,Amount) \
 {\
 	if (Amount != 0.0f) \
+	{ \
 		if (FThreadStats::IsCollectingData() || !GET_STATISEVERYFRAME(Stat)) \
+		{ \
 			FThreadStats::AddMessage(GET_STATFNAME(Stat), EStatOperation::Subtract, double(Amount));\
+			TRACE_STAT_ADD(GET_STATFNAME(Stat), -double(Amount)); \
+		} \
+	} \
 }
 #define DEC_DWORD_STAT_BY(Stat,Amount) \
 {\
 	if (Amount != 0) \
+	{ \
 		if (FThreadStats::IsCollectingData() || !GET_STATISEVERYFRAME(Stat)) \
+		{ \
 			FThreadStats::AddMessage(GET_STATFNAME(Stat), EStatOperation::Subtract, int64(Amount));\
+			TRACE_STAT_ADD(GET_STATFNAME(Stat), -int64(Amount)); \
+		} \
+	} \
 }
 #define DEC_DWORD_STAT_FNAME_BY(StatFName,Amount) \
 {\
 	if (Amount != 0) \
- 			FThreadStats::AddMessage(StatFName, EStatOperation::Subtract, int64(Amount));\
+	{ \
+		FThreadStats::AddMessage(StatFName, EStatOperation::Subtract, int64(Amount));\
+		TRACE_STAT_ADD(StatFName, -int64(Amount)); \
+	} \
 }
 #define DEC_MEMORY_STAT_BY(Stat,Amount) \
 {\
 	if (Amount != 0) \
+	{ \
 		if (FThreadStats::IsCollectingData() || !GET_STATISEVERYFRAME(Stat)) \
+		{ \
 			FThreadStats::AddMessage(GET_STATFNAME(Stat), EStatOperation::Subtract, int64(Amount));\
+			TRACE_STAT_ADD(GET_STATFNAME(Stat), -int64(Amount)); \
+		} \
+	} \
 }
 #define SET_MEMORY_STAT(Stat,Value) \
 {\
 	if (FThreadStats::IsCollectingData() || !GET_STATISEVERYFRAME(Stat)) \
+	{ \
 		FThreadStats::AddMessage(GET_STATFNAME(Stat), EStatOperation::Set, int64(Value));\
+		TRACE_STAT_SET(GET_STATFNAME(Stat), int64(Value)); \
+	} \
 }
 #define SET_DWORD_STAT(Stat,Value) \
 {\
 	if (FThreadStats::IsCollectingData() || !GET_STATISEVERYFRAME(Stat)) \
+	{ \
 		FThreadStats::AddMessage(GET_STATFNAME(Stat), EStatOperation::Set, int64(Value));\
+		TRACE_STAT_SET(GET_STATFNAME(Stat), int64(Value)); \
+	} \
 }
 #define SET_FLOAT_STAT(Stat,Value) \
 {\
 	if (FThreadStats::IsCollectingData() || !GET_STATISEVERYFRAME(Stat)) \
+	{ \
 		FThreadStats::AddMessage(GET_STATFNAME(Stat), EStatOperation::Set, double(Value));\
+		TRACE_STAT_SET(GET_STATFNAME(Stat), double(Value)); \
+	} \
 }
+
 #define STAT_ADD_CUSTOMMESSAGE_NAME(Stat,Value) \
 {\
 	FThreadStats::AddMessage(GET_STATFNAME(Stat), EStatOperation::SpecialMessageMarker, FName(Value));\
@@ -2031,52 +2114,75 @@ struct FStat_##StatName\
 #define INC_DWORD_STAT_FName(Stat) \
 {\
 	FThreadStats::AddMessage(Stat, EStatOperation::Add, int64(1));\
+	TRACE_STAT_INCREMENT(Stat); \
 }
 #define INC_FLOAT_STAT_BY_FName(Stat, Amount) \
 {\
 	if (Amount != 0.0f) \
+	{ \
 		FThreadStats::AddMessage(Stat, EStatOperation::Add, double(Amount));\
+		TRACE_STAT_ADD(Stat, double(Amount)); \
+	} \
 }
 #define INC_DWORD_STAT_BY_FName(Stat, Amount) \
 {\
 	if (Amount != 0) \
+	{ \
 		FThreadStats::AddMessage(Stat, EStatOperation::Add, int64(Amount));\
+		TRACE_STAT_ADD(Stat, int64(Amount)); \
+	} \
 }
 #define INC_MEMORY_STAT_BY_FName(Stat, Amount) \
 {\
 	if (Amount != 0) \
+	{ \
 		FThreadStats::AddMessage(Stat, EStatOperation::Add, int64(Amount));\
+		TRACE_STAT_ADD(Stat, int64(Amount)); \
+	} \
 }
 #define DEC_DWORD_STAT_FName(Stat) \
 {\
 	FThreadStats::AddMessage(Stat, EStatOperation::Subtract, int64(1));\
+	TRACE_STAT_DECREMENT(Stat); \
 }
 #define DEC_FLOAT_STAT_BY_FName(Stat,Amount) \
 {\
 	if (Amount != 0.0f) \
+	{ \
 		FThreadStats::AddMessage(Stat, EStatOperation::Subtract, double(Amount));\
+		TRACE_STAT_ADD(Stat, -double(Amount)); \
+	} \
 }
 #define DEC_DWORD_STAT_BY_FName(Stat,Amount) \
 {\
 	if (Amount != 0) \
+	{ \
 		FThreadStats::AddMessage(Stat, EStatOperation::Subtract, int64(Amount));\
+		TRACE_STAT_ADD(Stat, -int64(Amount)); \
+	} \
 }
 #define DEC_MEMORY_STAT_BY_FName(Stat,Amount) \
 {\
 	if (Amount != 0) \
+	{ \
 		FThreadStats::AddMessage(Stat, EStatOperation::Subtract, int64(Amount));\
+		TRACE_STAT_ADD(Stat, -int64(Amount)); \
+	} \
 }
 #define SET_MEMORY_STAT_FName(Stat,Value) \
 {\
 	FThreadStats::AddMessage(Stat, EStatOperation::Set, int64(Value));\
+	TRACE_STAT_SET(Stat, int64(Value)); \
 }
 #define SET_DWORD_STAT_FName(Stat,Value) \
 {\
 	FThreadStats::AddMessage(Stat, EStatOperation::Set, int64(Value));\
+	TRACE_STAT_SET(Stat, int64(Value)); \
 }
 #define SET_FLOAT_STAT_FName(Stat,Value) \
 {\
 	FThreadStats::AddMessage(Stat, EStatOperation::Set, double(Value));\
+	TRACE_STAT_SET(Stat, double(Value)); \
 }
 
 
diff --git a/Engine/Source/Runtime/CoreUObject/Private/Serialization/AsyncLoading.cpp b/Engine/Source/Runtime/CoreUObject/Private/Serialization/AsyncLoading.cpp
--- a/Engine/Source/Runtime/CoreUObject/Private/Serialization/AsyncLoading.cpp
+++ b/Engine/Source/Runtime/CoreUObject/Private/Serialization/AsyncLoading.cpp
@@ -41,6 +41,8 @@
 #include "HAL/LowLevelMemTracker.h"
 #include "ProfilingDebugging/CsvProfiler.h"
 #include "UObject/GCScopeLock.h"
+#include "ProfilingDebugging/MiscTrace.h"
+#include "Serialization/LoadTimeTracePrivate.h"
 
 #define FIND_MEMORY_STOMPS (1 && (PLATFORM_WINDOWS || PLATFORM_UNIX) && !WITH_EDITORONLY_DATA)
 
@@ -1788,7 +1790,11 @@ void FAsyncLoadingThread::QueueEvent_CreateLinker(FAsyncPackage* Package, int32
 	check(Package);
 	Package->AddNode(EEventLoadNode::Package_LoadSummary);
 	FWeakAsyncPackagePtr WeakPtr(Package);
-	EventQueue.AddAsyncEvent(Package->GetPriority(), GRandomizeLoadOrder ? GetRandomSerialNumber() : Package->SerialNumber, EventSystemPriority,
+
+	TAsyncLoadPriority UserPriority = Package->GetPriority();
+	int32 PackageSerialNumber = GRandomizeLoadOrder ? GetRandomSerialNumber() : Package->SerialNumber;
+	TRACE_LOADTIME_QUEUE_EVENT(Package, LoadTimeProfilerPackageEventType_CreateLinker, UserPriority, PackageSerialNumber, EventSystemPriority);
+	EventQueue.AddAsyncEvent(UserPriority, PackageSerialNumber, EventSystemPriority,
 		TFunction<void(FAsyncLoadEventArgs& Args)>(
 			[WeakPtr, this](FAsyncLoadEventArgs& Args)
 			{
@@ -1796,6 +1802,7 @@ void FAsyncLoadingThread::QueueEvent_CreateLinker(FAsyncPackage* Package, int32
 				check(Pkg);
 				if (Pkg)
 				{
+					TRACE_LOADTIME_ASYNC_PACKAGE_SCOPE(Pkg, LoadTimeProfilerPackageEventType_CreateLinker);
 					Pkg->SetTimeLimit(Args, TEXT("Create Linker"));
 					Pkg->Event_CreateLinker();
 					Args.OutLastObjectWorkWasPerformedOn = Pkg->GetLinkerRoot();
@@ -1850,7 +1857,10 @@ void FAsyncLoadingThread::QueueEvent_FinishLinker(FWeakAsyncPackagePtr WeakPtr,
 	FAsyncPackage* Pkg = GetPackage(WeakPtr);
 	if (Pkg)
 	{
-		EventQueue.AddAsyncEvent(Pkg->GetPriority(), GRandomizeLoadOrder ? GetRandomSerialNumber() : Pkg->SerialNumber, EventSystemPriority,
+		TAsyncLoadPriority UserPriority = Pkg->GetPriority();
+		int32 PackageSerialNumber = GRandomizeLoadOrder ? GetRandomSerialNumber() : Pkg->SerialNumber;
+		TRACE_LOADTIME_QUEUE_EVENT(Pkg, LoadTimeProfilerPackageEventType_FinishLinker, UserPriority, PackageSerialNumber, EventSystemPriority)
+		EventQueue.AddAsyncEvent(UserPriority, PackageSerialNumber, EventSystemPriority,
 			TFunction<void(FAsyncLoadEventArgs& Args)>(
 				[WeakPtr, this](FAsyncLoadEventArgs& Args)
 		{
@@ -1858,6 +1868,7 @@ void FAsyncLoadingThread::QueueEvent_FinishLinker(FWeakAsyncPackagePtr WeakPtr,
 			check(PkgInner);
 			if (PkgInner)
 			{
+				TRACE_LOADTIME_ASYNC_PACKAGE_SCOPE(PkgInner, LoadTimeProfilerPackageEventType_FinishLinker);
 				PkgInner->SetTimeLimit(Args, TEXT("Finish Linker"));
 				PkgInner->Event_FinishLinker();
 			}
@@ -1994,13 +2005,18 @@ void FAsyncLoadingThread::QueueEvent_StartImportPackages(FAsyncPackage* Package,
 {
 	check(Package);
 	FWeakAsyncPackagePtr WeakPtr(Package);
-	EventQueue.AddAsyncEvent(Package->GetPriority(), GRandomizeLoadOrder ? GetRandomSerialNumber() : Package->SerialNumber, EventSystemPriority,
+
+	TAsyncLoadPriority UserPriority = Package->GetPriority();
+	int32 PackageSerialNumber = GRandomizeLoadOrder ? GetRandomSerialNumber() : Package->SerialNumber;
+	TRACE_LOADTIME_QUEUE_EVENT(Package, LoadTimeProfilerPackageEventType_StartImportPackages, UserPriority, PackageSerialNumber, EventSystemPriority);
+	EventQueue.AddAsyncEvent(UserPriority, PackageSerialNumber, EventSystemPriority,
 		TFunction<void(FAsyncLoadEventArgs& Args)>(
 			[WeakPtr, this](FAsyncLoadEventArgs& Args)
 	{
 		FAsyncPackage* Pkg = GetPackage(WeakPtr);
 		if (Pkg)
 		{
+			TRACE_LOADTIME_ASYNC_PACKAGE_SCOPE(Pkg, LoadTimeProfilerPackageEventType_StartImportPackages);
 			Pkg->SetTimeLimit(Args, TEXT("Start Import Packages"));
 			Pkg->Event_StartImportPackages();
 		}
@@ -2161,6 +2177,7 @@ static const FName PrestreamPackageClassNameLoad = FName("PrestreamPackage");
 EAsyncPackageState::Type FAsyncPackage::LoadImports_Event()
 {
 	SCOPE_CYCLE_COUNTER(STAT_FAsyncPackage_LoadImports);
+	SCOPED_LOADTIMER(LoadImports_Event);
 	LastObjectWorkWasPerformedOn = LinkerRoot;
 	LastTypeOfWorkPerformed = TEXT("loading imports event");
 
@@ -2332,13 +2349,17 @@ void FAsyncLoadingThread::QueueEvent_SetupImports(FAsyncPackage* Package, int32
 	Package->AsyncPackageLoadingState = EAsyncPackageLoadingState::SetupImports;
 	check(Package);
 	FWeakAsyncPackagePtr WeakPtr(Package);
-	EventQueue.AddAsyncEvent(Package->GetPriority(), GRandomizeLoadOrder ? GetRandomSerialNumber() : Package->SerialNumber, EventSystemPriority,
+	TAsyncLoadPriority UserPriority = Package->GetPriority();
+	int32 PackageSerialNumber = GRandomizeLoadOrder ? GetRandomSerialNumber() : Package->SerialNumber;
+	TRACE_LOADTIME_QUEUE_EVENT(Package, LoadTimeProfilerPackageEventType_SetupImports, UserPriority, PackageSerialNumber, EventSystemPriority);
+	EventQueue.AddAsyncEvent(UserPriority, PackageSerialNumber, EventSystemPriority,
 		TFunction<void(FAsyncLoadEventArgs& Args)>(
 			[WeakPtr, this](FAsyncLoadEventArgs& Args)
 	{
 		FAsyncPackage* Pkg = GetPackage(WeakPtr);
 		if (Pkg)
 		{
+			TRACE_LOADTIME_ASYNC_PACKAGE_SCOPE(Pkg, LoadTimeProfilerPackageEventType_SetupImports);
 			Pkg->SetTimeLimit(Args, TEXT("Setup Imports"));
 			Pkg->Event_SetupImports();
 		}
@@ -2691,13 +2712,17 @@ void FAsyncLoadingThread::QueueEvent_SetupExports(FAsyncPackage* Package, int32
 	check(Package->AsyncPackageLoadingState == EAsyncPackageLoadingState::SetupExports);
 	check(Package);
 	FWeakAsyncPackagePtr WeakPtr(Package);
-	EventQueue.AddAsyncEvent(Package->GetPriority(), GRandomizeLoadOrder ? GetRandomSerialNumber() : Package->SerialNumber, EventSystemPriority,
+	TAsyncLoadPriority UserPriority = Package->GetPriority();
+	int32 PackageSerialNumber = GRandomizeLoadOrder ? GetRandomSerialNumber() : Package->SerialNumber;
+	TRACE_LOADTIME_QUEUE_EVENT(Package, LoadTimeProfilerPackageEventType_SetupExports, UserPriority, PackageSerialNumber, EventSystemPriority);
+	EventQueue.AddAsyncEvent(UserPriority, PackageSerialNumber, EventSystemPriority,
 		TFunction<void(FAsyncLoadEventArgs& Args)>(
 			[WeakPtr, this](FAsyncLoadEventArgs& Args)
 	{
 		FAsyncPackage* Pkg = GetPackage(WeakPtr);
 		if (Pkg)
 		{
+			TRACE_LOADTIME_ASYNC_PACKAGE_SCOPE(Pkg, LoadTimeProfilerPackageEventType_SetupExports);
 			Pkg->SetTimeLimit(Args, TEXT("Setup Exports"));
 			Pkg->Event_SetupExports();
 		}
@@ -2725,6 +2750,9 @@ void FAsyncLoadingThread::QueueEvent_ProcessImportsAndExports(FAsyncPackage* Pac
 	check(Package->AsyncPackageLoadingState == EAsyncPackageLoadingState::ProcessNewImportsAndExports);
 	check(Package);
 	FWeakAsyncPackagePtr WeakPtr(Package);
+	TAsyncLoadPriority UserPriority = Package->GetPriority();
+	int32 PackageSerialNumber = GRandomizeLoadOrder ? GetRandomSerialNumber() : Package->SerialNumber;
+	TRACE_LOADTIME_QUEUE_EVENT(Package, LoadTimeProfilerPackageEventType_ProcessImportsAndExports, UserPriority, PackageSerialNumber, EventSystemPriority);
 	EventQueue.AddAsyncEvent(Package->GetPriority(), GRandomizeLoadOrder ? GetRandomSerialNumber() : Package->SerialNumber, EventSystemPriority,
 		TFunction<void(FAsyncLoadEventArgs& Args)>(
 			[WeakPtr, this](FAsyncLoadEventArgs& Args)
@@ -2732,6 +2760,7 @@ void FAsyncLoadingThread::QueueEvent_ProcessImportsAndExports(FAsyncPackage* Pac
 		FAsyncPackage* Pkg = GetPackage(WeakPtr);
 		if (Pkg)
 		{
+			TRACE_LOADTIME_ASYNC_PACKAGE_SCOPE(Pkg, LoadTimeProfilerPackageEventType_ProcessImportsAndExports);
 			Pkg->SetTimeLimit(Args, TEXT("ProcessImportsAndExports"));
 			Pkg->Event_ProcessImportsAndExports();
 		}
@@ -2744,13 +2773,17 @@ void FAsyncLoadingThread::QueueEvent_ProcessPostloadWait(FAsyncPackage* Package,
 	check(Package->AsyncPackageLoadingState == EAsyncPackageLoadingState::WaitingForPostLoad);
 	check(Package);
 	FWeakAsyncPackagePtr WeakPtr(Package);
-	EventQueue.AddAsyncEvent(Package->GetPriority(), GRandomizeLoadOrder ? GetRandomSerialNumber() : Package->SerialNumber, EventSystemPriority,
+	TAsyncLoadPriority UserPriority = Package->GetPriority();
+	int32 PackageSerialNumber = GRandomizeLoadOrder ? GetRandomSerialNumber() : Package->SerialNumber;
+	TRACE_LOADTIME_QUEUE_EVENT(Package, LoadTimeProfilerPackageEventType_PostLoadWait, UserPriority, PackageSerialNumber, EventSystemPriority);
+	EventQueue.AddAsyncEvent(UserPriority, PackageSerialNumber, EventSystemPriority,
 		TFunction<void(FAsyncLoadEventArgs& Args)>(
 			[WeakPtr, this](FAsyncLoadEventArgs& Args)
 	{
 		FAsyncPackage* Pkg = GetPackage(WeakPtr);
 		if (Pkg)
 		{
+			TRACE_LOADTIME_ASYNC_PACKAGE_SCOPE(Pkg, LoadTimeProfilerPackageEventType_PostLoadWait);
 			Pkg->SetTimeLimit(Args, TEXT("Process Process Postload Wait"));
 			Pkg->Event_ProcessPostloadWait();
 		}
@@ -2763,13 +2796,17 @@ void FAsyncLoadingThread::QueueEvent_ExportsDone(FAsyncPackage* Package, int32 E
 	check(Package->AsyncPackageLoadingState == EAsyncPackageLoadingState::ProcessNewImportsAndExports);
 	check(Package);
 	FWeakAsyncPackagePtr WeakPtr(Package);
-	EventQueue.AddAsyncEvent(Package->GetPriority(), GRandomizeLoadOrder ? GetRandomSerialNumber() : Package->SerialNumber, EventSystemPriority,
+	TAsyncLoadPriority UserPriority = Package->GetPriority();
+	int32 PackageSerialNumber = GRandomizeLoadOrder ? GetRandomSerialNumber() : Package->SerialNumber;
+	TRACE_LOADTIME_QUEUE_EVENT(Package, LoadTimeProfilerPackageEventType_ExportsDone, UserPriority, PackageSerialNumber, EventSystemPriority);
+	EventQueue.AddAsyncEvent(UserPriority, PackageSerialNumber, EventSystemPriority,
 		TFunction<void(FAsyncLoadEventArgs& Args)>(
 			[WeakPtr, this](FAsyncLoadEventArgs& Args)
 	{
 		FAsyncPackage* Pkg = GetPackage(WeakPtr);
 		if (Pkg)
 		{
+			TRACE_LOADTIME_ASYNC_PACKAGE_SCOPE(Pkg, LoadTimeProfilerPackageEventType_ExportsDone);
 			Pkg->SetTimeLimit(Args, TEXT("Exports Done"));
 			Pkg->Event_ExportsDone();
 		}
@@ -2782,6 +2819,9 @@ void FAsyncLoadingThread::QueueEvent_StartPostLoad(FAsyncPackage* Package, int32
 	check(Package->AsyncPackageLoadingState == EAsyncPackageLoadingState::ReadyForPostLoad);
 	check(Package);
 	FWeakAsyncPackagePtr WeakPtr(Package);
+	TAsyncLoadPriority UserPriority = Package->GetPriority();
+	int32 PackageSerialNumber = GRandomizeLoadOrder ? GetRandomSerialNumber() : Package->SerialNumber;
+	TRACE_LOADTIME_QUEUE_EVENT(Package, LoadTimeProfilerPackageEventType_StartPostLoad, UserPriority, PackageSerialNumber, EventSystemPriority);
 	EventQueue.AddAsyncEvent(Package->GetPriority(), GRandomizeLoadOrder ? GetRandomSerialNumber() : Package->SerialNumber, EventSystemPriority,
 		TFunction<void(FAsyncLoadEventArgs& Args)>(
 			[WeakPtr, this](FAsyncLoadEventArgs& Args)
@@ -2789,6 +2829,7 @@ void FAsyncLoadingThread::QueueEvent_StartPostLoad(FAsyncPackage* Package, int32
 		FAsyncPackage* Pkg = GetPackage(WeakPtr);
 		if (Pkg)
 		{
+			TRACE_LOADTIME_ASYNC_PACKAGE_SCOPE(Pkg, LoadTimeProfilerPackageEventType_StartPostLoad);
 			Pkg->SetTimeLimit(Args, TEXT("Start Post Load"));
 			Pkg->Event_StartPostload();
 		}
@@ -3207,6 +3248,8 @@ void FAsyncPackage::EventDrivenCreateExport(int32 LocalExportIndex)
 	SCOPED_LOADTIMER(Package_CreateExports);
 	FObjectExport& Export = Linker->ExportMap[LocalExportIndex];
 
+	TRACE_LOADTIME_CREATE_EXPORT_SCOPE(Linker, &Export.Object, Export.SerialOffset, Export.SerialSize, Export.bIsAsset);
+
 	LLM_SCOPE(ELLMTag::AsyncLoading);
 	LLM_SCOPED_TAG_WITH_OBJECT_IN_SET(GetLinkerRoot(), ELLMTagSet::Assets);
 	LLM_SCOPED_TAG_WITH_OBJECT_IN_SET((Export.DynamicType == FObjectExport::EDynamicType::DynamicType) ? UDynamicClass::StaticClass() : 
@@ -3589,6 +3632,8 @@ void FAsyncPackage::EventDrivenSerializeExport(int32 LocalExportIndex)
 
 		Object->ClearFlags(RF_NeedLoad);
 
+		TRACE_LOADTIME_OBJECT_SCOPE(Object, LoadTimeProfilerObjectEventType_Serialize);
+
 		FUObjectSerializeContext* LoadContext = GetSerializeContext();
 		UObject* PrevSerializedObject = LoadContext->SerializedObject;
 		LoadContext->SerializedObject = Object;
@@ -3651,6 +3696,7 @@ void FAsyncPackage::EventDrivenSerializeExport(int32 LocalExportIndex)
 
 void FAsyncPackage::StartPrecacheRequest()
 {
+	SCOPED_LOADTIMER(StartPrecacheRequests);
 	if (Linker->bDynamicClassLinker)
 	{
 //native blueprint 
@@ -3792,6 +3838,7 @@ int64 FAsyncPackage::PrecacheRequestReady(IAsyncReadRequest * Read)
 
 void FAsyncPackage::MakeNextPrecacheRequestCurrent()
 {
+	SCOPED_LOADTIMER(MakeNextPrecacheRequestCurrent);
 	LLM_SCOPE(ELLMTag::AsyncLoading);
 
 	check(ReadyPrecacheRequests.Num());
@@ -3824,6 +3871,7 @@ void FAsyncPackage::MakeNextPrecacheRequestCurrent()
 
 void FAsyncPackage::FlushPrecacheBuffer()
 {
+	SCOPED_LOADTIMER(FlushPrecacheBuffer);
 	CurrentBlockOffset = -1;
 	CurrentBlockBytes = -1;
 	if (!Linker->bDynamicClassLinker)
@@ -3838,6 +3886,7 @@ int32 GCurrentExportIndex = -1;
 
 EAsyncPackageState::Type FAsyncPackage::ProcessImportsAndExports_Event()
 {
+	SCOPED_LOADTIMER(ProcessImportsAndExports_Event);
 	check(Linker);
 	bool bDidSomething = true;
 	int32 LoopIterations = 0;
@@ -4729,8 +4778,6 @@ bool FAsyncPackage::AreAllDependenciesFullyLoaded(TSet<UPackage*>& VisitedPackag
 
 EAsyncPackageState::Type FAsyncLoadingThread::ProcessLoadedPackages(bool bUseTimeLimit, bool bUseFullTimeLimit, float TimeLimit, bool& bDidSomething, FFlushTree* FlushTree)
 {
-	SCOPED_LOADTIMER(TickAsyncLoading_ProcessLoadedPackages);
-
 	EAsyncPackageState::Type Result = EAsyncPackageState::Complete;
 
 	// This is for debugging purposes only. @todo remove
@@ -4766,11 +4813,11 @@ EAsyncPackageState::Type FAsyncLoadingThread::ProcessLoadedPackages(bool bUseTim
 	bDidSomething = LoadedPackagesToProcess.Num() > 0;
 	for (int32 PackageIndex = 0; PackageIndex < LoadedPackagesToProcess.Num() && !IsAsyncLoadingSuspended(); ++PackageIndex)
 	{
-		SCOPED_LOADTIMER(ProcessLoadedPackagesTime);
-
 		FAsyncPackage* Package = LoadedPackagesToProcess[PackageIndex];
 		if (Package->GetDependencyRefCount() == 0)
 		{
+			SCOPED_LOADTIMER(ProcessLoadedPackagesTime);
+
 			Result = Package->PostLoadDeferredObjects(TickStartTime, bUseTimeLimit, TimeLimit);
 			if (Result == EAsyncPackageState::Complete)
 			{
@@ -4824,7 +4871,6 @@ EAsyncPackageState::Type FAsyncLoadingThread::ProcessLoadedPackages(bool bUseTim
 				const bool bInternalCallbacks = false;
 				const EAsyncLoadingResult::Type LoadingResult = Package->HasLoadFailed() ? EAsyncLoadingResult::Failed : EAsyncLoadingResult::Succeeded;
 				Package->CallCompletionCallbacks(bInternalCallbacks, LoadingResult);
-
 #if WITH_EDITOR
 				// In the editor we need to find any assets and add them to list for later callback
 				Package->GetLoadedAssets(LoadedAssets);
@@ -5159,6 +5205,10 @@ void FAsyncLoadingThread::StartThread()
 		bThreadStarted = true;
 		FPlatformMisc::MemoryBarrier();
 		Thread = FRunnableThread::Create(this, TEXT("FAsyncLoadingThread"), 0, TPri_Normal);
+		if (Thread)
+		{
+			TRACE_SET_THREAD_GROUP(Thread->GetThreadID(), TraceThreadGroup_AsyncLoading);
+		}
 	}
 }
 
@@ -5173,6 +5223,8 @@ uint32 FAsyncLoadingThread::Run()
 
 	AsyncLoadingThreadID = FPlatformTLS::GetCurrentThreadId();
 
+	TRACE_LOADTIME_START_ASYNC_LOADING();
+
 	if (!IsInGameThread())
 	{
 		FPlatformProcess::SetThreadAffinityMask(FPlatformAffinity::GetAsyncLoadingThreadMask());
@@ -5341,6 +5393,7 @@ void FAsyncLoadingThread::SuspendLoading()
 #endif
 	if (IsMultithreaded() && SuspendCount == 1)
 	{
+		TRACE_LOADTIME_SUSPEND_ASYNC_LOADING();
 		ThreadSuspendedEvent->Wait();
 	}
 }
@@ -5356,6 +5409,7 @@ void FAsyncLoadingThread::ResumeLoading()
 	if (IsMultithreaded() && SuspendCount == 0)
 	{
 		ThreadResumedEvent->Wait();
+		TRACE_LOADTIME_RESUME_ASYNC_LOADING();
 	}
 }
 
@@ -5495,11 +5549,13 @@ FAsyncPackage::FAsyncPackage(const FAsyncPackageDesc& InDesc)
 , FinishObjectsTime(0.0)
 #endif // PERF_TRACK_DETAILED_ASYNC_STATS
 {
+	TRACE_LOADTIME_NEW_ASYNC_PACKAGE(this, *InDesc.Name.ToString());
 	AddRequestID(InDesc.RequestID);
 }
 
 FAsyncPackage::~FAsyncPackage()
 {
+	TRACE_LOADTIME_DESTROY_ASYNC_PACKAGE(this);
 #if DO_CHECK
 	if (GEventDrivenLoaderEnabled)
 	{
@@ -5580,6 +5636,7 @@ void FAsyncPackage::AddRequestID(int32 Id)
 		}
 		RequestIDs.Add(Id);
 		AsyncLoadingThread.AddPendingRequest(Id);
+		TRACE_LOADTIME_ASYNC_PACKAGE_REQUEST_ASSOCIATION(this, Id);
 	}
 }
 
@@ -5730,6 +5787,7 @@ EAsyncPackageState::Type FAsyncPackage::TickAsyncPackage(bool InbUseTimeLimit, b
 
 	ReentryCount++;
 
+	TRACE_LOADTIME_ASYNC_PACKAGE_SCOPE(this, LoadTimeProfilerPackageEventType_Tick);
 	SCOPE_CYCLE_COUNTER(STAT_FAsyncPackage_Tick);
 	SCOPED_LOADTIMER(Package_Tick);
 
@@ -6095,6 +6153,7 @@ EAsyncPackageState::Type FAsyncPackage::CreateLinker()
 		check(Linker);
 		check(Linker->AsyncRoot == nullptr || Linker->AsyncRoot == this);
 		Linker->AsyncRoot = this;
+		TRACE_LOADTIME_ASYNC_PACKAGE_LINKER_ASSOCIATION(this, Linker);
 
 		UE_LOG(LogStreaming, Verbose, TEXT("FAsyncPackage::CreateLinker for %s finished."), *Desc.NameToLoad.ToString());
 	}
@@ -6663,7 +6722,10 @@ EAsyncPackageState::Type FAsyncPackage::PostLoadObjects()
 				check(!GEventDrivenLoaderEnabled || !Object->HasAnyFlags(RF_NeedLoad));
 
 				ThreadContext.CurrentlyPostLoadedObjectByALT = Object;
-				Object->ConditionalPostLoad();
+				{
+					TRACE_LOADTIME_OBJECT_SCOPE(Object, LoadTimeProfilerObjectEventType_PostLoad);
+					Object->ConditionalPostLoad();
+				}
 				ThreadContext.CurrentlyPostLoadedObjectByALT = nullptr;
 
 				LastObjectWorkWasPerformedOn = Object;
@@ -6728,7 +6790,10 @@ EAsyncPackageState::Type FAsyncPackage::PostLoadDeferredObjects(double InTickSta
 		FScopeCycleCounterUObject ConstructorScope(Object, GET_STATID(STAT_FAsyncPackage_PostLoadObjectsGameThread));
 
 		PackageScope.ThreadContext.CurrentlyPostLoadedObjectByALT = Object;
-		Object->ConditionalPostLoad();
+		{
+			TRACE_LOADTIME_OBJECT_SCOPE(Object, LoadTimeProfilerObjectEventType_PostLoad);
+			Object->ConditionalPostLoad();
+		}
 		PackageScope.ThreadContext.CurrentlyPostLoadedObjectByALT = nullptr;
 
 		if (ObjLoadedInPostLoad.Num())
@@ -7164,6 +7229,7 @@ int32 LoadPackageAsync(const FString& InName, const FGuid* InGuid /*= nullptr*/,
 		// Generate new request ID and add it immediately to the global request list (it needs to be there before we exit
 		// this function, otherwise it would be added when the packages are being processed on the async thread).
 		RequestID = GPackageRequestID.Increment();
+		TRACE_LOADTIME_BEGIN_REQUEST(RequestID);
 		FAsyncLoadingThread::Get().AddPendingRequest(RequestID);
 
 		// Allocate delegate on Game Thread, it is not safe to copy delegates by value on other threads
@@ -7247,7 +7313,7 @@ void FlushAsyncLoading(int32 PackageID /* = INDEX_NONE */)
 
 	checkf(IsInGameThread(), TEXT("Unable to FlushAsyncLoading from any thread other than the game thread."));
 
-	if (IsAsyncLoading())
+ 	if (IsAsyncLoading())
 	{
 		FAsyncLoadingThread& AsyncThread = FAsyncLoadingThread::Get();
 		// Flushing async loading while loading is suspend will result in infinite stall
@@ -7260,6 +7326,8 @@ void FlushAsyncLoading(int32 PackageID /* = INDEX_NONE */)
 			return;
 		}
 
+		TRACE_LOADTIME_FLUSH_ASYNCLOADING_SCOPE(PackageID);
+
 		FCoreDelegates::OnAsyncLoadingFlush.Broadcast();
 
 #if NO_LOGGING == 0
diff --git a/Engine/Source/Runtime/CoreUObject/Private/Serialization/LoadTimeTrace.cpp b/Engine/Source/Runtime/CoreUObject/Private/Serialization/LoadTimeTrace.cpp
new file mode 100644
--- /dev/null
+++ b/Engine/Source/Runtime/CoreUObject/Private/Serialization/LoadTimeTrace.cpp
@@ -0,0 +1,476 @@
+// Copyright 1998-2019 Epic Games, Inc. All Rights Reserved.
+
+#include "Serialization/LoadTimeTrace.h"
+#include "LoadTimeTracePrivate.h"
+#include "Trace/Trace.h"
+#include "Misc/CString.h"
+#include "HAL/PlatformTime.h"
+#include "HAL/PlatformTLS.h"
+#include "UObject/UObjectGlobals.h"
+
+#if LOADTIMEPROFILERTRACE_ENABLED
+
+UE_TRACE_EVENT_BEGIN(LoadTime, StartAsyncLoading)
+	UE_TRACE_EVENT_FIELD(uint64, Cycle)
+	UE_TRACE_EVENT_FIELD(uint32, ThreadId)
+UE_TRACE_EVENT_END()
+
+UE_TRACE_EVENT_BEGIN(LoadTime, SuspendAsyncLoading)
+	UE_TRACE_EVENT_FIELD(uint64, Cycle)
+UE_TRACE_EVENT_END()
+
+UE_TRACE_EVENT_BEGIN(LoadTime, ResumeAsyncLoading)
+	UE_TRACE_EVENT_FIELD(uint64, Cycle)
+UE_TRACE_EVENT_END()
+
+UE_TRACE_EVENT_BEGIN(LoadTime, NewPackage)
+	UE_TRACE_EVENT_FIELD(const UPackage*, Package)
+UE_TRACE_EVENT_END()
+
+UE_TRACE_EVENT_BEGIN(LoadTime, NewLinker)
+	UE_TRACE_EVENT_FIELD(const FLinkerLoad*, Linker)
+UE_TRACE_EVENT_END()
+
+UE_TRACE_EVENT_BEGIN(LoadTime, DestroyLinker)
+	UE_TRACE_EVENT_FIELD(const FLinkerLoad*, Linker)
+UE_TRACE_EVENT_END()
+
+UE_TRACE_EVENT_BEGIN(LoadTime, PackageSummary)
+	UE_TRACE_EVENT_FIELD(const FLinkerLoad*, Linker)
+	UE_TRACE_EVENT_FIELD(uint32, TotalHeaderSize)
+	UE_TRACE_EVENT_FIELD(uint32, NameCount)
+	UE_TRACE_EVENT_FIELD(uint32, ImportCount)
+	UE_TRACE_EVENT_FIELD(uint32, ExportCount)
+UE_TRACE_EVENT_END()
+
+UE_TRACE_EVENT_BEGIN(LoadTime, BeginLinkerScope)
+	UE_TRACE_EVENT_FIELD(const FLinkerLoad*, Linker)
+	UE_TRACE_EVENT_FIELD(uint64, Cycle)
+	UE_TRACE_EVENT_FIELD(uint32, ThreadId)
+UE_TRACE_EVENT_END()
+
+UE_TRACE_EVENT_BEGIN(LoadTime, EndLinkerScope)
+	UE_TRACE_EVENT_FIELD(uint64, Cycle)
+	UE_TRACE_EVENT_FIELD(uint32, ThreadId)
+UE_TRACE_EVENT_END()
+
+UE_TRACE_EVENT_BEGIN(LoadTime, BeginCreateExport)
+	UE_TRACE_EVENT_FIELD(uint64, Cycle)
+	UE_TRACE_EVENT_FIELD(const FLinkerLoad*, Linker)
+	UE_TRACE_EVENT_FIELD(uint64, SerialOffset)
+	UE_TRACE_EVENT_FIELD(uint64, SerialSize)
+	UE_TRACE_EVENT_FIELD(uint32, ThreadId)
+	UE_TRACE_EVENT_FIELD(bool, IsAsset)
+UE_TRACE_EVENT_END()
+
+UE_TRACE_EVENT_BEGIN(LoadTime, EndCreateExport)
+	UE_TRACE_EVENT_FIELD(uint64, Cycle)
+	UE_TRACE_EVENT_FIELD(const UObject*, Object)
+	UE_TRACE_EVENT_FIELD(const UClass*, Class)
+	UE_TRACE_EVENT_FIELD(uint32, ThreadId)
+UE_TRACE_EVENT_END()
+
+UE_TRACE_EVENT_BEGIN(LoadTime, BeginObjectScope)
+	UE_TRACE_EVENT_FIELD(uint64, Cycle)
+	UE_TRACE_EVENT_FIELD(const UObject*, Object)
+	UE_TRACE_EVENT_FIELD(uint32, ThreadId)
+	UE_TRACE_EVENT_FIELD(uint8, EventType)
+UE_TRACE_EVENT_END()
+
+UE_TRACE_EVENT_BEGIN(LoadTime, EndObjectScope)
+	UE_TRACE_EVENT_FIELD(uint64, Cycle)
+	UE_TRACE_EVENT_FIELD(uint32, ThreadId)
+UE_TRACE_EVENT_END()
+
+UE_TRACE_EVENT_BEGIN(LoadTime, LinkerArchiveAssociation)
+	UE_TRACE_EVENT_FIELD(const FLinkerLoad*, Linker)
+	UE_TRACE_EVENT_FIELD(const FArchive*, Archive)
+UE_TRACE_EVENT_END()
+
+UE_TRACE_EVENT_BEGIN(LoadTime, BeginRequest)
+	UE_TRACE_EVENT_FIELD(uint64, Cycle)
+	UE_TRACE_EVENT_FIELD(uint64, RequestId)
+	UE_TRACE_EVENT_FIELD(uint32, ThreadId)
+UE_TRACE_EVENT_END()
+
+UE_TRACE_EVENT_BEGIN(LoadTime, EndRequest)
+	UE_TRACE_EVENT_FIELD(uint64, Cycle)
+	UE_TRACE_EVENT_FIELD(uint64, RequestId)
+UE_TRACE_EVENT_END()
+
+UE_TRACE_EVENT_BEGIN(LoadTime, BeginLoadMap)
+	UE_TRACE_EVENT_FIELD(uint64, Cycle)
+	UE_TRACE_EVENT_FIELD(uint32, ThreadId)
+UE_TRACE_EVENT_END()
+
+UE_TRACE_EVENT_BEGIN(LoadTime, EndLoadMap)
+	UE_TRACE_EVENT_FIELD(uint64, Cycle)
+	UE_TRACE_EVENT_FIELD(uint32, ThreadId)
+UE_TRACE_EVENT_END()
+
+UE_TRACE_EVENT_BEGIN(LoadTime, NewStreamableHandle)
+	UE_TRACE_EVENT_FIELD(const FStreamableHandle*, StreamableHandle)
+	UE_TRACE_EVENT_FIELD(bool, IsCombined)
+UE_TRACE_EVENT_END()
+
+UE_TRACE_EVENT_BEGIN(LoadTime, DestroyStreamableHandle)
+	UE_TRACE_EVENT_FIELD(const FStreamableHandle*, StreamableHandle)
+UE_TRACE_EVENT_END()
+
+UE_TRACE_EVENT_BEGIN(LoadTime, BeginLoadStreamableHandle)
+	UE_TRACE_EVENT_FIELD(uint64, Cycle)
+	UE_TRACE_EVENT_FIELD(const FStreamableHandle*, StreamableHandle)
+	UE_TRACE_EVENT_END()
+
+UE_TRACE_EVENT_BEGIN(LoadTime, EndLoadStreamableHandle)
+	UE_TRACE_EVENT_FIELD(uint64, Cycle)
+	UE_TRACE_EVENT_FIELD(const FStreamableHandle*, StreamableHandle)
+	UE_TRACE_EVENT_END()
+
+UE_TRACE_EVENT_BEGIN(LoadTime, BeginWaitForStreamableHandle)
+	UE_TRACE_EVENT_FIELD(uint64, Cycle)
+	UE_TRACE_EVENT_FIELD(const FStreamableHandle*, StreamableHandle)
+	UE_TRACE_EVENT_FIELD(uint32, ThreadId)
+UE_TRACE_EVENT_END()
+
+UE_TRACE_EVENT_BEGIN(LoadTime, EndWaitForStreamableHandle)
+	UE_TRACE_EVENT_FIELD(uint64, Cycle)
+	UE_TRACE_EVENT_FIELD(uint32, ThreadId)
+UE_TRACE_EVENT_END()
+
+UE_TRACE_EVENT_BEGIN(LoadTime, NewAsyncPackage)
+	UE_TRACE_EVENT_FIELD(const FAsyncPackage*, AsyncPackage)
+UE_TRACE_EVENT_END()
+
+UE_TRACE_EVENT_BEGIN(LoadTime, DestroyAsyncPackage)
+	UE_TRACE_EVENT_FIELD(const FAsyncPackage*, AsyncPackage)
+UE_TRACE_EVENT_END()
+
+UE_TRACE_EVENT_BEGIN(LoadTime, StreamableHandleRequestAssociation)
+	UE_TRACE_EVENT_FIELD(const FStreamableHandle*, StreamableHandle)
+	UE_TRACE_EVENT_FIELD(uint64, RequestId)
+UE_TRACE_EVENT_END()
+
+UE_TRACE_EVENT_BEGIN(LoadTime, AsyncPackageRequestAssociation)
+	UE_TRACE_EVENT_FIELD(const FAsyncPackage*, AsyncPackage)
+	UE_TRACE_EVENT_FIELD(uint64, RequestId)
+UE_TRACE_EVENT_END()
+
+UE_TRACE_EVENT_BEGIN(LoadTime, LoadImportDependency)
+	UE_TRACE_EVENT_FIELD(const FAsyncPackage*, Package)
+	UE_TRACE_EVENT_FIELD(const FAsyncPackage*, DependentPackage)
+UE_TRACE_EVENT_END()
+
+UE_TRACE_EVENT_BEGIN(LoadTime, AsyncPackageLinkerAssociation)
+	UE_TRACE_EVENT_FIELD(const FAsyncPackage*, AsyncPackage)
+	UE_TRACE_EVENT_FIELD(const FLinkerLoad*, Linker)
+UE_TRACE_EVENT_END()
+
+UE_TRACE_EVENT_BEGIN(LoadTime, BeginAsyncPackageScope)
+	UE_TRACE_EVENT_FIELD(const FAsyncPackage*, AsyncPackage)
+	UE_TRACE_EVENT_FIELD(uint64, Cycle)
+	UE_TRACE_EVENT_FIELD(uint32, ThreadId)
+	UE_TRACE_EVENT_FIELD(uint8, EventType)
+UE_TRACE_EVENT_END()
+
+UE_TRACE_EVENT_BEGIN(LoadTime, EndAsyncPackageScope)
+	UE_TRACE_EVENT_FIELD(uint64, Cycle)
+	UE_TRACE_EVENT_FIELD(uint32, ThreadId)
+UE_TRACE_EVENT_END()
+
+UE_TRACE_EVENT_BEGIN(LoadTime, BeginFlushAsyncLoading)
+	UE_TRACE_EVENT_FIELD(uint64, Cycle)
+	UE_TRACE_EVENT_FIELD(uint64, RequestId)
+	UE_TRACE_EVENT_FIELD(uint32, ThreadId)
+UE_TRACE_EVENT_END()
+
+UE_TRACE_EVENT_BEGIN(LoadTime, EndFlushAsyncLoading)
+	UE_TRACE_EVENT_FIELD(uint64, Cycle)
+	UE_TRACE_EVENT_FIELD(uint32, ThreadId)
+UE_TRACE_EVENT_END()
+
+UE_TRACE_EVENT_BEGIN(LoadTime, QueueEvent)
+	UE_TRACE_EVENT_FIELD(uint64, Cycle)
+	UE_TRACE_EVENT_FIELD(const FAsyncPackage*, AsyncPackage)
+	UE_TRACE_EVENT_FIELD(uint32, ThreadId)
+	UE_TRACE_EVENT_FIELD(uint32, UserPriority)
+	UE_TRACE_EVENT_FIELD(uint32, PackageSerialNumber)
+	UE_TRACE_EVENT_FIELD(uint32, EventSystemPriority)
+	UE_TRACE_EVENT_FIELD(uint8, EventType)
+UE_TRACE_EVENT_END();
+
+UE_TRACE_EVENT_BEGIN(LoadTime, ClassInfo)
+	UE_TRACE_EVENT_FIELD(const UClass*, Class)
+UE_TRACE_EVENT_END()
+
+void FLoadTimeProfilerTracePrivate::OutputStartAsyncLoading()
+{
+	UE_TRACE_LOG(LoadTime, StartAsyncLoading)
+		<< StartAsyncLoading.Cycle(FPlatformTime::Cycles64())
+		<< StartAsyncLoading.ThreadId(FPlatformTLS::GetCurrentThreadId());
+}
+
+PRAGMA_DISABLE_SHADOW_VARIABLE_WARNINGS
+void FLoadTimeProfilerTracePrivate::OutputSuspendAsyncLoading()
+{
+	UE_TRACE_LOG(LoadTime, SuspendAsyncLoading)
+		<< SuspendAsyncLoading.Cycle(FPlatformTime::Cycles64());
+}
+
+void FLoadTimeProfilerTracePrivate::OutputResumeAsyncLoading()
+{
+	UE_TRACE_LOG(LoadTime, ResumeAsyncLoading)
+		<< ResumeAsyncLoading.Cycle(FPlatformTime::Cycles64());
+}
+PRAGMA_ENABLE_SHADOW_VARIABLE_WARNINGS
+
+void FLoadTimeProfilerTracePrivate::OutputNewLinker(const FLinkerLoad* Linker)
+{
+	UE_TRACE_LOG(LoadTime, NewLinker)
+		<< NewLinker.Linker(Linker);
+}
+
+void FLoadTimeProfilerTracePrivate::OutputDestroyLinker(const FLinkerLoad* Linker)
+{
+	UE_TRACE_LOG(LoadTime, DestroyLinker)
+		<< DestroyLinker.Linker(Linker);
+}
+
+void FLoadTimeProfilerTracePrivate::OutputBeginRequest(uint64 RequestId)
+{
+	UE_TRACE_LOG(LoadTime, BeginRequest)
+		<< BeginRequest.Cycle(FPlatformTime::Cycles64())
+		<< BeginRequest.RequestId(RequestId)
+		<< BeginRequest.ThreadId(FPlatformTLS::GetCurrentThreadId());
+}
+
+void FLoadTimeProfilerTracePrivate::OutputEndRequest(uint64 RequestId)
+{
+	UE_TRACE_LOG(LoadTime, EndRequest)
+		<< EndRequest.Cycle(FPlatformTime::Cycles64())
+		<< EndRequest.RequestId(RequestId);
+}
+
+void FLoadTimeProfilerTrace::OutputNewStreamableHandle(const FStreamableHandle* StreamableHandle, const TCHAR* DebugName, bool IsCombined)
+{
+	uint16 NameSize = (FCString::Strlen(DebugName) + 1) * sizeof(TCHAR);
+	UE_TRACE_LOG(LoadTime, NewStreamableHandle, NameSize)
+		<< NewStreamableHandle.StreamableHandle(StreamableHandle)
+		<< NewStreamableHandle.IsCombined(IsCombined)
+		<< NewStreamableHandle.Attachment(DebugName, NameSize);
+}
+
+void FLoadTimeProfilerTrace::OutputDestroyStreamableHandle(const FStreamableHandle* StreamableHandle)
+{
+	UE_TRACE_LOG(LoadTime, DestroyStreamableHandle)
+		<< DestroyStreamableHandle.StreamableHandle(StreamableHandle);
+}
+
+void FLoadTimeProfilerTrace::OutputBeginLoadStreamableHandle(const FStreamableHandle* StreamableHandle)
+{
+	UE_TRACE_LOG(LoadTime, BeginLoadStreamableHandle)
+		<< BeginLoadStreamableHandle.StreamableHandle(StreamableHandle)
+		<< BeginLoadStreamableHandle.Cycle(FPlatformTime::Cycles64());
+}
+
+void FLoadTimeProfilerTrace::OutputEndLoadStreamableHandle(const FStreamableHandle* StreamableHandle)
+{
+	UE_TRACE_LOG(LoadTime, EndLoadStreamableHandle)
+		<< EndLoadStreamableHandle.StreamableHandle(StreamableHandle)
+		<< EndLoadStreamableHandle.Cycle(FPlatformTime::Cycles64());
+}
+
+void FLoadTimeProfilerTracePrivate::OutputNewAsyncPackage(const FAsyncPackage* AsyncPackage, const TCHAR* PackageName)
+{
+	uint16 NameSize = (FCString::Strlen(PackageName) + 1) * sizeof(TCHAR);
+	UE_TRACE_LOG(LoadTime, NewAsyncPackage, NameSize)
+		<< NewAsyncPackage.AsyncPackage(AsyncPackage)
+		<< NewAsyncPackage.Attachment(PackageName, NameSize);
+}
+
+void FLoadTimeProfilerTracePrivate::OutputDestroyAsyncPackage(const FAsyncPackage* AsyncPackage)
+{
+	UE_TRACE_LOG(LoadTime, DestroyAsyncPackage)
+		<< DestroyAsyncPackage.AsyncPackage(AsyncPackage);
+}
+
+void FLoadTimeProfilerTracePrivate::OutputPackageSummary(const FLinkerLoad* Linker, uint32 TotalHeaderSize, uint32 NameCount, uint32 ImportCount, uint32 ExportCount)
+{
+	UE_TRACE_LOG(LoadTime, PackageSummary)
+		<< PackageSummary.Linker(Linker)
+		<< PackageSummary.TotalHeaderSize(TotalHeaderSize)
+		<< PackageSummary.NameCount(NameCount)
+		<< PackageSummary.ImportCount(ImportCount)
+		<< PackageSummary.ExportCount(ExportCount);
+}
+
+void FLoadTimeProfilerTrace::OutputStreamableHandleRequestAssociation(const FStreamableHandle* StreamableHandle, uint64 RequestId)
+{
+	UE_TRACE_LOG(LoadTime, StreamableHandleRequestAssociation)
+		<< StreamableHandleRequestAssociation.StreamableHandle(StreamableHandle)
+		<< StreamableHandleRequestAssociation.RequestId(RequestId);
+}
+
+void FLoadTimeProfilerTracePrivate::OutputAsyncPackageRequestAssociation(const FAsyncPackage* AsyncPackage, uint64 RequestId)
+{
+	UE_TRACE_LOG(LoadTime, AsyncPackageRequestAssociation)
+		<< AsyncPackageRequestAssociation.AsyncPackage(AsyncPackage)
+		<< AsyncPackageRequestAssociation.RequestId(RequestId);
+}
+
+void FLoadTimeProfilerTracePrivate::OutputLoadImportDependency(const FAsyncPackage* Package, const FAsyncPackage* DependentPackage)
+{
+	UE_TRACE_LOG(LoadTime, LoadImportDependency)
+		<< LoadImportDependency.Package(Package)
+		<< LoadImportDependency.DependentPackage(DependentPackage);
+}
+
+void FLoadTimeProfilerTracePrivate::OutputAsyncPackageLinkerAssociation(const FAsyncPackage* AsyncPackage, const FLinkerLoad* Linker)
+{
+	UE_TRACE_LOG(LoadTime, AsyncPackageLinkerAssociation)
+		<< AsyncPackageLinkerAssociation.AsyncPackage(AsyncPackage)
+		<< AsyncPackageLinkerAssociation.Linker(Linker);
+}
+
+void FLoadTimeProfilerTracePrivate::OutputLinkerArchiveAssociation(const FLinkerLoad* Linker, const FArchive* Archive)
+{
+	UE_TRACE_LOG(LoadTime, LinkerArchiveAssociation)
+		<< LinkerArchiveAssociation.Linker(Linker)
+		<< LinkerArchiveAssociation.Archive(Archive);
+}
+
+void FLoadTimeProfilerTracePrivate::OutputQueueEvent(const FAsyncPackage* AsyncPackage, ELoadTimeProfilePackageEventType EventType, TAsyncLoadPriority UserPriority, int32 PackageSerialNumber, int32 EventSystemPriority)
+{
+	UE_TRACE_LOG(LoadTime, QueueEvent)
+		<< QueueEvent.Cycle(FPlatformTime::Cycles64())
+		<< QueueEvent.AsyncPackage(AsyncPackage)
+		<< QueueEvent.ThreadId(FPlatformTLS::GetCurrentThreadId())
+		<< QueueEvent.UserPriority(UserPriority)
+		<< QueueEvent.PackageSerialNumber(PackageSerialNumber)
+		<< QueueEvent.EventSystemPriority(EventSystemPriority)
+		<< QueueEvent.EventType(EventType);
+}
+
+void FLoadTimeProfilerTracePrivate::OutputClassInfo(const UClass* Class, const TCHAR* Name)
+{
+	uint16 NameSize = (FCString::Strlen(Name) + 1) * sizeof(TCHAR);
+	UE_TRACE_LOG(LoadTime, ClassInfo, NameSize)
+		<< ClassInfo.Class(Class)
+		<< ClassInfo.Attachment(Name, NameSize);
+}
+
+FLoadTimeProfilerTracePrivate::FAsyncPackageScope::FAsyncPackageScope(const FAsyncPackage* AsyncPackage, ELoadTimeProfilePackageEventType EventType)
+{
+	UE_TRACE_LOG(LoadTime, BeginAsyncPackageScope)
+		<< BeginAsyncPackageScope.AsyncPackage(AsyncPackage)
+		<< BeginAsyncPackageScope.Cycle(FPlatformTime::Cycles64())
+		<< BeginAsyncPackageScope.ThreadId(FPlatformTLS::GetCurrentThreadId())
+		<< BeginAsyncPackageScope.EventType(EventType);
+}
+
+FLoadTimeProfilerTracePrivate::FAsyncPackageScope::~FAsyncPackageScope()
+{
+	UE_TRACE_LOG(LoadTime, EndAsyncPackageScope)
+		<< EndAsyncPackageScope.Cycle(FPlatformTime::Cycles64())
+		<< EndAsyncPackageScope.ThreadId(FPlatformTLS::GetCurrentThreadId());
+}
+
+FLoadTimeProfilerTracePrivate::FCreateExportScope::FCreateExportScope(const FLinkerLoad* Linker, uint64 SerialOffset, uint64 SerialSize, bool IsAsset, const UObject* const* InObject)
+	: Object(InObject)
+{
+	UE_TRACE_LOG(LoadTime, BeginCreateExport)
+		<< BeginCreateExport.Cycle(FPlatformTime::Cycles64())
+		<< BeginCreateExport.Linker(Linker)
+		<< BeginCreateExport.SerialOffset(SerialOffset)
+		<< BeginCreateExport.SerialSize(SerialSize)
+		<< BeginCreateExport.ThreadId(FPlatformTLS::GetCurrentThreadId())
+		<< BeginCreateExport.IsAsset(IsAsset);
+}
+
+FLoadTimeProfilerTracePrivate::FCreateExportScope::~FCreateExportScope()
+{
+	UE_TRACE_LOG(LoadTime, EndCreateExport)
+		<< EndCreateExport.Cycle(FPlatformTime::Cycles64())
+		<< EndCreateExport.Object(*Object)
+		<< EndCreateExport.Class(*Object ? (*Object)->GetClass() : nullptr)
+		<< EndCreateExport.ThreadId(FPlatformTLS::GetCurrentThreadId());
+}
+
+FLoadTimeProfilerTracePrivate::FObjectScope::FObjectScope(const UObject* Object, ELoadTimeProfilerObjectEventType EventType)
+{
+	UE_TRACE_LOG(LoadTime, BeginObjectScope)
+		<< BeginObjectScope.Cycle(FPlatformTime::Cycles64())
+		<< BeginObjectScope.Object(Object)
+		<< BeginObjectScope.ThreadId(FPlatformTLS::GetCurrentThreadId())
+		<< BeginObjectScope.EventType(EventType);
+}
+
+FLoadTimeProfilerTracePrivate::FObjectScope::~FObjectScope()
+{
+	UE_TRACE_LOG(LoadTime, EndObjectScope)
+		<< EndObjectScope.Cycle(FPlatformTime::Cycles64())
+		<< EndObjectScope.ThreadId(FPlatformTLS::GetCurrentThreadId());
+}
+
+FLoadTimeProfilerTracePrivate::FLinkerScope::FLinkerScope(const FLinkerLoad* Linker)
+{
+	UE_TRACE_LOG(LoadTime, BeginLinkerScope)
+		<< BeginLinkerScope.Linker(Linker)
+		<< BeginLinkerScope.Cycle(FPlatformTime::Cycles64())
+		<< BeginLinkerScope.ThreadId(FPlatformTLS::GetCurrentThreadId());
+}
+
+FLoadTimeProfilerTracePrivate::FLinkerScope::~FLinkerScope()
+{
+	UE_TRACE_LOG(LoadTime, EndLinkerScope)
+		<< EndLinkerScope.Cycle(FPlatformTime::Cycles64())
+		<< EndLinkerScope.ThreadId(FPlatformTLS::GetCurrentThreadId());
+}
+
+FLoadTimeProfilerTracePrivate::FFlushAsyncLoadingScope::FFlushAsyncLoadingScope(uint64 RequestId)
+{
+	UE_TRACE_LOG(LoadTime, BeginFlushAsyncLoading)
+		<< BeginFlushAsyncLoading.Cycle(FPlatformTime::Cycles64())
+		<< BeginFlushAsyncLoading.RequestId(RequestId)
+		<< BeginFlushAsyncLoading.ThreadId(FPlatformTLS::GetCurrentThreadId());
+}
+
+FLoadTimeProfilerTracePrivate::FFlushAsyncLoadingScope::~FFlushAsyncLoadingScope()
+{
+	UE_TRACE_LOG(LoadTime, EndFlushAsyncLoading)
+		<< EndFlushAsyncLoading.Cycle(FPlatformTime::Cycles64())
+		<< EndFlushAsyncLoading.ThreadId(FPlatformTLS::GetCurrentThreadId());
+}
+
+FLoadTimeProfilerTrace::FWaitForStreamableHandleScope::FWaitForStreamableHandleScope(const FStreamableHandle* StreamableHandle)
+{
+	UE_TRACE_LOG(LoadTime, BeginWaitForStreamableHandle)
+		<< BeginWaitForStreamableHandle.Cycle(FPlatformTime::Cycles64())
+		<< BeginWaitForStreamableHandle.StreamableHandle(StreamableHandle)
+		<< BeginWaitForStreamableHandle.ThreadId(FPlatformTLS::GetCurrentThreadId());
+}
+
+FLoadTimeProfilerTrace::FWaitForStreamableHandleScope::~FWaitForStreamableHandleScope()
+{
+	UE_TRACE_LOG(LoadTime, EndWaitForStreamableHandle)
+		<< EndWaitForStreamableHandle.Cycle(FPlatformTime::Cycles64())
+		<< EndWaitForStreamableHandle.ThreadId(FPlatformTLS::GetCurrentThreadId());
+}
+
+FLoadTimeProfilerTrace::FLoadMapScope::FLoadMapScope(const TCHAR* Name)
+{
+	uint16 NameSize = (FCString::Strlen(Name) + 1) * sizeof(TCHAR);
+	UE_TRACE_LOG(LoadTime, BeginLoadMap, NameSize)
+		<< BeginLoadMap.Cycle(FPlatformTime::Cycles64())
+		<< BeginLoadMap.ThreadId(FPlatformTLS::GetCurrentThreadId())
+		<< BeginLoadMap.Attachment(Name, NameSize);
+}
+
+FLoadTimeProfilerTrace::FLoadMapScope::~FLoadMapScope()
+{
+	UE_TRACE_LOG(LoadTime, EndLoadMap)
+		<< EndLoadMap.Cycle(FPlatformTime::Cycles64())
+		<< EndLoadMap.ThreadId(FPlatformTLS::GetCurrentThreadId());
+}
+
+#endif
\ No newline at end of file
diff --git a/Engine/Source/Runtime/CoreUObject/Private/Serialization/LoadTimeTracePrivate.h b/Engine/Source/Runtime/CoreUObject/Private/Serialization/LoadTimeTracePrivate.h
new file mode 100644
--- /dev/null
+++ b/Engine/Source/Runtime/CoreUObject/Private/Serialization/LoadTimeTracePrivate.h
@@ -0,0 +1,149 @@
+// Copyright 1998-2019 Epic Games, Inc. All Rights Reserved.
+
+#pragma once
+
+#include "Serialization/LoadTimeTrace.h"
+
+#if LOADTIMEPROFILERTRACE_ENABLED
+
+class FArchive;
+struct FAsyncPackage;
+class FLinkerLoad;
+class UObject;
+class UPackage;
+
+struct FLoadTimeProfilerTracePrivate
+{
+	static void OutputStartAsyncLoading();
+	static void OutputSuspendAsyncLoading();
+	static void OutputResumeAsyncLoading();
+	static void OutputNewLinker(const FLinkerLoad* Linker);
+	static void OutputDestroyLinker(const FLinkerLoad* Linker);
+	static void OutputNewAsyncPackage(const FAsyncPackage* AsyncPackage, const TCHAR* PackageName);
+	static void OutputDestroyAsyncPackage(const FAsyncPackage* AsyncPackage);
+	static void OutputBeginRequest(uint64 RequestId);
+	static void OutputEndRequest(uint64 RequestId);
+	static void OutputPackageSummary(const FLinkerLoad* Linker, uint32 TotalHeaderSize, uint32 NameCount, uint32 ImportCount, uint32 ExportCount);
+	static void OutputLoadImportDependency(const FAsyncPackage* Package, const FAsyncPackage* DependentPackage);
+	static void OutputAsyncPackageRequestAssociation(const FAsyncPackage* AsyncPackage, uint64 RequestId);
+	static void OutputAsyncPackageLinkerAssociation(const FAsyncPackage* AsyncPackage, const FLinkerLoad* Linker);
+	static void OutputLinkerArchiveAssociation(const FLinkerLoad* Linker, const FArchive* Archive);
+	static void OutputQueueEvent(const FAsyncPackage* AsyncPackage, ELoadTimeProfilePackageEventType EventType, TAsyncLoadPriority UserPriority, int32 PackageSerialNumber, int32 EventSystemPriority);
+	static void OutputClassInfo(const UClass* Class, const TCHAR* Name);
+
+	struct FAsyncPackageScope
+	{
+		FAsyncPackageScope(const FAsyncPackage* AsyncPackage, ELoadTimeProfilePackageEventType EventType);
+		~FAsyncPackageScope();
+	};
+
+	struct FCreateExportScope
+	{
+		FCreateExportScope(const FLinkerLoad* Linker, uint64 SerialOffset, uint64 SerialSize, bool IsAsset, const UObject* const* InObject);
+		~FCreateExportScope();
+
+	private:
+		const UObject* const* Object;
+	};
+
+	struct FObjectScope
+	{
+		FObjectScope(const UObject* Object, ELoadTimeProfilerObjectEventType EventType);
+		~FObjectScope();
+	};
+
+	struct FLinkerScope
+	{
+		FLinkerScope(const FLinkerLoad* Linker);
+		~FLinkerScope();
+	};
+
+	struct FFlushAsyncLoadingScope
+	{
+		FFlushAsyncLoadingScope(uint64 RequestId);
+		~FFlushAsyncLoadingScope();
+	};
+};
+
+#define TRACE_LOADTIME_START_ASYNC_LOADING() \
+	FLoadTimeProfilerTracePrivate::OutputStartAsyncLoading();
+
+#define TRACE_LOADTIME_SUSPEND_ASYNC_LOADING() \
+	FLoadTimeProfilerTracePrivate::OutputSuspendAsyncLoading();
+
+#define TRACE_LOADTIME_RESUME_ASYNC_LOADING() \
+	FLoadTimeProfilerTracePrivate::OutputResumeAsyncLoading();
+
+#define TRACE_LOADTIME_NEW_LINKER(Linker) \
+	FLoadTimeProfilerTracePrivate::OutputNewLinker(Linker);
+
+#define TRACE_LOADTIME_DESTROY_LINKER(Linker) \
+	FLoadTimeProfilerTracePrivate::OutputDestroyLinker(Linker);
+
+#define TRACE_LOADTIME_PACKAGE_SUMMARY(Linker, TotalHeaderSize, NameCount, ImportCount, ExportCount) \
+	FLoadTimeProfilerTracePrivate::OutputPackageSummary(Linker, TotalHeaderSize, NameCount, ImportCount, ExportCount);
+
+#define TRACE_LOADTIME_LINKER_ARCHIVE_ASSOCIATION(Linker, Archive) \
+	FLoadTimeProfilerTracePrivate::OutputLinkerArchiveAssociation(Linker, Archive);
+
+#define TRACE_LOADTIME_BEGIN_REQUEST(RequestId) \
+	FLoadTimeProfilerTracePrivate::OutputBeginRequest(RequestId);
+
+#define TRACE_LOADTIME_END_REQUEST(RequestId) \
+	FLoadTimeProfilerTracePrivate::OutputEndRequest(RequestId);
+
+#define TRACE_LOADTIME_NEW_ASYNC_PACKAGE(AsyncPackage, PackageName) \
+	FLoadTimeProfilerTracePrivate::OutputNewAsyncPackage(AsyncPackage, PackageName)
+
+#define TRACE_LOADTIME_DESTROY_ASYNC_PACKAGE(AsyncPackage) \
+	FLoadTimeProfilerTracePrivate::OutputDestroyAsyncPackage(AsyncPackage);
+
+#define TRACE_LOADTIME_ASYNC_PACKAGE_REQUEST_ASSOCIATION(AsyncPackage, RequestId) \
+	FLoadTimeProfilerTracePrivate::OutputAsyncPackageRequestAssociation(AsyncPackage, RequestId);
+
+#define TRACE_LOADTIME_ASYNC_PACKAGE_LINKER_ASSOCIATION(AsyncPackage, Linker) \
+	FLoadTimeProfilerTracePrivate::OutputAsyncPackageLinkerAssociation(AsyncPackage, Linker);
+
+#define TRACE_LOADTIME_QUEUE_EVENT(AsyncPackage, EventType, UserPriority, PackageSerialNumber, EventSystemPriority) \
+	FLoadTimeProfilerTracePrivate::OutputQueueEvent(AsyncPackage, EventType, UserPriority, PackageSerialNumber, EventSystemPriority);
+
+#define TRACE_LOADTIME_ASYNC_PACKAGE_SCOPE(AsyncPackage, EventType) \
+	FLoadTimeProfilerTracePrivate::FAsyncPackageScope __LoadTimeTraceAsyncPackageScope(AsyncPackage, EventType);
+
+#define TRACE_LOADTIME_CREATE_EXPORT_SCOPE(Linker, Object, SerialOffset, SerialSize, IsAsset) \
+	FLoadTimeProfilerTracePrivate::FCreateExportScope __LoadTimeTraceCreateExportScope(Linker, SerialOffset, SerialSize, IsAsset, Object);
+
+#define TRACE_LOADTIME_OBJECT_SCOPE(Object, EventType) \
+	FLoadTimeProfilerTracePrivate::FObjectScope __LoadTimeTraceObjectScope(Object, EventType);
+
+#define TRACE_LOADTIME_FLUSH_ASYNCLOADING_SCOPE(RequestId) \
+	FLoadTimeProfilerTracePrivate::FFlushAsyncLoadingScope __LoadTimeTraceFlushAsyncLoadingScope(RequestId);
+
+#define TRACE_LOADTIME_CLASS_INFO(Class, Name) \
+	FMiscTrace::OutputClassInfo(Class, Name);
+
+#else
+
+#define TRACE_LOADTIME_START_ASYNC_LOADING(...)
+#define TRACE_LOADTIME_SUSPEND_ASYNC_LOADING(...)
+#define TRACE_LOADTIME_RESUME_ASYNC_LOADING(...)
+#define TRACE_LOADTIME_NEW_PACKAGE(...)
+#define TRACE_LOADTIME_NEW_LINKER(...)
+#define TRACE_LOADTIME_DESTROY_LINKER(...)
+#define TRACE_LOADTIME_PACKAGE_SUMMARY(...)
+#define TRACE_LOADTIME_LINKER_ARCHIVE_ASSOCIATION(...)
+#define TRACE_LOADTIME_BEGIN_REQUEST(...)
+#define TRACE_LOADTIME_END_REQUEST(...)
+#define TRACE_LOADTIME_NEW_ASYNC_PACKAGE(...)
+#define TRACE_LOADTIME_DESTROY_ASYNC_PACKAGE(...)
+#define TRACE_LOADTIME_ASYNC_PACKAGE_REQUEST_ASSOCIATION(...)
+#define TRACE_LOADTIME_ASYNC_PACKAGE_LINKER_ASSOCIATION(...)
+#define TRACE_LOADTIME_FLUSH_ASYNCLOADING_SCOPE(...)
+#define TRACE_LOADTIME_QUEUE_EVENT(...)
+#define TRACE_LOADTIME_ASYNC_PACKAGE_SCOPE(...)
+#define TRACE_LOADTIME_TIMER_SCOPE(...)
+#define TRACE_LOADTIME_CREATE_EXPORT_SCOPE(...)
+#define TRACE_LOADTIME_OBJECT_SCOPE(...)
+#define TRACE_LOADTIME_CLASS_INFO(...)
+
+#endif
\ No newline at end of file
diff --git a/Engine/Source/Runtime/CoreUObject/Private/UObject/Class.cpp b/Engine/Source/Runtime/CoreUObject/Private/UObject/Class.cpp
--- a/Engine/Source/Runtime/CoreUObject/Private/UObject/Class.cpp
+++ b/Engine/Source/Runtime/CoreUObject/Private/UObject/Class.cpp
@@ -43,6 +43,7 @@
 #include "UObject/FrameworkObjectVersion.h"
 #include "UObject/GarbageCollection.h"
 #include "UObject/UObjectThreadContext.h"
+#include "Serialization/LoadTimeTracePrivate.h"
 
 // This flag enables some expensive class tree validation that is meant to catch mutations of 
 // the class tree outside of SetSuperStruct. It has been disabled because loading blueprints 
@@ -81,6 +82,7 @@ COREUOBJECT_API void InitializePrivateStaticClass(
 	const TCHAR* Name
 	)
 {
+	TRACE_LOADTIME_CLASS_INFO(TClass_PrivateStaticClass, Name);
 	NotifyRegistrationEvent(PackageName, Name, ENotifyRegistrationType::NRT_Class, ENotifyRegistrationPhase::NRP_Started);
 
 	/* No recursive ::StaticClass calls allowed. Setup extras. */
diff --git a/Engine/Source/Runtime/CoreUObject/Private/UObject/LinkerLoad.cpp b/Engine/Source/Runtime/CoreUObject/Private/UObject/LinkerLoad.cpp
--- a/Engine/Source/Runtime/CoreUObject/Private/UObject/LinkerLoad.cpp
+++ b/Engine/Source/Runtime/CoreUObject/Private/UObject/LinkerLoad.cpp
@@ -35,6 +35,7 @@
 #include "Serialization/Formatters/BinaryArchiveFormatter.h"
 #include "Serialization/Formatters/JsonArchiveInputFormatter.h"
 #include "Serialization/ArchiveUObjectFromStructuredArchive.h"
+#include "Serialization/LoadTimeTracePrivate.h"
 #include "HAL/FileManager.h"
 #include "UObject/CoreRedirects.h"
 
@@ -857,10 +858,12 @@ FLinkerLoad::FLinkerLoad(UPackage* InParent, const TCHAR* InFilename, uint32 InL
 #endif
 
 	OwnerThread = FPlatformTLS::GetCurrentThreadId();
+	TRACE_LOADTIME_NEW_LINKER(this);
 }
 
 FLinkerLoad::~FLinkerLoad()
 {
+	TRACE_LOADTIME_DESTROY_LINKER(this);
 #if !UE_BUILD_SHIPPING && !UE_BUILD_TEST
 	FLinkerManager::Get().GetLiveLinkers().Remove(this);
 #endif
@@ -1030,6 +1033,8 @@ FLinkerLoad::ELinkerStatus FLinkerLoad::CreateLoader(
 					Loader = IFileManager::Get().CreateFileReader(*Filename);
 				}
 
+				TRACE_LOADTIME_LINKER_ARCHIVE_ASSOCIATION(this, Loader);
+
 				if (!Loader)
 				{
 					UE_LOG(LogLinker, Warning, TEXT("Error opening file '%s'."), *Filename);
@@ -1160,6 +1165,8 @@ FLinkerLoad::ELinkerStatus FLinkerLoad::SerializePackageFileSummary()
 		// Read summary from file.
 		StructuredArchiveRootRecord.GetValue() << NAMED_FIELD(Summary);
 
+		TRACE_LOADTIME_PACKAGE_SUMMARY(this, Summary.TotalHeaderSize, Summary.NameCount, Summary.ImportCount, Summary.ExportCount);
+
 		// Check tag.
 		if( Summary.Tag != PACKAGE_FILE_TAG )
 		{
@@ -1431,7 +1438,7 @@ FLinkerLoad::ELinkerStatus FLinkerLoad::SerializeNameMap()
 
 	while (bFinishedPrecaching && NameMapIndex < Summary.NameCount && !IsTimeLimitExceeded(TEXT("serializing name map"), 100))
 	{
-		SCOPED_LOADTIMER(LinkerLoad_SerializeNameMap_ProcessingEntries);
+		//SCOPED_LOADTIMER(LinkerLoad_SerializeNameMap_ProcessingEntries);
 
 		FNameEntrySerialized NameEntry(ENAME_LinkerConstructor);
 
@@ -3430,7 +3437,6 @@ UObject* FLinkerLoad::Create( UClass* ObjectClass, FName ObjectName, UObject* Ou
 
 void FLinkerLoad::Preload( UObject* Object )
 {
-
 	//check(IsValidLowLevel());
 	check(Object);
 
@@ -3537,6 +3543,7 @@ void FLinkerLoad::Preload( UObject* Object )
 				Object->ClearFlags ( RF_NeedLoad );
 
 				{
+					TRACE_LOADTIME_OBJECT_SCOPE(Export.Object, LoadTimeProfilerObjectEventType_Serialize);
 					SCOPE_CYCLE_COUNTER(STAT_LinkerSerialize);
 #if USE_CIRCULAR_DEPENDENCY_LOAD_DEFERRING
 					// communicate with FLinkerPlaceholderBase, what object is currently serializing in
@@ -3868,6 +3875,8 @@ UObject* FLinkerLoad::CreateExport( int32 Index )
 	// Check whether we already loaded the object and if not whether the context flags allow loading it.
 	if( !Export.Object && !FilterExport(Export) ) // for some acceptable position, it was not "not for" 
 	{
+		TRACE_LOADTIME_CREATE_EXPORT_SCOPE(this, &Export.Object, Export.SerialOffset, Export.SerialSize, Export.bIsAsset);
+
 		FUObjectSerializeContext* CurrentLoadContext = GetSerializeContext();
 		check(!GEventDrivenLoaderEnabled || !bLockoutLegacyOperations || !EVENT_DRIVEN_ASYNC_LOAD_ACTIVE_AT_RUNTIME);
 		check(Export.ObjectName!=NAME_None || !(Export.ObjectFlags&RF_Public));
diff --git a/Engine/Source/Runtime/CoreUObject/Private/UObject/UObjectGlobals.cpp b/Engine/Source/Runtime/CoreUObject/Private/UObject/UObjectGlobals.cpp
--- a/Engine/Source/Runtime/CoreUObject/Private/UObject/UObjectGlobals.cpp
+++ b/Engine/Source/Runtime/CoreUObject/Private/UObject/UObjectGlobals.cpp
@@ -37,6 +37,7 @@
 #include "Serialization/DuplicatedObject.h"
 #include "Serialization/DuplicatedDataReader.h"
 #include "Serialization/DuplicatedDataWriter.h"
+#include "Serialization/LoadTimeTracePrivate.h"
 #include "Misc/PackageName.h"
 #include "UObject/LinkerLoad.h"
 #include "Blueprint/BlueprintSupport.h"
diff --git a/Engine/Source/Runtime/CoreUObject/Public/Serialization/LoadTimeTrace.h b/Engine/Source/Runtime/CoreUObject/Public/Serialization/LoadTimeTrace.h
new file mode 100644
--- /dev/null
+++ b/Engine/Source/Runtime/CoreUObject/Public/Serialization/LoadTimeTrace.h
@@ -0,0 +1,89 @@
+// Copyright 1998-2019 Epic Games, Inc. All Rights Reserved.
+
+#pragma once
+
+#if !IS_PROGRAM && !UE_BUILD_SHIPPING
+#define LOADTIMEPROFILERTRACE_ENABLED 0
+#else
+#define LOADTIMEPROFILERTRACE_ENABLED 0
+#endif
+
+enum ELoadTimeProfilePackageEventType
+{
+	LoadTimeProfilerPackageEventType_CreateLinker,
+	LoadTimeProfilerPackageEventType_FinishLinker,
+	LoadTimeProfilerPackageEventType_StartImportPackages,
+	LoadTimeProfilerPackageEventType_SetupImports,
+	LoadTimeProfilerPackageEventType_SetupExports,
+	LoadTimeProfilerPackageEventType_ProcessImportsAndExports,
+	LoadTimeProfilerPackageEventType_ExportsDone,
+	LoadTimeProfilerPackageEventType_PostLoadWait,
+	LoadTimeProfilerPackageEventType_StartPostLoad,
+	LoadTimeProfilerPackageEventType_Tick,
+	LoadTimeProfilerPackageEventType_Finish,
+	LoadTimeProfilerPackageEventType_None,
+};
+
+enum ELoadTimeProfilerObjectEventType
+{
+	LoadTimeProfilerObjectEventType_Create,
+	LoadTimeProfilerObjectEventType_Serialize,
+	LoadTimeProfilerObjectEventType_PostLoad,
+	LoadTimeProfilerObjectEventType_None
+};
+
+#if LOADTIMEPROFILERTRACE_ENABLED
+
+struct FStreamableHandle;
+
+struct FLoadTimeProfilerTrace
+{
+	COREUOBJECT_API static void OutputNewStreamableHandle(const FStreamableHandle* StreamableHandle, const TCHAR* DebugName, bool IsCombined);
+	COREUOBJECT_API static void OutputDestroyStreamableHandle(const FStreamableHandle* StreamableHandle);
+	COREUOBJECT_API static void OutputBeginLoadStreamableHandle(const FStreamableHandle* StreamableHandle);
+	COREUOBJECT_API static void OutputEndLoadStreamableHandle(const FStreamableHandle* StreamableHandle);
+	COREUOBJECT_API static void OutputStreamableHandleRequestAssociation(const FStreamableHandle* StreamableHandle, uint64 RequestId);
+
+	struct FLoadMapScope
+	{
+		COREUOBJECT_API FLoadMapScope(const TCHAR* Name);
+		COREUOBJECT_API ~FLoadMapScope();
+	};
+
+	struct FWaitForStreamableHandleScope
+	{
+		COREUOBJECT_API FWaitForStreamableHandleScope(const FStreamableHandle* StreamableHandle);
+		COREUOBJECT_API ~FWaitForStreamableHandleScope();
+	};
+};
+
+#define TRACE_LOADTIME_LOAD_MAP_SCOPE(Name) \
+	FLoadTimeProfilerTrace::FLoadMapScope __LoadTimeTraceLoadMapScope(Name);
+
+#define TRACE_LOADTIME_NEW_STREAMABLE_HANDLE(StreamableHandle, DebugName, IsCombined) \
+	FLoadTimeProfilerTrace::OutputNewStreamableHandle(StreamableHandle, DebugName, IsCombined);
+
+#define TRACE_LOADTIME_DESTROY_STREAMABLE_HANDLE(StreamableHandle) \
+	FLoadTimeProfilerTrace::OutputDestroyStreamableHandle(StreamableHandle);
+
+#define TRACE_LOADTIME_BEGIN_LOAD_STREAMABLE_HANDLE(StreamableHandle) \
+	FLoadTimeProfilerTrace::OutputBeginLoadStreamableHandle(StreamableHandle);
+
+#define TRACE_LOADTIME_END_LOAD_STREAMABLE_HANDLE(StreamableHandle) \
+	FLoadTimeProfilerTrace::OutputEndLoadStreamableHandle(StreamableHandle);
+
+#define TRACE_LOADTIME_STREAMABLE_HANDLE_REQUEST_ASSOCIATION(StreamableHandle, RequestId) \
+	FLoadTimeProfilerTrace::OutputStreamableHandleRequestAssociation(StreamableHandle, RequestId);
+
+#define TRACE_LOADTIME_WAIT_FOR_STREAMABLE_HANDLE_SCOPE(StreamableHandle) \
+	FLoadTimeProfilerTrace::FWaitForStreamableHandleScope __LoadTimeTraceWaitForStreamableHandleScope(StreamableHandle);
+
+#else
+#define TRACE_LOADTIME_LOAD_MAP_SCOPE(...)
+#define TRACE_LOADTIME_NEW_STREAMABLE_HANDLE(...)
+#define TRACE_LOADTIME_DESTROY_STREAMABLE_HANDLE(...)
+#define TRACE_LOADTIME_BEGIN_LOAD_STREAMABLE_HANDLE(...)
+#define TRACE_LOADTIME_END_LOAD_STREAMABLE_HANDLE(...)
+#define TRACE_LOADTIME_STREAMABLE_HANDLE_REQUEST_ASSOCIATION(...)
+#define TRACE_LOADTIME_WAIT_FOR_STREAMABLE_HANDLE_SCOPE(...)
+#endif
\ No newline at end of file
diff --git a/Engine/Source/Runtime/Engine/Private/StreamableManager.cpp b/Engine/Source/Runtime/Engine/Private/StreamableManager.cpp
--- a/Engine/Source/Runtime/Engine/Private/StreamableManager.cpp
+++ b/Engine/Source/Runtime/Engine/Private/StreamableManager.cpp
@@ -7,6 +7,7 @@
 #include "UObject/UObjectThreadContext.h"
 #include "HAL/IConsoleManager.h"
 #include "Tickable.h"
+#include "Serialization/LoadTimeTrace.h"
 
 DEFINE_LOG_CATEGORY_STATIC(LogStreamableManager, Log, All);
 
@@ -199,6 +200,8 @@ bool FStreamableHandle::BindUpdateDelegate(FStreamableUpdateDelegate NewDelegate
 
 EAsyncPackageState::Type FStreamableHandle::WaitUntilComplete(float Timeout, bool bStartStalledHandles)
 {
+	TRACE_LOADTIME_WAIT_FOR_STREAMABLE_HANDLE_SCOPE(this);
+
 	if (HasLoadCompleted())
 	{
 		return EAsyncPackageState::Complete;
@@ -503,6 +506,7 @@ void FStreamableHandle::StartStalledHandle()
 
 FStreamableHandle::~FStreamableHandle()
 {
+	TRACE_LOADTIME_DESTROY_STREAMABLE_HANDLE(this);
 	check(IsInGameThread());
 
 	if (IsActive())
@@ -516,6 +520,7 @@ FStreamableHandle::~FStreamableHandle()
 
 void FStreamableHandle::CompleteLoad()
 {
+	TRACE_LOADTIME_END_LOAD_STREAMABLE_HANDLE(this);
 	// Only complete if it's still active
 	if (IsActive())
 	{
@@ -890,7 +895,8 @@ FStreamable* FStreamableManager::StreamInternal(const FSoftObjectPath& InTargetN
 
 			Existing->bAsyncLoadRequestOutstanding = true;
 			Existing->bLoadFailed = false;
-			LoadPackageAsync(Package, FLoadPackageAsyncDelegate::CreateSP(Handle, &FStreamableHandle::AsyncLoadCallbackWrapper, TargetName), Priority);
+			int32 RequestId = LoadPackageAsync(Package, FLoadPackageAsyncDelegate::CreateSP(Handle, &FStreamableHandle::AsyncLoadCallbackWrapper, TargetName), Priority);
+			TRACE_LOADTIME_STREAMABLE_HANDLE_REQUEST_ASSOCIATION(&Handle.Get(), RequestId);
 		}
 	}
 	return Existing;
@@ -918,6 +924,8 @@ TSharedPtr<FStreamableHandle> FStreamableManager::RequestAsyncLoad(const TArray<
 	NewRequest->DebugName = DebugName;
 	NewRequest->Priority = Priority;
 
+	TRACE_LOADTIME_NEW_STREAMABLE_HANDLE(&NewRequest.Get(), *DebugName, false);
+
 	// Remove null requests
 
 	for (int32 i = NewRequest->RequestedAssets.Num() - 1; i >= 0 ; i--)
@@ -1047,6 +1055,8 @@ TSharedPtr<FStreamableHandle> FStreamableManager::RequestSyncLoad(const FSoftObj
 
 void FStreamableManager::StartHandleRequests(TSharedRef<FStreamableHandle> Handle)
 {
+	TRACE_LOADTIME_BEGIN_LOAD_STREAMABLE_HANDLE(&Handle.Get());
+
 	TArray<FStreamable *> ExistingStreamables;
 	ExistingStreamables.Reserve(Handle->RequestedAssets.Num());
 
@@ -1290,6 +1300,7 @@ TSharedPtr<FStreamableHandle> FStreamableManager::CreateCombinedHandle(const TAr
 	NewRequest->OwningManager = this;
 	NewRequest->bIsCombinedHandle = true;
 	NewRequest->DebugName = DebugName;
+	TRACE_LOADTIME_NEW_STREAMABLE_HANDLE(&NewRequest.Get(), *DebugName, true);
 
 	for (TSharedPtr<FStreamableHandle> ChildHandle : ChildHandles)
 	{
diff --git a/Engine/Source/Runtime/Engine/Private/UnrealEngine.cpp b/Engine/Source/Runtime/Engine/Private/UnrealEngine.cpp
--- a/Engine/Source/Runtime/Engine/Private/UnrealEngine.cpp
+++ b/Engine/Source/Runtime/Engine/Private/UnrealEngine.cpp
@@ -139,6 +139,7 @@ UnrealEngine.cpp: Implements the UEngine class and helpers.
 #include "Sound/AudioSettings.h"
 #include "Streaming/Texture2DUpdate.h"
 #include "Rendering/SkeletalMeshRenderData.h"
+#include "Serialization/LoadTimeTrace.h"
 
 #if WITH_EDITOR
 #include "Settings/LevelEditorPlaySettings.h"
@@ -12087,7 +12088,9 @@ void UEngine::TickWorldTravel(FWorldContext& Context, float DeltaSeconds)
 
 bool UEngine::LoadMap( FWorldContext& WorldContext, FURL URL, class UPendingNetGame* Pending, FString& Error )
 {
+	TRACE_LOADTIME_LOAD_MAP_SCOPE(*URL.Map);
 	STAT_ADD_CUSTOMMESSAGE_NAME( STAT_NamedMarker, *(FString( TEXT( "LoadMap - " ) + URL.Map )) );
+	TRACE_BOOKMARK(TEXT("LoadMap - %s"), *URL.Map);
 
 	DECLARE_SCOPE_CYCLE_COUNTER(TEXT("UEngine::LoadMap"), STAT_LoadMap, STATGROUP_LoadTime);
 
@@ -12635,6 +12638,7 @@ bool UEngine::LoadMap( FWorldContext& WorldContext, FURL URL, class UPendingNetG
 	}
 
 	STAT_ADD_CUSTOMMESSAGE_NAME( STAT_NamedMarker, *(FString( TEXT( "LoadMapComplete - " ) + URL.Map )) );
+	TRACE_BOOKMARK(TEXT("LoadMapComplete - %s"), *URL.Map);
 	MALLOC_PROFILER( FMallocProfiler::SnapshotMemoryLoadMapEnd( URL.Map ); )
 
 		double StopTime = FPlatformTime::Seconds();
@@ -13337,6 +13341,7 @@ static void AsyncMapChangeLevelLoadCompletionCallback(const FName& PackageName,
 	}
 
 	STAT_ADD_CUSTOMMESSAGE_NAME( STAT_NamedMarker, *(FString( TEXT( "PrepareMapChangeComplete - " ) + PackageName.ToString() )) );
+	TRACE_BOOKMARK(TEXT("PrepareMapChangeComplete - %s"), *PackageName.ToString());
 }
 
 
@@ -13382,6 +13387,7 @@ bool UEngine::PrepareMapChange(FWorldContext &Context, const TArray<FName>& Leve
 		for (const FName LevelName : Context.LevelsToLoadForPendingMapChange)
 		{
 			STAT_ADD_CUSTOMMESSAGE_NAME( STAT_NamedMarker, *(FString( TEXT( "PrepareMapChange - " ) + LevelName.ToString() )) );
+			TRACE_BOOKMARK(TEXT("PrepareMapChange - %s"), *LevelName.ToString());
 			LoadPackageAsync(LevelName.ToString(),
 				FLoadPackageAsyncDelegate::CreateStatic(&AsyncMapChangeLevelLoadCompletionCallback, Context.ContextHandle)
 			);
diff --git a/Engine/Source/Runtime/Launch/Launch.Build.cs b/Engine/Source/Runtime/Launch/Launch.Build.cs
--- a/Engine/Source/Runtime/Launch/Launch.Build.cs
+++ b/Engine/Source/Runtime/Launch/Launch.Build.cs
@@ -31,6 +31,7 @@ public class Launch : ModuleRules
 				"Slate",
 				"SlateCore",
 				"Sockets",
+				"Trace",
 				"Overlay",
 				"UtilityShaders",
 				"PreLoadScreen"
diff --git a/Engine/Source/Runtime/Launch/Private/LaunchEngineLoop.cpp b/Engine/Source/Runtime/Launch/Private/LaunchEngineLoop.cpp
--- a/Engine/Source/Runtime/Launch/Private/LaunchEngineLoop.cpp
+++ b/Engine/Source/Runtime/Launch/Private/LaunchEngineLoop.cpp
@@ -28,6 +28,8 @@
 #include "HAL/FileManagerGeneric.h"
 #include "HAL/ExceptionHandling.h"
 #include "Stats/StatsMallocProfilerProxy.h"
+#include "Trace/Trace.h"
+#include "ProfilingDebugging/MiscTrace.h"
 #if WITH_ENGINE
 #include "HAL/PlatformSplash.h"
 #endif
@@ -1155,6 +1157,8 @@ DECLARE_CYCLE_STAT( TEXT( "FEngineLoop::PreInit.AfterStats" ), STAT_FEngineLoop_
 
 int32 FEngineLoop::PreInit(const TCHAR* CmdLine)
 {
+	TRACE_REGISTER_GAME_THREAD(FPlatformTLS::GetCurrentThreadId());
+
 	SCOPED_BOOT_TIMING("FEngineLoop::PreInit");
 
 #if PLATFORM_WINDOWS
@@ -1202,6 +1206,18 @@ int32 FEngineLoop::PreInit(const TCHAR* CmdLine)
 		return -1;
 	}
 
+	{
+		const TCHAR* TraceHostPtr = TEXT("127.0.0.1");
+
+		FString TraceHost;
+		if (FParse::Value(CmdLine, TEXT("tracehost"), TraceHost))
+		{
+			TraceHostPtr = *TraceHost;
+		}
+
+		Trace::Connect(TraceHostPtr);
+	}
+
 #if WITH_ENGINE
 	FCoreUObjectDelegates::PostGarbageCollectConditionalBeginDestroy.AddStatic(DeferredPhysResourceCleanup);
 #endif
@@ -1759,6 +1775,7 @@ int32 FEngineLoop::PreInit(const TCHAR* CmdLine)
 
 
 		{
+			TRACE_THREAD_GROUP_SCOPE(TraceThreadGroup_ThreadPool);
 			GThreadPool = FQueuedThreadPool::Allocate();
 			int32 NumThreadsInThreadPool = FPlatformMisc::NumberOfWorkerThreadsToSpawn();
 
@@ -1770,6 +1787,7 @@ int32 FEngineLoop::PreInit(const TCHAR* CmdLine)
 			verify(GThreadPool->Create(NumThreadsInThreadPool, StackSize * 1024, TPri_SlightlyBelowNormal));
 		}
 		{
+			TRACE_THREAD_GROUP_SCOPE(TraceThreadGroup_BackgroundThreadPool);
 			GBackgroundPriorityThreadPool = FQueuedThreadPool::Allocate();
 			int32 NumThreadsInThreadPool = 2;
 			if (FPlatformProperties::IsServerOnly())
@@ -1781,12 +1799,15 @@ int32 FEngineLoop::PreInit(const TCHAR* CmdLine)
 		}
 
 #if WITH_EDITOR
-		// when we are in the editor we like to do things like build lighting and such
-		// this thread pool can be used for those purposes
-		GLargeThreadPool = FQueuedThreadPool::Allocate();
-		int32 NumThreadsInLargeThreadPool = FMath::Max(FPlatformMisc::NumberOfCoresIncludingHyperthreads() - 2, 2);
+		{
+			TRACE_THREAD_GROUP_SCOPE(TraceThreadGroup_LargeThreadPool);
+			// when we are in the editor we like to do things like build lighting and such
+			// this thread pool can be used for those purposes
+			GLargeThreadPool = FQueuedThreadPool::Allocate();
+			int32 NumThreadsInLargeThreadPool = FMath::Max(FPlatformMisc::NumberOfCoresIncludingHyperthreads() - 2, 2);
 
-		verify(GLargeThreadPool->Create(NumThreadsInLargeThreadPool, 128 * 1024));
+			verify(GLargeThreadPool->Create(NumThreadsInLargeThreadPool, 128 * 1024));
+		}
 #endif
 	}
 
@@ -1836,6 +1857,7 @@ int32 FEngineLoop::PreInit(const TCHAR* CmdLine)
 	if (FPlatformProcess::SupportsMultithreading())
 	{
 		{
+			TRACE_THREAD_GROUP_SCOPE(TraceThreadGroup_IOThreadPool);
 			SCOPED_BOOT_TIMING("GIOThreadPool->Create");
 			GIOThreadPool = FQueuedThreadPool::Allocate();
 			int32 NumThreadsInThreadPool = FPlatformMisc::NumberOfIOWorkerThreadsToSpawn();
@@ -3570,6 +3592,7 @@ int32 FEngineLoop::Init()
 void FEngineLoop::Exit()
 {
 	STAT_ADD_CUSTOMMESSAGE_NAME( STAT_NamedMarker, TEXT( "EngineLoop.Exit" ) );
+	TRACE_BOOKMARK(TEXT("EngineLoop.Exit"));
 
 	GIsRunning	= 0;
 	GLogConsole	= nullptr;
@@ -3920,6 +3943,7 @@ uint64 FScopedSampleMallocChurn::DumpFrame = 0;
 
 static inline void BeginFrameRenderThread(FRHICommandListImmediate& RHICmdList, uint64 CurrentFrameCounter)
 {
+	TRACE_BEGIN_FRAME(TraceFrameType_Rendering);
 	GRHICommandList.LatchBypass();
 	GFrameNumberRenderThread++;
 
@@ -3961,6 +3985,7 @@ static inline void EndFrameRenderThread(FRHICommandListImmediate& RHICmdList)
 	FPlatformMisc::EndNamedEvent();
 #endif
 #endif // !UE_BUILD_SHIPPING 
+	TRACE_END_FRAME(TraceFrameType_Rendering);
 }
 
 #if BUILD_EMBEDDED_APP
@@ -4019,6 +4044,8 @@ void FEngineLoop::Tick()
 	}
 
 	{
+		TRACE_BEGIN_FRAME(TraceFrameType_Game);
+
 		SCOPE_CYCLE_COUNTER(STAT_FrameTime);
 
 		#if WITH_PROFILEGPU && !UE_BUILD_SHIPPING
@@ -4448,6 +4475,7 @@ void FEngineLoop::Tick()
 #if UE_GC_TRACK_OBJ_AVAILABLE
 		SET_DWORD_STAT(STAT_Hash_NumObjects, GUObjectArray.GetObjectArrayNumMinusAvailable());
 #endif
+		TRACE_END_FRAME(TraceFrameType_Game);
 	}
 
 #if BUILD_EMBEDDED_APP
diff --git a/Engine/Source/Runtime/RHI/Private/GpuProfilerTrace.cpp b/Engine/Source/Runtime/RHI/Private/GpuProfilerTrace.cpp
new file mode 100644
--- /dev/null
+++ b/Engine/Source/Runtime/RHI/Private/GpuProfilerTrace.cpp
@@ -0,0 +1,79 @@
+// Copyright 1998-2019 Epic Games, Inc. All Rights Reserved.
+#include "GpuProfilerTrace.h"
+#include "Trace/Trace.h"
+#include "Misc/CString.h"
+#include "GPUProfiler.h"
+#include "ProfilingDebugging/MiscTrace.h"
+#include "HAL/PlatformMisc.h"
+
+#if GPUPROFILERTRACE_ENABLED
+
+UE_TRACE_EVENT_BEGIN(GpuProfiler, EventSpec)
+	UE_TRACE_EVENT_FIELD(const void*, EventType)
+	UE_TRACE_EVENT_FIELD(uint16, NameLength)
+UE_TRACE_EVENT_END()
+
+UE_TRACE_EVENT_BEGIN(GpuProfiler, Frame)
+	UE_TRACE_EVENT_FIELD(uint64, TimestampBase)
+	UE_TRACE_EVENT_FIELD(uint32, RenderingFrameNumber)
+UE_TRACE_EVENT_END()
+
+FGpuProfilerTrace::FFrame FGpuProfilerTrace::CurrentFrame;
+
+FGpuProfilerTrace::FEventType::FEventType(const TCHAR* Name)
+{
+	uint16 NameLength = FCString::Strlen(Name);
+	uint16 NameSize = (NameLength + 1) * sizeof(TCHAR);
+	UE_TRACE_LOG(GpuProfiler, EventSpec, NameSize) 
+		<< EventSpec.EventType(this)
+		<< EventSpec.NameLength(NameLength)
+		<< EventSpec.Attachment(Name, NameSize);
+}
+
+void FGpuProfilerTrace::BeginFrame()
+{
+	CurrentFrame.TimestampBase = 0;
+	CurrentFrame.EventBufferSize = 0;
+}
+
+void FGpuProfilerTrace::BeginEvent(const FEventType* EventType, uint32 FrameNumber, uint64 TimestampMicroseconds)
+{
+	if (CurrentFrame.EventBufferSize >= MaxEventBufferSize - 17) // 9 + 8
+	{
+		return;
+	}
+	if (CurrentFrame.TimestampBase == 0)
+	{
+		CurrentFrame.TimestampBase = TimestampMicroseconds;
+		CurrentFrame.LastTimestamp = CurrentFrame.TimestampBase;
+		CurrentFrame.RenderingFrameNumber = FrameNumber;
+	}
+	uint8* BufferPtr = CurrentFrame.EventBuffer + CurrentFrame.EventBufferSize;
+	uint64 TimestampDelta = TimestampMicroseconds - CurrentFrame.LastTimestamp;
+	CurrentFrame.LastTimestamp = TimestampMicroseconds;
+	FTraceUtils::Encode7bit((TimestampDelta << 1ull) | 0x1, BufferPtr);
+	*reinterpret_cast<uint64*>(BufferPtr) = uint64(EventType);
+	CurrentFrame.EventBufferSize = BufferPtr - CurrentFrame.EventBuffer + sizeof(uint64);
+}
+
+void FGpuProfilerTrace::EndEvent(uint64 TimestampMicroseconds)
+{
+	uint64 TimestampDelta = TimestampMicroseconds - CurrentFrame.LastTimestamp;
+	CurrentFrame.LastTimestamp = TimestampMicroseconds;
+	uint8* BufferPtr = CurrentFrame.EventBuffer + CurrentFrame.EventBufferSize;
+	FTraceUtils::Encode7bit(TimestampDelta << 1ull, BufferPtr);
+	CurrentFrame.EventBufferSize = BufferPtr - CurrentFrame.EventBuffer;
+}
+
+void FGpuProfilerTrace::EndFrame()
+{
+	if (CurrentFrame.EventBufferSize)
+	{
+		UE_TRACE_LOG(GpuProfiler, Frame, CurrentFrame.EventBufferSize)
+			<< Frame.TimestampBase(CurrentFrame.TimestampBase)
+			<< Frame.RenderingFrameNumber(CurrentFrame.RenderingFrameNumber)
+			<< Frame.Attachment(CurrentFrame.EventBuffer, CurrentFrame.EventBufferSize);
+	}
+}
+
+#endif
