TBA: Fix typo in json input formatter

diff --git a/Engine/Source/Runtime/CoreUObject/Private/Serialization/Formatters/JsonArchiveInputFormatter.cpp b/Engine/Source/Runtime/CoreUObject/Private/Serialization/Formatters/JsonArchiveInputFormatter.cpp
--- a/Engine/Source/Runtime/CoreUObject/Private/Serialization/Formatters/JsonArchiveInputFormatter.cpp
+++ b/Engine/Source/Runtime/CoreUObject/Private/Serialization/Formatters/JsonArchiveInputFormatter.cpp
@@ -250,7 +250,7 @@ void FJsonArchiveInputFormatter::Serialize(int32& Value)
 
 void FJsonArchiveInputFormatter::Serialize(int64& Value)
 {
-	Value = (uint64)ValueStack.Top()->AsNumber();
+	Value = (int64)ValueStack.Top()->AsNumber();
 }
 
 void FJsonArchiveInputFormatter::Serialize(float& Value)
