TBA: Fix incorrect text serialization of soft object ptrs

diff --git a/Engine/Source/Runtime/CoreUObject/Private/Serialization/Formatters/JsonArchiveOutputFormatter.cpp b/Engine/Source/Runtime/CoreUObject/Private/Serialization/Formatters/JsonArchiveOutputFormatter.cpp
--- a/Engine/Source/Runtime/CoreUObject/Private/Serialization/Formatters/JsonArchiveOutputFormatter.cpp
+++ b/Engine/Source/Runtime/CoreUObject/Private/Serialization/Formatters/JsonArchiveOutputFormatter.cpp
@@ -314,14 +314,8 @@ void FJsonArchiveOutputFormatter::Serialize(FWeakObjectPtr& Value)
 
 void FJsonArchiveOutputFormatter::Serialize(FSoftObjectPtr& Value)
 {
-	if (Value.IsValid())
-	{
-		SerializeStringInternal(FString::Printf(TEXT("Object:%s"), *Value.Get()->GetFullName()));
-	}
-	else
-	{
-		WriteValue(TEXT("null"));
-	}
+	FSoftObjectPath Path = Value.ToSoftObjectPath();
+	Serialize(Path);
 }
 
 void FJsonArchiveOutputFormatter::Serialize(FSoftObjectPath& Value)
