Copying //UE4/Dev-Core to Main (//UE4/Main)

diff --git a/Engine/Source/Runtime/CoreUObject/Private/Serialization/AsyncLoading.h b/Engine/Source/Runtime/CoreUObject/Private/Serialization/AsyncLoading.h
new file mode 100644
--- /dev/null
+++ b/Engine/Source/Runtime/CoreUObject/Private/Serialization/AsyncLoading.h
@@ -0,0 +1,960 @@
+// Copyright 1998-2019 Epic Games, Inc. All Rights Reserved.
+
+/*=============================================================================
+	AsyncLoading.h: Unreal async loading definitions.
+=============================================================================*/
+
+#pragma once
+
+#include "CoreMinimal.h"
+#include "HAL/ThreadSafeCounter.h"
+#include "UObject/ObjectResource.h"
+#include "UObject/GCObject.h"
+#include "Serialization/AsyncPackage.h"
+#include "UObject/Package.h"
+#include "Templates/Casts.h"
+#include "UObject/ObjectRedirector.h"
+#include "UObject/UObjectThreadContext.h"
+#include "Templates/RefCounting.h"
+#include "Serialization/AsyncPackageLoader.h"
+class IAsyncReadRequest;
+struct FAsyncPackage;
+struct FFlushTree;
+class FAsyncLoadingThread;
+
+struct FAsyncPackage;
+/** [EDL] Async Package Loading State */
+enum class EAsyncPackageLoadingState : uint8
+{
+	NewPackage,
+	WaitingForSummary,
+	StartImportPackages,
+	WaitingForImportPackages,
+	SetupImports,
+	SetupExports,
+	ProcessNewImportsAndExports,
+	WaitingForPostLoad,
+	ReadyForPostLoad,
+	PostLoad_Etc,
+	PackageComplete,
+};
+
+/** [EDL] This version is an ordinary pointer. We can swap in the safer version to verify assumptions */
+struct FUnsafeWeakAsyncPackagePtr
+{
+	FAsyncPackage* Package;
+	FUnsafeWeakAsyncPackagePtr(FAsyncPackage* InPackage = nullptr)
+		: Package(InPackage)
+	{
+	}
+	FORCEINLINE FAsyncPackage& GetPackage() const
+	{
+		check(Package);
+		return *Package;
+	}
+	FORCEINLINE friend uint32 GetTypeHash(const FUnsafeWeakAsyncPackagePtr& WeakPtr)
+	{
+		return GetTypeHash(WeakPtr.Package);
+	}
+	FORCEINLINE bool operator==(const FUnsafeWeakAsyncPackagePtr& Other) const
+	{
+		return Package == Other.Package;
+	}
+	FName HumanReadableStringForDebugging() const;
+};
+
+/**  [EDL] Weak pointer to the async package */
+struct FWeakAsyncPackagePtr
+{
+	FName PackageName;
+	int32 SerialNumber;
+	FWeakAsyncPackagePtr(FAsyncPackage* Package = nullptr);
+	FAsyncPackage& GetPackage() const;
+	FORCEINLINE friend uint32 GetTypeHash(const FWeakAsyncPackagePtr& WeakPtr)
+	{
+		return GetTypeHash(WeakPtr.PackageName);
+	}
+	FORCEINLINE bool operator==(const FWeakAsyncPackagePtr& Other) const
+	{
+		return SerialNumber == Other.SerialNumber;
+	}
+	FName HumanReadableStringForDebugging() const
+	{
+		return PackageName;
+	}
+};
+
+#define VERIFY_WEAK_ASYNC_PACKAGE_PTRS (0) //(DO_CHECK)
+
+#if VERIFY_WEAK_ASYNC_PACKAGE_PTRS
+typedef FWeakAsyncPackagePtr FCheckedWeakAsyncPackagePtr;
+#else
+typedef FUnsafeWeakAsyncPackagePtr FCheckedWeakAsyncPackagePtr;
+#endif
+
+/** [EDL] Event Load Node */
+enum class EEventLoadNode
+{
+	Package_LoadSummary,
+	Package_SetupImports,
+	Package_ExportsSerialized,
+	Package_NumPhases,
+
+	ImportOrExport_Create = 0,
+	ImportOrExport_Serialize,
+	Import_NumPhases,
+
+	Export_StartIO = Import_NumPhases,
+	Export_NumPhases,
+
+	MAX_NumPhases = Package_NumPhases,
+	Invalid_Value = -1
+};
+
+#define USE_IMPLICIT_ARCS (1) // saves time and memory by not actually adding the arcs that are always present, they are:
+
+// Import: EEventLoadNode::ImportOrExport_Create -> EEventLoadNode::ImportOrExport_Serialize: can't consider this import serialized until we hook it up after creation
+// Import: EEventLoadNode::ImportOrExport_Serialize -> EEventLoadNode::Package_ExportsSerialized: can't consider the package done with event driven loading until all imports are serialized
+
+// Export: EEventLoadNode::ImportOrExport_Create -> EEventLoadNode::Export_StartIO: can't do the IO request until it is created
+// Export: EEventLoadNode::Export_StartIO -> EEventLoadNode::ImportOrExport_Serialize: can't serialize until the IO request is ready
+// Import: EEventLoadNode::ImportOrExport_Serialize -> EEventLoadNode::Package_ExportsSerialized: can't consider the package done with event driven loading until all exports are serialized
+
+/** [EDL] Event Load Node Pointer */
+struct FEventLoadNodePtr
+{
+	FCheckedWeakAsyncPackagePtr WaitingPackage;
+	FPackageIndex ImportOrExportIndex;  // IsNull()==true for PACKAGE_ Phases
+	EEventLoadNode Phase;
+
+	FEventLoadNodePtr()
+		: Phase(EEventLoadNode::Invalid_Value)
+	{
+	}
+	FORCEINLINE friend uint32 GetTypeHash(const FEventLoadNodePtr& NodePtr)
+	{
+		return GetTypeHash(NodePtr.WaitingPackage) ^ (GetTypeHash(NodePtr.ImportOrExportIndex) << 13) ^ (uint32(NodePtr.Phase) << 27);
+	}
+
+#if 0
+	FORCEINLINE bool operator<(const FEventLoadNodePtr& Other) const // not a total order
+	{
+		return WaitingPackage.SerialNumber < Other.WaitingPackage.SerialNumber;
+	}
+#endif
+	FORCEINLINE bool operator==(const FEventLoadNodePtr& Other) const
+	{
+		return WaitingPackage == Other.WaitingPackage && 
+			ImportOrExportIndex == Other.ImportOrExportIndex && 
+			Phase == Other.Phase;
+	}
+#if USE_IMPLICIT_ARCS
+	FORCEINLINE int32 NumImplicitArcs()
+	{
+		return ImportOrExportIndex.IsNull() ? 0 : 1; // only import and export nodes have implicit arcs, and every one of them has exactly one
+	}
+	FORCEINLINE FEventLoadNodePtr GetImplicitArc()
+	{
+		check(!ImportOrExportIndex.IsNull()); // package nodes don't have implicit arcs
+		FEventLoadNodePtr Result;
+		Result.WaitingPackage = WaitingPackage;
+		if (Phase == EEventLoadNode::ImportOrExport_Serialize)
+		{
+			Result.Phase = EEventLoadNode::Package_ExportsSerialized;
+			check(Result.ImportOrExportIndex.IsNull());
+		}
+		else
+		{
+			Result.ImportOrExportIndex = ImportOrExportIndex;
+			if (Phase == EEventLoadNode::ImportOrExport_Create)
+			{
+				Result.Phase = ImportOrExportIndex.IsImport() ? EEventLoadNode::ImportOrExport_Serialize : EEventLoadNode::Export_StartIO;
+			}
+			else
+			{
+				check(Phase == EEventLoadNode::Export_StartIO);
+				Result.Phase = EEventLoadNode::ImportOrExport_Serialize;
+			}
+		}
+		return Result;
+	}
+
+#endif
+
+	FString HumanReadableStringForDebugging() const;
+};
+
+/** [EDL] Event Load Node */
+struct FEventLoadNode
+{
+#if USE_IMPLICIT_ARCS
+	typedef TArray<FEventLoadNodePtr> TNodesWaitingForMeArray;
+#else
+	typedef TArray<FEventLoadNodePtr, TInlineAllocator<1> > TNodesWaitingForMeArray;
+#endif
+
+	TNodesWaitingForMeArray NodesWaitingForMe;
+	int32 NumPrerequistes;
+	bool bFired;
+	bool bAddedToGraph;
+	FORCEINLINE FEventLoadNode()
+		: NumPrerequistes(0)
+		, bFired(false)
+		, bAddedToGraph(false)
+	{
+	}
+};
+
+/** [EDL] Event Load Node Array */
+struct FEventLoadNodeArray
+{
+	FEventLoadNode PackageNodes[int(EEventLoadNode::Package_NumPhases)];
+	FEventLoadNode* Array;
+	int32 TotalNumberOfImportExportNodes;
+	int32 TotalNumberOfNodesAdded;
+	int32 NumImports;
+	int32 NumExports;
+	int32 OffsetToImports;
+	int32 OffsetToExports;
+
+	FEventLoadNodeArray()
+		: Array(nullptr)
+		, TotalNumberOfImportExportNodes(0)
+		, TotalNumberOfNodesAdded(0)
+		, NumImports(0)
+		, NumExports(0)
+		, OffsetToImports(0)
+		, OffsetToExports(0)
+	{
+	}
+
+
+	bool AddNode(FEventLoadNodePtr Node)
+	{
+		FEventLoadNode& NodeRef(PtrToNode(Node));
+		check(!NodeRef.bAddedToGraph);
+		NodeRef.bAddedToGraph = true;
+		return (++TotalNumberOfNodesAdded) == 1; // first add, the caller will add this to a list of things with outstanding nodes
+	}
+	bool RemoveNode(FEventLoadNodePtr Node)
+	{
+		FEventLoadNode& NodeRef(PtrToNode(Node));
+		check(NodeRef.bAddedToGraph);
+		NodeRef.bAddedToGraph = false;
+		return (--TotalNumberOfNodesAdded) == 0; // first add, the caller will add this to a list of things with outstanding nodes
+	}
+	FORCEINLINE FEventLoadNode& GetNode(FEventLoadNodePtr Node, bool bCheckAdded = true)
+	{
+		FEventLoadNode& NodeRef(PtrToNode(Node));
+		check(!bCheckAdded || NodeRef.bAddedToGraph);
+		return NodeRef;
+	}
+
+	void Init(int32 InNumImports, int32 InNumExports);
+	void Shutdown();
+	void GetAddedNodes(TArray<FEventLoadNodePtr>& OutAddedNodes, FAsyncPackage* Owner);
+private:
+	FORCEINLINE FEventLoadNode& PtrToNode(FEventLoadNodePtr Node)
+	{
+		int32 Index;
+		if (Node.ImportOrExportIndex.IsNull())
+		{
+			Index = int32(Node.Phase);
+			check(Index < int32(EEventLoadNode::Package_NumPhases) && Index >= 0);
+			return PackageNodes[Index];
+		}
+		check(TotalNumberOfImportExportNodes); // otherwise Init was not called yet
+		if (Node.ImportOrExportIndex.IsImport())
+		{
+			check(int32(Node.Phase) < int32(EEventLoadNode::Import_NumPhases));
+			Index = OffsetToImports + Node.ImportOrExportIndex.ToImport() * int32(EEventLoadNode::Import_NumPhases) + int32(Node.Phase);
+			check(Index >= OffsetToImports && Index < OffsetToExports);
+		}
+		else
+		{
+			check(int32(Node.Phase) < int32(EEventLoadNode::Export_NumPhases));
+			Index = OffsetToExports + Node.ImportOrExportIndex.ToExport() * int32(EEventLoadNode::Export_NumPhases) + int32(Node.Phase);
+			check(Index >= OffsetToExports);
+		}
+		check(Array && Index < TotalNumberOfImportExportNodes && Index >= 0);
+		return Array[Index];
+	}
+
+};
+
+/** [EDL] Event Load Graph */
+struct FEventLoadGraph
+{
+	TSet<FCheckedWeakAsyncPackagePtr> PackagesWithNodes;
+	TArray<int32> IndicesToFire;
+
+	FEventLoadNodeArray& GetArray(FEventLoadNodePtr& Node);
+
+	FEventLoadNode& GetNode(FEventLoadNodePtr& NodeToGet);
+
+	void AddNode(FEventLoadNodePtr& NewNode, bool bHoldForLater = false, int32 NumImplicitPrereqs = 0);
+	void DoneAddingPrerequistesFireIfNone(FEventLoadNodePtr& NewNode, bool bWasHeldForLater = false);
+	void AddArc(FEventLoadNodePtr& PrereqisiteNode, FEventLoadNodePtr& DependentNode);
+	void RemoveNode(FEventLoadNodePtr& NodeToRemove);
+	void NodeWillBeFiredExternally(FEventLoadNodePtr& NodeThatWasFired);
+	void CheckForCycles(bool bDoSlowTests = (!UE_BUILD_SHIPPING && !UE_BUILD_TEST));
+	bool CheckForCyclesInner(const TMultiMap<FEventLoadNodePtr, FEventLoadNodePtr>& Arcs, TSet<FEventLoadNodePtr>& Visited, TSet<FEventLoadNodePtr>& Stack, const FEventLoadNodePtr& Visit);
+};
+
+/** [EDL] Event Load Graph */
+struct FAsyncLoadEventArgs
+{
+	double TickStartTime;
+	const TCHAR* OutLastTypeOfWorkPerformed;
+	UObject* OutLastObjectWorkWasPerformedOn;
+	float TimeLimit;
+	bool bUseTimeLimit;
+	bool bUseFullTimeLimit;
+
+	FAsyncLoadEventArgs()
+		: TickStartTime(0.0)
+		, OutLastTypeOfWorkPerformed(nullptr)
+		, OutLastObjectWorkWasPerformedOn(nullptr)
+		, TimeLimit(0.0f)
+		, bUseTimeLimit(false)
+		, bUseFullTimeLimit(true)
+	{
+	}
+};
+
+/** [EDL] this is a little wrapper that does random pops for further randomization */
+class FImportOrImportIndexArray : public TArray<int32>
+{
+public:
+	void HeapPop(int32& OutItem, bool bAllowShrinking = true);
+};
+
+/**
+* Structure containing intermediate data required for async loading of all imports and exports of a
+* FLinkerLoad.
+*/
+
+struct FAsyncPackage : public FGCObject
+{
+	friend struct FScopedAsyncPackageEvent;
+	/**
+	 * Constructor
+	 */
+	FAsyncPackage(FAsyncLoadingThread& InThread, const FAsyncPackageDesc& InDesc, IEDLBootNotificationManager& InEDLBootNotificationManager);
+	~FAsyncPackage();
+
+	/**
+	 * Ticks the async loading code.
+	 *
+	 * @param	InbUseTimeLimit		Whether to use a time limit
+	 * @param	InbUseFullTimeLimit	If true use the entire time limit, even if you have to block on IO
+	 * @param	InTimeLimit			Soft limit to time this function may take
+	 *
+	 * @return	true if package has finished loading, false otherwise
+	 */
+	EAsyncPackageState::Type TickAsyncPackage(bool bUseTimeLimit, bool bInbUseFullTimeLimit, float& InOutTimeLimit, FFlushTree* FlushTree);
+
+	/** Fills the package dependency tree required to flush a specific package */
+	void PopulateFlushTree(struct FFlushTree* FlushTree);	
+	
+	/** Marks a specific request as complete */
+	void MarkRequestIDsAsComplete();
+
+	/**
+	 * @return Estimated load completion percentage.
+	 */
+	FORCEINLINE float GetLoadPercentage() const
+	{
+		return LoadPercentage;
+	}
+
+	/**
+	 * @return Time load begun. This is NOT the time the load was requested in the case of other pending requests.
+	 */
+	double GetLoadStartTime() const;
+
+	/**
+	 * Emulates ResetLoaders for the package's Linker objects, hence deleting it. 
+	 */
+	void ResetLoader();
+
+	/**
+	* Disassociates linker from this package
+	*/
+	void DetachLinker();
+
+	/**
+	* Flushes linker cache for all objects loaded with this package
+	*/
+	void FlushObjectLinkerCache();
+
+	/**
+	 * Returns the name of the package to load.
+	 */
+	FORCEINLINE const FName& GetPackageName() const
+	{
+		return Desc.Name;
+	}
+
+	/**
+	 * Returns the name to load of the package.
+	 */
+	FORCEINLINE const FName& GetPackageNameToLoad() const
+	{
+		return Desc.NameToLoad;
+	}
+
+	void AddCompletionCallback(TUniquePtr<FLoadPackageAsyncDelegate>&& Callback, bool bInternal);
+
+	/** Gets the number of references to this package from other packages in the dependency tree. */
+	FORCEINLINE int32 GetDependencyRefCount() const
+	{
+		return DependencyRefCount.GetValue();
+	}
+
+	FORCEINLINE UPackage* GetLinkerRoot() const
+	{
+		return LinkerRoot;
+	}
+
+	/** Returns true if the package has finished loading. */
+	FORCEINLINE bool HasFinishedLoading() const
+	{
+		return bLoadHasFinished;
+	}
+
+	/** Returns package loading priority. */
+	FORCEINLINE TAsyncLoadPriority GetPriority() const
+	{
+		return Desc.Priority;
+	}
+
+	/** Returns package loading priority. */
+	FORCEINLINE void SetPriority(TAsyncLoadPriority InPriority)
+	{
+		Desc.Priority = InPriority;
+	}
+
+	/** Returns true if loading has failed */
+	FORCEINLINE bool HasLoadFailed() const
+	{
+		return bLoadHasFailed;
+	}
+
+	/** Marks the threaded loading phase as complete for this package */
+	void ThreadedLoadingHasFinished()
+	{
+		bThreadedLoadingFinished = true;
+	}
+
+	/** Returns true if the threaded loaded phas has completed for this package */
+	bool HasThreadedLoadingFinished() const
+	{
+		return bThreadedLoadingFinished;
+	}
+
+	/** Adds new request ID to the existing package */
+	void AddRequestID(int32 Id);
+
+	/**
+	* Cancel loading this package.
+	*/
+	void Cancel();
+
+	/**
+	* Set the package that spawned this package as a dependency.
+	*/
+	void SetDependencyRootPackage(FAsyncPackage* InDependencyRootPackage)
+	{
+		DependencyRootPackage = InDependencyRootPackage;
+	}
+
+	/** Returns true if this package is already being loaded in the current callstack */
+	bool IsBeingProcessedRecursively() const
+	{
+		return ReentryCount > 1;
+	}
+
+	/** FGCObject Interface */
+	virtual void AddReferencedObjects(FReferenceCollector& Collector) override;
+	virtual FString GetReferencerName() const override
+	{ 
+		return FString::Printf(TEXT("FAsyncPackage %s"), *GetPackageName().ToString()); 
+	}
+
+	/** Adds a new object referenced by this package */
+	void AddObjectReference(UObject* InObject);
+
+	/** Removes all objects from the list and clears async loading flags */
+	void EmptyReferencedObjects();
+
+	/** Returns the UPackage wrapped by this, if it is valid */
+	UPackage* GetLoadedPackage();
+
+#if WITH_EDITOR
+	/** Gets all assets loaded by this async package, used in the editor */
+	void GetLoadedAssets(TArray<FWeakObjectPtr>& AssetList);
+#endif
+
+	/** Checks if all dependencies (imported packages) of this package have been fully loaded */
+	bool AreAllDependenciesFullyLoaded(TSet<UPackage*>& VisitedPackages);
+
+	/** Returs true if this package loaded objects that can create GC clusters */
+	bool HasClusterObjects() const
+	{
+		return DeferredClusterObjects.Num() > 0;
+	}
+
+	/** Creates GC clusters from loaded objects */
+	EAsyncPackageState::Type CreateClusters(double InTickStartTime, bool bInUseTimeLimit, float& InOutTimeLimit);
+
+private:	
+
+	/** Checks if all dependencies (imported packages) of this package have been fully loaded */
+	static bool AreAllDependenciesFullyLoadedInternal(FAsyncPackage* Package, TSet<UPackage*>& VisitedPackages, FString& OutError);
+
+	struct FCompletionCallback
+	{
+		bool bIsInternal;
+		bool bCalled;
+		TUniquePtr<FLoadPackageAsyncDelegate> Callback;
+
+		FCompletionCallback()
+			: bIsInternal(false)
+			, bCalled(false)
+		{
+		}
+		FCompletionCallback(bool bInInternal, TUniquePtr<FLoadPackageAsyncDelegate>&& InCallback)
+			: bIsInternal(bInInternal)
+			, bCalled(false)
+			, Callback(MoveTemp(InCallback))
+		{
+		}
+	};
+
+	/** Basic information associated with this package */
+	FAsyncPackageDesc Desc;
+	/** Linker which is going to have its exports and imports loaded									*/
+	FLinkerLoad*				Linker;
+	/** Package which is going to have its exports and imports loaded									*/
+	UPackage*				LinkerRoot;
+	/** Call backs called when we finished loading this package											*/
+	TArray<FCompletionCallback>	CompletionCallbacks;
+	/** Pending Import packages - we wait until all of them have been fully loaded. */
+	TArray<FAsyncPackage*> PendingImportedPackages;
+	/** Referenced imports - list of packages we need until we finish loading this package. */
+	TArray<FAsyncPackage*> ReferencedImports;
+	/** Root package if this package was loaded as a dependency of another. NULL otherwise */
+	FAsyncPackage* DependencyRootPackage;
+	/** Number of references to this package from other packages in the dependency tree. */
+	FThreadSafeCounter	DependencyRefCount;
+	/** Current index into linkers import table used to spread creation over several frames				*/
+	int32							LoadImportIndex;
+	/** Current index into linkers import table used to spread creation over several frames				*/
+	int32							ImportIndex;
+	/** Current index into linkers export table used to spread creation over several frames				*/
+	int32							ExportIndex;
+	/** Current index into ObjLoaded array used to spread routing PreLoad over several frames			*/
+	int32							PreLoadIndex;
+	/** Current index into ObjLoaded array used to spread routing PreLoad over several frames			*/
+	int32							PreLoadSortIndex;
+	/** Current index into ObjLoaded array used to spread routing PreLoad over several frames			*/
+	int32							FinishExternalReadDependenciesIndex;
+	/** Current index into ObjLoaded array used to spread routing PostLoad over several frames			*/
+	int32							PostLoadIndex;
+	/** Current index into DeferredPostLoadObjects array used to spread routing PostLoad over several frames			*/
+	int32						DeferredPostLoadIndex;
+	/** Current index into DeferredFinalizeObjects array used to spread routing PostLoad over several frames			*/
+	int32						DeferredFinalizeIndex;
+	/** Current index into DeferredClusterObjects array used to spread routing CreateClusters over several frames			*/
+	int32						DeferredClusterIndex;
+	/** Currently used time limit for this tick.														*/
+	float						TimeLimit;
+	/** Whether we are using a time limit for this tick.												*/
+	bool						bUseTimeLimit;
+	/** Whether we should use the entire time limit, even if we're blocked on I/O						*/
+	bool						bUseFullTimeLimit;
+	/** Whether we already exceed the time limit this tick.												*/
+	bool						bTimeLimitExceeded;
+	/** True if our load has failed */
+	bool						bLoadHasFailed;
+	/** True if our load has finished */
+	bool						bLoadHasFinished;
+	/** True if threaded loading has finished for this package */
+	bool						bThreadedLoadingFinished;
+	/** True if this package was created by this async package */
+	bool						bCreatedLinkerRoot;
+	/** The time taken when we started the tick.														*/
+	double						TickStartTime;
+	/** Last object work was performed on. Used for debugging/ logging purposes.						*/
+	UObject*					LastObjectWorkWasPerformedOn;
+	/** Last type of work performed on object.															*/
+	const TCHAR*				LastTypeOfWorkPerformed;
+	/** Time load begun. This is NOT the time the load was requested in the case of pending requests.	*/
+	double						LoadStartTime;
+	/** Estimated load percentage.																		*/
+	float						LoadPercentage;
+	/** Objects to be post loaded on the game thread */
+	TArray<UObject*> DeferredPostLoadObjects;
+	/** Objects to be finalized on the game thread */
+	TArray<UObject*> DeferredFinalizeObjects;
+	/** Objects loaded while loading this package */
+	TArray<UObject*> PackageObjLoaded;
+	/** Packages that were loaded synchronously while async loading this package or packages added by verify import */
+	TArray<FLinkerLoad*> DelayedLinkerClosePackages;
+	/** Objects to create GC clusters from */
+	TArray<UObject*> DeferredClusterObjects;
+
+	/** List of all request handles */
+	TArray<int32> RequestIDs;
+#if WITH_EDITORONLY_DATA
+	/** Index of the meta-data object within the linkers export table (unset if not yet processed, although may still be INDEX_NONE if there is no meta-data) */
+	TOptional<int32> MetaDataIndex;
+#endif // WITH_EDITORONLY_DATA
+	/** Number of times we recursed to load this package. */
+	int32 ReentryCount;
+	/** List of objects referenced by this package */
+	TSet<UObject*> ReferencedObjects;
+	/** Critical section for referenced objects list */
+	FCriticalSection ReferencedObjectsCritical;
+	/** Cached async loading thread object this package was created by */
+	FAsyncLoadingThread& AsyncLoadingThread;
+	class IEDLBootNotificationManager& EDLBootNotificationManager;
+	/** Packages that have been imported by this async package */
+	TSet<UPackage*> ImportedPackages;
+
+public:
+	
+
+	/** [EDL] Begin Event driven loader specific stuff */
+
+	EAsyncPackageLoadingState AsyncPackageLoadingState;
+	int32 SerialNumber;
+
+	FImportOrImportIndexArray ImportsThatAreNowCreated;
+	FImportOrImportIndexArray ImportsThatAreNowSerialized;
+	FImportOrImportIndexArray ExportsThatCanBeCreated;
+	FImportOrImportIndexArray ExportsThatCanHaveIOStarted;
+	FImportOrImportIndexArray ExportsThatCanBeSerialized;
+	TArray<IAsyncReadRequest*> ReadyPrecacheRequests;
+
+	struct FExportIORequest
+	{
+		int64 Offset;
+		int64 BytesToRead;
+		int32 FirstExportCovered;
+		int32 LastExportCovered;
+		TArray<int32> ExportsToRead;
+
+		FExportIORequest()
+			: Offset(-1)
+			, BytesToRead(-1)
+			, FirstExportCovered(-1)
+			, LastExportCovered(-1)
+		{
+		}
+	};
+
+	TMap<IAsyncReadRequest*, FExportIORequest> PrecacheRequests;
+	TMap<int32, IAsyncReadRequest*> ExportIndexToPrecacheRequest;
+	int64 CurrentBlockOffset;
+	int64 CurrentBlockBytes;
+	TSet<int32> ExportsInThisBlock;
+
+	TMap<TPair<FName, FPackageIndex>, FPackageIndex> ObjectNameWithOuterToExport;
+
+	TSet<FWeakAsyncPackagePtr> PackagesIMayBeWaitingForBeforePostload; // these need to be reexamined and perhaps deleted or collapsed
+
+	TSet<FWeakAsyncPackagePtr> PackagesIAmWaitingForBeforePostload;    // These are linked with PackagesIAmWaitingForBeforePostload, so we can be sure the other package will let us know
+	TSet<FWeakAsyncPackagePtr> OtherPackagesWaitingForMeBeforePostload;
+
+	TArray<FCheckedWeakAsyncPackagePtr> PackagesWaitingToLinkImports;
+
+	int32 ImportAddNodeIndex;
+	int32 ExportAddNodeIndex;
+
+	bool bProcessImportsAndExportsInFlight;
+	bool bProcessPostloadWaitInFlight;
+	bool bAllExportsSerialized;
+
+	void Event_CreateLinker();
+	void Event_FinishLinker();
+	void Event_StartImportPackages();
+	void Event_SetupImports();
+	void Event_SetupExports();
+	void Event_ProcessImportsAndExports();
+	void Event_ExportsDone();
+	void Event_ProcessPostloadWait();
+	void Event_StartPostload();
+
+	void MarkNewObjectForLoadIfItIsAnExport(UObject *Object);
+	bool AnyImportsAndExportWorkOutstanding();
+	void ConditionalQueueProcessImportsAndExports(bool bRequeueForTimeout = false);
+	void ConditionalQueueProcessPostloadWait();
+
+
+	EAsyncPackageState::Type LoadImports_Event();
+	EAsyncPackageState::Type SetupImports_Event();
+	EAsyncPackageState::Type SetupExports_Event();
+	EAsyncPackageState::Type ProcessImportsAndExports_Event();
+
+	FObjectImport* FindExistingImport(int32 LocalImportIndex);
+	void LinkImport(int32 LocalImportIndex);
+	void EventDrivenCreateExport(int32 LocalExportIndex);
+	void StartPrecacheRequest();
+	void EventDrivenSerializeExport(int32 LocalExportIndex);
+	int64 PrecacheRequestReady(IAsyncReadRequest * Req);
+	void MakeNextPrecacheRequestCurrent();
+	void FlushPrecacheBuffer();
+	void EventDrivenLoadingComplete();
+
+	void DumpDependencies(const TCHAR* Label, UObject* Obj);
+	void DumpDependencies(const TCHAR* Label, FLinkerLoad* DumpLinker, FPackageIndex Index);
+
+	UObject* EventDrivenIndexToObject(FPackageIndex Index, bool bCheckSerialized, FPackageIndex DumpIndex = FPackageIndex());
+	template<class T>
+	T* CastEventDrivenIndexToObject(FPackageIndex Index, bool bCheckSerialized, FPackageIndex DumpIndex = FPackageIndex())
+	{
+		UObject* Result = EventDrivenIndexToObject(Index, bCheckSerialized, DumpIndex);
+		if (!Result)
+		{
+			return nullptr;
+		}
+		return CastChecked<T>(Result);
+	}
+
+	FEventLoadNodeArray EventNodeArray;	
+
+	FEventLoadNodePtr AddNode(EEventLoadNode Phase, FPackageIndex ImportOrExportIndex = FPackageIndex(), bool bHoldForLater = false, int32 NumImplicitPrereqs = 0);
+	void DoneAddingPrerequistesFireIfNone(EEventLoadNode Phase, FPackageIndex ImportOrExportIndex = FPackageIndex(), bool bWasHeldForLater = false);
+	void RemoveNode(EEventLoadNode Phase, FPackageIndex ImportOrExportIndex = FPackageIndex());
+	void NodeWillBeFiredExternally(EEventLoadNode Phase, FPackageIndex ImportOrExportIndex = FPackageIndex());
+	void AddArc(FEventLoadNodePtr& PrereqisiteNode, FEventLoadNodePtr& DependentNode);
+	void RemoveAllNodes();
+	void FireNode(FEventLoadNodePtr& NodeToFire);
+
+	FString GetDebuggingPath(FPackageIndex Idx);
+
+	void SetTimeLimit(FAsyncLoadEventArgs& Args, const TCHAR* WorkType)
+	{
+		Args.OutLastTypeOfWorkPerformed = WorkType;
+		Args.OutLastObjectWorkWasPerformedOn = LinkerRoot;
+		TickStartTime = Args.TickStartTime;
+		LastTypeOfWorkPerformed = WorkType;
+		LastObjectWorkWasPerformedOn = LinkerRoot;
+		TimeLimit = Args.TimeLimit;
+		bUseTimeLimit = Args.bUseTimeLimit;
+		bUseFullTimeLimit = Args.bUseFullTimeLimit;
+	}
+
+	FAsyncLoadingThread& GetOwnerThread()
+	{
+		return AsyncLoadingThread;
+	}
+
+	/** [EDL] End Event driven loader specific stuff */
+
+
+#if PERF_TRACK_DETAILED_ASYNC_STATS
+	/** Number of times Tick function has been called.													*/
+	int32							TickCount;
+	/** Number of iterations in loop inside Tick.														*/
+	int32							TickLoopCount;
+
+	/** Number of iterations for CreateLinker.															*/
+	int32							CreateLinkerCount;
+	/** Number of iterations for FinishLinker.															*/
+	int32							FinishLinkerCount;
+	/** Number of iterations for CreateImports.															*/
+	int32							CreateImportsCount;
+	/** Number of iterations for CreateExports.															*/
+	int32							CreateExportsCount;
+	/** Number of iterations for PreLoadObjects.														*/
+	int32							PreLoadObjectsCount;
+	/** Number of iterations for PostLoadObjects.														*/
+	int32							PostLoadObjectsCount;
+	/** Number of iterations for FinishObjects.															*/
+	int32							FinishObjectsCount;
+
+	/** Total time spent in Tick.																		*/
+	double						TickTime;
+	/** Total time spent in	CreateLinker.																*/
+	double						CreateLinkerTime;
+	/** Total time spent in FinishLinker.																*/
+	double						FinishLinkerTime;
+	/** Total time spent in CreateImports.																*/
+	double						CreateImportsTime;
+	/** Total time spent in CreateExports.																*/
+	double						CreateExportsTime;
+	/** Total time spent in PreLoadObjects.																*/
+	double						PreLoadObjectsTime;
+	/** Total time spent in PostLoadObjects.															*/
+	double						PostLoadObjectsTime;
+	/** Total time spent in FinishObjects.																*/
+	double						FinishObjectsTime;
+
+#endif
+
+	void CallCompletionCallbacks(bool bInternalOnly, EAsyncLoadingResult::Type LoadingResult);
+
+	/**
+	* Route PostLoad to deferred objects.
+	*
+	* @return true if we finished calling PostLoad on all loaded objects and no new ones were created, false otherwise
+	*/
+	EAsyncPackageState::Type PostLoadDeferredObjects(double InTickStartTime, bool bInUseTimeLimit, float& InOutTimeLimit);
+
+	/** Close any linkers that have been open as a result of synchronous load during async loading */
+	void CloseDelayedLinkers();
+
+private:
+	/**
+	 * Gives up time slice if time limit is enabled.
+	 *
+	 * @return true if time slice can be given up, false otherwise
+	 */
+	bool GiveUpTimeSlice();
+	/**
+	 * Returns whether time limit has been exceeded.
+	 *
+	 * @return true if time limit has been exceeded (and is used), false otherwise
+	 */
+	bool IsTimeLimitExceeded();
+
+	/**
+	 * Begin async loading process. Simulates parts of BeginLoad.
+	 *
+	 * Objects created during BeginAsyncLoad and EndAsyncLoad will have EInternalObjectFlags::AsyncLoading set
+	 */
+	void BeginAsyncLoad();
+	/**
+	 * End async loading process. Simulates parts of EndLoad(). FinishObjects 
+	 * simulates some further parts once we're fully done loading the package.
+	 */
+	void EndAsyncLoad();
+	/**
+	 * Create linker async. Linker is not finalized at this point.
+	 *
+	 * @return true
+	 */
+	EAsyncPackageState::Type CreateLinker();
+	/**
+	 * Finalizes linker creation till time limit is exceeded.
+	 *
+	 * @return true if linker is finished being created, false otherwise
+	 */
+	EAsyncPackageState::Type FinishLinker();
+	/**
+	 * Loads imported packages..
+	 *
+	 * @param FlushTree Package dependency tree to be flushed
+	 * @return true if we finished loading all imports, false otherwise
+	 */
+	EAsyncPackageState::Type LoadImports(FFlushTree* FlushTree);
+	/**
+	 * Create imports till time limit is exceeded.
+	 *
+	 * @return true if we finished creating all imports, false otherwise
+	 */
+	EAsyncPackageState::Type CreateImports();
+
+#if WITH_EDITORONLY_DATA
+	/**
+	 * Creates and loads meta-data for the package.
+	 *
+	 * @return true if we finished creating meta-data, false otherwise.
+	 */
+	EAsyncPackageState::Type CreateMetaData();
+#endif // WITH_EDITORONLY_DATA
+	/**
+	 * Create exports till time limit is exceeded.
+	 *
+	 * @return true if we finished creating and preloading all exports, false otherwise.
+	 */
+	EAsyncPackageState::Type CreateExports();
+	/**
+	 * Preloads aka serializes all loaded objects.
+	 *
+	 * @return true if we finished serializing all loaded objects, false otherwise.
+	 */
+	EAsyncPackageState::Type PreLoadObjects();
+	/**
+	 * Route PostLoad to all loaded objects. This might load further objects!
+	 *
+	 * @return true if we finished calling PostLoad on all loaded objects and no new ones were created, false otherwise
+	 */
+	EAsyncPackageState::Type PostLoadObjects();
+	/**
+	 * Finish up objects and state, which means clearing the EInternalObjectFlags::AsyncLoading flag on newly created ones
+	 *
+	 * @return true
+	 */
+	EAsyncPackageState::Type FinishObjects();
+
+	/**
+	 * Finalizes external dependencies till time limit is exceeded
+	 *
+	 * @return Complete if all dependencies are finished, TimeOut otherwise
+	 */
+	EAsyncPackageState::Type FinishExternalReadDependencies();
+
+	/**
+	 * Function called when pending import package has been loaded.
+	 */
+	void ImportFullyLoadedCallback(const FName& PackageName, UPackage* LoadedPackage, EAsyncLoadingResult::Type Result);
+	/**
+	 * Adds dependency tree to the list if packages to wait for until their linkers have been created.
+	 *
+	 * @param ImportedPackage Package imported either directly or by one of the imported packages
+	 * @param FlushTree Package dependency tree to be flushed
+	 */
+	void AddDependencyTree(FAsyncPackage& ImportedPackage, TSet<FAsyncPackage*>& SearchedPackages, FFlushTree* FlushTree);
+	/**
+	 * Adds a unique package to the list of packages to wait for until their linkers have been created.
+	 *
+	 * @param PendingImport Package imported either directly or by one of the imported packages
+	 * @param FlushTree Package dependency tree to be flushed
+	 */
+	bool AddUniqueLinkerDependencyPackage(FAsyncPackage& PendingImport, FFlushTree* FlushTree);
+	/**
+	 * Adds a package to the list of pending import packages.
+	 *
+	 * @param PendingImport Name of the package imported either directly or by one of the imported packages
+	 * @param FlushTree Package dependency tree to be flushed
+	 */
+	void AddImportDependency(const FName& PendingImport, FFlushTree* FlushTree);
+	/**
+	 * Removes references to any imported packages.
+	 */
+	void FreeReferencedImports();
+
+	/**
+	* Updates load percentage stat
+	*/
+	void UpdateLoadPercentage();
+
+#if PERF_TRACK_DETAILED_ASYNC_STATS
+	/** Add this time taken for object of class Class to have CreateExport called, to the stats we track. */
+	void TrackCreateExportTimeForClass(const UClass* Class, double Time);
+
+	/** Add this time taken for object of class Class to have PostLoad called, to the stats we track. */
+	void TrackPostLoadTimeForClass(const UClass* Class, double Time);
+#endif // PERF_TRACK_DETAILED_ASYNC_STATS
+
+public:
+
+	/** Serialization context for this package */
+	FUObjectSerializeContext* GetSerializeContext();
+};
+
+struct FScopedAsyncPackageEvent
+{
+	/** Current scope package */
+	FAsyncPackage* Package;
+	/** Outer scope package */
+	FAsyncPackage* PreviousPackage;
+	
+	FScopedAsyncPackageEvent(FAsyncPackage* InPackage);
+	~FScopedAsyncPackageEvent();
+};
+
