Copying //UE4/Dev-Core to Main (//UE4/Main)

diff --git a/Engine/Source/Runtime/CoreUObject/Private/Serialization/AsyncLoading2.cpp b/Engine/Source/Runtime/CoreUObject/Private/Serialization/AsyncLoading2.cpp
new file mode 100644
--- /dev/null
+++ b/Engine/Source/Runtime/CoreUObject/Private/Serialization/AsyncLoading2.cpp
@@ -0,0 +1,5373 @@
+// Copyright 1998-2019 Epic Games, Inc. All Rights Reserved.
+
+/*=============================================================================
+	UnAsyncLoading.cpp: Unreal async loading code.
+=============================================================================*/
+
+#include "Serialization/AsyncLoading2.h"
+#include "HAL/PlatformFilemanager.h"
+#include "HAL/FileManager.h"
+#include "HAL/Event.h"
+#include "HAL/RunnableThread.h"
+#include "Misc/ScopeLock.h"
+#include "Stats/StatsMisc.h"
+#include "Misc/CoreStats.h"
+#include "HAL/IConsoleManager.h"
+#include "Misc/CoreDelegates.h"
+#include "Misc/CommandLine.h"
+#include "Misc/App.h"
+#include "Misc/MessageDialog.h"
+#include "Misc/PackageName.h"
+#include "UObject/PackageFileSummary.h"
+#include "UObject/Linker.h"
+#include "UObject/SoftObjectPath.h"
+#include "UObject/LinkerLoad.h"
+#include "Serialization/DeferredMessageLog.h"
+#include "UObject/UObjectThreadContext.h"
+#include "Misc/Paths.h"
+#include "Misc/ExclusiveLoadPackageTimeTracker.h"
+#include "ProfilingDebugging/LoadTimeTracker.h"
+#include "HAL/ThreadHeartBeat.h"
+#include "HAL/ExceptionHandling.h"
+#include "UObject/UObjectHash.h"
+#include "Templates/UniquePtr.h"
+#include "Serialization/BufferReader.h"
+#include "Async/TaskGraphInterfaces.h"
+#include "Blueprint/BlueprintSupport.h"
+#include "HAL/LowLevelMemTracker.h"
+#include "ProfilingDebugging/CsvProfiler.h"
+#include "UObject/GCScopeLock.h"
+#include "ProfilingDebugging/MiscTrace.h"
+#include "Serialization/LoadTimeTracePrivate.h"
+#include "ProfilingDebugging/CountersTrace.h"
+#include "Serialization/AsyncPackage.h"
+#include "Serialization/Zenaphore.h"
+#include "IO/IoDispatcher.h"
+#include "Containers/CircularQueue.h"
+#include "UObject/GCObject.h"
+#include "UObject/ObjectRedirector.h"
+
+#if UE_BUILD_DEVELOPMENT || UE_BUILD_DEBUG
+//PRAGMA_DISABLE_OPTIMIZATION
+#define ALT2_VERIFY_ASYNC_FLAGS
+#endif
+
+class FSimpleArchive
+	: public FArchive
+{
+public:
+	FSimpleArchive(const uint8* BufferPtr, uint64 BufferSize)
+	{
+#if DEVIRTUALIZE_FLinkerLoad_Serialize
+		ActiveFPLB->OriginalFastPathLoadBuffer = BufferPtr;
+		ActiveFPLB->StartFastPathLoadBuffer = BufferPtr;
+		ActiveFPLB->EndFastPathLoadBuffer = BufferPtr + BufferSize;
+#endif
+	}
+
+	int64 TotalSize() override
+	{
+#if DEVIRTUALIZE_FLinkerLoad_Serialize
+		return ActiveFPLB->EndFastPathLoadBuffer - ActiveFPLB->OriginalFastPathLoadBuffer;
+#else
+		return 0;
+#endif
+	}
+
+	int64 Tell() override
+	{
+#if DEVIRTUALIZE_FLinkerLoad_Serialize
+		return ActiveFPLB->StartFastPathLoadBuffer - ActiveFPLB->OriginalFastPathLoadBuffer;
+#else
+		return 0;
+#endif
+	}
+
+	void Seek(int64 Position) override
+	{
+#if DEVIRTUALIZE_FLinkerLoad_Serialize
+		ActiveFPLB->StartFastPathLoadBuffer = ActiveFPLB->OriginalFastPathLoadBuffer + Position;
+		check(ActiveFPLB->StartFastPathLoadBuffer <= ActiveFPLB->EndFastPathLoadBuffer);
+#endif
+	}
+
+	void Serialize(void* Data, int64 Length) override
+	{
+#if DEVIRTUALIZE_FLinkerLoad_Serialize
+		if (!Length || ArIsError)
+		{
+			return;
+		}
+		check(ActiveFPLB->StartFastPathLoadBuffer + Length <= ActiveFPLB->EndFastPathLoadBuffer);
+		FMemory::Memcpy(Data, ActiveFPLB->StartFastPathLoadBuffer, Length);
+		ActiveFPLB->StartFastPathLoadBuffer += Length;
+#endif
+	}
+};
+
+struct FAsyncPackage2;
+class FAsyncLoadingThread2Impl;
+
+struct FPackageStoreEntryRuntime
+{
+	int32* Slimports = nullptr;
+	UPackage* Package = nullptr;
+};
+
+struct FGlobalPackageId
+{
+	int32 Id;
+
+	bool operator==(FGlobalPackageId RHS) const
+	{
+		return Id == RHS.Id;
+	}
+
+	friend uint32 GetTypeHash(FGlobalPackageId GlobalPackageId)
+	{
+		return GetTypeHash(GlobalPackageId.Id);
+	}
+};
+
+struct FGlobalImport
+{
+	FName ObjectName;
+	FPackageIndex GlobalIndex;
+	FPackageIndex OuterIndex;
+	FPackageIndex OutermostIndex;
+	int32 Pad;
+};
+
+struct FGlobalImportRuntime
+{
+	/** Persistent data */
+	int32 Count = 0;
+	FName* Names;
+	FPackageIndex* Outers;
+	FPackageIndex* Packages;
+	UObject** Objects;
+	/** Reference tracking for GC management */
+	TAtomic<int32>* RefCounts;
+	TArray<UObject*> KeepAliveObjects;
+	void OnPreGarbageCollect();
+	void OnPostGarbageCollect();
+};
+	
+struct FPackageStoreEntrySerialized
+{
+	FGuid Guid;
+	FName Name;
+	FName FileName;
+	uint32 PackageFlags;
+	int32 ImportCount;
+	int32 ImportOffset;
+	int32 SlimportCount;
+	int32 SlimportOffset;
+	int32 ExportCount;
+	int32 ExportOffset;
+	int32 PreloadDependencyCount;
+	int32 PreloadDependencyOffset;
+	int32 Pad;
+	int64 BulkDataStartOffset;
+};
+
+class FGlobalNameMap
+{
+public:
+	void Load(const FString& FilePath)
+	{
+		TUniquePtr<FArchive> Ar(IFileManager::Get().CreateFileReader(*FilePath));
+		check(Ar);
+
+		int32 NameCount;
+		*Ar << NameCount;
+		NameEntries.Reserve(NameCount);
+		FNameEntrySerialized SerializedNameEntry(ENAME_LinkerConstructor);
+
+		for (int32 I = 0; I < NameCount; ++I)
+		{
+			*Ar << SerializedNameEntry;
+			NameEntries.Emplace(FName(SerializedNameEntry).GetDisplayIndex());
+			EntryToIndex.Add(NameEntries[I], I);
+		}
+	}
+
+	void Save(const FString& FilePath)
+	{
+		TUniquePtr<FArchive> Ar(IFileManager::Get().CreateFileWriter(*FilePath));
+		check(Ar);
+
+		int32 NameCount = NameEntries.Num();
+		*Ar << NameCount;
+		
+		for (int32 I = 0; I < NameCount; ++I)
+		{
+			FName::GetEntry(NameEntries[I])->Write(*Ar);
+		}
+	}
+
+	FName GetName(const uint32 NameIndex, const uint32 NameNumber) const
+	{
+		FNameEntryId NameEntry = NameEntries[NameIndex];
+		return FName::CreateFromDisplayId(NameEntry, NameNumber);
+	}
+
+	FName FromSerializedName(const FName& SerializedName) const
+	{
+		const int32 EntryIndex = SerializedName.GetComparisonIndex().ToUnstableInt();
+		FNameEntryId NameEntry = NameEntries[EntryIndex];
+		return FName::CreateFromDisplayId(NameEntry, SerializedName.GetNumber());
+	}
+
+	const int32* GetIndex(const FName& Name) const
+	{
+		return EntryToIndex.Find(Name.GetDisplayIndex());
+	}
+
+	int32 GetOrCreateIndex(const FName& Name)
+	{
+		if (const int32* ExistingIndex = EntryToIndex.Find(Name.GetDisplayIndex()))
+		{
+			return *ExistingIndex;
+		}
+		else
+		{
+			int32 NewIndex = NameEntries.Add(Name.GetDisplayIndex());
+			EntryToIndex.Add(NameEntries[NewIndex], NewIndex);
+			return NewIndex;
+		}
+	}
+
+	const TArray<FNameEntryId>& GetNameEntries() const
+	{
+		return NameEntries;
+	}
+
+private:
+	TArray<FNameEntryId> NameEntries;
+	TMap<FNameEntryId, int32> EntryToIndex;
+};
+
+enum class EChunkType : uint8
+{
+	None,
+	PackageSummary,
+	ExportData,
+	BulkData
+};
+
+FIoChunkId CreateChunkId(uint32 NameIndex, uint32 NameNumber, uint16 ChunkIndex, EChunkType ChunkType)
+{
+	uint8 Data[12] = {0};
+
+	*reinterpret_cast<uint32*>(&Data[0]) = NameIndex;
+	*reinterpret_cast<int32*>(&Data[4]) = NameNumber;
+	*reinterpret_cast<uint16*>(&Data[8]) = ChunkIndex;
+	*reinterpret_cast<uint8*>(&Data[10]) = static_cast<uint8>(ChunkType);
+
+	FIoChunkId ChunkId;
+	ChunkId.Set(Data, 12);
+	return ChunkId;
+}
+
+FIoChunkId CreateChunkId(const FIoChunkId& ChunkId, uint16 ChunkIndex, EChunkType ChunkType)
+{
+	FIoChunkId Out = ChunkId;
+
+	uint8* Data = reinterpret_cast<uint8*>(&Out);
+	*reinterpret_cast<uint16*>(&Data[8]) = ChunkIndex;
+	*reinterpret_cast<uint8*>(&Data[10]) = static_cast<uint8>(ChunkType);
+
+	return Out;
+}
+
+FName GetChunkName(const FIoChunkId& ChunkId, const FGlobalNameMap& GlobalNameMap)
+{
+	const uint8* Data = reinterpret_cast<const uint8*>(&ChunkId);
+	const uint32 NameIndex = *reinterpret_cast<const uint32*>(&Data[0]);
+	const int32 NameNumber = *reinterpret_cast<const int32*>(&Data[4]);
+
+	return GlobalNameMap.GetName(NameIndex, NameNumber);
+}
+
+EChunkType GetChunkType(const FIoChunkId& ChunkId)
+{
+	const uint8* Data = reinterpret_cast<const uint8*>(&ChunkId);
+	const uint8 ChunkType = *reinterpret_cast<const uint8*>(&Data[10]);
+	return static_cast<EChunkType>(ChunkType);
+}
+
+uint16 GetChunkIndex(const FIoChunkId& ChunkId)
+{
+	const uint8* Data = reinterpret_cast<const uint8*>(&ChunkId);
+	return *reinterpret_cast<const uint16*>(&Data[8]);
+
+}
+
+struct FPackageSummary
+{
+	FGuid Guid;
+	uint32 PackageFlags;
+	int32 ImportCount;
+	int32 ExportCount;
+	int32 PreloadDependencyCount;
+	int32 ExportOffset;
+	int32 GraphDataOffset;
+	int32 GraphDataSize;
+	int32 BulkDataStartOffset;
+};
+
+class FIoRequestQueue final : private FRunnable
+{
+public:
+	struct FCompletionEvent { FAsyncPackage2* Package; FIoBatch IoBatch; FIoRequest IoRequest; };
+
+	FIoRequestQueue(FIoDispatcher& InIoDispatcher, FZenaphore& InZenaphore, const uint32 InCapacity = 1024);
+	~FIoRequestQueue();
+	void EnqueueRequest(FAsyncPackage2* Package, const FIoChunkId& ChunkId);
+	TOptional<FCompletionEvent> DequeueCompletionEvent();
+	bool HasPendingRequests() const;
+	bool WaitForRequest(float SecondsToWait);
+
+private:
+	struct FSubmissionRequest { FAsyncPackage2* Package; FIoChunkId ChunkId; };
+	using FPendingRequest = FCompletionEvent;
+
+	uint32 Run() override;
+	void Stop() override;
+
+	FIoDispatcher& IoDispatcher;
+	FZenaphore& Zenaphore;
+	FRunnableThread* Thread = nullptr;
+	FEvent* WakeUpEvent = nullptr;
+	bool bIsRunning = false;
+	TCircularQueue<FSubmissionRequest> RequestQueue;
+	FCriticalSection RequestQueueLock;
+	TCircularQueue<FCompletionEvent> CompletionQueue;
+	FCriticalSection CompletionQueueLock;
+	TArray<FPendingRequest> PendingRequests;
+	TAtomic<int32> NumPendingRequests;
+};
+
+FIoRequestQueue::FIoRequestQueue(FIoDispatcher& InIoDispatcher, FZenaphore& InZenaphore, const uint32 InCapacity)
+: IoDispatcher(InIoDispatcher)
+, Zenaphore(InZenaphore)
+, RequestQueue(InCapacity)
+, CompletionQueue(InCapacity)
+, NumPendingRequests(0)
+{
+	WakeUpEvent = FGenericPlatformProcess::GetSynchEventFromPool(false);
+	Thread = FRunnableThread::Create(this, TEXT("IoRequestQueue"), 0, TPri_Normal);
+}
+
+FIoRequestQueue::~FIoRequestQueue()
+{
+	Stop();
+	Thread->Kill(true);
+	Thread = nullptr;
+	FPlatformProcess::ReturnSynchEventToPool(WakeUpEvent);
+}
+
+void FIoRequestQueue::EnqueueRequest(FAsyncPackage2* Package, const FIoChunkId& ChunkId)
+{
+	FScopeLock _(&RequestQueueLock);
+	bool bEnqueued = RequestQueue.Enqueue(FSubmissionRequest { Package, ChunkId });
+	NumPendingRequests++;
+	check(bEnqueued);
+	WakeUpEvent->Trigger();
+}
+
+TOptional<FIoRequestQueue::FCompletionEvent> FIoRequestQueue::DequeueCompletionEvent()
+{
+	TOptional<FCompletionEvent> Out;
+	FScopeLock _(&CompletionQueueLock);
+	if (const FCompletionEvent* CompletionEvent = CompletionQueue.Peek())
+	{
+		Out = *CompletionEvent;
+		CompletionQueue.Dequeue();
+	}
+
+	return Out;
+}
+
+bool FIoRequestQueue::HasPendingRequests() const
+{
+	return NumPendingRequests > 0 || !CompletionQueue.IsEmpty() || !RequestQueue.IsEmpty();
+}
+
+bool FIoRequestQueue::WaitForRequest(float SecondsToWait)
+{
+	return true;
+}
+
+void FIoRequestQueue::Stop()
+{
+	if (bIsRunning)
+	{
+		bIsRunning = false;
+		WakeUpEvent->Trigger();
+	}
+}
+
+uint32 FIoRequestQueue::Run()
+{
+	bIsRunning = true;
+
+	while (bIsRunning)
+	{
+		if (!bIsRunning)
+		{
+			break;
+		}
+
+		// Process incoming
+		{
+			FSubmissionRequest Request;
+			bool bDequeuedRequest = false;
+			{
+				FScopeLock _(&RequestQueueLock);
+				bDequeuedRequest = RequestQueue.Dequeue(Request);
+			}
+			
+			if (bDequeuedRequest)
+			{
+				FIoBatch IoBatch = IoDispatcher.NewBatch();
+				FIoRequest IoRequest = IoBatch.Read(Request.ChunkId, FIoReadOptions());
+				IoBatch.Issue();
+				PendingRequests.Add(FPendingRequest { Request.Package, IoBatch, IoRequest });
+			}
+		}
+
+		// Process pending
+		{
+			int32 NumPending = PendingRequests.Num();
+			for (int32 I = 0; I < PendingRequests.Num();)
+			{
+				FPendingRequest& PendingRequest = PendingRequests[I];
+				if (PendingRequest.IoRequest.Status().IsCompleted())
+				{
+					{
+						FScopeLock _(&CompletionQueueLock);
+						bool bEnqueued = CompletionQueue.Enqueue(PendingRequest);
+						check(bEnqueued);
+					}
+					PendingRequests.RemoveAtSwap(I, 1, false);
+					NumPendingRequests--;
+				}
+				else
+				{
+					++I;
+				}
+			}
+
+			if (NumPending != PendingRequests.Num())
+			{
+				Zenaphore.NotifyOne();
+			}
+		}
+
+		if (PendingRequests.Num() == 0 && NumPendingRequests == 0)
+		{
+			WakeUpEvent->Wait();
+		}
+	}
+
+	return 0;
+}
+
+enum class EAsyncPackageLoadingState2 : uint8
+{
+	NewPackage,
+	WaitingForSummary,
+	StartImportPackages,
+	WaitingForImportPackages,
+	SetupImports,
+	SetupExports,
+	ProcessNewImportsAndExports,
+	PostLoad_Etc,
+	PackageComplete,
+};
+
+enum EEventLoadNode2 : uint8
+{
+	Package_CreateLinker,
+	Package_LoadSummary,
+	Package_ImportPackages,
+	Package_SetupImports,
+	Package_SetupExports,
+	Package_ExportsSerialized,
+	Package_PostLoad,
+	Package_Tick,
+	Package_Delete,
+	Package_NumPhases,
+
+	ImportOrExport_Create = 0,
+	ImportOrExport_Serialize,
+	Import_NumPhases,
+
+	Export_StartIO = Import_NumPhases,
+	Export_NumPhases,
+};
+
+class FEventLoadGraphAllocator;
+struct FAsyncLoadEventSpec;
+struct FAsyncLoadingThreadState2;
+
+/** [EDL] Event Load Node */
+class FEventLoadNode2
+{
+public:
+	FEventLoadNode2(const FAsyncLoadEventSpec* InSpec, FAsyncPackage2* InPackage, int32 InImportOrExportIndex);
+	void DependsOn(FEventLoadNode2* Other);
+	void AddBarrier();
+	void ReleaseBarrier();
+	void Execute(FAsyncLoadingThreadState2& ThreadState);
+
+	int32 GetBarrierCount()
+	{
+		return BarrierCount.Load();
+	}
+
+	bool IsDone()
+	{
+		return !!bDone.Load();
+	}
+
+private:
+	void ProcessDependencies(FAsyncLoadingThreadState2& ThreadState);
+	void Fire(FAsyncLoadingThreadState2& ThreadState);
+
+	union
+	{
+		FEventLoadNode2* SingleDependent;
+		FEventLoadNode2** MultipleDependents;
+	};
+	uint32 DependenciesCount = 0;
+	uint32 DependenciesCapacity = 0;
+	TAtomic<int32> BarrierCount { 0 };
+	TAtomic<uint8> DependencyWriterCount { 0 };
+	TAtomic<uint8> bDone { 0 };
+#if UE_BUILD_DEVELOPMENT
+	TAtomic<uint8> bFired { 0 };
+#endif
+
+	const FAsyncLoadEventSpec* Spec;
+	FAsyncPackage2* Package;
+	int32 ImportOrExportIndex;
+};
+
+class FAsyncLoadEventGraphAllocator
+{
+public:
+	FEventLoadNode2* AllocNodes(uint32 Count)
+	{
+		//TRACE_CPUPROFILER_EVENT_SCOPE(AllocNodes);
+		SIZE_T Size = Count * sizeof(FEventLoadNode2);
+		TotalNodeCount += Count;
+		TotalAllocated += Size;
+		return reinterpret_cast<FEventLoadNode2*>(FMemory::Malloc(Size));
+	}
+
+	void FreeNodes(FEventLoadNode2* Nodes, uint32 Count)
+	{
+		//TRACE_CPUPROFILER_EVENT_SCOPE(FreeNodes);
+		FMemory::Free(Nodes);
+		SIZE_T Size = Count * sizeof(FEventLoadNode2);
+		TotalAllocated -= Size;
+		TotalNodeCount -= Count;
+	}
+
+	FEventLoadNode2** AllocArcs(uint32 Count)
+	{
+		//TRACE_CPUPROFILER_EVENT_SCOPE(AllocArcs);
+		SIZE_T Size = Count * sizeof(FEventLoadNode2*);
+		TotalArcCount += Count;
+		TotalAllocated += Size;
+		return reinterpret_cast<FEventLoadNode2**>(FMemory::Malloc(Size));
+	}
+
+	void FreeArcs(FEventLoadNode2** Arcs, uint32 Count)
+	{
+		//TRACE_CPUPROFILER_EVENT_SCOPE(FreeArcs);
+		FMemory::Free(Arcs);
+		SIZE_T Size = Count * sizeof(FEventLoadNode2*);
+		TotalAllocated -= Size;
+		TotalArcCount -= Count;
+	}
+
+	TAtomic<int64> TotalNodeCount { 0 };
+	TAtomic<int64> TotalArcCount { 0 };
+	TAtomic<int64> TotalAllocated { 0 };
+};
+
+class FAsyncLoadEventQueue2
+{
+public:
+	FAsyncLoadEventQueue2();
+	~FAsyncLoadEventQueue2();
+
+	void SetZenaphore(FZenaphore* InZenaphore)
+	{
+		Zenaphore = InZenaphore;
+	}
+
+	bool PopAndExecute(FAsyncLoadingThreadState2& ThreadState);
+	void Push(FEventLoadNode2* Node);
+
+private:
+	FZenaphore* Zenaphore = nullptr;
+	TAtomic<uint64> Head { 0 };
+	TAtomic<uint64> Tail { 0 };
+	TAtomic<FEventLoadNode2*> Entries[524288];
+};
+
+struct FAsyncLoadEventSpec
+{
+	typedef EAsyncPackageState::Type(*FAsyncLoadEventFunc)(FAsyncPackage2*, int32);
+	FAsyncLoadEventFunc Func = nullptr;
+	FAsyncLoadEventQueue2* EventQueue = nullptr;
+	bool bExecuteImmediately = false;
+};
+
+struct FAsyncLoadingThreadState2
+	: public FTlsAutoCleanup
+{
+	static FAsyncLoadingThreadState2* Create(FAsyncLoadEventGraphAllocator& GraphAllocator)
+	{
+		check(TlsSlot != 0);
+		check(!FPlatformTLS::GetTlsValue(TlsSlot));
+		FAsyncLoadingThreadState2* State = new FAsyncLoadingThreadState2(GraphAllocator);
+		State->Register();
+		FPlatformTLS::SetTlsValue(TlsSlot, State);
+		return State;
+	}
+
+	static FAsyncLoadingThreadState2& Get()
+	{
+		check(TlsSlot != 0);
+		return *static_cast<FAsyncLoadingThreadState2*>(FPlatformTLS::GetTlsValue(TlsSlot));
+	}
+
+	FAsyncLoadingThreadState2(FAsyncLoadEventGraphAllocator& InGraphAllocator)
+		: GraphAllocator(InGraphAllocator)
+	{
+
+	}
+
+	void ProcessDeferredFrees()
+	{
+		if (DeferredFreeNodes.Num() || DeferredFreeArcs.Num())
+		{
+			TRACE_CPUPROFILER_EVENT_SCOPE(ProcessDeferredFrees);
+			for (TTuple<FEventLoadNode2*, uint32>& DeferredFreeNode : DeferredFreeNodes)
+			{
+				GraphAllocator.FreeNodes(DeferredFreeNode.Get<0>(), DeferredFreeNode.Get<1>());
+			}
+			DeferredFreeNodes.Reset();
+			for (TTuple<FEventLoadNode2**, uint32>& DeferredFreeArc : DeferredFreeArcs)
+			{
+				GraphAllocator.FreeArcs(DeferredFreeArc.Get<0>(), DeferredFreeArc.Get<1>());
+			}
+			DeferredFreeArcs.Reset();
+		}
+	}
+
+	void SetTimeLimit(bool bUseTimeLimit, float TimeLimit)
+	{
+
+	}
+
+	bool IsTimeLimitExceeded()
+	{
+		/*static double LastTestTime = -1.0;
+		bool bTimeLimitExceeded = false;
+		if (bUseTimeLimit)
+		{
+			double CurrentTime = FPlatformTime::Seconds();
+			bTimeLimitExceeded = CurrentTime - InTickStartTime > InTimeLimit;
+
+			if (bTimeLimitExceeded && GWarnIfTimeLimitExceeded)
+			{
+				IsTimeLimitExceededPrint(InTickStartTime, CurrentTime, LastTestTime, InTimeLimit, InLastTypeOfWorkPerformed, InLastObjectWorkWasPerformedOn);
+			}
+			LastTestTime = CurrentTime;
+		}
+		if (!bTimeLimitExceeded)
+		{
+			bTimeLimitExceeded = IsGarbageCollectionWaiting();
+			UE_CLOG(bTimeLimitExceeded, LogStreaming, Verbose, TEXT("Timing out async loading due to Garbage Collection request"));
+		}
+		return bTimeLimitExceeded;*/
+		return false;
+	}
+
+	FAsyncLoadEventGraphAllocator& GraphAllocator;
+	TArray<TTuple<FEventLoadNode2*, uint32>> DeferredFreeNodes;
+	TArray<TTuple<FEventLoadNode2**, uint32>> DeferredFreeArcs;
+	TArray<FEventLoadNode2*> NodesToFire;
+	bool bShouldFireNodes = true;
+	static uint32 TlsSlot;
+};
+
+uint32 FAsyncLoadingThreadState2::TlsSlot;
+
+/**
+* Structure containing intermediate data required for async loading of all imports and exports of a
+* FLinkerLoad.
+*/
+
+struct FAsyncPackage2
+	: public FGCObject
+{
+	friend struct FScopedAsyncPackageEvent2;
+
+	FAsyncPackage2(const FAsyncPackageDesc& InDesc, int32 InSerialNumber, FAsyncLoadingThread2Impl& InAsyncLoadingThread, IEDLBootNotificationManager& InEDLBootNotificationManager, FAsyncLoadEventGraphAllocator& InGraphAllocator, const FAsyncLoadEventSpec* EventSpecs, FGlobalPackageId InGlobalPackageId);
+	virtual ~FAsyncPackage2();
+
+
+	bool bAddedForDelete = false;
+
+	void AddRef()
+	{
+		++RefCount;
+	}
+
+	void ReleaseRef()
+	{
+		check(RefCount > 0);
+		if (--RefCount == 0)
+		{
+			GetNode(EEventLoadNode2::Package_Delete)->ReleaseBarrier();
+		}
+	}
+
+	void ClearImportedPackages()
+	{
+		for (FAsyncPackage2* ImportedAsyncPackage : ImportedAsyncPackages)
+		{
+			ImportedAsyncPackage->ReleaseRef();
+		}
+		ImportedAsyncPackages.Empty();
+		ReleaseGlobalImportObjectReferences();
+	}
+
+	void AddGlobalImportObjectReferences()
+	{
+		for (int32 LocalImportIndex = 0; LocalImportIndex < LocalImportCount; ++LocalImportIndex)
+		{
+			const int32 GlobalImportIndex = LocalImportIndices[LocalImportIndex];
+			++GlobalImportObjectRefCounts[GlobalImportIndex];
+		}
+	}
+
+	void ReleaseGlobalImportObjectReferences()
+	{
+		for (int32 LocalImportIndex = 0; LocalImportIndex < LocalImportCount; ++LocalImportIndex)
+		{
+			const int32 GlobalImportIndex = LocalImportIndices[LocalImportIndex];
+			--GlobalImportObjectRefCounts[GlobalImportIndex];
+		}
+	}
+
+	/** Marks a specific request as complete */
+	void MarkRequestIDsAsComplete();
+
+	/**
+	 * @return Estimated load completion percentage.
+	 */
+	FORCEINLINE float GetLoadPercentage() const
+	{
+		return LoadPercentage;
+	}
+
+	/**
+	 * @return Time load begun. This is NOT the time the load was requested in the case of other pending requests.
+	 */
+	double GetLoadStartTime() const;
+
+	/**
+	 * Emulates ResetLoaders for the package's Linker objects, hence deleting it.
+	 */
+	void ResetLoader();
+
+	/**
+	* Disassociates linker from this package
+	*/
+	void DetachLinker();
+
+	/**
+	* Flushes linker cache for all objects loaded with this package
+	*/
+	void FlushObjectLinkerCache();
+
+	/**
+	 * Returns the name of the package to load.
+	 */
+	FORCEINLINE const FName& GetPackageName() const
+	{
+		return Desc.Name;
+	}
+
+	/**
+	 * Returns the name to load of the package.
+	 */
+	FORCEINLINE const FName& GetPackageNameToLoad() const
+	{
+		return Desc.NameToLoad;
+	}
+
+	void AddCompletionCallback(TUniquePtr<FLoadPackageAsyncDelegate>&& Callback, bool bInternal);
+
+	FORCEINLINE UPackage* GetLinkerRoot() const
+	{
+		return LinkerRoot;
+	}
+
+	/** Returns true if the package has finished loading. */
+	FORCEINLINE bool HasFinishedLoading() const
+	{
+		return bLoadHasFinished;
+	}
+
+	/** Returns package loading priority. */
+	FORCEINLINE TAsyncLoadPriority GetPriority() const
+	{
+		return Desc.Priority;
+	}
+
+	/** Returns package loading priority. */
+	FORCEINLINE void SetPriority(TAsyncLoadPriority InPriority)
+	{
+		Desc.Priority = InPriority;
+	}
+
+	/** Returns true if loading has failed */
+	FORCEINLINE bool HasLoadFailed() const
+	{
+		return bLoadHasFailed;
+	}
+
+	/** Adds new request ID to the existing package */
+	void AddRequestID(int32 Id);
+
+	/**
+	* Cancel loading this package.
+	*/
+	void Cancel();
+
+	/** Returns true if this package is already being loaded in the current callstack */
+	bool IsBeingProcessedRecursively() const
+	{
+		return ReentryCount > 1;
+	}
+
+	/** FGCObject Interface */
+	virtual void AddReferencedObjects(FReferenceCollector& Collector) override {}
+	virtual FString GetReferencerName() const override
+	{
+		return FString::Printf(TEXT("FAsyncPackage %s"), *GetPackageName().ToString());
+	}
+
+	void AddOwnedObject(UObject* Object)
+	{
+		OwnedObjects.Add(Object);
+	}
+
+	void ClearOwnedObjects();
+
+	/** Returns the UPackage wrapped by this, if it is valid */
+	UPackage* GetLoadedPackage();
+
+#if WITH_EDITOR
+	/** Gets all assets loaded by this async package, used in the editor */
+	void GetLoadedAssets(TArray<FWeakObjectPtr>& AssetList);
+#endif
+
+	/** Checks if all dependencies (imported packages) of this package have been fully loaded */
+	bool AreAllDependenciesFullyLoaded(TSet<UPackage*>& VisitedPackages);
+
+	/** Returs true if this package loaded objects that can create GC clusters */
+	bool HasClusterObjects() const
+	{
+		return DeferredClusterObjects.Num() > 0;
+	}
+
+	/** Creates GC clusters from loaded objects */
+	EAsyncPackageState::Type CreateClusters();
+
+	void ImportPackagesRecursive();
+
+private:
+
+	/** Checks if all dependencies (imported packages) of this package have been fully loaded */
+	static bool AreAllDependenciesFullyLoadedInternal(FAsyncPackage2* Package, TSet<UPackage*>& VisitedPackages, FString& OutError);
+
+	struct FCompletionCallback
+	{
+		bool bIsInternal;
+		bool bCalled;
+		TUniquePtr<FLoadPackageAsyncDelegate> Callback;
+
+		FCompletionCallback()
+			: bIsInternal(false)
+			, bCalled(false)
+		{
+		}
+		FCompletionCallback(bool bInInternal, TUniquePtr<FLoadPackageAsyncDelegate>&& InCallback)
+			: bIsInternal(bInInternal)
+			, bCalled(false)
+			, Callback(MoveTemp(InCallback))
+		{
+		}
+	};
+
+	TAtomic<int32> RefCount{ 0 };
+
+	/** Basic information associated with this package */
+	FAsyncPackageDesc Desc;
+	/** Linker which is going to have its exports and imports loaded									*/
+	FLinkerLoad*				Linker;
+	/** Package which is going to have its exports and imports loaded									*/
+	UPackage*				LinkerRoot;
+	/** Call backs called when we finished loading this package											*/
+	TArray<FCompletionCallback>	CompletionCallbacks;
+	/** Current index into linkers import table used to spread creation over several frames				*/
+	int32							ImportIndex;
+	/** Current index into linkers export table used to spread creation over several frames				*/
+	int32							ExportIndex;
+	/** Current index into ObjLoaded array used to spread routing PreLoad over several frames			*/
+	int32							FinishExternalReadDependenciesIndex;
+	/** Current index into ObjLoaded array used to spread routing PostLoad over several frames			*/
+	int32							PostLoadIndex;
+	/** Current index into DeferredPostLoadObjects array used to spread routing PostLoad over several frames			*/
+	int32						DeferredPostLoadIndex;
+	/** Current index into DeferredFinalizeObjects array used to spread routing PostLoad over several frames			*/
+	int32						DeferredFinalizeIndex;
+	/** Current index into DeferredClusterObjects array used to spread routing CreateClusters over several frames			*/
+	int32						DeferredClusterIndex;
+	/** True if our load has failed */
+	bool						bLoadHasFailed;
+	/** True if our load has finished */
+	bool						bLoadHasFinished;
+	/** True if this package was created by this async package */
+	bool						bCreatedLinkerRoot;
+	/** Time load begun. This is NOT the time the load was requested in the case of pending requests.	*/
+	double						LoadStartTime;
+	/** Estimated load percentage.																		*/
+	float						LoadPercentage;
+	/** Objects to be post loaded on the game thread */
+	TArray<UObject*> DeferredPostLoadObjects;
+	/** Objects to be finalized on the game thread */
+	TArray<UObject*> DeferredFinalizeObjects;
+	/** Objects loaded while loading this package */
+	TArray<UObject*> PackageObjLoaded;
+	/** Packages that were loaded synchronously while async loading this package or packages added by verify import */
+	TArray<FLinkerLoad*> DelayedLinkerClosePackages;
+	/** Objects to create GC clusters from */
+	TArray<UObject*> DeferredClusterObjects;
+
+	/** List of all request handles */
+	TArray<int32> RequestIDs;
+#if WITH_EDITORONLY_DATA
+	/** Index of the meta-data object within the linkers export table (unset if not yet processed, although may still be INDEX_NONE if there is no meta-data) */
+	TOptional<int32> MetaDataIndex;
+#endif // WITH_EDITORONLY_DATA
+	/** Number of times we recursed to load this package. */
+	int32 ReentryCount;
+	TArray<FAsyncPackage2*> ImportedAsyncPackages;
+	/** List of OwnedObjects = Exports + UPackage + ObjectsCreatedFromExports */
+	TArray<UObject*> OwnedObjects;
+	/** Cached async loading thread object this package was created by */
+	FAsyncLoadingThread2Impl& AsyncLoadingThread;
+	IEDLBootNotificationManager& EDLBootNotificationManager;
+	FAsyncLoadEventGraphAllocator& GraphAllocator;
+	/** Package chunk id																				*/
+	FIoChunkId PackageChunkId;
+	/** Global package id used to index into the package store											*/
+	FGlobalPackageId GlobalPackageId;
+	/** Packages that have been imported by this async package */
+	TSet<UPackage*> ImportedPackages;
+
+	FEventLoadNode2* PackageNodes = nullptr;
+	FEventLoadNode2* ImportNodes = nullptr;
+	FEventLoadNode2* ExportNodes = nullptr;
+	uint32 ImportNodeCount = 0;
+	uint32 ExportNodeCount = 0;
+	
+	TUniquePtr<uint8[]> PackageSummaryBuffer;
+	TArray<FIoBuffer> ExportIoBuffers;
+
+	int32 GlobalImportCount = 0;
+	int32 LocalImportCount = 0;
+	int32* LocalImportIndices = nullptr;
+	FName* GlobalImportNames = nullptr;
+	FPackageIndex* GlobalImportOuters = nullptr;
+	FPackageIndex* GlobalImportPackages = nullptr;
+	UObject** GlobalImportObjects = nullptr;
+	TAtomic<int32>* GlobalImportObjectRefCounts = nullptr;
+public:
+
+	FAsyncLoadingThread2Impl& GetAsyncLoadingThread()
+	{
+		return AsyncLoadingThread;
+	}
+
+	FAsyncLoadEventGraphAllocator& GetGraphAllocator()
+	{
+		return GraphAllocator;
+	}
+
+	/** [EDL] Begin Event driven loader specific stuff */
+
+	EAsyncPackageLoadingState2 AsyncPackageLoadingState;
+	int32 SerialNumber;
+
+	TMap<TPair<FName, FPackageIndex>, FPackageIndex> ObjectNameWithOuterToExport;
+
+	bool bHasImportedPackagesRecursive = false;
+
+	bool bAllExportsSerialized;
+
+	static EAsyncPackageState::Type Event_CreateLinker(FAsyncPackage2* Package, int32);
+	static EAsyncPackageState::Type Event_FinishLinker(FAsyncPackage2* Package, int32);
+	static EAsyncPackageState::Type Event_StartImportPackages(FAsyncPackage2* Package, int32);
+	static EAsyncPackageState::Type Event_SetupImports(FAsyncPackage2* Package, int32);
+	static EAsyncPackageState::Type Event_SetupExports(FAsyncPackage2* Package, int32);
+	static EAsyncPackageState::Type Event_LinkImport(FAsyncPackage2* Package, int32 LocalImportIndex);
+	static EAsyncPackageState::Type Event_ImportSerialized(FAsyncPackage2* Package, int32 LocalImportIndex);
+	static EAsyncPackageState::Type Event_CreateExport(FAsyncPackage2* Package, int32 LocalExportIndex);
+	static EAsyncPackageState::Type Event_StartIO(FAsyncPackage2* Package, int32 LocalExportIndex);
+	static EAsyncPackageState::Type Event_SerializeExport(FAsyncPackage2* Package, int32 LocalExportIndex);
+	static EAsyncPackageState::Type Event_ExportsDone(FAsyncPackage2* Package, int32);
+	static EAsyncPackageState::Type Event_StartPostload(FAsyncPackage2* Package, int32);
+	static EAsyncPackageState::Type Event_Tick(FAsyncPackage2* Package, int32);
+	static EAsyncPackageState::Type Event_Delete(FAsyncPackage2* Package, int32);
+
+	void ProcessIoRequest(FIoRequest& IoRequest);
+
+	void MarkNewObjectForLoadIfItIsAnExport(UObject *Object);
+
+	EAsyncPackageState::Type SetupSlimports_Event();
+
+	UObject* FindExistingSlimport(int32 GlobalImportIndex);
+	void LinkSlimport(int32 LocalImportIndex, int32 GlobalImportIndex = -1);
+	void EventDrivenCreateExport(int32 LocalExportIndex);
+	void EventDrivenSerializeExport(int32 LocalExportIndex);
+	void EventDrivenLoadingComplete();
+
+	UObject* EventDrivenIndexToObject(FPackageIndex Index, bool bCheckSerialized);
+	template<class T>
+	T* CastEventDrivenIndexToObject(FPackageIndex Index, bool bCheckSerialized)
+	{
+		UObject* Result = EventDrivenIndexToObject(Index, bCheckSerialized);
+		if (!Result)
+		{
+			return nullptr;
+		}
+		return CastChecked<T>(Result);
+	}
+
+	FEventLoadNode2* GetNode(EEventLoadNode2 Phase, FPackageIndex ImportOrExportIndex = FPackageIndex());
+
+	/** [EDL] End Event driven loader specific stuff */
+
+	void CallCompletionCallbacks(bool bInternalOnly, EAsyncLoadingResult::Type LoadingResult);
+
+	/**
+	* Route PostLoad to deferred objects.
+	*
+	* @return true if we finished calling PostLoad on all loaded objects and no new ones were created, false otherwise
+	*/
+	EAsyncPackageState::Type PostLoadDeferredObjects();
+
+	/** Close any linkers that have been open as a result of synchronous load during async loading */
+	void CloseDelayedLinkers();
+
+private:
+	/**
+	 * Begin async loading process. Simulates parts of BeginLoad.
+	 *
+	 * Objects created during BeginAsyncLoad and EndAsyncLoad will have EInternalObjectFlags::AsyncLoading set
+	 */
+	void BeginAsyncLoad();
+	/**
+	 * End async loading process. Simulates parts of EndLoad(). FinishObjects
+	 * simulates some further parts once we're fully done loading the package.
+	 */
+	void EndAsyncLoad();
+	/**
+	 * Create linker async. Linker is not finalized at this point.
+	 *
+	 * @return true
+	 */
+	EAsyncPackageState::Type CreateLinker();
+	/**
+	 * Finalizes linker creation till time limit is exceeded.
+	 *
+	 * @return true if linker is finished being created, false otherwise
+	 */
+	EAsyncPackageState::Type FinishLinker();
+
+	/**
+	 * Route PostLoad to all loaded objects. This might load further objects!
+	 *
+	 * @return true if we finished calling PostLoad on all loaded objects and no new ones were created, false otherwise
+	 */
+	EAsyncPackageState::Type PostLoadObjects();
+	/**
+	 * Finish up objects and state, which means clearing the EInternalObjectFlags::AsyncLoading flag on newly created ones
+	 *
+	 * @return true
+	 */
+	EAsyncPackageState::Type FinishObjects();
+
+	/**
+	 * Finalizes external dependencies till time limit is exceeded
+	 *
+	 * @return Complete if all dependencies are finished, TimeOut otherwise
+	 */
+	EAsyncPackageState::Type FinishExternalReadDependencies();
+
+	/**
+	* Updates load percentage stat
+	*/
+	void UpdateLoadPercentage();
+
+public:
+
+	/** Serialization context for this package */
+	FUObjectSerializeContext* GetSerializeContext();
+};
+
+struct FScopedAsyncPackageEvent2
+{
+	/** Current scope package */
+	FAsyncPackage2* Package;
+	/** Outer scope package */
+	FAsyncPackage2* PreviousPackage;
+
+	FScopedAsyncPackageEvent2(FAsyncPackage2* InPackage);
+	~FScopedAsyncPackageEvent2();
+};
+
+class FAsyncLoadingThreadWorker : private FRunnable
+{
+public:
+	FAsyncLoadingThreadWorker(FAsyncLoadEventGraphAllocator& InGraphAllocator, FAsyncLoadEventQueue2& InEventQueue, FZenaphore& InZenaphore, TAtomic<int32>& InActiveWorkersCount)
+		: Zenaphore(InZenaphore)
+		, EventQueue(InEventQueue)
+		, GraphAllocator(InGraphAllocator)
+		, ActiveWorkersCount(InActiveWorkersCount)
+	{
+	}
+
+	void StartThread();
+	
+	void StopThread()
+	{
+		bStopRequested = true;
+		bSuspendRequested = true;
+		Zenaphore.NotifyAll();
+	}
+	
+	void SuspendThread()
+	{
+		bSuspendRequested = true;
+		Zenaphore.NotifyAll();
+	}
+	
+	void ResumeThread()
+	{
+		bSuspendRequested = false;
+	}
+	
+	int32 GetThreadId() const
+	{
+		return ThreadId;
+	}
+
+private:
+	virtual bool Init() override { return true; }
+	virtual uint32 Run() override;
+	virtual void Stop() override {};
+
+	FZenaphore& Zenaphore;
+	FAsyncLoadEventQueue2& EventQueue;
+	FAsyncLoadEventGraphAllocator& GraphAllocator;
+	TAtomic<int32>& ActiveWorkersCount;
+	FRunnableThread* Thread = nullptr;
+	TAtomic<bool> bStopRequested { false };
+	TAtomic<bool> bSuspendRequested { false };
+	int32 ThreadId = 0;
+};
+
+static FPackageIndex FindExportFromSlimport(FLinkerLoad* ImportLinker,
+	int32 GlobalImportIndex,
+	FPackageIndex* GlobalImportOuters,
+	FName* GlobalImportNames)
+{
+	check(ImportLinker && ImportLinker->AsyncRoot && ((FAsyncPackage2*)ImportLinker->AsyncRoot)->ObjectNameWithOuterToExport.Num());
+	FPackageIndex Result;
+	FPackageIndex ImportOuterIndex = GlobalImportOuters[GlobalImportIndex];
+	if (ImportOuterIndex.IsImport())
+	{
+		FName ObjectName = GlobalImportNames[GlobalImportIndex];
+		FPackageIndex ExportOuterIndex = FindExportFromSlimport(ImportLinker, ImportOuterIndex.ToImport(), GlobalImportOuters, GlobalImportNames);
+		FPackageIndex* PotentialExport = ((FAsyncPackage2*)ImportLinker->AsyncRoot)->ObjectNameWithOuterToExport.Find(TPair<FName, FPackageIndex>(ObjectName, ExportOuterIndex));
+		if (PotentialExport)
+		{
+			Result = *PotentialExport;
+		}
+	}
+	return Result;
+}
+
+static FPackageIndex FindExportFromObject2(FLinkerLoad* Linker, UObject *Object)
+{
+	check(Linker && Linker->AsyncRoot && ((FAsyncPackage2*)Linker->AsyncRoot)->ObjectNameWithOuterToExport.Num());
+	FPackageIndex Result;
+	UObject* Outer = Object->GetOuter();
+	if (Outer)
+	{
+		FPackageIndex OuterIndex = FindExportFromObject2(Linker, Outer);
+		FPackageIndex* PotentialExport = ((FAsyncPackage2*)Linker->AsyncRoot)->ObjectNameWithOuterToExport.Find(TPair<FName, FPackageIndex>(Object->GetFName(), OuterIndex));
+		if (PotentialExport)
+		{
+			Result = *PotentialExport;
+		}
+	}
+	return Result;
+}
+
+class FAsyncLoadingThread2Impl
+	: public FRunnable
+{
+	friend struct FAsyncPackage2;
+public:
+	FAsyncLoadingThread2Impl(IEDLBootNotificationManager& InEDLBootNotificationManager);
+	virtual ~FAsyncLoadingThread2Impl();
+
+private:
+	/** Thread to run the worker FRunnable on */
+	FRunnableThread* Thread;
+	TAtomic<bool> bStopRequested { false };
+	TAtomic<bool> bSuspendRequested { false };
+	TArray<FAsyncLoadingThreadWorker> Workers;
+	TAtomic<int32> ActiveWorkersCount { 0 };
+	bool bWorkersSuspended = false;
+
+	/** [ASYNC/GAME THREAD] true if the async thread is actually started. We don't start it until after we boot because the boot process on the game thread can create objects that are also being created by the loader */
+	bool bThreadStarted = false;
+
+	/** [ASYNC/GAME THREAD] Event used to signal loading should be cancelled */
+	FEvent* CancelLoadingEvent;
+	/** [ASYNC/GAME THREAD] Event used to signal that the async loading thread should be suspended */
+	FEvent* ThreadSuspendedEvent;
+	/** [ASYNC/GAME THREAD] Event used to signal that the async loading thread has resumed */
+	FEvent* ThreadResumedEvent;
+	/** [ASYNC/GAME THREAD] List of queued packages to stream */
+	TArray<FAsyncPackageDesc*> QueuedPackages;
+	/** [ASYNC/GAME THREAD] Package queue critical section */
+	FCriticalSection QueueCritical;
+	/** [ASYNC/GAME THREAD] Event used to signal there's queued packages to stream */
+	TArray<FAsyncPackage2*> LoadedPackages;
+	/** [ASYNC/GAME THREAD] Critical section for LoadedPackages list */
+	FCriticalSection LoadedPackagesCritical;
+	TArray<FAsyncPackage2*> LoadedPackagesToProcess;
+	TArray<FAsyncPackage2*> PackagesToDelete;
+#if WITH_EDITOR
+	TArray<FWeakObjectPtr> LoadedAssets;
+#endif
+
+	FCriticalSection AsyncPackagesCritical;
+	TMap<FName, FAsyncPackage2*> AsyncPackageNameLookup;
+
+	IEDLBootNotificationManager& EDLBootNotificationManager;
+
+	/** List of all pending package requests */
+	TSet<int32> PendingRequests;
+	/** Synchronization object for PendingRequests list */
+	FCriticalSection PendingRequestsCritical;
+
+	/** [ASYNC/GAME THREAD] Number of package load requests in the async loading queue */
+	TAtomic<uint32> QueuedPackagesCounter { 0 };
+	/** [ASYNC/GAME THREAD] Number of packages being loaded on the async thread and post loaded on the game thread */
+	FThreadSafeCounter ExistingAsyncPackagesCounter;
+
+	FThreadSafeCounter AsyncThreadReady;
+
+	/** When cancelling async loading: list of package requests to cancel */
+	TArray<FAsyncPackageDesc*> QueuedPackagesToCancel;
+	/** When cancelling async loading: list of packages to cancel */
+	TSet<FAsyncPackage2*> PackagesToCancel;
+
+	/** Async loading thread ID */
+	uint32 AsyncLoadingThreadID;
+
+	FThreadSafeCounter PackageRequestID;
+	FThreadSafeCounter AsyncPackageSerialNumber;
+
+	/** I/O Dispatcher */
+	FGlobalNameMap GlobalNameMap;
+	FIoStoreEnvironment IoStoreEnvironment;
+	TUniquePtr<FIoStoreReader> IoStoreReader;
+	FIoDispatcher IoDispatcher;
+	TUniquePtr<FIoRequestQueue> IoRequestQueue;
+
+	/** Package store */
+	FPackageStoreEntryRuntime* StoreEntriesRuntime = nullptr;
+	FPackageStoreEntrySerialized* StoreEntriesSerialized = nullptr;
+	int32* Slimports = nullptr;
+	int32 SlimportCount = 0;
+	FGlobalImportRuntime GlobalImportRuntime;
+	TMap<FName, FGlobalPackageId> PackageNameToGlobalPackageId;
+	int32 PackageCount = 0;
+
+	inline UObject** GetGlobalImportObjects(int32& OutCount)
+	{
+		OutCount = GlobalImportRuntime.Count;
+		return GlobalImportRuntime.Objects;
+	}
+
+	inline FPackageIndex* GetGlobalImportOuters(int32& OutCount) const
+	{
+		OutCount = GlobalImportRuntime.Count;
+		return GlobalImportRuntime.Outers;
+	}
+
+	inline FName* GetGlobalImportNames(int32& OutCount) const
+	{
+		OutCount = GlobalImportRuntime.Count;
+		return GlobalImportRuntime.Names;
+	}
+
+	inline FPackageIndex* GetGlobalImportPackages(int32& OutCount) const 
+	{
+		OutCount = GlobalImportRuntime.Count;
+		return GlobalImportRuntime.Packages;
+	}
+
+	inline TAtomic<int32>* GetGlobalImportObjectRefCounts()
+	{
+		return GlobalImportRuntime.RefCounts;
+	}
+
+	inline int32* GetPackageSlimports(FGlobalPackageId GlobalPackageId, int32& OutCount) const 
+	{
+		OutCount = StoreEntriesSerialized[GlobalPackageId.Id].SlimportCount;
+		return StoreEntriesRuntime[GlobalPackageId.Id].Slimports;
+	}
+
+	inline int32 GetPackageImportCount(FGlobalPackageId GlobalPackageId) const 
+	{
+		return StoreEntriesSerialized[GlobalPackageId.Id].ImportCount;
+	}
+
+	inline int32 GetPackageExportCount(FGlobalPackageId GlobalPackageId) const 
+	{
+		return StoreEntriesSerialized[GlobalPackageId.Id].ExportCount;
+	}
+
+	inline FString GetPackageFileName(FGlobalPackageId GlobalPackageId)
+	{
+		return StoreEntriesSerialized[GlobalPackageId.Id].FileName.ToString();
+	}
+
+public:
+
+	//~ Begin FRunnable Interface.
+	virtual bool Init();
+	virtual uint32 Run();
+	virtual void Stop();
+	//~ End FRunnable Interface
+
+	/** Start the async loading thread */
+	void StartThread();
+
+	/** [GC] Management of global import objects */
+	void OnPreGarbageCollect();
+	void OnPostGarbageCollect();
+
+	/** [EDL] Event queue */
+	FZenaphore AltZenaphore;
+	TArray<FZenaphore> WorkerZenaphores;
+	FAsyncLoadEventGraphAllocator GraphAllocator;
+	FAsyncLoadEventQueue2 EventQueue;
+	FAsyncLoadEventQueue2 AsyncEventQueue;
+	FAsyncLoadEventQueue2 CreateExportsEventQueue;
+	FAsyncLoadEventQueue2 SerializeExportsEventQueue;
+	TArray<FAsyncLoadEventQueue2*> AltEventQueues;
+	TArray<FAsyncLoadEventSpec> EventSpecs;
+
+	/** [EDL] Queues CreateLinker event */
+	void QueueEvent_CreateLinker(FAsyncPackage2* Pkg);
+
+	/** True if multithreaded async loading is currently being used. */
+	bool IsMultithreaded()
+	{
+		return bThreadStarted;
+	}
+
+	/** Sets the current state of async loading */
+	void EnterAsyncLoadingTick()
+	{
+		AsyncLoadingTickCounter++;
+	}
+
+	void LeaveAsyncLoadingTick()
+	{
+		AsyncLoadingTickCounter--;
+		check(AsyncLoadingTickCounter >= 0);
+	}
+
+	/** Gets the current state of async loading */
+	bool GetIsInAsyncLoadingTick() const
+	{
+		return !!AsyncLoadingTickCounter;
+	}
+
+	/** Returns true if packages are currently being loaded on the async thread */
+	bool IsAsyncLoadingPackages()
+	{
+		FPlatformMisc::MemoryBarrier();
+		return QueuedPackagesCounter != 0 || ExistingAsyncPackagesCounter.GetValue() != 0;
+	}
+
+	/** Returns true this codes runs on the async loading thread */
+	bool IsInAsyncLoadThread()
+	{
+		if (IsMultithreaded())
+		{
+			// We still need to report we're in async loading thread even if 
+			// we're on game thread but inside of async loading code (PostLoad mostly)
+			// to make it behave exactly like the non-threaded version
+			uint32 CurrentThreadId = FPlatformTLS::GetCurrentThreadId();
+			if (CurrentThreadId == AsyncLoadingThreadID ||
+				(IsInGameThread() && GetIsInAsyncLoadingTick()))
+			{
+				return true;
+			}
+			else
+			{
+				for (const FAsyncLoadingThreadWorker& Worker : Workers)
+				{
+					if (CurrentThreadId == Worker.GetThreadId())
+					{
+						return true;
+					}
+				}
+			}
+			return false;
+		}
+		else
+		{
+			return IsInGameThread() && GetIsInAsyncLoadingTick();
+		}
+	}
+
+	/** Returns true if async loading is suspended */
+	bool IsAsyncLoadingSuspended()
+	{
+		return bSuspendRequested;
+	}
+
+	void NotifyConstructedDuringAsyncLoading(UObject* Object, bool bSubObject);
+
+	void FireCompletedCompiledInImport(FGCObject* AsyncPackage, FPackageIndex Import);
+
+	/**
+	* [ASYNC THREAD] Finds an existing async package in the AsyncPackages by its name.
+	*
+	* @param PackageName async package name.
+	* @return Pointer to the package or nullptr if not found
+	*/
+	FORCEINLINE FAsyncPackage2* FindAsyncPackage(const FName& PackageName)
+	{
+		FScopeLock LockAsyncPackages(&AsyncPackagesCritical);
+		//checkSlow(IsInAsyncLoadThread());
+		return AsyncPackageNameLookup.FindRef(PackageName);
+	}
+
+	/**
+	* [ASYNC THREAD] Inserts package to queue according to priority.
+	*
+	* @param PackageName - async package name.
+	* @param InsertMode - Insert mode, describing how we insert this package into the request list
+	*/
+	void InsertPackage(FAsyncPackage2* Package, bool bReinsert = false);
+
+	FAsyncPackage2* FindOrInsertPackage(FAsyncPackageDesc* InDesc, bool& bInserted);
+
+	/**
+	* [ASYNC/GAME THREAD] Queues a package for streaming.
+	*
+	* @param Package package descriptor.
+	*/
+	void QueuePackage(FAsyncPackageDesc& Package);
+
+	/**
+	* [ASYNC* THREAD] Loads all packages
+	*
+	* @param OutPackagesProcessed Number of packages processed in this call.
+	* @return The current state of async loading
+	*/
+	EAsyncPackageState::Type ProcessAsyncLoadingFromGameThread(int32& OutPackagesProcessed);
+
+	/**
+	* [GAME THREAD] Ticks game thread side of async loading.
+	*
+	* @param bUseTimeLimit True if time limit should be used [time-slicing].
+	* @param bUseFullTimeLimit True if full time limit should be used [time-slicing].
+	* @param TimeLimit Maximum amount of time that can be spent in this call [time-slicing].
+	* @param FlushTree Package dependency tree to be flushed
+	* @return The current state of async loading
+	*/
+	EAsyncPackageState::Type TickAsyncLoadingFromGameThread(bool bUseTimeLimit, bool bUseFullTimeLimit, float TimeLimit, int32 FlushRequestID = INDEX_NONE);
+
+	/**
+	* [ASYNC THREAD] Main thread loop
+	*
+	* @param bUseTimeLimit True if time limit should be used [time-slicing].
+	* @param bUseFullTimeLimit True if full time limit should be used [time-slicing].
+	* @param TimeLimit Maximum amount of time that can be spent in this call [time-slicing].
+	* @param FlushTree Package dependency tree to be flushed
+	*/
+	EAsyncPackageState::Type TickAsyncThreadFromGameThread(bool& bDidSomething);
+
+	/** Initializes async loading thread */
+	void InitializeLoading();
+
+	void ShutdownLoading();
+
+	int32 LoadPackage(
+		const FString& InPackageName,
+		const FGuid* InGuid,
+		const TCHAR* InPackageToLoadFrom,
+		FLoadPackageAsyncDelegate InCompletionDelegate,
+		EPackageFlags InPackageFlags,
+		int32 InPIEInstanceID,
+		int32 InPackagePriority);
+
+	EAsyncPackageState::Type ProcessLoadingFromGameThread(bool bUseTimeLimit, bool bUseFullTimeLimit, float TimeLimit);
+
+	EAsyncPackageState::Type ProcessLoadingUntilCompleteFromGameThread(TFunctionRef<bool()> CompletionPredicate, float TimeLimit);
+
+	void CancelLoading();
+
+	void SuspendLoading();
+
+	void ResumeLoading();
+
+	void FlushLoading(int32 PackageId);
+
+	int32 GetNumAsyncPackages()
+	{
+		FPlatformMisc::MemoryBarrier();
+		return ExistingAsyncPackagesCounter.GetValue();
+	}
+
+	/**
+	 * [GAME THREAD] Gets the load percentage of the specified package
+	 * @param PackageName Name of the package to return async load percentage for
+	 * @return Percentage (0-100) of the async package load or -1 of package has not been found
+	 */
+	float GetAsyncLoadPercentage(const FName& PackageName);
+
+	/**
+	 * [ASYNC/GAME THREAD] Checks if a request ID already is added to the loading queue
+	 */
+	bool ContainsRequestID(int32 RequestID)
+	{
+		FScopeLock Lock(&PendingRequestsCritical);
+		return PendingRequests.Contains(RequestID);
+	}
+
+	/**
+	 * [ASYNC/GAME THREAD] Adds a request ID to the list of pending requests
+	 */
+	void AddPendingRequest(int32 RequestID)
+	{
+		FScopeLock Lock(&PendingRequestsCritical);
+		PendingRequests.Add(RequestID);
+	}
+
+	/**
+	 * [ASYNC/GAME THREAD] Removes a request ID from the list of pending requests
+	 */
+	void RemovePendingRequests(TArray<int32>& RequestIDs)
+	{
+		FScopeLock Lock(&PendingRequestsCritical);
+		for (int32 ID : RequestIDs)
+		{
+			PendingRequests.Remove(ID);
+			TRACE_LOADTIME_END_REQUEST(ID);
+		}
+	}
+
+private:
+
+	void SuspendWorkers();
+	void ResumeWorkers();
+
+	/**
+	* [GAME THREAD] Performs game-thread specific operations on loaded packages (not-thread-safe PostLoad, callbacks)
+	*
+	* @param bUseTimeLimit True if time limit should be used [time-slicing].
+	* @param bUseFullTimeLimit True if full time limit should be used [time-slicing].
+	* @param TimeLimit Maximum amount of time that can be spent in this call [time-slicing].
+	* @param FlushTree Package dependency tree to be flushed
+	* @return The current state of async loading
+	*/
+	EAsyncPackageState::Type ProcessLoadedPackagesFromGameThread(bool& bDidSomething, int32 FlushRequestID = INDEX_NONE);
+
+	bool CreateAsyncPackagesFromQueue();
+
+	FAsyncPackage2* CreateAsyncPackage(const FAsyncPackageDesc& Desc)
+	{
+		FGlobalPackageId GlobalPackageId = PackageNameToGlobalPackageId[Desc.NameToLoad];
+		return new FAsyncPackage2(Desc, AsyncPackageSerialNumber.Increment(), *this, EDLBootNotificationManager, GraphAllocator, EventSpecs.GetData(), GlobalPackageId);
+	}
+
+	/**
+	* [ASYNC THREAD] Internal helper function for updating the priorities of an existing package and all its dependencies
+	*/
+	void UpdateExistingPackagePriorities(FAsyncPackage2* InPackage, TAsyncLoadPriority InNewPriority);
+
+	/**
+	* [ASYNC THREAD] Adds a package to a list of packages that have finished loading on the async thread
+	*/
+	void AddToLoadedPackages(FAsyncPackage2* Package);
+
+	/** Number of times we re-entered the async loading tick, mostly used by singlethreaded ticking. Debug purposes only. */
+	int32 AsyncLoadingTickCounter;
+
+	/** Enqueue I/O request */
+	void EnqueueIoRequest(FAsyncPackage2* Package, const FIoChunkId& ChunkId);
+};
+
+/**
+ * Updates FUObjectThreadContext with the current package when processing it.
+ * FUObjectThreadContext::AsyncPackage is used by NotifyConstructedDuringAsyncLoading.
+ */
+struct FAsyncPackageScope2
+{
+	/** Outer scope package */
+	FGCObject* PreviousPackage;
+	/** Cached ThreadContext so we don't have to access it again */
+	FUObjectThreadContext& ThreadContext;
+
+	FAsyncPackageScope2(FGCObject* InPackage)
+		: ThreadContext(FUObjectThreadContext::Get())
+	{
+		PreviousPackage = ThreadContext.AsyncPackage;
+		ThreadContext.AsyncPackage = InPackage;
+	}
+	~FAsyncPackageScope2()
+	{
+		ThreadContext.AsyncPackage = PreviousPackage;
+	}
+};
+
+/** Just like TGuardValue for FAsyncLoadingThread::AsyncLoadingTickCounter but only works for the game thread */
+struct FAsyncLoadingTickScope2
+{
+	FAsyncLoadingThread2Impl& AsyncLoadingThread;
+	bool bNeedsToLeaveAsyncTick;
+
+	FAsyncLoadingTickScope2(FAsyncLoadingThread2Impl& InAsyncLoadingThread)
+		: AsyncLoadingThread(InAsyncLoadingThread)
+		, bNeedsToLeaveAsyncTick(false)
+	{
+		if (IsInGameThread())
+		{
+			AsyncLoadingThread.EnterAsyncLoadingTick();
+			bNeedsToLeaveAsyncTick = true;
+		}
+	}
+	~FAsyncLoadingTickScope2()
+	{
+		if (bNeedsToLeaveAsyncTick)
+		{
+			AsyncLoadingThread.LeaveAsyncLoadingTick();
+		}
+	}
+};
+
+void FAsyncLoadingThread2Impl::InitializeLoading()
+{
+	FString RootDir;
+	if (!FParse::Value(FCommandLine::Get(), TEXT("-zendir="), RootDir))
+	{
+		UE_LOG(LogStreaming, Error, TEXT("Failed to initialize package loader. No root directory specified"));
+		return;
+	}
+	
+	UE_LOG(LogStreaming, Log, TEXT("Initializing package loader using directory '%s'"), *RootDir);
+
+	{
+		TRACE_CPUPROFILER_EVENT_SCOPE(LoadGlobalNameMap);
+
+		FString GlobalNameMapFilePath = RootDir / TEXT("Container.namemap");
+		UE_LOG(LogStreaming, Log, TEXT("Loading global name map '%s'"), *GlobalNameMapFilePath);
+		GlobalNameMap.Load(GlobalNameMapFilePath);
+	}
+
+	{
+		TRACE_CPUPROFILER_EVENT_SCOPE(InitIoDispatcher);
+
+		IoStoreEnvironment.InitializeFileEnvironment(RootDir);
+		IoStoreReader = MakeUnique<FIoStoreReader>(IoStoreEnvironment);
+		FIoStatus ReaderStatus = IoStoreReader->Initialize(TEXT("PackageLoader"));
+
+		UE_CLOG(!ReaderStatus.IsOk(), LogStreaming, Error, TEXT("Failed to initialize I/O dispatcher: '%s'"), *ReaderStatus.ToString());
+
+		IoDispatcher.Mount(IoStoreReader.Get());
+		IoRequestQueue = MakeUnique<FIoRequestQueue>(IoDispatcher, AltZenaphore, 131072);
+	}
+
+	{
+		TRACE_CPUPROFILER_EVENT_SCOPE(LoadPackageStoreToc);
+
+		TUniquePtr<FArchive> StoreTocArchive(IFileManager::Get().CreateFileReader(*(RootDir / TEXT("megafile.ustoretoc"))));
+		check(StoreTocArchive)
+
+		const int32 PackageByteCount = StoreTocArchive->TotalSize();
+		PackageCount = PackageByteCount / sizeof(FPackageStoreEntrySerialized);
+		StoreEntriesSerialized = reinterpret_cast<FPackageStoreEntrySerialized*>(FMemory::Malloc(PackageByteCount));
+		// In-place loading
+		StoreTocArchive->Serialize(StoreEntriesSerialized, PackageByteCount);
+		// FName fixup
+		PackageNameToGlobalPackageId.Reserve(PackageCount);
+		for (int I = 0; I < PackageCount; ++I)
+		{
+			StoreEntriesSerialized[I].Name = GlobalNameMap.FromSerializedName(StoreEntriesSerialized[I].Name);
+			StoreEntriesSerialized[I].FileName = GlobalNameMap.FromSerializedName(StoreEntriesSerialized[I].FileName);
+			PackageNameToGlobalPackageId.Add(StoreEntriesSerialized[I].Name, FGlobalPackageId { I });
+		}
+	}
+
+	{
+		TRACE_CPUPROFILER_EVENT_SCOPE(LoadPackageStoreSlimports);
+
+		TUniquePtr<FArchive> SlimportArchive(IFileManager::Get().CreateFileReader(*(RootDir / TEXT("megafile.uslimport"))));
+		check(SlimportArchive);
+
+		const int32 SlimportByteCount = SlimportArchive->TotalSize();
+		SlimportCount = SlimportByteCount / sizeof(int32);
+		Slimports = reinterpret_cast<int32*>(FMemory::Malloc(SlimportByteCount));
+		SlimportArchive->Serialize(Slimports, SlimportByteCount);
+	}
+
+	{
+		TRACE_CPUPROFILER_EVENT_SCOPE(LoadPackageStoreGlimports);
+
+		TUniquePtr<FArchive> ImportArchive(IFileManager::Get().CreateFileReader(*(RootDir / TEXT("megafile.uglimport"))));
+		check(ImportArchive);
+
+		const int32 ImportByteCount = ImportArchive->TotalSize();
+		GlobalImportRuntime.Count = ImportByteCount / sizeof(FGlobalImport);
+		FGlobalImport* Imports = reinterpret_cast<FGlobalImport*>(FMemory::Malloc(ImportByteCount));
+		// In-place loading
+		ImportArchive->Serialize(Imports, ImportByteCount);
+
+		GlobalImportRuntime.Names = new FName[GlobalImportRuntime.Count];
+		GlobalImportRuntime.Outers = new FPackageIndex[GlobalImportRuntime.Count];
+		GlobalImportRuntime.Packages = new FPackageIndex[GlobalImportRuntime.Count];
+		GlobalImportRuntime.Objects = new UObject*[GlobalImportRuntime.Count];
+		GlobalImportRuntime.RefCounts = new TAtomic<int32>[GlobalImportRuntime.Count];
+
+		for (int I = 0; I < GlobalImportRuntime.Count; ++I)
+		{
+			FGlobalImport& Import = Imports[I];
+			Import.ObjectName = GlobalNameMap.FromSerializedName(Import.ObjectName);
+			GlobalImportRuntime.Names[I] = Import.ObjectName;
+			GlobalImportRuntime.Outers[I] = Import.OuterIndex;
+			GlobalImportRuntime.Packages[I] = Import.OutermostIndex;
+			GlobalImportRuntime.Objects[I] = nullptr;
+			GlobalImportRuntime.RefCounts[I] = 0;
+		}
+	}
+
+	StoreEntriesRuntime = new FPackageStoreEntryRuntime[PackageCount];
+	for (int I = 0; I < PackageCount; ++I)
+	{
+		FPackageStoreEntrySerialized& EntrySerialized = StoreEntriesSerialized[I];
+		FPackageStoreEntryRuntime& EntryRuntime = StoreEntriesRuntime[I];
+		EntryRuntime.Slimports = Slimports + EntrySerialized.SlimportOffset / sizeof(int32);
+	}
+
+	AsyncThreadReady.Increment();
+}
+
+void FAsyncLoadingThread2Impl::QueuePackage(FAsyncPackageDesc& Package)
+{
+	//TRACE_CPUPROFILER_EVENT_SCOPE(QueuePackage);
+	{
+		FScopeLock QueueLock(&QueueCritical);
+		++QueuedPackagesCounter;
+		QueuedPackages.Add(new FAsyncPackageDesc(Package, MoveTemp(Package.PackageLoadedDelegate)));
+	}
+	AltZenaphore.NotifyOne();
+}
+
+void FAsyncLoadingThread2Impl::UpdateExistingPackagePriorities(FAsyncPackage2* InPackage, TAsyncLoadPriority InNewPriority)
+{
+	check(!IsInGameThread() || !IsMultithreaded());
+	InPackage->SetPriority(InNewPriority);
+	return;
+}
+
+FAsyncPackage2* FAsyncLoadingThread2Impl::FindOrInsertPackage(FAsyncPackageDesc* InDesc, bool& bInserted)
+{
+	//TRACE_CPUPROFILER_EVENT_SCOPE(FindOrInsertPackage);
+	FAsyncPackage2* Package = nullptr;
+	bInserted = false;
+	{
+		FScopeLock LockAsyncPackages(&AsyncPackagesCritical);
+		Package = AsyncPackageNameLookup.FindRef(InDesc->Name);
+		if (!Package)
+		{
+			Package = CreateAsyncPackage(*InDesc);
+			Package->AddRef();
+			ExistingAsyncPackagesCounter.Increment();
+			AsyncPackageNameLookup.Add(Package->GetPackageName(), Package);
+			bInserted = true;
+		}
+		else if (InDesc->RequestID > 0)
+		{
+			Package->AddRequestID(InDesc->RequestID);
+		}
+		if (InDesc->PackageLoadedDelegate.IsValid())
+		{
+			const bool bInternalCallback = false;
+			Package->AddCompletionCallback(MoveTemp(InDesc->PackageLoadedDelegate), bInternalCallback);
+		}
+	}
+	if (bInserted)
+	{
+		QueueEvent_CreateLinker(Package);
+	}
+	return Package;
+}
+
+bool FAsyncLoadingThread2Impl::CreateAsyncPackagesFromQueue()
+{
+	TRACE_CPUPROFILER_EVENT_SCOPE(CreateAsyncPackagesFromQueue);
+	TArray<FAsyncPackageDesc*> QueueCopy;
+	{
+		FScopeLock QueueLock(&QueueCritical);
+		QueueCopy.Append(QueuedPackages);
+		QueuedPackages.Reset();
+	}
+
+	for (FAsyncPackageDesc* PackageRequest : QueueCopy)
+	{
+		bool bInserted;
+		FAsyncPackage2* Package = FindOrInsertPackage(PackageRequest, bInserted);
+		--QueuedPackagesCounter;
+		{
+			TRACE_CPUPROFILER_EVENT_SCOPE(ImportPackages);
+			Package->ImportPackagesRecursive();
+		}
+		delete PackageRequest;
+	}
+
+	return QueueCopy.Num() > 0;
+}
+
+FEventLoadNode2::FEventLoadNode2(const FAsyncLoadEventSpec* InSpec, FAsyncPackage2* InPackage, int32 InImportOrExportIndex)
+	: Spec(InSpec)
+	, Package(InPackage)
+	, ImportOrExportIndex(InImportOrExportIndex)
+{
+	check(Spec);
+	check(Package);
+}
+
+void FEventLoadNode2::DependsOn(FEventLoadNode2* Other)
+{
+	//TRACE_CPUPROFILER_EVENT_SCOPE(DependsOn);
+#if UE_BUILD_DEVELOPMENT
+	check(!bDone);
+	check(!bFired);
+#endif
+	uint8 Expected = 0;
+	while (!Other->DependencyWriterCount.CompareExchange(Expected, 1))
+	{
+		TRACE_CPUPROFILER_EVENT_SCOPE(DependsOnContested);
+		check(Expected == 1);
+		Expected = 0;
+	}
+	if (!Other->bDone.Load())
+	{
+		++BarrierCount;
+		if (Other->DependenciesCount == 0)
+		{
+			Other->SingleDependent = this;
+			Other->DependenciesCount = 1;
+		}
+		else
+		{
+			if (Other->DependenciesCount == 1)
+			{
+				FEventLoadNode2* FirstDependency = Other->SingleDependent;
+				uint32 NewDependenciesCapacity = 4;
+				Other->DependenciesCapacity = NewDependenciesCapacity;
+				Other->MultipleDependents = Package->GetGraphAllocator().AllocArcs(NewDependenciesCapacity);
+				Other->MultipleDependents[0] = FirstDependency;
+			}
+			else if (Other->DependenciesCount == Other->DependenciesCapacity)
+			{
+				FEventLoadNode2** OriginalDependents = Other->MultipleDependents;
+				uint32 OldDependenciesCapcity = Other->DependenciesCapacity;
+				SIZE_T OldDependenciesSize = OldDependenciesCapcity * sizeof(FEventLoadNode2*);
+				uint32 NewDependenciesCapacity = OldDependenciesCapcity * 2;
+				Other->DependenciesCapacity = NewDependenciesCapacity;
+				Other->MultipleDependents = Package->GetGraphAllocator().AllocArcs(NewDependenciesCapacity);
+				FMemory::Memcpy(Other->MultipleDependents, OriginalDependents, OldDependenciesSize);
+				Package->GetGraphAllocator().FreeArcs(OriginalDependents, OldDependenciesCapcity);
+			}
+			Other->MultipleDependents[Other->DependenciesCount++] = this;
+		}
+	}
+	Other->DependencyWriterCount.Store(0);
+}
+
+void FEventLoadNode2::AddBarrier()
+{
+#if UE_BUILD_DEVELOPMENT
+	check(!bDone);
+	check(!bFired);
+#endif
+	++BarrierCount;
+}
+
+void FEventLoadNode2::ReleaseBarrier()
+{
+	check(BarrierCount > 0);
+	if (--BarrierCount == 0)
+	{
+		Fire(FAsyncLoadingThreadState2::Get());
+	}
+}
+
+void FEventLoadNode2::Fire(FAsyncLoadingThreadState2& ThreadState)
+{
+	//TRACE_CPUPROFILER_EVENT_SCOPE(Fire);
+
+#if UE_BUILD_DEVELOPMENT
+	bFired.Store(1);
+#endif
+
+	if (Spec->bExecuteImmediately)
+	{
+		Execute(ThreadState);
+	}
+	else
+	{
+		Spec->EventQueue->Push(this);
+	}
+}
+
+void FEventLoadNode2::Execute(FAsyncLoadingThreadState2& ThreadState)
+{
+	//TRACE_CPUPROFILER_EVENT_SCOPE(ExecuteEvent);
+	check(BarrierCount.Load() == 0);
+	EAsyncPackageState::Type State = Spec->Func(Package, ImportOrExportIndex);
+	check(State == EAsyncPackageState::Complete);
+	bDone.Store(1);
+	ProcessDependencies(ThreadState);
+}
+
+void FEventLoadNode2::ProcessDependencies(FAsyncLoadingThreadState2& ThreadState)
+{
+	//TRACE_CPUPROFILER_EVENT_SCOPE(ProcessDependencies);
+	if (DependencyWriterCount.Load() != 0)
+	{
+		TRACE_CPUPROFILER_EVENT_SCOPE(ConcurrentWriter);
+		do
+		{
+			FPlatformProcess::Sleep(0);
+		} while (DependencyWriterCount.Load() != 0);
+	}
+
+	if (DependenciesCount == 1)
+	{
+		check(SingleDependent->BarrierCount > 0);
+		if (--SingleDependent->BarrierCount == 0)
+		{
+			ThreadState.NodesToFire.Push(SingleDependent);
+		}
+	}
+	else if (DependenciesCount != 0)
+	{
+		FEventLoadNode2** Current = MultipleDependents;
+		FEventLoadNode2** End = MultipleDependents + DependenciesCount;
+		for (; Current < End; ++Current)
+		{
+			FEventLoadNode2* Dependent = *Current;
+			check(Dependent->BarrierCount > 0);
+			if (--Dependent->BarrierCount == 0)
+			{
+				ThreadState.NodesToFire.Push(Dependent);
+			}
+		}
+		ThreadState.DeferredFreeArcs.Add(MakeTuple(MultipleDependents, DependenciesCapacity));
+	}
+	if (ThreadState.bShouldFireNodes)
+	{
+		ThreadState.bShouldFireNodes = false;
+		while (ThreadState.NodesToFire.Num())
+		{
+			ThreadState.NodesToFire.Pop(false)->Fire(ThreadState);
+		}
+		ThreadState.bShouldFireNodes = true;
+	}
+}
+
+FAsyncLoadEventQueue2::FAsyncLoadEventQueue2()
+{
+	FMemory::Memzero(Entries, sizeof(Entries));
+}
+
+FAsyncLoadEventQueue2::~FAsyncLoadEventQueue2()
+{
+}
+
+void FAsyncLoadEventQueue2::Push(FEventLoadNode2* Node)
+{
+	uint64 LocalHead = Head.IncrementExchange();
+	FEventLoadNode2* Expected = nullptr;
+	if (!Entries[LocalHead % UE_ARRAY_COUNT(Entries)].CompareExchange(Expected, Node))
+	{
+		*(volatile int*)0 = 0; // queue is full: TODO
+	}
+	if (Zenaphore)
+	{
+		Zenaphore->NotifyOne();
+	}
+}
+
+bool FAsyncLoadEventQueue2::PopAndExecute(FAsyncLoadingThreadState2& ThreadState)
+{
+	FEventLoadNode2* Node = nullptr;
+	{
+		uint64 LocalHead = Head.Load();
+		uint64 LocalTail = Tail.Load();
+		for (;;)
+		{
+			if (LocalTail >= LocalHead)
+			{
+				break;
+			}
+			if (Tail.CompareExchange(LocalTail, LocalTail + 1))
+			{
+				while (!Node)
+				{
+					Node = Entries[LocalTail % UE_ARRAY_COUNT(Entries)].Exchange(nullptr);
+				}
+				break;
+			}
+		}
+	}
+
+	if (Node)
+	{
+		//TRACE_CPUPROFILER_EVENT_SCOPE(Execute);
+		Node->Execute(ThreadState);
+		return true;
+	}
+	else
+	{
+		return false;
+	}
+}
+
+FScopedAsyncPackageEvent2::FScopedAsyncPackageEvent2(FAsyncPackage2* InPackage)
+	:Package(InPackage)
+{
+	check(Package);
+
+	// Update the thread context with the current package. This is used by NotifyConstructedDuringAsyncLoading.
+	FUObjectThreadContext& ThreadContext = FUObjectThreadContext::Get();
+	PreviousPackage = static_cast<FAsyncPackage2*>(ThreadContext.AsyncPackage);
+	ThreadContext.AsyncPackage = Package;
+
+	Package->BeginAsyncLoad();
+}
+
+FScopedAsyncPackageEvent2::~FScopedAsyncPackageEvent2()
+{
+	Package->EndAsyncLoad();
+
+	// Restore the package from the outer scope
+	FUObjectThreadContext& ThreadContext = FUObjectThreadContext::Get();
+	ThreadContext.AsyncPackage = PreviousPackage;
+}
+
+void FAsyncLoadingThreadWorker::StartThread()
+{
+	Thread = FRunnableThread::Create(this, TEXT("FAsyncLoadingThreadWorker"), 0, TPri_Normal);
+	ThreadId = Thread->GetThreadID();
+	TRACE_SET_THREAD_GROUP(ThreadId, "AsyncLoading");
+}
+
+uint32 FAsyncLoadingThreadWorker::Run()
+{
+	LLM_SCOPE(ELLMTag::AsyncLoading);
+
+	if (!IsInGameThread())
+	{
+		FPlatformProcess::SetThreadAffinityMask(FPlatformAffinity::GetAsyncLoadingThreadMask());
+	}
+
+	FAsyncLoadingThreadState2::Create(GraphAllocator);
+
+	FZenaphoreWaiter Waiter(Zenaphore, TEXT("WaitForEvents"));
+
+	FAsyncLoadingThreadState2& ThreadState = FAsyncLoadingThreadState2::Get();
+
+	bool bSuspended = false;
+	while (!bStopRequested)
+	{
+		if (bSuspended)
+		{
+			if (!bSuspendRequested.Load(EMemoryOrder::SequentiallyConsistent))
+			{
+				bSuspended = false;
+			}
+			else
+			{
+				FPlatformProcess::Sleep(0.001f);
+			}
+		}
+		else
+		{
+			bool bDidSomething = false;
+			{
+				FGCScopeGuard GCGuard;
+				TRACE_CPUPROFILER_EVENT_SCOPE(AsyncLoadingTime);
+				++ActiveWorkersCount;
+				do
+				{
+					bDidSomething = EventQueue.PopAndExecute(ThreadState);
+					
+					if (bSuspendRequested.Load(EMemoryOrder::Relaxed))
+					{
+						bSuspended = true;
+						bDidSomething = true;
+						break;
+					}
+				} while (bDidSomething);
+				--ActiveWorkersCount;
+			}
+			if (!bDidSomething)
+			{
+				ThreadState.ProcessDeferredFrees();
+				Waiter.Wait();
+			}
+		}
+	}
+	return 0;
+}
+
+void FAsyncLoadingThread2Impl::QueueEvent_CreateLinker(FAsyncPackage2* Package)
+{
+	check(Package);
+	Package->GetNode(EEventLoadNode2::Package_CreateLinker)->ReleaseBarrier();
+}
+
+FUObjectSerializeContext* FAsyncPackage2::GetSerializeContext()
+{
+	return FUObjectThreadContext::Get().GetSerializeContext();
+}
+
+EAsyncPackageState::Type FAsyncPackage2::Event_CreateLinker(FAsyncPackage2* Package, int32)
+{
+	TRACE_CPUPROFILER_EVENT_SCOPE(Event_CreateLinker);
+	// Keep track of time when we start loading.
+	if (Package->LoadStartTime == 0.0)
+	{
+		double Now = FPlatformTime::Seconds();
+		Package->LoadStartTime = Now;
+	}
+	FScopedAsyncPackageEvent2 Scope(Package);
+	//SCOPED_LOADTIMER(Package_CreateLinker);
+	check(!Package->Linker);
+	Package->CreateLinker();
+	check(Package->Linker);
+	check(Package->AsyncPackageLoadingState == EAsyncPackageLoadingState2::NewPackage);
+	Package->AsyncPackageLoadingState = EAsyncPackageLoadingState2::WaitingForSummary;
+	return EAsyncPackageState::Complete;
+}
+
+EAsyncPackageState::Type FAsyncPackage2::Event_FinishLinker(FAsyncPackage2* Package, int32)
+{
+	TRACE_CPUPROFILER_EVENT_SCOPE(Event_FinishLinker);
+	FScopedAsyncPackageEvent2 Scope(Package);
+	EAsyncPackageState::Type Result = Package->FinishLinker();
+	check(Result == EAsyncPackageState::Complete);
+	check(Package->Linker/* && Linker->HasFinishedInitialization()*/);
+	check(Package->AsyncPackageLoadingState == EAsyncPackageLoadingState2::WaitingForSummary);
+	Package->AsyncPackageLoadingState = EAsyncPackageLoadingState2::StartImportPackages;
+	return EAsyncPackageState::Complete;
+}
+
+EAsyncPackageState::Type FAsyncPackage2::Event_StartImportPackages(FAsyncPackage2* Package, int32)
+{
+	TRACE_CPUPROFILER_EVENT_SCOPE(Event_StartImportPackages);
+
+	LLM_SCOPE(ELLMTag::AsyncLoading);
+
+	{
+		TRACE_CPUPROFILER_EVENT_SCOPE(SetupSerializedArcs);
+		const TArray<FNameEntryId>* GlobalNameMap = &Package->AsyncLoadingThread.GlobalNameMap.GetNameEntries();
+
+		const FPackageSummary* PackageSummary = reinterpret_cast<const FPackageSummary*>(Package->PackageSummaryBuffer.Get());
+		uint64 ZenHeaderDataSize = PackageSummary->GraphDataSize;
+		const uint8* ZenHeaderData = Package->PackageSummaryBuffer.Get() + PackageSummary->GraphDataOffset;
+		FSimpleArchive ZenHeaderArchive(ZenHeaderData, ZenHeaderDataSize);
+		int32 InternalArcCount;
+		ZenHeaderArchive << InternalArcCount;
+		for (int32 InternalArcIndex = 0; InternalArcIndex < InternalArcCount; ++InternalArcIndex)
+		{
+			int32 FromNodeIndex;
+			int32 ToNodeIndex;
+			ZenHeaderArchive << FromNodeIndex;
+			ZenHeaderArchive << ToNodeIndex;
+			Package->PackageNodes[ToNodeIndex].DependsOn(Package->PackageNodes + FromNodeIndex);
+		}
+		int32 ImportedPackagesCount;
+		ZenHeaderArchive << ImportedPackagesCount;
+		for (int32 ImportedPackageIndex = 0; ImportedPackageIndex < ImportedPackagesCount; ++ImportedPackageIndex)
+		{
+			int32 ImportedPackageNameIndex;
+			int32 ImportedPackageNameNumber;
+			ZenHeaderArchive << ImportedPackageNameIndex << ImportedPackageNameNumber;
+			FNameEntryId MappedName = (*GlobalNameMap)[ImportedPackageNameIndex];
+			FName ImportedPackageName = FName::CreateFromDisplayId(MappedName, ImportedPackageNameNumber);
+			FAsyncPackage2* ImportedPackage = Package->AsyncLoadingThread.FindAsyncPackage(ImportedPackageName);
+			int32 ExternalArcCount;
+			ZenHeaderArchive << ExternalArcCount;
+			for (int32 ExternalArcIndex = 0; ExternalArcIndex < ExternalArcCount; ++ExternalArcIndex)
+			{
+				int32 FromNodeIndex;
+				int32 ToNodeIndex;
+				ZenHeaderArchive << FromNodeIndex;
+				ZenHeaderArchive << ToNodeIndex;
+				if (ImportedPackage)
+				{
+					Package->PackageNodes[ToNodeIndex].DependsOn(ImportedPackage->PackageNodes + FromNodeIndex);
+				}
+			}
+		}
+	}
+
+	check(Package->AsyncPackageLoadingState == EAsyncPackageLoadingState2::StartImportPackages);
+	Package->AsyncPackageLoadingState = EAsyncPackageLoadingState2::WaitingForImportPackages;
+	return EAsyncPackageState::Complete;
+}
+
+static UObject* GFindExistingSlimport(int32 GlobalImportIndex,
+	UObject** GlobalImportObjects,
+	FPackageIndex* GlobalImportOuters,
+	FName* GlobalImportNames)
+{
+	UObject*& Object = GlobalImportObjects[GlobalImportIndex];
+	if (!Object)
+	{
+		const FPackageIndex& OuterIndex = GlobalImportOuters[GlobalImportIndex];
+		const FName& ObjectName = GlobalImportNames[GlobalImportIndex];
+		if (OuterIndex.IsNull())
+		{
+			Object = StaticFindObjectFast(UPackage::StaticClass(), nullptr, ObjectName, true);
+		}
+		else
+		{
+			UObject* Outer = GFindExistingSlimport(OuterIndex.ToImport(), GlobalImportObjects, GlobalImportOuters, GlobalImportNames);
+			if (Outer)
+			{
+				Object = StaticFindObjectFast(UObject::StaticClass(), Outer, ObjectName, false, true);
+			}
+		}
+	}
+	return Object;
+}
+
+UObject* FAsyncPackage2::FindExistingSlimport(int32 GlobalImportIndex)
+{
+	return GFindExistingSlimport(GlobalImportIndex,
+		GlobalImportObjects,
+		GlobalImportOuters,
+		GlobalImportNames);
+}
+
+void FAsyncPackage2::ImportPackagesRecursive()
+{
+	if (bHasImportedPackagesRecursive)
+	{
+		return;
+	}
+	bHasImportedPackagesRecursive = true;
+
+#if 1
+	FMemStack& MemStack = FMemStack::Get();
+	FMemMark Mark(MemStack);
+	TArrayView<FPackageIndex> LocalImportedPackages(new(MemStack)FPackageIndex[LocalImportCount], LocalImportCount);
+
+	for (int32 LocalImportIndex = 0; LocalImportIndex < LocalImportCount; ++LocalImportIndex)
+	{
+		const int32 GlobalImportIndex = LocalImportIndices[LocalImportIndex];
+		const FPackageIndex ImportedPackageIndex = GlobalImportPackages[GlobalImportIndex];
+		UObject* ImportedObject = FindExistingSlimport(GlobalImportIndex);
+		const UPackage* Package = (UPackage*)GlobalImportObjects[ImportedPackageIndex.ToImport()];
+
+		if (Package && Package->HasAnyPackageFlags(PKG_CompiledIn))
+		{
+			continue;
+		}
+
+		if (ImportedObject)
+		{
+			if (IsFullyLoadedObj(ImportedObject))
+			{
+				continue;
+			}
+		}
+
+		if (LocalImportedPackages.Contains(ImportedPackageIndex))
+		{
+			continue;
+		}
+
+		LocalImportedPackages[LocalImportIndex] = ImportedPackageIndex;
+
+		const FName ImportedPackageName = GlobalImportNames[ImportedPackageIndex.ToImport()];
+		FAsyncPackageDesc Info(INDEX_NONE, ImportedPackageName);
+		Info.Priority = Desc.Priority;
+		bool bInserted;
+		FAsyncPackage2* ImportedAsyncPackage = AsyncLoadingThread.FindOrInsertPackage(&Info, bInserted);
+		ImportedAsyncPackage->AddRef();
+		ImportedAsyncPackages.Add(ImportedAsyncPackage);
+		if (bInserted)
+		{
+			ImportedAsyncPackage->ImportPackagesRecursive();
+		}
+
+		// we can't set up our imports until all packages we are importing have loaded their summary
+		FEventLoadNode2* SetupImportsNode = GetNode(EEventLoadNode2::Package_SetupImports);
+		FEventLoadNode2* ImportedPackageLoadSummaryNode = ImportedAsyncPackage->GetNode(EEventLoadNode2::Package_LoadSummary);
+		SetupImportsNode->DependsOn(ImportedPackageLoadSummaryNode);
+	}
+#else
+	bool bNeedToLoadPackage = false;
+	bool bNeedToCheckImportObjectsInPackage = true;
+	bool bHasAddedNeedToLoadPackage = false;
+	FPackageIndex LastImportedPackageIndex;
+
+	// slimports are sorted, first comes a new package, then N additional objects
+	for (int32 LocalImportIndex = 0; LocalImportIndex < LocalImportCount; ++LocalImportIndex)
+	{
+		const int32 GlobalImportIndex = LocalImportIndices[LocalImportIndex];
+		const FPackageIndex ImportedPackageIndex = GlobalImportPackages[GlobalImportIndex];
+
+		const bool bIsANewImportPackage = ImportedPackageIndex != LastImportedPackageIndex;
+		if (bIsANewImportPackage)
+		{
+			LastImportedPackageIndex = ImportedPackageIndex;
+			bHasAddedNeedToLoadPackage = false;
+
+			UPackage* ImportedPackage = (UPackage*)FindExistingSlimport(ImportedPackageIndex.ToImport());
+			if (ImportedPackage)
+			{
+				bNeedToLoadPackage = false;
+				if (ImportedPackage->HasAnyPackageFlags(PKG_CompiledIn))
+				{
+					bNeedToCheckImportObjectsInPackage = false;
+				}
+				else
+				{
+					AddObjectReference(ImportedPackage);
+					bNeedToCheckImportObjectsInPackage = true;
+				}
+			}
+			else
+			{
+				bNeedToLoadPackage = true;
+				bNeedToCheckImportObjectsInPackage = false;
+			}
+		}
+		else if (bNeedToCheckImportObjectsInPackage)
+		{
+			UObject* ImportedObject = FindExistingSlimport(GlobalImportIndex);
+			if (ImportedObject)
+			{
+				AddObjectReference(ImportedObject);
+				bNeedToLoadPackage |= !IsFullyLoadedObj(ImportedObject);
+			}
+			else
+			{
+				bNeedToLoadPackage = true;
+			}
+		}
+
+		if (bNeedToLoadPackage && !bHasAddedNeedToLoadPackage)
+		{
+			bHasAddedNeedToLoadPackage = true;
+
+			const FName ImportedPackageName = GlobalImportNames[ImportedPackageIndex.ToImport()];
+			FAsyncPackageDesc Info(INDEX_NONE, ImportedPackageName);
+			Info.Priority = Desc.Priority;
+			bool bInserted;
+			FAsyncPackage2* ImportedAsyncPackage = AsyncLoadingThread.FindOrInsertPackage(&Info, bInserted);
+			ImportedAsyncPackage->AddRef();
+			ImportedAsyncPackages.Add(ImportedAsyncPackage);
+			if (bInserted)
+			{
+				ImportedAsyncPackage->ImportPackagesRecursive();
+			}
+
+			FEventLoadNode2* SetupImportsNode = GetNode(EEventLoadNode2::Package_SetupImports);
+			FEventLoadNode2* ImportedPackageLoadSummaryNode = ImportedAsyncPackage->GetNode(EEventLoadNode2::Package_LoadSummary);
+			SetupImportsNode->DependsOn(ImportedPackageLoadSummaryNode);
+		}
+	}
+#endif
+}
+
+EAsyncPackageState::Type FAsyncPackage2::Event_SetupImports(FAsyncPackage2* Package, int32)
+{
+	TRACE_CPUPROFILER_EVENT_SCOPE(Event_SetupImports);
+
+	check(Package->AsyncPackageLoadingState == EAsyncPackageLoadingState2::WaitingForImportPackages);
+	Package->AsyncPackageLoadingState = EAsyncPackageLoadingState2::SetupImports;
+	{
+		FScopedAsyncPackageEvent2 Scope(Package);
+		verify(Package->SetupSlimports_Event() == EAsyncPackageState::Complete);
+	}
+	check(Package->AsyncPackageLoadingState == EAsyncPackageLoadingState2::SetupImports);
+	check(Package->ImportIndex == Package->LocalImportCount);
+	Package->AsyncPackageLoadingState = EAsyncPackageLoadingState2::SetupExports;
+	return EAsyncPackageState::Complete;
+}
+
+EAsyncPackageState::Type FAsyncPackage2::SetupSlimports_Event()
+{
+	if (!GIsInitialLoad)
+	{
+		ImportIndex = LocalImportCount;
+		return EAsyncPackageState::Complete;
+	}
+
+	while (ImportIndex < LocalImportCount)
+	{
+		const int32 LocalImportIndex = ImportIndex++;
+		const int32 GlobalImportIndex = LocalImportIndices[LocalImportIndex];
+
+		// skip packages
+		if (GlobalImportOuters[GlobalImportIndex].IsNull())
+		{
+			continue;
+		}
+
+		// find package of import object
+		FPackageIndex PackageIndex = GlobalImportPackages[GlobalImportIndex];
+		const FName& PackageName = GlobalImportNames[PackageIndex.ToImport()];
+		UObject*& Package = GlobalImportObjects[PackageIndex.ToImport()];
+		UPackage* ImportPackage = Package ? CastChecked<UPackage>(Package) : nullptr;
+		if (!ImportPackage)
+		{
+			ImportPackage = FindObjectFast<UPackage>(NULL, PackageName, false, false);
+			Package = ImportPackage; // This will write to global import table!!!
+		}
+		check(ImportPackage);
+
+		// do initial loading stuff for compiled in packages
+		FLinkerLoad* ImportLinker = ImportPackage->LinkerLoad;
+		const bool bDynamicImport = ImportLinker && ImportLinker->bDynamicClassLinker;
+		if (!ImportLinker && ImportPackage->HasAnyPackageFlags(PKG_CompiledIn) && !bDynamicImport)
+		{
+			FPackageIndex OuterMostIndex = FPackageIndex::FromImport(GlobalImportIndex);
+			FPackageIndex OuterMostNonPackageIndex = OuterMostIndex;
+			while (true)
+			{
+				check(!OuterMostIndex.IsNull() && OuterMostIndex.IsImport());
+				FPackageIndex NextOuterMostIndex = GlobalImportOuters[OuterMostIndex.ToImport()];
+
+				if (NextOuterMostIndex.IsNull())
+				{
+					break;
+				}
+				OuterMostNonPackageIndex = OuterMostIndex;
+				OuterMostIndex = NextOuterMostIndex;
+			}
+			FName OuterMostNonPackageObjectName = GlobalImportNames[OuterMostNonPackageIndex.ToImport()];
+			check(GlobalImportOuters[OuterMostIndex.ToImport()].IsNull());
+			check(PackageName == GlobalImportNames[OuterMostIndex.ToImport()]);
+			// OuterMostNonPackageIndex is used here because if it is a CDO or subobject, etc,
+			// we wait for the outermost thing that is not a package
+			const bool bWaitingForCompiledInImport = EDLBootNotificationManager.AddWaitingPackage(
+				this, PackageName, OuterMostNonPackageObjectName, FPackageIndex::FromImport(LocalImportIndex));
+			if (bWaitingForCompiledInImport)
+			{
+				GetNode(EEventLoadNode2::ImportOrExport_Create, FPackageIndex::FromImport(LocalImportIndex))->AddBarrier();
+			}
+		}
+	}
+
+	return ImportIndex == LocalImportCount ? EAsyncPackageState::Complete : EAsyncPackageState::TimeOut;
+}
+
+EAsyncPackageState::Type FAsyncPackage2::Event_SetupExports(FAsyncPackage2* Package, int32)
+{
+	TRACE_CPUPROFILER_EVENT_SCOPE(Event_SetupExports);
+
+	Package->ExportIndex = Package->Linker->ExportMap.Num();
+	check(Package->AsyncPackageLoadingState == EAsyncPackageLoadingState2::SetupExports);
+	Package->AsyncPackageLoadingState = EAsyncPackageLoadingState2::ProcessNewImportsAndExports;
+	return EAsyncPackageState::Complete;
+}
+
+EAsyncPackageState::Type FAsyncPackage2::Event_LinkImport(FAsyncPackage2* Package, int32 LocalImportIndex)
+{
+	TRACE_CPUPROFILER_EVENT_SCOPE(Event_LinkImport);
+	Package->LinkSlimport(LocalImportIndex);
+	return EAsyncPackageState::Complete;
+}
+
+EAsyncPackageState::Type FAsyncPackage2::Event_ImportSerialized(FAsyncPackage2* Package, int32 LocalImportIndex)
+{
+	TRACE_CPUPROFILER_EVENT_SCOPE(Event_ImportSerialized);
+
+	int32 GlobalImportIndex = Package->LocalImportIndices[LocalImportIndex];
+	UObject* Object = Package->GlobalImportObjects[GlobalImportIndex];
+	if (Object)
+	{
+		checkf(!Object->HasAnyFlags(RF_NeedLoad), TEXT("%s had RF_NeedLoad yet it was marked as serialized."), *Object->GetFullName()); // otherwise is isn't done serializing, is it?
+	}
+	return EAsyncPackageState::Complete;
+}
+
+EAsyncPackageState::Type FAsyncPackage2::Event_CreateExport(FAsyncPackage2* Package, int32 LocalExportIndex)
+{
+	TRACE_CPUPROFILER_EVENT_SCOPE(Event_CreateExport);
+
+	check(Package->AsyncPackageLoadingState == EAsyncPackageLoadingState2::ProcessNewImportsAndExports);
+
+	FScopedAsyncPackageEvent2 Scope(Package);
+
+	Package->EventDrivenCreateExport(LocalExportIndex);
+	return EAsyncPackageState::Complete;
+}
+
+EAsyncPackageState::Type FAsyncPackage2::Event_SerializeExport(FAsyncPackage2* Package, int32 LocalExportIndex)
+{
+	TRACE_CPUPROFILER_EVENT_SCOPE(Event_SerializeExport);
+
+	FScopedAsyncPackageEvent2 Scope(Package);
+
+	Package->EventDrivenSerializeExport(LocalExportIndex);
+	FObjectExport& Export = Package->Linker->ExportMap[LocalExportIndex];
+	UObject* Object = Export.Object;
+	check(!Object || !Object->HasAnyFlags(RF_NeedLoad));
+	return EAsyncPackageState::Complete;
+}
+
+void FAsyncPackage2::LinkSlimport(int32 LocalImportIndex, int32 GlobalImportIndex)
+{
+	if (GlobalImportIndex == -1)
+	{
+		check(LocalImportIndex >= 0 && LocalImportIndex < LocalImportCount);
+		GlobalImportIndex = LocalImportIndices[LocalImportIndex];
+	}
+
+	UObject*& Object = GlobalImportObjects[GlobalImportIndex];
+
+	if (Object)
+	{
+		return;
+	}
+
+	const FName ObjectName = GlobalImportNames[GlobalImportIndex];
+	const FPackageIndex PackageIndex = GlobalImportPackages[GlobalImportIndex];
+
+	UObject*& ImportPackage = GlobalImportObjects[PackageIndex.ToImport()];
+	// UPackage* ImportPackage = (UPackage*)FindExistingSlimport(PackageIndex.ToImport());
+
+	if (!ImportPackage)
+	{
+		const FName PackageName = GlobalImportNames[PackageIndex.ToImport()];
+		ImportPackage = FindObjectFast<UPackage>(NULL, PackageName, false, false);
+		check(ImportPackage);
+	}
+	if (PackageIndex.ToImport() != GlobalImportIndex)
+	{
+		UObject* Outer = ImportPackage;
+		FPackageIndex OuterIndex = GlobalImportOuters[GlobalImportIndex];
+		if (PackageIndex != OuterIndex)
+		{
+			if (!GlobalImportObjects[OuterIndex.ToImport()])
+			{
+				LinkSlimport(/*dummy_local_index*/-1, OuterIndex.ToImport());
+			}
+			Outer = GlobalImportObjects[OuterIndex.ToImport()];
+		}
+		check(Outer);
+
+		Object = StaticFindObjectFast(UObject::StaticClass(), Outer, ObjectName, false, true);
+		if (!Object)
+		{
+			TRACE_CPUPROFILER_EVENT_SCOPE(LinkImport_SpinWait);
+			while (!Object)
+			{
+				Object = StaticFindObjectFast(UObject::StaticClass(), Outer, ObjectName, false, true);
+			}
+		}
+	}
+	check(Object);
+}
+
+UObject* FAsyncPackage2::EventDrivenIndexToObject(FPackageIndex Index, bool bCheckSerialized)
+{
+	UObject* Result = nullptr;
+	if (Index.IsNull())
+	{
+		return Result;
+	}
+	if (Index.IsExport())
+	{
+		Result = Linker->Exp(Index).Object;
+	}
+	else if (Index.IsImport())
+	{
+		int32 GlobalImportIndex = LocalImportIndices[Index.ToImport()];
+		Result = FindExistingSlimport(GlobalImportIndex);
+		check(Result);
+	}
+	if (bCheckSerialized && !IsFullyLoadedObj(Result))
+	{
+		FEventLoadNode2* MyDependentNode = GetNode(EEventLoadNode2::ImportOrExport_Serialize, Index);
+		if (!Result)
+		{
+			UE_LOG(LogStreaming, Error, TEXT("Missing Dependency, request for %s but it hasn't been created yet."), *Linker->GetPathName(Index));
+		}
+		else if (!MyDependentNode || MyDependentNode->GetBarrierCount() > 0)
+		{
+			UE_LOG(LogStreaming, Fatal, TEXT("Missing Dependency, request for %s but it was still waiting for serialization."), *Linker->GetPathName(Index));
+		}
+		else
+		{
+			UE_LOG(LogStreaming, Fatal, TEXT("Missing Dependency, request for %s but it was still has RF_NeedLoad."), *Linker->GetPathName(Index));
+		}
+	}
+	if (Result)
+	{
+		UE_CLOG(Result->HasAnyInternalFlags(EInternalObjectFlags::Unreachable), LogStreaming, Fatal, TEXT("Returning an object  (%s) from EventDrivenIndexToObject that is unreachable."), *Result->GetFullName());
+	}
+	return Result;
+}
+
+
+void FAsyncPackage2::EventDrivenCreateExport(int32 LocalExportIndex)
+{
+	//SCOPED_LOADTIMER(Package_CreateExports);
+	FObjectExport& Export = Linker->ExportMap[LocalExportIndex];
+
+	TRACE_LOADTIME_CREATE_EXPORT_SCOPE(Linker, &Export.Object, Export.SerialOffset, Export.SerialSize, Export.bIsAsset);
+
+	LLM_SCOPE(ELLMTag::AsyncLoading);
+	LLM_SCOPED_TAG_WITH_OBJECT_IN_SET(GetLinkerRoot(), ELLMTagSet::Assets);
+	LLM_SCOPED_TAG_WITH_OBJECT_IN_SET((Export.DynamicType == FObjectExport::EDynamicType::DynamicType) ? UDynamicClass::StaticClass() : 
+		CastEventDrivenIndexToObject<UClass>(Export.ClassIndex, false), ELLMTagSet::AssetClasses);
+
+
+	// Check whether we already loaded the object and if not whether the context flags allow loading it.
+	//check(!Export.Object || Export.Object->HasAnyFlags(RF_ClassDefaultObject)); // we should not have this yet, unless it is a CDO
+	check(!Export.Object); // we should not have this yet
+	if (!Export.Object && !Export.bExportLoadFailed)
+	{
+		FUObjectSerializeContext* LoadContext = GetSerializeContext();
+
+		if (!Linker->FilterExport(Export)) // for some acceptable position, it was not "not for" 
+		{
+			check(Export.ObjectName != NAME_None || !(Export.ObjectFlags&RF_Public));
+			check(LoadContext->HasStartedLoading());
+			if (Export.DynamicType == FObjectExport::EDynamicType::DynamicType)
+			{
+//native blueprint 
+
+				Export.Object = ConstructDynamicType(*Linker->GetExportPathName(LocalExportIndex), EConstructDynamicType::OnlyAllocateClassObject);
+				check(Export.Object);
+				UDynamicClass* DC = Cast<UDynamicClass>(Export.Object);
+				UObject* DCD = DC ? DC->GetDefaultObject(false) : nullptr;
+				if (GIsInitialLoad || GUObjectArray.IsOpenForDisregardForGC())
+				{
+					Export.Object->AddToRoot();
+					if (DCD)
+					{
+						DCD->AddToRoot();
+					}
+				}
+				// TODO: Replace with DCD->SetInternalFlags(EObjectInternalFlag::Async)?
+				// But then it also need to be cleared in ClearOwnedObjects()
+				// if (DCD)
+				// {
+				// 	AddObjectReference(DCD);
+				// }
+				UE_LOG(LogStreaming, Verbose, TEXT("EventDrivenCreateExport: Created dynamic class %s"), *Export.Object->GetFullName());
+				if (Export.Object)
+				{
+					Export.Object->SetLinker(Linker, LocalExportIndex);
+				}
+			}
+			else if (Export.DynamicType == FObjectExport::EDynamicType::ClassDefaultObject)
+			{
+				UClass* LoadClass = nullptr;
+				if (!Export.ClassIndex.IsNull())
+				{
+					LoadClass = CastEventDrivenIndexToObject<UClass>(Export.ClassIndex, true);
+				}
+				if (!LoadClass)
+				{
+					UE_LOG(LogStreaming, Error, TEXT("Could not find class %s to create %s"), *Linker->ImpExp(Export.ClassIndex).ObjectName.ToString(), *Export.ObjectName.ToString());
+					Export.bExportLoadFailed = true;
+					return;
+				}
+				Export.Object = LoadClass->GetDefaultObject(true);
+				if (Export.Object)
+				{
+					Export.Object->SetLinker(Linker, LocalExportIndex);
+				}
+			}
+			else
+			{
+				UClass* LoadClass = nullptr;
+				if (Export.ClassIndex.IsNull())
+				{
+					LoadClass = UClass::StaticClass();
+				}
+				else
+				{
+					LoadClass = CastEventDrivenIndexToObject<UClass>(Export.ClassIndex, true);
+				}
+				if (!LoadClass)
+				{
+					UE_LOG(LogStreaming, Error, TEXT("Could not find class %s to create %s"), *Linker->ImpExp(Export.ClassIndex).ObjectName.ToString(), *Export.ObjectName.ToString());
+					Export.bExportLoadFailed = true;
+					return;
+				}
+				UObject* ThisParent = nullptr;
+				if (!Export.OuterIndex.IsNull())
+				{
+					ThisParent = EventDrivenIndexToObject(Export.OuterIndex, false);
+				}
+				else if (Export.bForcedExport)
+				{
+					// see FLinkerLoad::CreateExport, there may be some more we can do here
+					check(!Export.bForcedExport); // this is leftover from seekfree loading I think
+				}
+				else
+				{
+					check(LinkerRoot);
+					ThisParent = LinkerRoot;
+				}
+				check(!dynamic_cast<UObjectRedirector*>(ThisParent));
+				if (!ThisParent)
+				{
+					UE_LOG(LogStreaming, Error, TEXT("Could not find outer %s to create %s"), *Linker->ImpExp(Export.OuterIndex).ObjectName.ToString(), *Export.ObjectName.ToString());
+					Export.bExportLoadFailed = true;
+					return;
+				}
+
+				// Try to find existing object first in case we're a forced export to be able to reconcile. Also do it for the
+				// case of async loading as we cannot in-place replace objects.
+
+				UObject* ActualObjectWithTheName = StaticFindObjectFastInternal(NULL, ThisParent, Export.ObjectName, true);
+
+				// Always attempt to find object in memory first
+				if (ActualObjectWithTheName && (ActualObjectWithTheName->GetClass() == LoadClass))
+				{
+					Export.Object = ActualObjectWithTheName;
+				}
+
+				// Object is found in memory.
+				if (Export.Object)
+				{
+					check(!Export.bForcedExport);
+					Export.Object->SetLinker(Linker, LocalExportIndex);
+
+					// If this object was allocated but never loaded (components created by a constructor, CDOs, etc) make sure it gets loaded
+					// Do this for all subobjects created in the native constructor.
+					if (!Export.Object->HasAnyFlags(RF_LoadCompleted))
+					{
+						UE_LOG(LogStreaming, VeryVerbose, TEXT("Note2: %s was constructed during load and is an export and so needs loading."), *Export.Object->GetFullName());
+						UE_CLOG(!Export.Object->HasAllFlags(RF_WillBeLoaded), LogStreaming, Fatal, TEXT("%s was found in memory and is an export but does not have all load flags."), *Export.Object->GetFullName());
+						if(Export.Object->HasAnyFlags(RF_ClassDefaultObject))
+						{
+							// never call PostLoadSubobjects on class default objects, this matches the behavior of the old linker where
+							// StaticAllocateObject prevents setting of RF_NeedPostLoad and RF_NeedPostLoadSubobjects, but FLinkerLoad::Preload
+							// assigns RF_NeedPostLoad for blueprint CDOs:
+							Export.Object->SetFlags(RF_NeedLoad | RF_NeedPostLoad | RF_WasLoaded);
+						}
+						else
+						{
+							Export.Object->SetFlags(RF_NeedLoad | RF_NeedPostLoad | RF_NeedPostLoadSubobjects | RF_WasLoaded);
+						}
+						Export.Object->ClearFlags(RF_WillBeLoaded);
+						if (IsInGameThread())
+						{
+							Export.Object->SetInternalFlags(EInternalObjectFlags::Async);
+							AddOwnedObject(Export.Object);
+						}
+					}
+					else
+					{
+						Export.Object->SetInternalFlags(EInternalObjectFlags::Async);
+						AddOwnedObject(Export.Object);
+					}
+					check(Export.Object->HasAnyInternalFlags(EInternalObjectFlags::Async));
+				}
+				else
+				{
+					if (ActualObjectWithTheName && !ActualObjectWithTheName->GetClass()->IsChildOf(LoadClass))
+					{
+						UE_LOG(LogLinker, Error, TEXT("Failed import: class '%s' name '%s' outer '%s'. There is another object (of '%s' class) at the path."),
+							*LoadClass->GetName(), *Export.ObjectName.ToString(), *ThisParent->GetName(), *ActualObjectWithTheName->GetClass()->GetName());
+						Export.bExportLoadFailed = true; // I am not sure if this is an actual fail or not...looked like it in the original code
+						return;
+					}
+
+					// Find the Archetype object for the one we are loading.
+					check(!Export.TemplateIndex.IsNull());
+					UObject* Template = EventDrivenIndexToObject(Export.TemplateIndex, true);
+					if (!Template)
+					{
+						UE_LOG(LogStreaming, Error, TEXT("Cannot construct %s in %s because we could not find its template %s"), *Export.ObjectName.ToString(), *Linker->GetArchiveName(), *Linker->GetImportPathName(Export.TemplateIndex));
+						Export.bExportLoadFailed = true;
+						return;
+					}
+					// we also need to ensure that the template has set up any instances
+					Template->ConditionalPostLoadSubobjects();
+
+
+					check(!GVerifyObjectReferencesOnly); // not supported with the event driven loader
+					// Create the export object, marking it with the appropriate flags to
+					// indicate that the object's data still needs to be loaded.
+					EObjectFlags ObjectLoadFlags = Export.ObjectFlags;
+					ObjectLoadFlags = EObjectFlags(ObjectLoadFlags | RF_NeedLoad | RF_NeedPostLoad | RF_NeedPostLoadSubobjects | RF_WasLoaded);
+
+					FName NewName = Export.ObjectName;
+
+					// If we are about to create a CDO, we need to ensure that all parent sub-objects are loaded
+					// to get default value initialization to work.
+		#if DO_CHECK
+					if ((ObjectLoadFlags & RF_ClassDefaultObject) != 0)
+					{
+						UClass* SuperClass = LoadClass->GetSuperClass();
+						UObject* SuperCDO = SuperClass ? SuperClass->GetDefaultObject() : nullptr;
+						check(!SuperCDO || Template == SuperCDO); // the template for a CDO is the CDO of the super
+						if (SuperClass && !SuperClass->IsNative())
+						{
+							check(SuperCDO);
+							if (SuperClass->HasAnyFlags(RF_NeedLoad))
+							{
+								UE_LOG(LogStreaming, Fatal, TEXT("Super %s had RF_NeedLoad while creating %s"), *SuperClass->GetFullName(), *Export.ObjectName.ToString());
+								Export.bExportLoadFailed = true;
+								return;
+							}
+							if (SuperCDO->HasAnyFlags(RF_NeedLoad))
+							{
+								UE_LOG(LogStreaming, Fatal, TEXT("Super CDO %s had RF_NeedLoad while creating %s"), *SuperCDO->GetFullName(), *Export.ObjectName.ToString());
+								Export.bExportLoadFailed = true;
+								return;
+							}
+							TArray<UObject*> SuperSubObjects;
+							GetObjectsWithOuter(SuperCDO, SuperSubObjects, /*bIncludeNestedObjects=*/ false, /*ExclusionFlags=*/ RF_NoFlags, /*InternalExclusionFlags=*/ EInternalObjectFlags::Native);
+
+							for (UObject* SubObject : SuperSubObjects)
+							{
+								if (SubObject->HasAnyFlags(RF_NeedLoad))
+								{
+									UE_LOG(LogStreaming, Fatal, TEXT("Super CDO subobject %s had RF_NeedLoad while creating %s"), *SubObject->GetFullName(), *Export.ObjectName.ToString());
+									Export.bExportLoadFailed = true;
+									return;
+								}
+							}
+						}
+						else
+						{
+							check(Template->IsA(LoadClass));
+						}
+					}
+		#endif
+					if (LoadClass->HasAnyFlags(RF_NeedLoad))
+					{
+						UE_LOG(LogStreaming, Fatal, TEXT("LoadClass %s had RF_NeedLoad while creating %s"), *LoadClass->GetFullName(), *Export.ObjectName.ToString());
+						Export.bExportLoadFailed = true;
+						return;
+					}
+					{
+						UObject* LoadCDO = LoadClass->GetDefaultObject();
+						if (LoadCDO->HasAnyFlags(RF_NeedLoad))
+						{
+							UE_LOG(LogStreaming, Fatal, TEXT("Class CDO %s had RF_NeedLoad while creating %s"), *LoadCDO->GetFullName(), *Export.ObjectName.ToString());
+							Export.bExportLoadFailed = true;
+							return;
+						}
+					}
+					if (Template->HasAnyFlags(RF_NeedLoad))
+					{
+						UE_LOG(LogStreaming, Fatal, TEXT("Template %s had RF_NeedLoad while creating %s"), *Template->GetFullName(), *Export.ObjectName.ToString());
+						Export.bExportLoadFailed = true;
+						return;
+					}
+
+					Export.Object = StaticConstructObject_Internal
+						(
+							LoadClass,
+							ThisParent,
+							NewName,
+							ObjectLoadFlags,
+							EInternalObjectFlags::None,
+							Template,
+							false,
+							nullptr,
+							true
+						);
+
+					if (GIsInitialLoad || GUObjectArray.IsOpenForDisregardForGC())
+					{
+						Export.Object->AddToRoot();
+					}
+					if (IsInGameThread())
+					{
+						Export.Object->SetInternalFlags(EInternalObjectFlags::Async);
+						AddOwnedObject(Export.Object);
+					}
+					check(Export.Object->HasAnyInternalFlags(EInternalObjectFlags::Async));
+					Export.Object->SetLinker(Linker, LocalExportIndex);
+					check(Export.Object->GetClass() == LoadClass);
+					check(NewName == Export.ObjectName);
+				}
+			}
+		}
+	}
+	if (!Export.Object)
+	{
+		Export.bExportLoadFailed = true;
+	}
+	check(Export.Object || Export.bExportLoadFailed);
+}
+
+
+void FAsyncPackage2::MarkNewObjectForLoadIfItIsAnExport(UObject *Object)
+{
+	if (!Object->HasAnyFlags(RF_WillBeLoaded | RF_LoadCompleted | RF_NeedLoad))
+	{
+		FPackageIndex MaybeExportIndex = FindExportFromObject2(Linker, Object);
+		if (MaybeExportIndex.IsExport())
+		{
+			UE_LOG(LogStreaming, VeryVerbose, TEXT("Note: %s was constructed during load and is an export and so needs loading."), *Object->GetFullName());
+			Object->SetFlags(RF_WillBeLoaded);
+		}
+	}
+}
+
+void FAsyncPackage2::EventDrivenSerializeExport(int32 LocalExportIndex)
+{
+	//SCOPED_LOADTIMER(Package_PreLoadObjects);
+
+	FObjectExport& Export = Linker->ExportMap[LocalExportIndex];
+
+	LLM_SCOPE(ELLMTag::UObject);
+	LLM_SCOPED_TAG_WITH_OBJECT_IN_SET(GetLinkerRoot(), ELLMTagSet::Assets);
+	LLM_SCOPED_TAG_WITH_OBJECT_IN_SET((Export.DynamicType == FObjectExport::EDynamicType::DynamicType) ? UDynamicClass::StaticClass() :
+		CastEventDrivenIndexToObject<UClass>(Export.ClassIndex, false), ELLMTagSet::AssetClasses);
+
+	UObject* Object = Export.Object;
+	if (Object && Linker->bDynamicClassLinker)
+	{
+		//native blueprint 
+		UDynamicClass* UD = Cast<UDynamicClass>(Object);
+		if (UD)
+		{
+			check(Export.DynamicType == FObjectExport::EDynamicType::DynamicType);
+			UObject* LocObj = ConstructDynamicType(*Linker->GetExportPathName(LocalExportIndex), EConstructDynamicType::CallZConstructor);
+			check(UD == LocObj);
+		}
+		Object->ClearFlags(RF_NeedLoad | RF_WillBeLoaded);
+	}
+	else if (Object && Object->HasAnyFlags(RF_NeedLoad))
+	{
+		check(Object->GetLinker() == Linker);
+		check(Object->GetLinkerIndex() == LocalExportIndex);
+		UClass* Cls = nullptr;
+
+		// If this is a struct, make sure that its parent struct is completely loaded
+		if (UStruct* Struct = dynamic_cast<UStruct*>(Object))
+		{
+			UStruct* SuperStruct = nullptr;
+			if (!Export.SuperIndex.IsNull())
+			{
+				SuperStruct = CastEventDrivenIndexToObject<UStruct>(Export.SuperIndex, true);
+				if (!SuperStruct)
+				{
+					// see FLinkerLoad::CreateExport, there may be some more we can do here
+					UE_LOG(LogStreaming, Fatal, TEXT("Could not find SuperStruct %s to create %s"), *Linker->ImpExp(Export.SuperIndex).ObjectName.ToString(), *Export.ObjectName.ToString());
+					Export.bExportLoadFailed = true;
+					return;
+				}
+			}
+			if (SuperStruct)
+			{
+				Struct->SetSuperStruct(SuperStruct);
+				if (UClass* ClassObject = dynamic_cast<UClass*>(Object))
+				{
+					ClassObject->Bind();
+				}
+			}
+		}
+		
+		const FPackageFileSummary& Summary = Linker->Summary;
+		const FCustomVersionContainer& SummaryVersions = Summary.GetCustomVersionContainer();
+
+		FSimpleArchive Ar(ExportIoBuffers[LocalExportIndex].Data(), ExportIoBuffers[LocalExportIndex].DataSize());
+		Ar.SetUE4Ver(Linker->Summary.GetFileVersionUE4());
+		Ar.SetLicenseeUE4Ver(Linker->Summary.GetFileVersionLicenseeUE4());
+		Ar.SetEngineVer(Linker->Summary.SavedByEngineVersion);
+		Ar.SetCustomVersions(SummaryVersions);
+
+		FArchive* OldLoader = Linker->Loader;
+		Linker->Loader = &Ar;
+
+		Object->ClearFlags(RF_NeedLoad);
+
+		TRACE_LOADTIME_OBJECT_SCOPE(Object, LoadTimeProfilerObjectEventType_Serialize);
+
+		FUObjectSerializeContext* LoadContext = GetSerializeContext();
+		UObject* PrevSerializedObject = LoadContext->SerializedObject;
+		LoadContext->SerializedObject = Object;
+		Linker->bForceSimpleIndexToObject = true;
+
+		// Find the Archetype object for the one we are loading. This is piped to GetArchetypeFromLoader
+		check(!Export.TemplateIndex.IsNull());
+		UObject* Template = EventDrivenIndexToObject(Export.TemplateIndex, true);
+		check(Template);
+
+		check(!Linker->TemplateForGetArchetypeFromLoader);
+		Linker->TemplateForGetArchetypeFromLoader = Template;
+
+		if (Object->HasAnyFlags(RF_ClassDefaultObject))
+		{
+			Object->GetClass()->SerializeDefaultObject(Object, *Linker);
+		}
+		else
+		{
+			Object->Serialize(*Linker);
+		}
+		check(Linker->TemplateForGetArchetypeFromLoader == Template);
+		Linker->TemplateForGetArchetypeFromLoader = nullptr;
+
+		Object->SetFlags(RF_LoadCompleted);
+		LoadContext->SerializedObject = PrevSerializedObject;
+		Linker->bForceSimpleIndexToObject = false;
+
+		UE_CLOG(Ar.Tell() != Export.SerialSize, LogStreaming, Warning, TEXT("Serialize mismatch, ObjectName='%s'"), *Object->GetFullName());
+
+		Linker->Loader = OldLoader;
+
+#if DO_CHECK
+		if (Object->HasAnyFlags(RF_ClassDefaultObject) && Object->GetClass()->HasAnyClassFlags(CLASS_CompiledFromBlueprint))
+		{
+			check(Object->HasAllFlags(RF_NeedPostLoad | RF_WasLoaded));
+			//Object->SetFlags(RF_NeedPostLoad | RF_WasLoaded);
+		}
+#endif
+	}
+
+	// push stats so that we don't overflow number of tags per thread during blocking loading
+	LLM_PUSH_STATS_FOR_ASSET_TAGS();
+}
+
+EAsyncPackageState::Type FAsyncPackage2::Event_StartIO(FAsyncPackage2* Package, int32 LocalExportIndex)
+{
+	TRACE_CPUPROFILER_EVENT_SCOPE(Event_StartIO);
+
+	check(Package->AsyncPackageLoadingState == EAsyncPackageLoadingState2::ProcessNewImportsAndExports);
+	FObjectExport& Export = Package->Linker->ExportMap[LocalExportIndex];
+	if (Package->Linker->bDynamicClassLinker || //native blueprint, there is no IO for these
+		(Export.Object && !Export.Object->HasAnyFlags(RF_NeedLoad)))
+	{
+		Package->GetNode(EEventLoadNode2::ImportOrExport_Serialize, FPackageIndex::FromExport(LocalExportIndex))->ReleaseBarrier();
+	}
+	else
+	{
+		Package->AsyncLoadingThread.EnqueueIoRequest(Package, CreateChunkId(Package->PackageChunkId, LocalExportIndex, EChunkType::ExportData));
+	}
+	return EAsyncPackageState::Complete;
+}
+
+EAsyncPackageState::Type FAsyncPackage2::Event_ExportsDone(FAsyncPackage2* Package, int32)
+{
+	TRACE_CPUPROFILER_EVENT_SCOPE(Event_ExportsDone);
+
+	check(Package->AsyncPackageLoadingState == EAsyncPackageLoadingState2::ProcessNewImportsAndExports);
+	Package->bAllExportsSerialized = true;
+	Package->AsyncPackageLoadingState = EAsyncPackageLoadingState2::PostLoad_Etc;
+	return EAsyncPackageState::Complete;
+}
+
+EAsyncPackageState::Type FAsyncPackage2::Event_StartPostload(FAsyncPackage2* Package, int32)
+{
+	TRACE_CPUPROFILER_EVENT_SCOPE(Event_StartPostLoad);
+
+	LLM_SCOPE(ELLMTag::AsyncLoading);
+
+	check(Package->bAllExportsSerialized);
+	{
+		check(Package->PackageObjLoaded.Num() == 0);
+		Package->PackageObjLoaded.Reserve(Package->Linker->ExportMap.Num());
+		for (int32 LocalExportIndex = 0; LocalExportIndex < Package->Linker->ExportMap.Num(); LocalExportIndex++)
+		{
+			FObjectExport& Export = Package->Linker->ExportMap[LocalExportIndex];
+			UObject* Object = Export.Object;
+			if (Object && (Object->HasAnyFlags(RF_NeedPostLoad) || Package->Linker->bDynamicClassLinker || Object->HasAnyInternalFlags(EInternalObjectFlags::AsyncLoading)))
+			{
+				check(Object->IsValidLowLevelFast());
+				Package->PackageObjLoaded.Add(Object);
+			}
+		}
+	}
+	Package->EventDrivenLoadingComplete();
+	return EAsyncPackageState::Complete;
+}
+
+EAsyncPackageState::Type FAsyncPackage2::Event_Tick(FAsyncPackage2* Package, int32)
+{
+	TRACE_CPUPROFILER_EVENT_SCOPE(Event_Tick);
+
+	check(!Package->HasFinishedLoading());
+	
+	check(!Package->ReentryCount);
+	Package->ReentryCount++;
+
+	// Keep track of time when we start loading.
+	check(Package->LoadStartTime > 0.0f);
+
+	FAsyncPackageScope2 PackageScope(Package);
+
+	// Make sure we finish our work if there's no time limit. The loop is required as PostLoad
+	// might cause more objects to be loaded in which case we need to Preload them again.
+	EAsyncPackageState::Type LoadingState = EAsyncPackageState::Complete;
+	do
+	{
+		// Reset value to true at beginning of loop.
+		LoadingState = EAsyncPackageState::Complete;
+
+		// Begin async loading, simulates BeginLoad
+		Package->BeginAsyncLoad();
+
+		if (LoadingState == EAsyncPackageState::Complete && !Package->bLoadHasFailed)
+		{
+			SCOPED_LOADTIMER(Package_ExternalReadDependencies);
+			LoadingState = Package->FinishExternalReadDependencies();
+		}
+
+		// Call PostLoad on objects, this could cause new objects to be loaded that require
+		// another iteration of the PreLoad loop.
+		if (LoadingState == EAsyncPackageState::Complete && !Package->bLoadHasFailed)
+		{
+			SCOPED_LOADTIMER(Package_PostLoadObjects);
+			LoadingState = Package->PostLoadObjects();
+		}
+
+		// End async loading, simulates EndLoad
+		Package->EndAsyncLoad();
+
+		// Finish objects (removing EInternalObjectFlags::AsyncLoading, dissociate imports and forced exports, 
+		// call completion callback, ...
+		// If the load has failed, perform completion callbacks and then quit
+		if (LoadingState == EAsyncPackageState::Complete || Package->bLoadHasFailed)
+		{
+			LoadingState = Package->FinishObjects();
+		}
+	} while (!FAsyncLoadingThreadState2::Get().IsTimeLimitExceeded() && LoadingState == EAsyncPackageState::TimeOut);
+
+	if (Package->LinkerRoot && LoadingState == EAsyncPackageState::Complete)
+	{
+		Package->LinkerRoot->MarkAsFullyLoaded();
+	}
+
+	// Mark this package as loaded if everything completed.
+	Package->bLoadHasFinished = (LoadingState == EAsyncPackageState::Complete);
+
+	if (Package->bLoadHasFinished)
+	{
+		check(Package->AsyncPackageLoadingState == EAsyncPackageLoadingState2::PostLoad_Etc);
+		Package->AsyncPackageLoadingState = EAsyncPackageLoadingState2::PackageComplete;
+	}
+
+	Package->ReentryCount--;
+	check(Package->ReentryCount >= 0);
+
+	if (LoadingState == EAsyncPackageState::TimeOut)
+	{
+		return EAsyncPackageState::TimeOut;
+	}
+	check(LoadingState == EAsyncPackageState::Complete);
+	// We're done, at least on this thread, so we can remove the package now.
+	Package->AsyncLoadingThread.AddToLoadedPackages(Package);
+	return EAsyncPackageState::Complete;
+}
+
+EAsyncPackageState::Type FAsyncPackage2::Event_Delete(FAsyncPackage2* Package, int32)
+{
+	delete Package;
+	return EAsyncPackageState::Complete;
+}
+
+void FAsyncPackage2::ProcessIoRequest(FIoRequest& IoRequest)
+{
+	UE_CLOG(!IoRequest.IsOk(), LogStreaming, Warning, TEXT("I/O Error: '%s', package: '%s'"), *IoRequest.Status().ToString(), *Desc.NameToLoad.ToString()); 
+
+	FIoChunkId ChunkId = IoRequest.GetChunkId();
+
+	switch (GetChunkType(ChunkId))
+	{
+	case EChunkType::PackageSummary:
+	{
+		FIoBuffer SummaryIoBuffer = IoRequest.GetChunk();
+		PackageSummaryBuffer = MakeUnique<uint8[]>(SummaryIoBuffer.DataSize());
+		FMemory::Memcpy(PackageSummaryBuffer.Get(), SummaryIoBuffer.Data(), SummaryIoBuffer.DataSize());
+		GetNode(EEventLoadNode2::Package_LoadSummary)->ReleaseBarrier();
+		break;
+	}
+	case EChunkType::ExportData:
+	{
+		const int32 LocalExportIndex = GetChunkIndex(ChunkId);
+		FIoBuffer IoBuffer = IoRequest.GetChunk();
+		check(ExportIoBuffers.Num() > LocalExportIndex);
+		ExportIoBuffers[LocalExportIndex] = IoBuffer;
+		GetNode(EEventLoadNode2::ImportOrExport_Serialize, FPackageIndex::FromExport(LocalExportIndex))->ReleaseBarrier();
+		break;
+	}
+	}
+}
+
+void FAsyncPackage2::EventDrivenLoadingComplete()
+{
+	AsyncPackageLoadingState = EAsyncPackageLoadingState2::PostLoad_Etc;
+	GetNode(Package_Tick)->ReleaseBarrier();
+}
+
+FEventLoadNode2* FAsyncPackage2::GetNode(EEventLoadNode2 Phase, FPackageIndex ImportOrExportIndex)
+{
+	if (ImportOrExportIndex.IsNull())
+	{
+		return PackageNodes + Phase;
+	}
+	else if (ImportOrExportIndex.IsImport())
+	{
+		int32 ImportNodeIndex = ImportOrExportIndex.ToImport() * EEventLoadNode2::Import_NumPhases + Phase;
+		return ImportNodes + ImportNodeIndex;
+	}
+	else
+	{
+		int32 ExportNodeIndex = ImportOrExportIndex.ToExport() * EEventLoadNode2::Export_NumPhases + Phase;
+		return ExportNodes + ExportNodeIndex;
+	}
+}
+
+void FAsyncLoadingThread2Impl::AddToLoadedPackages(FAsyncPackage2* Package)
+{
+	FScopeLock LoadedLock(&LoadedPackagesCritical);
+	check(!LoadedPackages.Contains(Package));
+	LoadedPackages.Add(Package);
+}
+
+EAsyncPackageState::Type FAsyncLoadingThread2Impl::ProcessAsyncLoadingFromGameThread(int32& OutPackagesProcessed)
+{
+	SCOPED_LOADTIMER(AsyncLoadingTime);
+
+	check(IsInGameThread());
+
+	// If we're not multithreaded and flushing async loading, update the thread heartbeat
+	const bool bNeedsHeartbeatTick = !FAsyncLoadingThread2Impl::IsMultithreaded();
+	OutPackagesProcessed = 0;
+
+	FAsyncLoadingTickScope2 InAsyncLoadingTick(*this);
+	uint32 LoopIterations = 0;
+
+	FAsyncLoadingThreadState2& ThreadState = FAsyncLoadingThreadState2::Get();
+
+	while (true)
+	{
+		do 
+		{
+			ThreadState.ProcessDeferredFrees();
+
+			if (bNeedsHeartbeatTick && (++LoopIterations) % 32 == 31)
+			{
+				// Update heartbeat after 32 events
+				FThreadHeartBeat::Get().HeartBeat();
+			}
+
+			if (ThreadState.IsTimeLimitExceeded())
+			{
+				return EAsyncPackageState::TimeOut;
+			}
+
+			if (IsAsyncLoadingSuspended())
+			{
+				return EAsyncPackageState::TimeOut;
+			}
+
+			if (QueuedPackagesCounter)
+			{
+				CreateAsyncPackagesFromQueue();
+				OutPackagesProcessed++;
+				break;
+			}
+
+			bool bProcessedResourceCompletionEvents = false;
+			TOptional<FIoRequestQueue::FCompletionEvent> IoCompletionEvent = IoRequestQueue->DequeueCompletionEvent();
+			while (IoCompletionEvent)
+			{
+				FAsyncPackage2* Package = IoCompletionEvent.GetValue().Package;
+				Package->ProcessIoRequest(IoCompletionEvent.GetValue().IoRequest);
+				IoDispatcher.FreeBatch(IoCompletionEvent.GetValue().IoBatch);
+				IoCompletionEvent = IoRequestQueue->DequeueCompletionEvent();
+			}
+			if (bProcessedResourceCompletionEvents)
+			{
+				OutPackagesProcessed++;
+				break;
+			}
+
+			bool bPopped = false;
+			for (FAsyncLoadEventQueue2* Queue : AltEventQueues)
+			{
+				if (Queue->PopAndExecute(ThreadState))
+				{
+					bPopped = true;
+					break;
+				}
+			}
+			if (bPopped)
+			{
+				OutPackagesProcessed++;
+				break;
+			}
+
+			if (IoRequestQueue->HasPendingRequests())
+			{
+				FPlatformProcess::Sleep(0.001f);
+			}
+			else
+			{
+				return EAsyncPackageState::Complete;
+			}
+		} while (false);
+	}
+	check(false);
+	return EAsyncPackageState::Complete;
+}
+
+bool FAsyncPackage2::AreAllDependenciesFullyLoadedInternal(FAsyncPackage2* Package, TSet<UPackage*>& VisitedPackages, FString& OutError)
+{
+	for (UPackage* ImportPackage : Package->ImportedPackages)
+	{
+		if (VisitedPackages.Contains(ImportPackage))
+		{
+			continue;
+		}
+		VisitedPackages.Add(ImportPackage);
+
+		if (FLinkerLoad* ImportPackageLinker = FLinkerLoad::FindExistingLinkerForPackage(ImportPackage))
+		{
+			if (ImportPackageLinker->AsyncRoot)
+			{
+				FAsyncPackage2* AsyncRoot = (FAsyncPackage2*)ImportPackageLinker->AsyncRoot;
+				if (!AsyncRoot->bAllExportsSerialized)
+				{
+					OutError = FString::Printf(TEXT("%s Doesn't have all exports Serialized"), *Package->GetPackageName().ToString());
+					return false;
+				}
+				if (AsyncRoot->DeferredPostLoadIndex < AsyncRoot->DeferredPostLoadObjects.Num())
+				{
+					OutError = FString::Printf(TEXT("%s Doesn't have all objects processed by DeferredPostLoad"), *Package->GetPackageName().ToString());
+					return false;
+				}
+				for (FObjectExport& Export : ImportPackageLinker->ExportMap)
+				{
+					if (Export.Object && Export.Object->HasAnyFlags(RF_NeedPostLoad|RF_NeedLoad))
+					{
+						OutError = FString::Printf(TEXT("%s has not been %s"), *Export.Object->GetFullName(), 
+							Export.Object->HasAnyFlags(RF_NeedLoad) ? TEXT("Serialized") : TEXT("PostLoaded"));
+						return false;
+					}
+				}
+
+				if (!AreAllDependenciesFullyLoadedInternal(AsyncRoot, VisitedPackages, OutError))
+				{
+					OutError = Package->GetPackageName().ToString() + TEXT("->") + OutError;
+					return false;
+				}
+			}
+		}
+	}
+	return true;
+}
+
+bool FAsyncPackage2::AreAllDependenciesFullyLoaded(TSet<UPackage*>& VisitedPackages)
+{
+	TRACE_CPUPROFILER_EVENT_SCOPE(AreAllDependenciesFullyLoaded);
+	VisitedPackages.Reset();
+	FString Error;
+	bool bLoaded = AreAllDependenciesFullyLoadedInternal(this, VisitedPackages, Error);
+	if (!bLoaded)
+	{
+		UE_LOG(LogStreaming, Verbose, TEXT("AreAllDependenciesFullyLoaded: %s"), *Error);
+	}
+	return bLoaded;
+}
+
+EAsyncPackageState::Type FAsyncLoadingThread2Impl::ProcessLoadedPackagesFromGameThread(bool& bDidSomething, int32 FlushRequestID)
+{
+	EAsyncPackageState::Type Result = EAsyncPackageState::Complete;
+
+	// This is for debugging purposes only. @todo remove
+	volatile int32 CurrentAsyncLoadingCounter = AsyncLoadingTickCounter;
+
+	{
+		FScopeLock LoadedPackagesLock(&LoadedPackagesCritical);
+		if (LoadedPackages.Num() != 0)
+		{
+			LoadedPackagesToProcess.Append(LoadedPackages);
+			LoadedPackages.Reset();
+		}
+	}
+	if (IsMultithreaded() && ENamedThreads::GetRenderThread() == ENamedThreads::GameThread) // render thread tasks are actually being sent to the game thread.
+	{
+		// The async loading thread might have queued some render thread tasks (we don't have a render thread yet, so these are actually sent to the game thread)
+		// We need to process them now before we do any postloads.
+		FTaskGraphInterface::Get().ProcessThreadUntilIdle(ENamedThreads::GameThread);
+		if (FAsyncLoadingThreadState2::Get().IsTimeLimitExceeded())
+		{
+			return EAsyncPackageState::TimeOut;
+		}
+	}
+
+		
+	bDidSomething = LoadedPackagesToProcess.Num() > 0;
+	for (int32 PackageIndex = 0; PackageIndex < LoadedPackagesToProcess.Num() && !IsAsyncLoadingSuspended(); ++PackageIndex)
+	{
+		FAsyncPackage2* Package = LoadedPackagesToProcess[PackageIndex];
+		SCOPED_LOADTIMER(ProcessLoadedPackagesTime);
+
+		Result = Package->PostLoadDeferredObjects();
+		if (Result == EAsyncPackageState::Complete)
+		{
+			{
+				FScopeLock LockAsyncPackages(&AsyncPackagesCritical);
+				AsyncPackageNameLookup.Remove(Package->GetPackageName());
+				Package->ClearOwnedObjects();
+			}
+
+			// Remove the package from the list before we trigger the callbacks, 
+			// this is to ensure we can re-enter FlushAsyncLoading from any of the callbacks
+			LoadedPackagesToProcess.RemoveAt(PackageIndex--);
+
+			// Emulates ResetLoaders on the package linker's linkerroot.
+			if (!Package->IsBeingProcessedRecursively())
+			{
+				Package->ResetLoader();
+			}
+
+			// Close linkers opened by synchronous loads during async loading
+			Package->CloseDelayedLinkers();
+
+			// Incremented on the Async Thread, now decrement as we're done with this package				
+			const int32 NewExistingAsyncPackagesCounterValue = ExistingAsyncPackagesCounter.Decrement();
+
+			UE_CLOG(NewExistingAsyncPackagesCounterValue < 0, LogStreaming, Fatal, TEXT("ExistingAsyncPackagesCounter is negative, this means we loaded more packages then requested so there must be a bug in async loading code."));
+
+			// Call external callbacks
+			const bool bInternalCallbacks = false;
+			const EAsyncLoadingResult::Type LoadingResult = Package->HasLoadFailed() ? EAsyncLoadingResult::Failed : EAsyncLoadingResult::Succeeded;
+			Package->CallCompletionCallbacks(bInternalCallbacks, LoadingResult);
+#if WITH_EDITOR
+			// In the editor we need to find any assets and add them to list for later callback
+			Package->GetLoadedAssets(LoadedAssets);
+#endif
+			// We don't need the package anymore
+			check(!Package->bAddedForDelete);
+			check(!PackagesToDelete.Contains(Package));
+			PackagesToDelete.Add(Package);
+			Package->bAddedForDelete = true;
+			Package->MarkRequestIDsAsComplete();
+
+			if (FlushRequestID != INDEX_NONE && !ContainsRequestID(FlushRequestID))
+			{
+				// The only package we care about has finished loading, so we're good to exit
+				break;
+			}
+		}
+		else
+		{
+			break;
+		}
+	}
+	bDidSomething = bDidSomething || PackagesToDelete.Num() > 0;
+
+	// Delete packages we're done processing and are no longer dependencies of anything else
+	if (Result != EAsyncPackageState::TimeOut)
+	{
+		// For performance reasons this set is created here and reset inside of AreAllDependenciesFullyLoaded
+		TSet<UPackage*> VisitedPackages;
+
+		for (int32 PackageIndex = 0; PackageIndex < PackagesToDelete.Num(); ++PackageIndex)
+		{
+			FAsyncPackage2* Package = PackagesToDelete[PackageIndex];
+			if (!Package->IsBeingProcessedRecursively())
+			{
+				bool bSafeToDelete = false;
+				if (Package->HasClusterObjects())
+				{
+					// This package will create GC clusters but first check if all dependencies of this package have been fully loaded
+					if (Package->AreAllDependenciesFullyLoaded(VisitedPackages))
+					{
+						if (Package->CreateClusters() == EAsyncPackageState::Complete)
+						{
+							// All clusters created, it's safe to delete the package
+							bSafeToDelete = true;
+						}
+						else
+						{
+							// Cluster creation timed out
+							Result = EAsyncPackageState::TimeOut;
+							break;
+						}
+					}
+				}
+				else
+				{
+					// No clusters to create so it's safe to delete
+					bSafeToDelete = true;
+				}
+
+				if (bSafeToDelete)
+				{
+					PackagesToDelete.RemoveAtSwap(PackageIndex--);
+					Package->ClearImportedPackages();
+					Package->ReleaseRef();
+				}
+			}
+
+			// push stats so that we don't overflow number of tags per thread during blocking loading
+			LLM_PUSH_STATS_FOR_ASSET_TAGS();
+		}
+	}
+
+	if (Result == EAsyncPackageState::Complete)
+	{
+#if WITH_EDITORONLY_DATA
+		// This needs to happen after loading new blueprints in the editor, and this is handled in EndLoad for synchronous loads
+		FBlueprintSupport::FlushReinstancingQueue();
+#endif
+
+#if WITH_EDITOR
+		// In editor builds, call the asset load callback. This happens in both editor and standalone to match EndLoad
+		TArray<FWeakObjectPtr> TempLoadedAssets = LoadedAssets;
+		LoadedAssets.Reset();
+
+		// Make a copy because LoadedAssets could be modified by one of the OnAssetLoaded callbacks
+		for (const FWeakObjectPtr& WeakAsset : TempLoadedAssets)
+		{
+			// It may have been unloaded/marked pending kill since being added, ignore those cases
+			if (UObject* LoadedAsset = WeakAsset.Get())
+			{
+				FCoreUObjectDelegates::OnAssetLoaded.Broadcast(LoadedAsset);
+			}
+		}
+#endif
+
+		// We're not done until all packages have been deleted
+		Result = PackagesToDelete.Num() ? EAsyncPackageState::PendingImports : EAsyncPackageState::Complete;
+	}
+
+	return Result;
+}
+
+EAsyncPackageState::Type FAsyncLoadingThread2Impl::TickAsyncLoadingFromGameThread(bool bUseTimeLimit, bool bUseFullTimeLimit, float TimeLimit, int32 FlushRequestID)
+{
+	LLM_SCOPE(ELLMTag::AsyncLoading);
+
+	//TRACE_INT_VALUE(QueuedPackagesCounter, QueuedPackagesCounter);
+	//TRACE_INT_VALUE(GraphNodeCount, GraphAllocator.TotalNodeCount);
+	//TRACE_INT_VALUE(GraphArcCount, GraphAllocator.TotalArcCount);
+	//TRACE_MEMORY_VALUE(GraphMemory, GraphAllocator.TotalAllocated);
+
+
+	check(IsInGameThread());
+	check(!IsGarbageCollecting());
+
+	const bool bLoadingSuspended = IsAsyncLoadingSuspended();
+	EAsyncPackageState::Type Result = bLoadingSuspended ? EAsyncPackageState::PendingImports : EAsyncPackageState::Complete;
+
+	if (!bLoadingSuspended)
+	{
+		FAsyncLoadingThreadState2::Get().SetTimeLimit(bUseTimeLimit, TimeLimit);
+
+		// First make sure there's no objects pending to be unhashed. This is important in uncooked builds since we don't 
+		// detach linkers immediately there and we may end up in getting unreachable objects from Linkers in CreateImports
+		if (FPlatformProperties::RequiresCookedData() == false && IsIncrementalUnhashPending() && IsAsyncLoadingPackages())
+		{
+			// Call ConditionalBeginDestroy on all pending objects. CBD is where linkers get detached from objects.
+			UnhashUnreachableObjects(false);
+		}
+
+		const bool bIsMultithreaded = FAsyncLoadingThread2Impl::IsMultithreaded();
+		double TickStartTime = FPlatformTime::Seconds();
+
+		bool bDidSomething = false;
+		{
+			Result = ProcessLoadedPackagesFromGameThread(bDidSomething, FlushRequestID);
+			double TimeLimitUsedForProcessLoaded = FPlatformTime::Seconds() - TickStartTime;
+			UE_CLOG(!GIsEditor && bUseTimeLimit && TimeLimitUsedForProcessLoaded > .1f, LogStreaming, Warning, TEXT("Took %6.2fms to ProcessLoadedPackages"), float(TimeLimitUsedForProcessLoaded) * 1000.0f);
+		}
+
+		if (!bIsMultithreaded && Result != EAsyncPackageState::TimeOut)
+		{
+			Result = TickAsyncThreadFromGameThread(bDidSomething);
+		}
+
+		if (Result != EAsyncPackageState::TimeOut)
+		{
+			{
+				FScopeLock QueueLock(&QueueCritical);
+				FScopeLock LoadedLock(&LoadedPackagesCritical);
+				// Flush deferred messages
+				if (ExistingAsyncPackagesCounter.GetValue() == 0)
+				{
+					bDidSomething = true;
+					FDeferredMessageLog::Flush();
+				}
+			}
+			if (!bDidSomething)
+			{
+				if (bIsMultithreaded)
+				{
+					if (GIsInitialLoad)
+					{
+						bDidSomething = EDLBootNotificationManager.ConstructWaitingBootObjects(); // with the ASL, we always create new boot objects when we have nothing else to do
+					}
+				}
+				else
+				{
+					if (GIsInitialLoad)
+					{
+						bDidSomething = EDLBootNotificationManager.FireCompletedCompiledInImports(); // no ASL, first try to fire any completed boot objects, and if there are none, then create some boot objects
+						if (!bDidSomething)
+						{
+							bDidSomething = EDLBootNotificationManager.ConstructWaitingBootObjects();
+						}
+					}
+				}
+			}
+		}
+
+		// Call update callback once per tick on the game thread
+		FCoreDelegates::OnAsyncLoadingFlushUpdate.Broadcast();
+	}
+
+	return Result;
+}
+
+FAsyncLoadingThread2Impl::FAsyncLoadingThread2Impl(IEDLBootNotificationManager& InEDLBootNotificationManager)
+	: Thread(nullptr)
+	, EDLBootNotificationManager(InEDLBootNotificationManager)
+{
+	GEventDrivenLoaderEnabled = true;
+
+#if LOADTIMEPROFILERTRACE_ENABLED
+	FLoadTimeProfilerTracePrivate::Init();
+#endif
+
+	AltEventQueues.Add(&SerializeExportsEventQueue);
+	AltEventQueues.Add(&AsyncEventQueue);
+	AltEventQueues.Add(&EventQueue);
+	AltEventQueues.Add(&CreateExportsEventQueue);
+
+	EventSpecs.AddDefaulted(EEventLoadNode2::Package_NumPhases + EEventLoadNode2::Import_NumPhases + EEventLoadNode2::Export_NumPhases);
+	EventSpecs[EEventLoadNode2::Package_CreateLinker] = { &FAsyncPackage2::Event_CreateLinker, &AsyncEventQueue, false };
+	EventSpecs[EEventLoadNode2::Package_LoadSummary] = { &FAsyncPackage2::Event_FinishLinker, &AsyncEventQueue, false };
+	EventSpecs[EEventLoadNode2::Package_ImportPackages] = { &FAsyncPackage2::Event_StartImportPackages, &AsyncEventQueue, false };
+	EventSpecs[EEventLoadNode2::Package_SetupImports] = { &FAsyncPackage2::Event_SetupImports, &AsyncEventQueue, false };
+	EventSpecs[EEventLoadNode2::Package_SetupExports] = { &FAsyncPackage2::Event_SetupExports, &AsyncEventQueue, true };
+	EventSpecs[EEventLoadNode2::Package_ExportsSerialized] = { &FAsyncPackage2::Event_ExportsDone, &AsyncEventQueue, true };
+	EventSpecs[EEventLoadNode2::Package_PostLoad] = { &FAsyncPackage2::Event_StartPostload, &AsyncEventQueue, false };
+	EventSpecs[EEventLoadNode2::Package_Tick] = { &FAsyncPackage2::Event_Tick, &EventQueue, false };
+	EventSpecs[EEventLoadNode2::Package_Delete] = { &FAsyncPackage2::Event_Delete, &AsyncEventQueue, false };
+
+	EventSpecs[EEventLoadNode2::Package_NumPhases + EEventLoadNode2::ImportOrExport_Create] = { &FAsyncPackage2::Event_LinkImport, &AsyncEventQueue, false };
+	EventSpecs[EEventLoadNode2::Package_NumPhases + EEventLoadNode2::ImportOrExport_Serialize] = { &FAsyncPackage2::Event_ImportSerialized, &AsyncEventQueue, true };
+
+	EventSpecs[EEventLoadNode2::Package_NumPhases + EEventLoadNode2::Import_NumPhases + EEventLoadNode2::ImportOrExport_Create] = { &FAsyncPackage2::Event_CreateExport, &CreateExportsEventQueue, false };
+	EventSpecs[EEventLoadNode2::Package_NumPhases + EEventLoadNode2::Import_NumPhases + EEventLoadNode2::ImportOrExport_Serialize] = { &FAsyncPackage2::Event_SerializeExport, &SerializeExportsEventQueue, false };
+	EventSpecs[EEventLoadNode2::Package_NumPhases + EEventLoadNode2::Import_NumPhases + EEventLoadNode2::Export_StartIO] = { &FAsyncPackage2::Event_StartIO, &AsyncEventQueue, false };
+
+	CancelLoadingEvent = FPlatformProcess::GetSynchEventFromPool();
+	ThreadSuspendedEvent = FPlatformProcess::GetSynchEventFromPool();
+	ThreadResumedEvent = FPlatformProcess::GetSynchEventFromPool();
+	AsyncLoadingTickCounter = 0;
+
+	FAsyncLoadingThreadState2::TlsSlot = FPlatformTLS::AllocTlsSlot();
+	FAsyncLoadingThreadState2::Create(GraphAllocator);
+}
+
+FAsyncLoadingThread2Impl::~FAsyncLoadingThread2Impl()
+{
+	if (Thread)
+	{
+		ShutdownLoading();
+	}
+}
+
+void FAsyncLoadingThread2Impl::ShutdownLoading()
+{
+	FCoreUObjectDelegates::GetPreGarbageCollectDelegate().RemoveAll(this);
+	FCoreUObjectDelegates::GetPostGarbageCollect().RemoveAll(this);
+
+	delete Thread;
+	Thread = nullptr;
+	FPlatformProcess::ReturnSynchEventToPool(CancelLoadingEvent);
+	CancelLoadingEvent = nullptr;
+	FPlatformProcess::ReturnSynchEventToPool(ThreadSuspendedEvent);
+	ThreadSuspendedEvent = nullptr;
+	FPlatformProcess::ReturnSynchEventToPool(ThreadResumedEvent);
+	ThreadResumedEvent = nullptr;
+}
+
+void FAsyncLoadingThread2Impl::StartThread()
+{
+
+
+	// Make sure the GC sync object is created before we start the thread (apparently this can happen before we call InitUObject())
+	FGCCSyncObject::Create();
+
+	if (!Thread)
+	{
+		UE_LOG(LogStreaming, Log, TEXT("Starting Async Loading Thread."));
+		bThreadStarted = true;
+		FPlatformMisc::MemoryBarrier();
+		Thread = FRunnableThread::Create(this, TEXT("FAsyncLoadingThread"), 0, TPri_Normal);
+		if (Thread)
+		{
+			TRACE_SET_THREAD_GROUP(Thread->GetThreadID(), "AsyncLoading");
+		}
+		
+		int32 WorkerCount = 3;
+		FParse::Value(FCommandLine::Get(), TEXT("-zenworkercount="), WorkerCount);
+		
+		if (WorkerCount > 0)
+		{
+			for (FAsyncLoadEventQueue2* Queue : AltEventQueues)
+			{
+				Queue->SetZenaphore(&AltZenaphore);
+			}
+
+			WorkerZenaphores.AddDefaulted(FMath::Max(3, WorkerCount));
+			Workers.Reserve(WorkerCount);
+			for (int32 WorkerIndex = 0; WorkerIndex < WorkerCount; ++WorkerIndex)
+			{
+				if (WorkerIndex == 0)
+				{
+					Workers.Emplace(GraphAllocator, SerializeExportsEventQueue, WorkerZenaphores[0], ActiveWorkersCount);
+					SerializeExportsEventQueue.SetZenaphore(&WorkerZenaphores[0]);
+					AltEventQueues.Remove(&SerializeExportsEventQueue);
+				}
+				else if (WorkerIndex == 1)
+				{
+					Workers.Emplace(GraphAllocator, CreateExportsEventQueue, WorkerZenaphores[1], ActiveWorkersCount);
+					CreateExportsEventQueue.SetZenaphore(&WorkerZenaphores[1]);
+					AltEventQueues.Remove(&CreateExportsEventQueue);
+				}
+				else
+				{
+					Workers.Emplace(GraphAllocator, AsyncEventQueue, WorkerZenaphores[2], ActiveWorkersCount);
+					AsyncEventQueue.SetZenaphore(&WorkerZenaphores[2]);
+					AltEventQueues.Remove(&AsyncEventQueue);
+				}
+				Workers[WorkerIndex].StartThread();
+			}
+		}
+
+		FCoreUObjectDelegates::GetPreGarbageCollectDelegate().AddRaw(this, &FAsyncLoadingThread2Impl::OnPreGarbageCollect);
+		FCoreUObjectDelegates::GetPostGarbageCollect().AddRaw(this, &FAsyncLoadingThread2Impl::OnPostGarbageCollect);
+	}
+}
+
+bool FAsyncLoadingThread2Impl::Init()
+{
+	return true;
+}
+
+void FAsyncLoadingThread2Impl::SuspendWorkers()
+{
+	if (bWorkersSuspended)
+	{
+		return;
+	}
+	TRACE_CPUPROFILER_EVENT_SCOPE(SuspendWorkers);
+	for (FAsyncLoadingThreadWorker& Worker : Workers)
+	{
+		Worker.SuspendThread();
+	}
+	while (ActiveWorkersCount > 0)
+	{
+		FPlatformProcess::SleepNoStats(0);
+	}
+	bWorkersSuspended = true;
+}
+
+void FAsyncLoadingThread2Impl::ResumeWorkers()
+{
+	if (!bWorkersSuspended)
+	{
+		return;
+	}
+	TRACE_CPUPROFILER_EVENT_SCOPE(ResumeWorkers);
+	for (FAsyncLoadingThreadWorker& Worker : Workers)
+	{
+		Worker.ResumeThread();
+	}
+	bWorkersSuspended = false;
+}
+
+uint32 FAsyncLoadingThread2Impl::Run()
+{
+	LLM_SCOPE(ELLMTag::AsyncLoading);
+
+	AsyncLoadingThreadID = FPlatformTLS::GetCurrentThreadId();
+
+	FAsyncLoadingThreadState2::Create(GraphAllocator);
+
+	TRACE_LOADTIME_START_ASYNC_LOADING();
+
+	if (!IsInGameThread())
+	{
+		FPlatformProcess::SetThreadAffinityMask(FPlatformAffinity::GetAsyncLoadingThreadMask());
+	}
+
+	FAsyncLoadingThreadState2& ThreadState = FAsyncLoadingThreadState2::Get();
+
+	FZenaphoreWaiter Waiter(AltZenaphore, TEXT("WaitForEvents"));
+	bool bIsSuspended = false;
+	while (!bStopRequested)
+	{
+		if (bIsSuspended)
+		{
+			if (!bSuspendRequested.Load(EMemoryOrder::SequentiallyConsistent) && !IsGarbageCollectionWaiting())
+			{
+				ThreadResumedEvent->Trigger();
+				bIsSuspended = false;
+				ResumeWorkers();
+			}
+			else
+			{
+				FPlatformProcess::Sleep(0.001f);
+			}
+		}
+		else
+		{
+			bool bDidSomething = false;
+			{
+				FGCScopeGuard GCGuard;
+				TRACE_CPUPROFILER_EVENT_SCOPE(AsyncLoadingTime);
+				do
+				{
+					bDidSomething = false;
+					for (;;)
+					{
+						TOptional<FIoRequestQueue::FCompletionEvent> IoCompletionEvent = IoRequestQueue->DequeueCompletionEvent();
+						if (IoCompletionEvent)
+						{
+							TRACE_CPUPROFILER_EVENT_SCOPE(ProcessResourceCompletionEvent);
+							FAsyncPackage2* Package = IoCompletionEvent.GetValue().Package;
+							Package->ProcessIoRequest(IoCompletionEvent.GetValue().IoRequest);
+							IoDispatcher.FreeBatch(IoCompletionEvent.GetValue().IoBatch);
+							bDidSomething = true;
+						}
+						else
+						{
+							break;
+						}
+					}
+
+					if (QueuedPackagesCounter)
+					{
+						if (CreateAsyncPackagesFromQueue())
+						{
+							bDidSomething = true;
+						}
+					}
+
+					bool bShouldSuspend = false;
+					bool bPopped = false;
+					do 
+					{
+						bPopped = false;
+						for (FAsyncLoadEventQueue2* Queue : AltEventQueues)
+						{
+							if (Queue->PopAndExecute(ThreadState))
+							{
+								bPopped = true;
+								bDidSomething = true;
+							}
+
+							if (bSuspendRequested.Load(EMemoryOrder::Relaxed) || IsGarbageCollectionWaiting())
+							{
+								bShouldSuspend = true;
+								bPopped = false;
+								break;
+							}
+						}
+					} while (bPopped);
+
+					if (bShouldSuspend || bSuspendRequested.Load(EMemoryOrder::Relaxed) || IsGarbageCollectionWaiting())
+					{
+						SuspendWorkers();
+						ThreadSuspendedEvent->Trigger();
+						bIsSuspended = true;
+						bDidSomething = true;
+						break;
+					}
+				} while (bDidSomething);
+			}
+			if (!bDidSomething)
+			{
+				ThreadState.ProcessDeferredFrees();
+				Waiter.Wait();
+			}
+		}
+	}
+	return 0;
+}
+
+EAsyncPackageState::Type FAsyncLoadingThread2Impl::TickAsyncThreadFromGameThread(bool& bDidSomething)
+{
+	check(IsInGameThread());
+	EAsyncPackageState::Type Result = EAsyncPackageState::Complete;
+	
+	int32 ProcessedRequests = 0;
+	if (AsyncThreadReady.GetValue())
+	{
+		if (GIsInitialLoad)
+		{
+			EDLBootNotificationManager.FireCompletedCompiledInImports();
+		}
+		if (IsGarbageCollectionWaiting() || FAsyncLoadingThreadState2::Get().IsTimeLimitExceeded())
+		{
+			Result = EAsyncPackageState::TimeOut;
+		}
+		else
+		{
+			FGCScopeGuard GCGuard;
+			Result = ProcessAsyncLoadingFromGameThread(ProcessedRequests);
+			bDidSomething = bDidSomething || ProcessedRequests > 0;
+		}
+	}
+
+	return Result;
+}
+
+void FAsyncLoadingThread2Impl::Stop()
+{
+	for (FAsyncLoadingThreadWorker& Worker : Workers)
+	{
+		Worker.StopThread();
+	}
+	bSuspendRequested = true;
+	bStopRequested = true;
+	AltZenaphore.NotifyAll();
+}
+
+void FAsyncLoadingThread2Impl::CancelLoading()
+{
+	check(false);
+	// TODO
+}
+
+void FAsyncLoadingThread2Impl::SuspendLoading()
+{
+	UE_CLOG(!IsInGameThread() || IsInSlateThread(), LogStreaming, Fatal, TEXT("Async loading can only be suspended from the main thread"));
+	if (!bSuspendRequested)
+	{
+		bSuspendRequested = true;
+		if (IsMultithreaded())
+		{
+			TRACE_LOADTIME_SUSPEND_ASYNC_LOADING();
+			AltZenaphore.NotifyAll();
+			ThreadSuspendedEvent->Wait();
+		}
+	}
+}
+
+void FAsyncLoadingThread2Impl::ResumeLoading()
+{
+	check(IsInGameThread() && !IsInSlateThread());
+	if (bSuspendRequested)
+	{
+		bSuspendRequested = false;
+		if (IsMultithreaded())
+		{
+			ThreadResumedEvent->Wait();
+			TRACE_LOADTIME_RESUME_ASYNC_LOADING();
+		}
+	}
+}
+
+float FAsyncLoadingThread2Impl::GetAsyncLoadPercentage(const FName& PackageName)
+{
+	float LoadPercentage = -1.0f;
+	FScopeLock LockAsyncPackages(&AsyncPackagesCritical);
+	FAsyncPackage2* Package = AsyncPackageNameLookup.FindRef(PackageName);
+	if (Package)
+	{
+		LoadPercentage = Package->GetLoadPercentage();
+	}
+	return LoadPercentage;
+}
+
+void FGlobalImportRuntime::OnPreGarbageCollect()
+{
+	int32 NumCleared = 0;
+
+	for (int32 GlobalImportIndex = 0; GlobalImportIndex < Count; ++GlobalImportIndex)
+	{
+		UObject*& Object = Objects[GlobalImportIndex];
+		if (!Object)
+		{
+			continue;
+		}
+
+		// Import objects in packages currently being loaded already have the Async flag set.
+		// They will never be destroyed during GC, and the object pointers are safe to keep.
+		const bool bHasAsyncFlag = Object->HasAnyInternalFlags(EInternalObjectFlags::Async);
+		if (bHasAsyncFlag)
+		{
+			continue;
+		}
+
+		const int32 RefCount = RefCounts[GlobalImportIndex];
+		check(RefCount >= 0);
+
+		if (RefCount > 0)
+		{
+			// Import objects in native packages will never be garbage collected and does not need marking.
+			FPackageIndex PackageIndex = Packages[GlobalImportIndex];
+			UPackage* Package = CastChecked<UPackage>(Objects[PackageIndex.ToImport()]);
+			if (Package->HasAnyPackageFlags(PKG_CompiledIn))
+			{
+				continue;
+			}
+			// Mark object to be kept alive during GC
+			Object->SetInternalFlags(EInternalObjectFlags::Async);
+			KeepAliveObjects.Add(Object);
+		}
+		else
+		{
+			// Clear object pointer since object may get destroyed during GC
+			Object = nullptr;
+			++NumCleared;
+		}
+	}
+
+	UE_LOG(LogStreaming, Log, TEXT("FGlobalImportRuntime::OnPreGarbageCollect - Marked %d objects, cleared %d object pointers"),
+		KeepAliveObjects.Num(), NumCleared);
+}
+
+void FGlobalImportRuntime::OnPostGarbageCollect()
+{
+	if (KeepAliveObjects.Num() == 0)
+	{
+		return;
+	}
+
+	for (UObject* Object : KeepAliveObjects)
+	{
+		Object->ClearInternalFlags(EInternalObjectFlags::Async);
+	}
+
+	const int32 UnmarkedCount = KeepAliveObjects.Num();
+	KeepAliveObjects.Reset();
+	UE_LOG(LogStreaming, Log, TEXT("FGlobalImportRuntime::UpdateGlobalImportsPostGC - Unmarked %d objects"),
+		UnmarkedCount);
+}
+
+void FAsyncLoadingThread2Impl::OnPreGarbageCollect()
+{
+	TRACE_CPUPROFILER_EVENT_SCOPE(AltPreGC);
+	if (!IsAsyncLoadingPackages())
+	{
+#ifdef ALT2_VERIFY_ASYNC_FLAGS
+		const EInternalObjectFlags AsyncFlags = EInternalObjectFlags::Async | EInternalObjectFlags::AsyncLoading;
+		for (int32 ObjectIndex = 0; ObjectIndex < GUObjectArray.GetObjectArrayNum(); ++ObjectIndex)
+		{
+			FUObjectItem* ObjectItem = &GUObjectArray.GetObjectItemArrayUnsafe()[ObjectIndex];
+			if (UObject* Obj = static_cast<UObject*>(ObjectItem->Object))
+			{
+				ensure(!Obj->HasAnyInternalFlags(AsyncFlags));
+			}
+		}
+#endif
+		return;
+	}
+	GlobalImportRuntime.OnPreGarbageCollect();
+}
+
+void FAsyncLoadingThread2Impl::OnPostGarbageCollect()
+{
+	TRACE_CPUPROFILER_EVENT_SCOPE(AltPostGC);
+	GlobalImportRuntime.OnPostGarbageCollect();
+}
+
+/**
+ * Call back into the async loading code to inform of the creation of a new object
+ * @param Object		Object created
+ * @param bSubObject	Object created as a sub-object of a loaded object
+ */
+void FAsyncLoadingThread2Impl::NotifyConstructedDuringAsyncLoading(UObject* Object, bool bSubObject)
+{
+	// Mark objects created during async loading process (e.g. from within PostLoad or CreateExport) as async loaded so they 
+	// cannot be found. This requires also keeping track of them so we can remove the async loading flag later one when we 
+	// finished routing PostLoad to all objects.
+	if (!bSubObject)
+	{
+		Object->SetInternalFlags(EInternalObjectFlags::AsyncLoading);
+	}
+	FUObjectThreadContext& ThreadContext = FUObjectThreadContext::Get();
+	check(ThreadContext.AsyncPackage); // Otherwise something is wrong and we're creating objects outside of async loading code
+	FAsyncPackage2* AsyncPackage2 = (FAsyncPackage2*)ThreadContext.AsyncPackage;
+	AsyncPackage2->AddOwnedObject(Object);
+	
+	// if this is in the package and is an export, then let mark it as needing load now
+	if (Object->GetOutermost() == AsyncPackage2->GetLinkerRoot() && 
+		AsyncPackage2->AsyncPackageLoadingState <= EAsyncPackageLoadingState2::ProcessNewImportsAndExports &&
+		AsyncPackage2->AsyncPackageLoadingState > EAsyncPackageLoadingState2::WaitingForSummary)
+	{
+		AsyncPackage2->MarkNewObjectForLoadIfItIsAnExport(Object);
+	}
+}
+
+void FAsyncLoadingThread2Impl::FireCompletedCompiledInImport(FGCObject* AsyncPackage, FPackageIndex Import)
+{
+	static_cast<FAsyncPackage2*>(AsyncPackage)->GetNode(EEventLoadNode2::ImportOrExport_Create, Import)->ReleaseBarrier();
+}
+
+void FAsyncLoadingThread2Impl::EnqueueIoRequest(FAsyncPackage2* Package, const FIoChunkId& ChunkId)
+{
+	IoRequestQueue->EnqueueRequest(Package, ChunkId);
+}
+
+/*-----------------------------------------------------------------------------
+	FAsyncPackage implementation.
+-----------------------------------------------------------------------------*/
+
+/**
+* Constructor
+*/
+FAsyncPackage2::FAsyncPackage2(const FAsyncPackageDesc& InDesc, int32 InSerialNumber, FAsyncLoadingThread2Impl& InAsyncLoadingThread, IEDLBootNotificationManager& InEDLBootNotificationManager, FAsyncLoadEventGraphAllocator& InGraphAllocator, const FAsyncLoadEventSpec* EventSpecs, FGlobalPackageId InGlobalPackageId)
+: Desc(InDesc)
+, Linker(nullptr)
+, LinkerRoot(nullptr)
+, ImportIndex(0)
+, ExportIndex(0)
+, FinishExternalReadDependenciesIndex(0)
+, PostLoadIndex(0)
+, DeferredPostLoadIndex(0)
+, DeferredFinalizeIndex(0)
+, DeferredClusterIndex(0)
+, bLoadHasFailed(false)
+, bLoadHasFinished(false)
+, bCreatedLinkerRoot(false)
+, LoadStartTime(0.0)
+, LoadPercentage(0)
+, ReentryCount(0)
+, AsyncLoadingThread(InAsyncLoadingThread)
+, EDLBootNotificationManager(InEDLBootNotificationManager)
+, GraphAllocator(InGraphAllocator)
+, GlobalPackageId(InGlobalPackageId)
+// Begin EDL specific properties
+, GlobalImportObjectRefCounts(AsyncLoadingThread.GetGlobalImportObjectRefCounts())
+, AsyncPackageLoadingState(EAsyncPackageLoadingState2::NewPackage)
+, SerialNumber(InSerialNumber)
+, bAllExportsSerialized(false)
+{
+	TRACE_CPUPROFILER_EVENT_SCOPE(NewAsyncPackage);
+	//TRACE_LOADTIME_NEW_ASYNC_PACKAGE(this, *InDesc.Name.ToString());
+	AddRequestID(InDesc.RequestID);
+
+	const int32* NameIndex = AsyncLoadingThread.GlobalNameMap.GetIndex(Desc.NameToLoad);
+	check(NameIndex != nullptr);
+	PackageChunkId = CreateChunkId(*NameIndex, Desc.NameToLoad.GetNumber(), 0, EChunkType::None);
+
+	{
+		LocalImportIndices = AsyncLoadingThread.GetPackageSlimports(GlobalPackageId, LocalImportCount);
+		GlobalImportNames = AsyncLoadingThread.GetGlobalImportNames(GlobalImportCount);
+		GlobalImportOuters = AsyncLoadingThread.GetGlobalImportOuters(GlobalImportCount);
+		GlobalImportPackages = AsyncLoadingThread.GetGlobalImportPackages(GlobalImportCount);
+		GlobalImportObjects = AsyncLoadingThread.GetGlobalImportObjects(GlobalImportCount);
+		AddGlobalImportObjectReferences();
+	}
+
+	{
+		TRACE_CPUPROFILER_EVENT_SCOPE(CreateNodes);
+		FPackageFileSummary Summary;
+		uint32 ImportCount = AsyncLoadingThread.GetPackageImportCount(GlobalPackageId);
+		ImportNodeCount = ImportCount * EEventLoadNode2::Import_NumPhases;
+		uint32 ExportCount = AsyncLoadingThread.GetPackageExportCount(GlobalPackageId);
+		ExportNodeCount = ExportCount * EEventLoadNode2::Export_NumPhases;
+
+		PackageNodes = GraphAllocator.AllocNodes(EEventLoadNode2::Package_NumPhases + ImportNodeCount + ExportNodeCount);
+		for (int32 Phase = 0; Phase < EEventLoadNode2::Package_NumPhases; ++Phase)
+		{
+			new (PackageNodes + Phase) FEventLoadNode2(EventSpecs + Phase, this, -1);
+		}
+
+		FEventLoadNode2* CreateLinkerNode = PackageNodes + EEventLoadNode2::Package_CreateLinker;
+		CreateLinkerNode->AddBarrier();
+
+		FEventLoadNode2* LoadSummaryNode = PackageNodes + EEventLoadNode2::Package_LoadSummary;
+		LoadSummaryNode->DependsOn(CreateLinkerNode);
+		LoadSummaryNode->AddBarrier();
+
+		FEventLoadNode2* ImportPackagesNode = PackageNodes + EEventLoadNode2::Package_ImportPackages;
+		FEventLoadNode2* SetupImportsNode = PackageNodes + EEventLoadNode2::Package_SetupImports;
+		FEventLoadNode2* SetupExportsNode = PackageNodes + EEventLoadNode2::Package_SetupExports;
+		FEventLoadNode2* ExportsSerializedNode = PackageNodes + EEventLoadNode2::Package_ExportsSerialized;
+		FEventLoadNode2* PostLoadNode = PackageNodes + EEventLoadNode2::Package_PostLoad;
+		FEventLoadNode2* TickNode = PackageNodes + EEventLoadNode2::Package_Tick;
+		TickNode->AddBarrier();
+
+		ImportPackagesNode->DependsOn(LoadSummaryNode);
+		SetupImportsNode->DependsOn(ImportPackagesNode);
+		SetupExportsNode->DependsOn(SetupImportsNode);
+		ExportsSerializedNode->DependsOn(SetupExportsNode);
+		PostLoadNode->DependsOn(ExportsSerializedNode);
+
+		FEventLoadNode2* DeleteNode = PackageNodes + EEventLoadNode2::Package_Delete;
+		DeleteNode->AddBarrier();
+		DeleteNode->DependsOn(TickNode);
+
+		// Add nodes for all imports and exports.
+		ImportNodes = PackageNodes + EEventLoadNode2::Package_NumPhases;
+		for (uint32 LocalImportIndex = 0; LocalImportIndex < ImportCount; ++LocalImportIndex)
+		{
+			uint32 NodeIndex = EEventLoadNode2::Import_NumPhases * LocalImportIndex;
+			FEventLoadNode2* CreateImportNode = ImportNodes + NodeIndex + EEventLoadNode2::ImportOrExport_Create;
+			new (CreateImportNode) FEventLoadNode2(EventSpecs + EEventLoadNode2::Package_NumPhases + EEventLoadNode2::ImportOrExport_Create, this, LocalImportIndex);
+			FEventLoadNode2* SerializeImportNode = ImportNodes + NodeIndex + EEventLoadNode2::ImportOrExport_Serialize;
+			new (SerializeImportNode) FEventLoadNode2(EventSpecs + EEventLoadNode2::Package_NumPhases + EEventLoadNode2::ImportOrExport_Serialize, this, LocalImportIndex);
+
+			CreateImportNode->DependsOn(SetupExportsNode); // Need to wait for SetupExports here because of preload dependencies
+			SerializeImportNode->DependsOn(CreateImportNode);
+			ExportsSerializedNode->DependsOn(SerializeImportNode);
+		}
+		ExportNodes = ImportNodes + ImportNodeCount;
+		for (uint32 LocalExportIndex = 0; LocalExportIndex < ExportCount; ++LocalExportIndex)
+		{
+			uint32 NodeIndex = EEventLoadNode2::Export_NumPhases * LocalExportIndex;
+			FEventLoadNode2* CreateExportNode = ExportNodes + NodeIndex + EEventLoadNode2::ImportOrExport_Create;
+			new (CreateExportNode) FEventLoadNode2(EventSpecs + EEventLoadNode2::Package_NumPhases + EEventLoadNode2::Import_NumPhases + EEventLoadNode2::ImportOrExport_Create, this, LocalExportIndex);
+			FEventLoadNode2* SerializeExportNode = ExportNodes + NodeIndex + EEventLoadNode2::ImportOrExport_Serialize;
+			new (SerializeExportNode) FEventLoadNode2(EventSpecs + EEventLoadNode2::Package_NumPhases + EEventLoadNode2::Import_NumPhases + EEventLoadNode2::ImportOrExport_Serialize, this, LocalExportIndex);
+			FEventLoadNode2* StartIONode = ExportNodes + NodeIndex + EEventLoadNode2::Export_StartIO;
+			new (StartIONode) FEventLoadNode2(EventSpecs + EEventLoadNode2::Package_NumPhases + EEventLoadNode2::Import_NumPhases + EEventLoadNode2::Export_StartIO, this, LocalExportIndex);
+			CreateExportNode->DependsOn(SetupExportsNode);
+			StartIONode->DependsOn(CreateExportNode);
+			ExportsSerializedNode->DependsOn(SerializeExportNode);
+			SerializeExportNode->AddBarrier();
+		}
+	}
+}
+
+FAsyncPackage2::~FAsyncPackage2()
+{
+	TRACE_CPUPROFILER_EVENT_SCOPE(DeleteAsyncPackage);
+
+	check(RefCount == 0);
+
+	FAsyncLoadingThreadState2::Get().DeferredFreeNodes.Add(MakeTuple(PackageNodes, EEventLoadNode2::Package_NumPhases + ImportNodeCount + ExportNodeCount));
+
+	//TRACE_LOADTIME_DESTROY_ASYNC_PACKAGE(this);
+
+	MarkRequestIDsAsComplete();
+	DetachLinker();
+	SerialNumber = 0; // the weak pointer will always fail now
+	
+	ensure(OwnedObjects.Num() == 0);
+}
+
+
+void FAsyncPackage2::ClearOwnedObjects()
+{
+	for (UObject* Object : OwnedObjects)
+	{
+		const EObjectFlags Flags = Object->GetFlags();
+		const EInternalObjectFlags InternalFlags = Object->GetInternalFlags();
+		EInternalObjectFlags InternalFlagsToClear = EInternalObjectFlags::None;
+
+		ensure(!(Flags & (RF_NeedPostLoad | RF_NeedPostLoadSubobjects)));
+		if (!!(InternalFlags & EInternalObjectFlags::AsyncLoading))
+		{
+			ensure(!(Flags & RF_WasLoaded));
+			InternalFlagsToClear |= EInternalObjectFlags::AsyncLoading;
+		}
+
+		if (!!(InternalFlags & EInternalObjectFlags::Async))
+		{
+			InternalFlagsToClear |= EInternalObjectFlags::Async;
+		}
+		Object->ClearInternalFlags(InternalFlagsToClear);
+	}
+	OwnedObjects.Empty();
+}
+
+void FAsyncPackage2::AddRequestID(int32 Id)
+{
+	if (Id > 0)
+	{
+		if (Desc.RequestID == INDEX_NONE)
+		{
+			// For debug readability
+			Desc.RequestID = Id;
+		}
+		RequestIDs.Add(Id);
+		AsyncLoadingThread.AddPendingRequest(Id);
+		//TRACE_LOADTIME_ASYNC_PACKAGE_REQUEST_ASSOCIATION(this, Id);
+	}
+}
+
+void FAsyncPackage2::MarkRequestIDsAsComplete()
+{
+	AsyncLoadingThread.RemovePendingRequests(RequestIDs);
+	RequestIDs.Reset();
+}
+
+/**
+ * @return Time load begun. This is NOT the time the load was requested in the case of other pending requests.
+ */
+double FAsyncPackage2::GetLoadStartTime() const
+{
+	return LoadStartTime;
+}
+
+/**
+ * Emulates ResetLoaders for the package's Linker objects, hence deleting it. 
+ */
+void FAsyncPackage2::ResetLoader()
+{
+	LLM_SCOPE(ELLMTag::AsyncLoading);
+
+	// Reset loader.
+	if (Linker)
+	{
+		check(Linker->AsyncRoot == this || Linker->AsyncRoot == nullptr);
+		Linker->AsyncRoot = nullptr;
+		// Flush cache and queue for delete
+		Linker->FlushCache();
+		Linker->Detach();
+		// FLinkerManager::Get().RemoveLinker(Linker);
+		Linker = nullptr;
+	}
+}
+
+void FAsyncPackage2::DetachLinker()
+{	
+	if (Linker)
+	{
+		Linker->FlushCache();
+		checkf(bLoadHasFinished || bLoadHasFailed, TEXT("FAsyncPackage::DetachLinker called before load finished on package \"%s\""), *this->GetPackageName().ToString());
+		check(Linker->AsyncRoot == this || Linker->AsyncRoot == nullptr);
+		Linker->AsyncRoot = nullptr;
+		Linker = nullptr;
+	}
+}
+
+void FAsyncPackage2::FlushObjectLinkerCache()
+{
+	for (UObject* Obj : PackageObjLoaded)
+	{
+		if (Obj)
+		{
+			FLinkerLoad* ObjLinker = Obj->GetLinker();
+			if (ObjLinker)
+			{
+				ObjLinker->FlushCache();
+			}
+		}
+	}
+}
+
+#if WITH_EDITOR 
+void FAsyncPackage2::GetLoadedAssets(TArray<FWeakObjectPtr>& AssetList)
+{
+	for (UObject* Obj : PackageObjLoaded)
+	{
+		if (Obj && !Obj->IsPendingKill() && Obj->IsAsset())
+		{
+			AssetList.AddUnique(Obj);
+		}
+	}
+}
+#endif
+
+/**
+ * Begin async loading process. Simulates parts of BeginLoad.
+ *
+ * Objects created during BeginAsyncLoad and EndAsyncLoad will have EInternalObjectFlags::AsyncLoading set
+ */
+void FAsyncPackage2::BeginAsyncLoad()
+{
+	if (IsInGameThread())
+	{
+		AsyncLoadingThread.EnterAsyncLoadingTick();
+	}
+
+	// this won't do much during async loading except increase the load count which causes IsLoading to return true
+	FUObjectSerializeContext* LoadContext = GetSerializeContext();
+	BeginLoad(LoadContext);
+}
+
+/**
+ * End async loading process. Simulates parts of EndLoad(). FinishObjects 
+ * simulates some further parts once we're fully done loading the package.
+ */
+void FAsyncPackage2::EndAsyncLoad()
+{
+	check(IsAsyncLoading());
+
+	// this won't do much during async loading except decrease the load count which causes IsLoading to return false
+	FUObjectSerializeContext* LoadContext = GetSerializeContext();
+	EndLoad(LoadContext);
+
+	if (IsInGameThread())
+	{
+		AsyncLoadingThread.LeaveAsyncLoadingTick();
+	}
+
+	if (!bLoadHasFailed)
+	{
+		// Mark the package as loaded, if we succeeded
+		LinkerRoot->SetFlags(RF_WasLoaded);
+	}
+}
+
+/**
+ * Create linker async. Linker is not finalized at this point.
+ *
+ * @return true
+ */
+EAsyncPackageState::Type FAsyncPackage2::CreateLinker()
+{
+	check(Linker == nullptr);
+
+	// Try to find existing package or create it if not already present.
+	UPackage* Package = nullptr;
+	{
+		Package = FindObjectFast<UPackage>(nullptr, Desc.Name);
+		if (!Package)
+		{
+			Package = NewObject<UPackage>(/*Outer*/nullptr, Desc.Name, RF_Public);
+			Package->SetPackageFlags(Desc.PackageFlags);
+			Package->FileName = Desc.NameToLoad;
+			bCreatedLinkerRoot = true;
+		}
+	}
+	check(!IsNativeCodePackage(Package));
+	if (IsInGameThread() || !bCreatedLinkerRoot)
+	{
+		Package->SetInternalFlags(EInternalObjectFlags::Async);
+		AddOwnedObject(Package);
+	}
+	check(Package->HasAnyInternalFlags(EInternalObjectFlags::Async));
+	LinkerRoot = Package;
+
+	check(!FLinkerLoad::FindExistingLinkerForPackage(Package));
+
+	{
+		uint32 LinkerFlags = LOAD_None | LOAD_Async | LOAD_NoVerify;
+		Linker = new FLinkerLoad(Package, *AsyncLoadingThread.GetPackageFileName(GlobalPackageId), LinkerFlags);
+		Linker->bIsAsyncLoader = false;
+		Linker->bLockoutLegacyOperations = true;
+		Linker->SetIsLoading(true);
+		Linker->SetIsPersistent(true);
+		Linker->AsyncRoot = this;
+	}
+	Package->LinkerLoad = Linker;
+
+	AsyncLoadingThread.EnqueueIoRequest(this, CreateChunkId(PackageChunkId, 0, EChunkType::PackageSummary));
+
+	UE_LOG(LogStreaming, Verbose, TEXT("FAsyncPackage::CreateLinker for %s finished."), *Desc.NameToLoad.ToString());
+	return EAsyncPackageState::Complete;
+}
+
+/**
+ * Finalizes linker creation till time limit is exceeded.
+ *
+ * @return true if linker is finished being created, false otherwise
+ */
+EAsyncPackageState::Type FAsyncPackage2::FinishLinker()
+{
+	LLM_SCOPE(ELLMTag::AsyncLoading);
+
+	//SCOPED_LOADTIMER(FinishLinkerTime);
+	{
+		SCOPED_LOADTIMER(LinkerLoad_FinalizeCreation);
+
+		const FPackageSummary* PackageSummary = reinterpret_cast<const FPackageSummary*>(PackageSummaryBuffer.Get());
+
+		{
+			SCOPED_LOADTIMER(LinkerLoad_SerializePackageFileSummary);
+
+			FPackageFileSummary& Summary	= Linker->Summary;
+
+			Summary.Tag						= PACKAGE_FILE_TAG;
+			Summary.Guid					= PackageSummary->Guid;
+			Summary.PackageFlags			= PackageSummary->PackageFlags;
+			Summary.ExportCount				= PackageSummary->ExportCount;
+			Summary.ImportCount				= PackageSummary->ImportCount;
+			Summary.PreloadDependencyCount	= PackageSummary->PreloadDependencyCount;
+			Summary.PreloadDependencyOffset	= 1; // HACK: circumvent check in FLinkerLoad::SerializePreloadDependencies()
+			Summary.BulkDataStartOffset		= PackageSummary->BulkDataStartOffset;
+
+			Summary.SetFileVersions(GPackageFileUE4Version, GPackageFileLicenseeUE4Version, /*unversioned*/true);
+
+			Linker->UpdateFromPackageFileSummary();
+		}
+
+		// TODO: FLinker should not be a FArchive - NameMap is only required for operator<<(FName& Name)
+		{
+			SCOPED_LOADTIMER(LinkerLoad_SerializeNameMap);
+			Linker->ActiveNameMap = &AsyncLoadingThread.GlobalNameMap.GetNameEntries();
+		}
+
+		// TODO: FLinker should not be a FArchive - Slimports are only required for operator<<( UObject*& Object )
+		{
+			SCOPED_LOADTIMER(LinkerLoad_SerializeImportMap);
+			int32 TmpImportCount = 0;
+			int32 TmpGlobalImportCount = 0;
+			Linker->LocalImportIndices = AsyncLoadingThread.GetPackageSlimports(GlobalPackageId, TmpImportCount);
+			Linker->GlobalImportObjects = AsyncLoadingThread.GetGlobalImportObjects(TmpGlobalImportCount);
+		}
+
+		if (PackageSummary->ExportCount)
+		{
+			SCOPED_LOADTIMER(LinkerLoad_SerializeExportMap);
+			const FObjectExport* Exports = reinterpret_cast<const FObjectExport*>(PackageSummaryBuffer.Get() + PackageSummary->ExportOffset);
+
+			Linker->ExportMap.AddUninitialized(PackageSummary->ExportCount);
+			FMemory::Memcpy(Linker->ExportMap.GetData(), Exports, PackageSummary->ExportCount * sizeof(FObjectExport));
+
+			for (FObjectExport& Export : Linker->ExportMap)
+			{
+				Export.ObjectName = AsyncLoadingThread.GlobalNameMap.FromSerializedName(Export.ObjectName);
+			}
+
+			// Foreach Export.bWasFiltered = FilterExport(Export); // TODO: Should never return true in packaged game, if so move to cooker!!!
+			Linker->ExportMapIndex = PackageSummary->ExportCount;
+
+			// ObjectNameWithOuterToExport has two use cases:
+			// - used for SetupImports during initial loading
+			// - MarkNewObjectForLoadIfItIsAnExport from NotifyConstructedDuringAsyncLoading
+			ObjectNameWithOuterToExport.Reserve(PackageSummary->ExportCount);
+			for (int32 LocalExportIndex = 0; LocalExportIndex < PackageSummary->ExportCount; ++LocalExportIndex)
+			{
+				FPackageIndex Index = FPackageIndex::FromExport(LocalExportIndex);
+				const FObjectExport& Export = Linker->Exp(Index);
+				ObjectNameWithOuterToExport.Add(TPair<FName, FPackageIndex>(Export.ObjectName, Export.OuterIndex), Index);
+			}
+
+			OwnedObjects.Reserve(PackageSummary->ExportCount + 1); // ExportCount + UPackage
+			ExportIoBuffers.SetNum(PackageSummary->ExportCount);
+		}
+
+		// Add this linker to the object manager's linker array.
+		// FLinkerManager::Get().AddLoader(Linker);
+	}
+
+	return EAsyncPackageState::Complete;
+}
+
+EAsyncPackageState::Type FAsyncPackage2::FinishExternalReadDependencies()
+{
+	if (FAsyncLoadingThreadState2::Get().IsTimeLimitExceeded())
+	{
+		return EAsyncPackageState::TimeOut;
+	}
+
+	FLinkerLoad* VisitedLinkerLoad = nullptr;
+	while (FinishExternalReadDependenciesIndex < PackageObjLoaded.Num())
+	{
+		UObject* Obj = PackageObjLoaded[FinishExternalReadDependenciesIndex];
+		FLinkerLoad* LinkerLoad = Obj ? Obj->GetLinker() : nullptr;
+		if (LinkerLoad && LinkerLoad != VisitedLinkerLoad)
+		{
+			if (!LinkerLoad->FinishExternalReadDependencies(0.0) || FAsyncLoadingThreadState2::Get().IsTimeLimitExceeded())
+			{
+				return EAsyncPackageState::TimeOut;
+			}
+			VisitedLinkerLoad = LinkerLoad;
+		}
+		FinishExternalReadDependenciesIndex++;
+	}
+		
+	return EAsyncPackageState::Complete;
+}
+
+/**
+ * Route PostLoad to all loaded objects. This might load further objects!
+ *
+ * @return true if we finished calling PostLoad on all loaded objects and no new ones were created, false otherwise
+ */
+EAsyncPackageState::Type FAsyncPackage2::PostLoadObjects()
+{
+	LLM_SCOPE(ELLMTag::UObject);
+
+	SCOPED_LOADTIMER(PostLoadObjectsTime);
+
+	FUObjectThreadContext& ThreadContext = FUObjectThreadContext::Get();
+	TGuardValue<bool> GuardIsRoutingPostLoad(ThreadContext.IsRoutingPostLoad, true);
+
+	FUObjectSerializeContext* LoadContext = GetSerializeContext();
+	TArray<UObject*>& ThreadObjLoaded = LoadContext->PRIVATE_GetObjectsLoadedInternalUseOnly();
+	if (ThreadObjLoaded.Num())
+	{
+		// New objects have been loaded. They need to go through PreLoad first so exit now and come back after they've been preloaded.
+		PackageObjLoaded.Append(ThreadObjLoaded);
+		ThreadObjLoaded.Reset();
+		return EAsyncPackageState::TimeOut;
+	}
+
+	int32 PreLoadIndex = PackageObjLoaded.Num();
+	
+	const bool bAsyncPostLoadEnabled = true;
+	const bool bIsMultithreaded = AsyncLoadingThread.IsMultithreaded();
+
+	// PostLoad objects.
+	while (PostLoadIndex < PackageObjLoaded.Num() && PostLoadIndex < PreLoadIndex && !FAsyncLoadingThreadState2::Get().IsTimeLimitExceeded())
+	{
+		UObject* Object = PackageObjLoaded[PostLoadIndex++];
+		if (Object)
+		{
+			if (!Object->IsReadyForAsyncPostLoad())
+			{
+				--PostLoadIndex;
+				break;
+			}
+			else if (!bIsMultithreaded || (bAsyncPostLoadEnabled && CanPostLoadOnAsyncLoadingThread(Object)))
+			{
+				check(!Object->HasAnyFlags(RF_NeedLoad));
+
+				ThreadContext.CurrentlyPostLoadedObjectByALT = Object;
+				{
+					TRACE_LOADTIME_OBJECT_SCOPE(Object, LoadTimeProfilerObjectEventType_PostLoad);
+					Object->ConditionalPostLoad();
+					Object->AtomicallyClearInternalFlags(EInternalObjectFlags::AsyncLoading);
+				}
+				ThreadContext.CurrentlyPostLoadedObjectByALT = nullptr;
+
+				if (ThreadObjLoaded.Num())
+				{
+					// New objects have been loaded. They need to go through PreLoad first so exit now and come back after they've been preloaded.
+					PackageObjLoaded.Append(ThreadObjLoaded);
+					ThreadObjLoaded.Reset();
+					return EAsyncPackageState::TimeOut;
+				}
+			}
+			else
+			{
+				DeferredPostLoadObjects.Add(Object);
+			}
+			// All object must be finalized on the game thread
+			DeferredFinalizeObjects.Add(Object);
+			check(Object->IsValidLowLevelFast());
+		}
+	}
+
+	PackageObjLoaded.Append(ThreadObjLoaded);
+	ThreadObjLoaded.Reset();
+
+	// New objects might have been loaded during PostLoad.
+	EAsyncPackageState::Type Result = (PreLoadIndex == PackageObjLoaded.Num()) && (PostLoadIndex == PackageObjLoaded.Num()) ? EAsyncPackageState::Complete : EAsyncPackageState::TimeOut;
+	return Result;
+}
+
+void CreateClustersFromPackage(FLinkerLoad* PackageLinker, TArray<UObject*>& OutClusterObjects);
+
+EAsyncPackageState::Type FAsyncPackage2::PostLoadDeferredObjects()
+{
+	SCOPED_LOADTIMER(PostLoadDeferredObjectsTime);
+
+	FAsyncPackageScope2 PackageScope(this);
+
+	EAsyncPackageState::Type Result = EAsyncPackageState::Complete;
+	TGuardValue<bool> GuardIsRoutingPostLoad(PackageScope.ThreadContext.IsRoutingPostLoad, true);
+	FAsyncLoadingTickScope2 InAsyncLoadingTick(AsyncLoadingThread);
+
+	FUObjectSerializeContext* LoadContext = GetSerializeContext();
+	TArray<UObject*>& ObjLoadedInPostLoad = LoadContext->PRIVATE_GetObjectsLoadedInternalUseOnly();
+	TArray<UObject*> ObjLoadedInPostLoadLocal;
+
+	STAT(double PostLoadStartTime = FPlatformTime::Seconds());
+
+	while (DeferredPostLoadIndex < DeferredPostLoadObjects.Num() && 
+		!AsyncLoadingThread.IsAsyncLoadingSuspended() &&
+		!FAsyncLoadingThreadState2::Get().IsTimeLimitExceeded())
+	{
+		UObject* Object = DeferredPostLoadObjects[DeferredPostLoadIndex++];
+		
+		check(Object);
+
+		PackageScope.ThreadContext.CurrentlyPostLoadedObjectByALT = Object;
+		{
+			TRACE_LOADTIME_OBJECT_SCOPE(Object, LoadTimeProfilerObjectEventType_PostLoad);
+			Object->ConditionalPostLoad();
+		}
+		PackageScope.ThreadContext.CurrentlyPostLoadedObjectByALT = nullptr;
+
+		if (ObjLoadedInPostLoad.Num())
+		{
+			// If there were any LoadObject calls inside of PostLoad, we need to pre-load those objects here. 
+			// There's no going back to the async tick loop from here.
+			UE_LOG(LogStreaming, Warning, TEXT("Detected %d objects loaded in PostLoad while streaming, this may cause hitches as we're blocking async loading to pre-load them."), ObjLoadedInPostLoad.Num());
+			
+			// Copy to local array because ObjLoadedInPostLoad can change while we're iterating over it
+			ObjLoadedInPostLoadLocal.Append(ObjLoadedInPostLoad);
+			ObjLoadedInPostLoad.Reset();
+
+			while (ObjLoadedInPostLoadLocal.Num())
+			{
+				// Make sure all objects loaded in PostLoad get post-loaded too
+				DeferredPostLoadObjects.Append(ObjLoadedInPostLoadLocal);
+
+				// Preload (aka serialize) the objects loaded in PostLoad.
+				for (UObject* PreLoadObject : ObjLoadedInPostLoadLocal)
+				{
+					if (PreLoadObject && PreLoadObject->GetLinker())
+					{
+						PreLoadObject->GetLinker()->Preload(PreLoadObject);
+					}
+				}
+
+				// Other objects could've been loaded while we were preloading, continue until we've processed all of them.
+				ObjLoadedInPostLoadLocal.Reset();
+				ObjLoadedInPostLoadLocal.Append(ObjLoadedInPostLoad);
+				ObjLoadedInPostLoad.Reset();
+			}			
+		}
+
+		UpdateLoadPercentage();
+	}
+
+	// New objects might have been loaded during PostLoad.
+	Result = (DeferredPostLoadIndex == DeferredPostLoadObjects.Num()) ? EAsyncPackageState::Complete : EAsyncPackageState::TimeOut;
+	/*if (Result == EAsyncPackageState::Complete)
+	{
+		for (int32 LocalExportIndex = 0; LocalExportIndex < Linker->ExportMap.Num(); LocalExportIndex++)
+		{
+			UObject* Object = Linker->ExportMap[LocalExportIndex].Object;
+			if (Object && Object->HasAnyFlags(RF_NeedLoad | RF_NeedPostLoad))
+			{
+				Result = EAsyncPackageState::TimeOut;
+			}
+		}
+	}*/
+	
+	if (Result == EAsyncPackageState::Complete)
+	{
+		TArray<UObject*> CDODefaultSubobjects;
+		// Clear async loading flags (we still want RF_Async, but EInternalObjectFlags::AsyncLoading can be cleared)
+		while (DeferredFinalizeIndex < DeferredFinalizeObjects.Num() &&
+			(DeferredPostLoadIndex % 100 != 0 || (!AsyncLoadingThread.IsAsyncLoadingSuspended() && !FAsyncLoadingThreadState2::Get().IsTimeLimitExceeded())))
+		{
+			UObject* Object = DeferredFinalizeObjects[DeferredFinalizeIndex++];
+			if (Object)
+			{
+				Object->AtomicallyClearInternalFlags(EInternalObjectFlags::AsyncLoading);
+			}
+
+			// CDO need special handling, no matter if it's listed in DeferredFinalizeObjects or created here for DynamicClass
+			UObject* CDOToHandle = nullptr;
+
+			// Dynamic Class doesn't require/use pre-loading (or post-loading). 
+			// The CDO is created at this point, because now it's safe to solve cyclic dependencies.
+			if (UDynamicClass* DynamicClass = Cast<UDynamicClass>(Object))
+			{
+				check((DynamicClass->ClassFlags & CLASS_Constructed) != 0);
+
+				//native blueprint 
+
+				check(DynamicClass->HasAnyClassFlags(CLASS_TokenStreamAssembled));
+				// this block should be removed entirely when and if we add the CDO to the fake export table
+				CDOToHandle = DynamicClass->GetDefaultObject(false);
+				UE_CLOG(!CDOToHandle, LogStreaming, Fatal, TEXT("EDL did not create the CDO for %s before it finished loading."), *DynamicClass->GetFullName());
+				CDOToHandle->AtomicallyClearInternalFlags(EInternalObjectFlags::AsyncLoading);
+			}
+			else
+			{
+				CDOToHandle = ((Object != nullptr) && Object->HasAnyFlags(RF_ClassDefaultObject)) ? Object : nullptr;
+			}
+
+			// Clear AsyncLoading in CDO's subobjects.
+			if(CDOToHandle != nullptr)
+			{
+				CDOToHandle->GetDefaultSubobjects(CDODefaultSubobjects);
+				for (UObject* SubObject : CDODefaultSubobjects)
+				{
+					if (SubObject && SubObject->HasAnyInternalFlags(EInternalObjectFlags::AsyncLoading))
+					{
+						SubObject->AtomicallyClearInternalFlags(EInternalObjectFlags::AsyncLoading);
+					}
+				}
+				CDODefaultSubobjects.Reset();
+			}
+		}
+		if (DeferredFinalizeIndex == DeferredFinalizeObjects.Num())
+		{
+			DeferredFinalizeIndex = 0;
+			DeferredFinalizeObjects.Reset();
+			Result = EAsyncPackageState::Complete;
+		}
+		else
+		{
+			Result = EAsyncPackageState::TimeOut;
+		}
+	
+		// Mark package as having been fully loaded and update load time.
+		if (Result == EAsyncPackageState::Complete && LinkerRoot && !bLoadHasFailed)
+		{
+			LinkerRoot->AtomicallyClearInternalFlags(EInternalObjectFlags::AsyncLoading);
+			LinkerRoot->MarkAsFullyLoaded();			
+			LinkerRoot->SetLoadTime(FPlatformTime::Seconds() - LoadStartTime);
+
+			if (Linker)
+			{
+				CreateClustersFromPackage(Linker, DeferredClusterObjects);
+			}
+		}
+
+		FSoftObjectPath::InvalidateTag();
+		FUniqueObjectGuid::InvalidateTag();
+	}
+
+	return Result;
+}
+
+EAsyncPackageState::Type FAsyncPackage2::CreateClusters()
+{
+	while (DeferredClusterIndex < DeferredClusterObjects.Num() &&
+			!AsyncLoadingThread.IsAsyncLoadingSuspended() &&
+			!FAsyncLoadingThreadState2::Get().IsTimeLimitExceeded())
+	{
+		UObject* ClusterRootObject = DeferredClusterObjects[DeferredClusterIndex++];
+		ClusterRootObject->CreateCluster();
+	}
+
+	EAsyncPackageState::Type Result;
+	if (DeferredClusterIndex == DeferredClusterObjects.Num())
+	{
+		DeferredClusterIndex = 0;
+		DeferredClusterObjects.Reset();
+		Result = EAsyncPackageState::Complete;
+	}
+	else
+	{
+		Result = EAsyncPackageState::TimeOut;
+	}
+
+	return Result;
+}
+
+EAsyncPackageState::Type FAsyncPackage2::FinishObjects()
+{
+	SCOPED_LOADTIMER(FinishObjectsTime);
+
+	FUObjectSerializeContext* LoadContext = GetSerializeContext();
+	check(!Linker || LoadContext == Linker->GetSerializeContext());		
+	TArray<UObject*>& ThreadObjLoaded = LoadContext->PRIVATE_GetObjectsLoadedInternalUseOnly();
+
+	EAsyncLoadingResult::Type LoadingResult;
+	if (!bLoadHasFailed)
+	{
+		ThreadObjLoaded.Reset();
+		LoadingResult = EAsyncLoadingResult::Succeeded;
+	}
+	else
+	{		
+		PackageObjLoaded.Append(ThreadObjLoaded);
+
+		// Cleanup objects from this package only
+		for (int32 ObjectIndex = PackageObjLoaded.Num() - 1; ObjectIndex >= 0; --ObjectIndex)
+		{
+			UObject* Object = PackageObjLoaded[ObjectIndex];
+			if (Object && Object->GetOutermost()->GetFName() == Desc.Name)
+			{
+				Object->ClearFlags(RF_NeedPostLoad | RF_NeedLoad | RF_NeedPostLoadSubobjects);
+				Object->MarkPendingKill();
+				PackageObjLoaded[ObjectIndex] = nullptr;
+			}
+		}
+
+		// Clean up UPackage so it can't be found later
+		if (LinkerRoot && !LinkerRoot->IsRooted())
+		{
+			if (bCreatedLinkerRoot)
+			{
+				LinkerRoot->ClearFlags(RF_NeedPostLoad | RF_NeedLoad | RF_NeedPostLoadSubobjects);
+				LinkerRoot->MarkPendingKill();
+				LinkerRoot->Rename(*MakeUniqueObjectName(GetTransientPackage(), UPackage::StaticClass()).ToString(), nullptr, REN_DontCreateRedirectors | REN_DoNotDirty | REN_ForceNoResetLoaders | REN_NonTransactional);
+			}
+			DetachLinker();
+		}
+
+		LoadingResult = EAsyncLoadingResult::Failed;
+	}
+
+	// Simulate what EndLoad does.
+	// FLinkerManager::Get().DissociateImportsAndForcedExports(); //@todo: this should be avoidable
+	PostLoadIndex = 0;
+	FinishExternalReadDependenciesIndex = 0;
+
+	// Keep the linkers to close until we finish loading and it's safe to close them too
+	LoadContext->MoveDelayedLinkerClosePackages(DelayedLinkerClosePackages);
+
+	if (Linker)
+	{
+		// Flush linker cache now to reduce peak memory usage (5.5-10x)
+		// We shouldn't need it anyway at this point and even if something attempts to read in PostLoad, 
+		// we're just going to re-cache then.
+		Linker->FlushCache();
+	}
+
+	const bool bInternalCallbacks = true;
+	CallCompletionCallbacks(bInternalCallbacks, LoadingResult);
+
+	for (UObject* Object : OwnedObjects)
+	{
+		if (!Object->HasAnyFlags(RF_NeedPostLoad | RF_NeedPostLoadSubobjects))
+		{
+			Object->ClearInternalFlags(EInternalObjectFlags::AsyncLoading);
+		}
+	}
+
+	return EAsyncPackageState::Complete;
+}
+
+void FAsyncPackage2::CloseDelayedLinkers()
+{
+	// Close any linkers that have been open as a result of blocking load while async loading
+	for (FLinkerLoad* LinkerToClose : DelayedLinkerClosePackages)
+	{
+		if (LinkerToClose->LinkerRoot != nullptr)
+		{
+			check(LinkerToClose);
+			//check(LinkerToClose->LinkerRoot);
+			FLinkerLoad* LinkerToReset = FLinkerLoad::FindExistingLinkerForPackage(LinkerToClose->LinkerRoot);
+			check(LinkerToReset == LinkerToClose);
+			if (LinkerToReset && LinkerToReset->AsyncRoot)
+			{
+				UE_LOG(LogStreaming, Error, TEXT("Linker cannot be reset right now...leaking %s"), *LinkerToReset->GetArchiveName());
+				continue;
+			}
+		}
+		check(LinkerToClose->LinkerRoot == nullptr);
+		check(LinkerToClose->AsyncRoot == nullptr);
+	}
+}
+
+void FAsyncPackage2::CallCompletionCallbacks(bool bInternal, EAsyncLoadingResult::Type LoadingResult)
+{
+	checkSlow(bInternal || !IsInAsyncLoadingThread());
+
+	UPackage* LoadedPackage = (!bLoadHasFailed) ? LinkerRoot : nullptr;
+	for (FCompletionCallback& CompletionCallback : CompletionCallbacks)
+	{
+		if (CompletionCallback.bIsInternal == bInternal && !CompletionCallback.bCalled)
+		{
+			CompletionCallback.bCalled = true;
+			CompletionCallback.Callback->ExecuteIfBound(Desc.Name, LoadedPackage, LoadingResult);
+		}
+	}
+}
+
+UPackage* FAsyncPackage2::GetLoadedPackage()
+{
+	UPackage* LoadedPackage = (!bLoadHasFailed) ? LinkerRoot : nullptr;
+	return LoadedPackage;
+}
+
+void FAsyncPackage2::Cancel()
+{
+	// Call any completion callbacks specified.
+	bLoadHasFailed = true;
+	const EAsyncLoadingResult::Type Result = EAsyncLoadingResult::Canceled;
+	CallCompletionCallbacks(true, Result);
+	CallCompletionCallbacks(false, Result);
+
+	FUObjectSerializeContext* LoadContext = GetSerializeContext();
+	if (LoadContext)
+	{
+		TArray<UObject*>& ThreadObjLoaded = LoadContext->PRIVATE_GetObjectsLoadedInternalUseOnly();
+		if (ThreadObjLoaded.Num())
+		{
+			PackageObjLoaded.Append(ThreadObjLoaded);
+			ThreadObjLoaded.Reset();
+		}
+	}
+
+	{
+		// Clear load flags from any referenced objects
+		ClearFlagsAndDissolveClustersFromLoadedObjects(PackageObjLoaded);
+		ClearFlagsAndDissolveClustersFromLoadedObjects(DeferredFinalizeObjects);
+	
+		// Release references
+		PackageObjLoaded.Empty();
+		DeferredFinalizeObjects.Empty();
+	}
+
+	if (LinkerRoot)
+	{
+		if (Linker)
+		{
+			Linker->FlushCache();
+		}
+		if (bCreatedLinkerRoot)
+		{
+			LinkerRoot->ClearFlags(RF_WasLoaded);
+			LinkerRoot->bHasBeenFullyLoaded = false;
+			LinkerRoot->Rename(*MakeUniqueObjectName(GetTransientPackage(), UPackage::StaticClass()).ToString(), nullptr, REN_DontCreateRedirectors | REN_DoNotDirty | REN_ForceNoResetLoaders | REN_NonTransactional);
+		}
+		ResetLoader();
+	}
+	FinishExternalReadDependenciesIndex = 0;
+}
+
+void FAsyncPackage2::AddCompletionCallback(TUniquePtr<FLoadPackageAsyncDelegate>&& Callback, bool bInternal)
+{
+	// This is to ensure that there is no one trying to subscribe to a already loaded package
+	//check(!bLoadHasFinished && !bLoadHasFailed);
+	CompletionCallbacks.Emplace(bInternal, MoveTemp(Callback));
+}
+
+void FAsyncPackage2::UpdateLoadPercentage()
+{
+	// PostLoadCount is just an estimate to prevent packages to go to 100% too quickly
+	// We may never reach 100% this way, but it's better than spending most of the load package time at 100%
+	float NewLoadPercentage = 0.0f;
+	if (Linker)
+	{
+		const int32 PostLoadCount = FMath::Max(DeferredPostLoadObjects.Num(), LocalImportCount);
+		NewLoadPercentage = 100.f * (ExportIndex + DeferredPostLoadIndex) / (Linker->ExportMap.Num() + PostLoadCount);		
+	}
+	else if (DeferredPostLoadObjects.Num() > 0)
+	{
+		NewLoadPercentage = static_cast<float>(DeferredPostLoadIndex) / DeferredPostLoadObjects.Num();
+	}
+	// It's also possible that we got so many objects to PostLoad that LoadPercantage will actually drop
+	LoadPercentage = FMath::Max(NewLoadPercentage, LoadPercentage);
+}
+
+int32 FAsyncLoadingThread2Impl::LoadPackage(const FString& InName, const FGuid* InGuid, const TCHAR* InPackageToLoadFrom, FLoadPackageAsyncDelegate InCompletionDelegate, EPackageFlags InPackageFlags, int32 InPIEInstanceID, int32 InPackagePriority)
+{
+	TRACE_CPUPROFILER_EVENT_SCOPE(LoadPackage);
+
+	int32 RequestID = INDEX_NONE;
+
+	static bool bOnce = false;
+	if (!bOnce)
+	{
+		bOnce = true;
+		FGCObject::StaticInit(); // otherwise this thing is created during async loading, but not associated with a package
+	}
+
+	// The comments clearly state that it should be a package name but we also handle it being a filename as this function is not perf critical
+	// and LoadPackage handles having a filename being passed in as well.
+	FString PackageName;
+	bool bValidPackageName = true;
+
+	if (FPackageName::IsValidLongPackageName(InName, /*bIncludeReadOnlyRoots*/true))
+	{
+		PackageName = InName;
+	}
+	// PackageName got populated by the conditional function
+	else if (!(FPackageName::IsPackageFilename(InName) && FPackageName::TryConvertFilenameToLongPackageName(InName, PackageName)))
+	{
+		// PackageName may get populated by the conditional function
+		FString ClassName;
+
+		if (!FPackageName::ParseExportTextPath(PackageName, &ClassName, &PackageName))
+		{
+			UE_LOG(LogStreaming, Warning, TEXT("LoadPackageAsync failed to begin to load a package because the supplied package name ")
+				TEXT("was neither a valid long package name nor a filename of a map within a content folder: '%s' (%s)"),
+				*PackageName, *InName);
+
+			bValidPackageName = false;
+		}
+	}
+
+	FString PackageNameToLoad(InPackageToLoadFrom);
+
+	if (bValidPackageName)
+	{
+		if (PackageNameToLoad.IsEmpty())
+		{
+			PackageNameToLoad = PackageName;
+		}
+		// Make sure long package name is passed to FAsyncPackage so that it doesn't attempt to 
+		// create a package with short name.
+		if (FPackageName::IsShortPackageName(PackageNameToLoad))
+		{
+			UE_LOG(LogStreaming, Warning, TEXT("Async loading code requires long package names (%s)."), *PackageNameToLoad);
+
+			bValidPackageName = false;
+		}
+	}
+
+	if (bValidPackageName)
+	{
+		if (FCoreDelegates::OnAsyncLoadPackage.IsBound())
+		{
+			FCoreDelegates::OnAsyncLoadPackage.Broadcast(InName);
+		}
+
+		// Generate new request ID and add it immediately to the global request list (it needs to be there before we exit
+		// this function, otherwise it would be added when the packages are being processed on the async thread).
+		RequestID = PackageRequestID.Increment();
+		TRACE_LOADTIME_BEGIN_REQUEST(RequestID);
+		AddPendingRequest(RequestID);
+
+		// Allocate delegate on Game Thread, it is not safe to copy delegates by value on other threads
+		TUniquePtr<FLoadPackageAsyncDelegate> CompletionDelegatePtr;
+		if (InCompletionDelegate.IsBound())
+		{
+			CompletionDelegatePtr.Reset(new FLoadPackageAsyncDelegate(InCompletionDelegate));
+		}
+
+		// Add new package request
+		FAsyncPackageDesc PackageDesc(RequestID, *PackageName, *PackageNameToLoad, InGuid ? *InGuid : FGuid(), MoveTemp(CompletionDelegatePtr), InPackageFlags, InPIEInstanceID, InPackagePriority);
+		QueuePackage(PackageDesc);
+	}
+	else
+	{
+		InCompletionDelegate.ExecuteIfBound(FName(*InName), nullptr, EAsyncLoadingResult::Failed);
+	}
+
+	return RequestID;
+}
+
+EAsyncPackageState::Type FAsyncLoadingThread2Impl::ProcessLoadingFromGameThread(bool bUseTimeLimit, bool bUseFullTimeLimit, float TimeLimit)
+{
+	TickAsyncLoadingFromGameThread(bUseTimeLimit, bUseFullTimeLimit, TimeLimit);
+	return IsAsyncLoading() ? EAsyncPackageState::TimeOut : EAsyncPackageState::Complete;
+}
+
+void FAsyncLoadingThread2Impl::FlushLoading(int32 PackageID)
+{
+	if (IsAsyncLoading())
+	{
+		// Flushing async loading while loading is suspend will result in infinite stall
+		UE_CLOG(bSuspendRequested, LogStreaming, Fatal, TEXT("Cannot Flush Async Loading while async loading is suspended"));
+
+		if (PackageID != INDEX_NONE && !ContainsRequestID(PackageID))
+		{
+			return;
+		}
+
+		FCoreDelegates::OnAsyncLoadingFlush.Broadcast();
+
+#if NO_LOGGING == 0
+		{
+			// Log the flush, but only display once per frame to avoid log spam.
+			static uint64 LastFrameNumber = -1;
+			if (LastFrameNumber != GFrameNumber)
+			{
+				UE_LOG(LogStreaming, Display, TEXT("Flushing async loaders."));
+				LastFrameNumber = GFrameNumber;
+			}
+			else
+			{
+				UE_LOG(LogStreaming, Log, TEXT("Flushing async loaders."));
+			}
+		}
+#endif
+
+		double StartTime = FPlatformTime::Seconds();
+
+		// Flush async loaders by not using a time limit. Needed for e.g. garbage collection.
+		{
+			while (IsAsyncLoading())
+			{
+				EAsyncPackageState::Type Result = TickAsyncLoadingFromGameThread(false, false, 0, PackageID);
+				if (PackageID != INDEX_NONE && !ContainsRequestID(PackageID))
+				{
+					break;
+				}
+
+				if (IsMultithreaded())
+				{
+					// Update the heartbeat and sleep. If we're not multithreading, the heartbeat is updated after each package has been processed
+					FThreadHeartBeat::Get().HeartBeat();
+					FPlatformProcess::SleepNoStats(0.0001f);
+				}
+
+				// push stats so that we don't overflow number of tags per thread during blocking loading
+				LLM_PUSH_STATS_FOR_ASSET_TAGS();
+			}
+		}
+
+		double EndTime = FPlatformTime::Seconds();
+		double ElapsedTime = EndTime - StartTime;
+
+		check(PackageID != INDEX_NONE || !IsAsyncLoading());
+	}
+}
+
+EAsyncPackageState::Type FAsyncLoadingThread2Impl::ProcessLoadingUntilCompleteFromGameThread(TFunctionRef<bool()> CompletionPredicate, float TimeLimit)
+{
+	if (!IsAsyncLoading())
+	{
+		return EAsyncPackageState::Complete;
+	}
+
+	// Flushing async loading while loading is suspend will result in infinite stall
+	UE_CLOG(bSuspendRequested, LogStreaming, Fatal, TEXT("Cannot Flush Async Loading while async loading is suspended"));
+
+	if (TimeLimit <= 0.0f)
+	{
+		// Set to one hour if no time limit
+		TimeLimit = 60 * 60;
+	}
+
+	while (IsAsyncLoading() && TimeLimit > 0 && !CompletionPredicate())
+	{
+		double TickStartTime = FPlatformTime::Seconds();
+		if (ProcessLoadingFromGameThread(true, true, TimeLimit) == EAsyncPackageState::Complete)
+		{
+			return EAsyncPackageState::Complete;
+		}
+
+		if (IsMultithreaded())
+		{
+			// Update the heartbeat and sleep. If we're not multithreading, the heartbeat is updated after each package has been processed
+			FThreadHeartBeat::Get().HeartBeat();
+			FPlatformProcess::SleepNoStats(0.0001f);
+		}
+
+		TimeLimit -= (FPlatformTime::Seconds() - TickStartTime);
+	}
+
+	return TimeLimit <= 0 ? EAsyncPackageState::TimeOut : EAsyncPackageState::Complete;
+}
+
+FAsyncLoadingThread2::FAsyncLoadingThread2(IEDLBootNotificationManager& InEDLBootNotificationManager)
+{
+	Impl = new FAsyncLoadingThread2Impl(InEDLBootNotificationManager);
+}
+
+FAsyncLoadingThread2::~FAsyncLoadingThread2()
+{
+	delete Impl;
+}
+
+void FAsyncLoadingThread2::InitializeLoading()
+{
+	Impl->InitializeLoading();
+}
+
+void FAsyncLoadingThread2::ShutdownLoading()
+{
+	Impl->ShutdownLoading();
+}
+
+void FAsyncLoadingThread2::StartThread()
+{
+	Impl->StartThread();
+}
+
+bool FAsyncLoadingThread2::IsMultithreaded()
+{
+	return Impl->IsMultithreaded();
+}
+
+bool FAsyncLoadingThread2::IsInAsyncLoadThread()
+{
+	return Impl->IsInAsyncLoadThread();
+}
+
+void FAsyncLoadingThread2::NotifyConstructedDuringAsyncLoading(UObject* Object, bool bSubObject)
+{
+	Impl->NotifyConstructedDuringAsyncLoading(Object, bSubObject);
+}
+
+void FAsyncLoadingThread2::FireCompletedCompiledInImport(FGCObject* AsyncPackage, FPackageIndex Import)
+{
+	Impl->FireCompletedCompiledInImport(AsyncPackage, Import);
+}
+
+int32 FAsyncLoadingThread2::LoadPackage(const FString& InPackageName, const FGuid* InGuid, const TCHAR* InPackageToLoadFrom, FLoadPackageAsyncDelegate InCompletionDelegate, EPackageFlags InPackageFlags, int32 InPIEInstanceID, int32 InPackagePriority)
+{
+	return Impl->LoadPackage(InPackageName, InGuid, InPackageToLoadFrom, InCompletionDelegate, InPackageFlags, InPIEInstanceID, InPackagePriority);
+}
+
+EAsyncPackageState::Type FAsyncLoadingThread2::ProcessLoading(bool bUseTimeLimit, bool bUseFullTimeLimit, float TimeLimit)
+{
+	return Impl->ProcessLoadingFromGameThread(bUseTimeLimit, bUseFullTimeLimit, TimeLimit);
+}
+
+EAsyncPackageState::Type FAsyncLoadingThread2::ProcessLoadingUntilComplete(TFunctionRef<bool()> CompletionPredicate, float TimeLimit)
+{
+	return Impl->ProcessLoadingUntilCompleteFromGameThread(CompletionPredicate, TimeLimit);
+}
+
+void FAsyncLoadingThread2::CancelLoading()
+{
+	Impl->CancelLoading();
+}
+
+void FAsyncLoadingThread2::SuspendLoading()
+{
+	Impl->SuspendLoading();
+}
+
+void FAsyncLoadingThread2::ResumeLoading()
+{
+	Impl->ResumeLoading();
+}
+
+void FAsyncLoadingThread2::FlushLoading(int32 PackageId)
+{
+	Impl->FlushLoading(PackageId);
+}
+
+int32 FAsyncLoadingThread2::GetNumAsyncPackages()
+{
+	return Impl->GetNumAsyncPackages();
+}
+
+float FAsyncLoadingThread2::GetAsyncLoadPercentage(const FName& PackageName)
+{
+	return Impl->GetAsyncLoadPercentage(PackageName);
+}
+
+bool FAsyncLoadingThread2::IsAsyncLoadingSuspended()
+{
+	return Impl->IsAsyncLoadingSuspended();
+}
+
+bool FAsyncLoadingThread2::IsAsyncLoadingPackages()
+{
+	return Impl->IsAsyncLoadingPackages();
+}
