Copying //UE4/Dev-Core to Main (//UE4/Main)

diff --git a/Engine/Source/Runtime/CoreUObject/Private/Serialization/AsyncLoading2.h b/Engine/Source/Runtime/CoreUObject/Private/Serialization/AsyncLoading2.h
new file mode 100644
--- /dev/null
+++ b/Engine/Source/Runtime/CoreUObject/Private/Serialization/AsyncLoading2.h
@@ -0,0 +1,43 @@
+// Copyright 1998-2019 Epic Games, Inc. All Rights Reserved.
+
+/*=============================================================================
+	AsyncLoading.h: Unreal async loading definitions.
+=============================================================================*/
+
+#pragma once
+
+#include "CoreMinimal.h"
+#include "Serialization/AsyncPackageLoader.h"
+
+class FAsyncLoadingThread2Impl;
+
+class FAsyncLoadingThread2
+	: public IAsyncPackageLoader
+{
+public:
+	FAsyncLoadingThread2(IEDLBootNotificationManager& InEDLBootNotificationManager);
+	virtual ~FAsyncLoadingThread2();
+
+	void InitializeLoading() override;
+	void ShutdownLoading() override;
+	void StartThread() override;
+	bool IsMultithreaded() override;
+	bool IsInAsyncLoadThread() override;
+	void NotifyConstructedDuringAsyncLoading(UObject* Object, bool bSubObject) override;
+	void FireCompletedCompiledInImport(FGCObject* AsyncPackage, FPackageIndex Import) override;
+	int32 LoadPackage(const FString& InPackageName, const FGuid* InGuid, const TCHAR* InPackageToLoadFrom, FLoadPackageAsyncDelegate InCompletionDelegate, EPackageFlags InPackageFlags, int32 InPIEInstanceID, int32 InPackagePriority) override;
+	EAsyncPackageState::Type ProcessLoading(bool bUseTimeLimit, bool bUseFullTimeLimit, float TimeLimit) override;
+	EAsyncPackageState::Type ProcessLoadingUntilComplete(TFunctionRef<bool()> CompletionPredicate, float TimeLimit) override;
+	void CancelLoading() override;
+	void SuspendLoading() override;
+	void ResumeLoading() override;
+	void FlushLoading(int32 PackageId) override;
+	int32 GetNumAsyncPackages() override;
+	float GetAsyncLoadPercentage(const FName& PackageName) override;
+	bool IsAsyncLoadingSuspended() override;
+	bool IsAsyncLoadingPackages() override;
+
+private:
+	FAsyncLoadingThread2Impl* Impl;
+};
+
diff --git a/Engine/Source/Runtime/CoreUObject/Private/Serialization/AsyncLoadingThread.h b/Engine/Source/Runtime/CoreUObject/Private/Serialization/AsyncLoadingThread.h
--- a/Engine/Source/Runtime/CoreUObject/Private/Serialization/AsyncLoadingThread.h
+++ b/Engine/Source/Runtime/CoreUObject/Private/Serialization/AsyncLoadingThread.h
@@ -17,6 +17,9 @@
 #include "HAL/ThreadSafeBool.h"
 #include "Misc/ConfigCacheIni.h"
 #include "HAL/LowLevelMemTracker.h"
+#include "Serialization/AsyncPackageLoader.h"
+
+struct FPrecacheCallbackHandler;
 
 /** [EDL] Event Driven loader event */
 struct FAsyncLoadEvent
@@ -132,16 +135,19 @@ struct FFlushTree
 	}
 };
 
+/** Holds the maximum package summary size that can be set via ini files
+  * This is used for the initial precache and should be large enough to hold the actual Sum.TotalHeaderSize 
+  */
 struct FMaxPackageSummarySize
 {
-	int32 Value;
-	FMaxPackageSummarySize();
+	static int32 Value;
+	static void Init();
 };
 
 /**
  * Async loading thread. Preloads/serializes packages on async loading thread. Postloads objects on the game thread.
  */
-class FAsyncLoadingThread : public FRunnable
+class FAsyncLoadingThread : public FRunnable, public IAsyncPackageLoader
 {
 	/** Structure that holds the async loading thread ini settings */
 	struct FAsyncLoadingThreadSettings
@@ -152,6 +158,8 @@ class FAsyncLoadingThread : public FRunnable
 		FAsyncLoadingThreadSettings();
 	};
 
+	IEDLBootNotificationManager& EDLBootNotificationManager;
+
 	/** Thread to run the worker FRunnable on */
 	FRunnableThread* Thread;
 	/** Stops this thread */
@@ -205,8 +213,6 @@ class FAsyncLoadingThread : public FRunnable
 public:
 	/** [EDL] Async Packages that are ready for tick */
 	TArray<FAsyncPackage*> AsyncPackagesReadyForTick;
-	/* This is used for the initial precache and should be large enough to find the actual Sum.TotalHeaderSize */
-	const FMaxPackageSummarySize MaxPackageSummarySize;
 private:
 
 #if THREADSAFE_UOBJECTS
@@ -228,6 +234,11 @@ private:
 
 	FThreadSafeCounter AsyncThreadReady;
 
+	/** When cancelling async loading: list of package requests to cancel */
+	TArray<FAsyncPackageDesc*> QueuedPackagesToCancel;
+	/** When cancelling async loading: list of packages to cancel */
+	TSet<FAsyncPackage*> PackagesToCancel;
+
 	/** Async loading thread ID */
 	static uint32 AsyncLoadingThreadID;
 
@@ -242,12 +253,11 @@ private:
 	/** Thread safe counter used to accumulate cycles spent on blocking. Using stats may generate to many stats messages. */
 	static FThreadSafeCounter BlockingCycles;
 #endif
+public:
 
-	FAsyncLoadingThread();
+	FAsyncLoadingThread(int32 InThreadIndex, IEDLBootNotificationManager& InEDLBootNotificationManager);
 	virtual ~FAsyncLoadingThread();
 
-public:
-
 	//~ Begin FRunnable Interface.
 	virtual bool Init();
 	virtual uint32 Run();
@@ -255,10 +265,33 @@ public:
 	//~ End FRunnable Interface
 
 	/** Returns the async loading thread singleton */
-	static FAsyncLoadingThread& Get();
+	static FAsyncLoadingThread& Get()
+	{
+		check(Instance);
+		return *Instance;
+	}
 
 	/** Start the async loading thread */
-	void StartThread();
+	void StartThread() override;
+
+	int32 LoadPackage(
+			const FString& InPackageName,
+			const FGuid* InGuid,
+			const TCHAR* InPackageToLoadFrom,
+			FLoadPackageAsyncDelegate InCompletionDelegate,
+			EPackageFlags InPackageFlags,
+			int32 InPIEInstanceID,
+			int32 InPackagePriority) override;
+
+	EAsyncPackageState::Type ProcessLoading(bool bUseTimeLimit, bool bUseFullTimeLimit, float TimeLimit) override;
+
+	EAsyncPackageState::Type ProcessLoadingUntilComplete(TFunctionRef<bool()> CompletionPredicate, float TimeLimit) override;
+
+	void FlushLoading(int32 PackageId) override;
+
+	void NotifyConstructedDuringAsyncLoading(UObject* Object, bool bSubObject) override;
+
+	void FireCompletedCompiledInImport(FGCObject* AsyncPacakge, FPackageIndex Import) override;
 
 	/** [EDL] Event queue */
 	FAsyncLoadEventQueue EventQueue;
@@ -300,38 +333,43 @@ public:
 	static FAsyncLoadingThreadSettings& GetAsyncLoadingThreadSettings();
 
 	/** True if multithreaded async loading is currently being used. */
-	static FORCEINLINE bool IsMultithreaded()
+	FORCEINLINE bool IsMultithreaded() override
 	{
 		return bThreadStarted;
 	}
 
 	/** Sets the current state of async loading */
-	void EnterAsyncLoadingTick()
+	static void EnterAsyncLoadingTick(int32 ThreadIndex)
 	{
-		AsyncLoadingTickCounter++;
+		AsyncLoadingTickCounter.Increment();
+		check(CurrentAsyncLoadingTickThreadIndex == INDEX_NONE || CurrentAsyncLoadingTickThreadIndex == ThreadIndex);
+		CurrentAsyncLoadingTickThreadIndex = ThreadIndex;
 	}
 
-	void LeaveAsyncLoadingTick()
+	static void LeaveAsyncLoadingTick(int32 ThreadIndex)
 	{
-		AsyncLoadingTickCounter--;
-		check(AsyncLoadingTickCounter >= 0);
+		int32 AsyncLoadingTickCounterValue = AsyncLoadingTickCounter.Decrement();
+		check(AsyncLoadingTickCounterValue >= 0);
+		check(CurrentAsyncLoadingTickThreadIndex == ThreadIndex);
+		if (AsyncLoadingTickCounterValue == 0)
+		{
+			CurrentAsyncLoadingTickThreadIndex = INDEX_NONE;
+		}
 	}
 
 	/** Gets the current state of async loading */
-	FORCEINLINE bool GetIsInAsyncLoadingTick() const
+	static bool GetIsInAsyncLoadingTick()
 	{
-		return !!AsyncLoadingTickCounter;
+		return !!AsyncLoadingTickCounter.GetValue();
 	}
 
-	/** Returns true if packages are currently being loaded on the async thread */
-	FORCEINLINE bool IsAsyncLoadingPackages()
+	FORCEINLINE bool IsAsyncLoadingPackages() override
 	{
 		FPlatformMisc::MemoryBarrier();
 		return QueuedPackagesCounter.GetValue() != 0 || ExistingAsyncPackagesCounter.GetValue() != 0;
 	}
 
-	/** Returns true this codes runs on the async loading thread */
-	static FORCEINLINE bool IsInAsyncLoadThread()
+	FORCEINLINE bool IsInAsyncLoadThread() override
 	{
 		bool bResult = false;
 		if (IsMultithreaded())
@@ -340,17 +378,17 @@ public:
 			// we're on game thread but inside of async loading code (PostLoad mostly)
 			// to make it behave exactly like the non-threaded version
 			bResult = FPlatformTLS::GetCurrentThreadId() == AsyncLoadingThreadID || 
-				(IsInGameThread() && FAsyncLoadingThread::Get().GetIsInAsyncLoadingTick());
+				(IsInGameThread() && GetIsInAsyncLoadingTick());
 		}
 		else
 		{
-			bResult = IsInGameThread() && FAsyncLoadingThread::Get().GetIsInAsyncLoadingTick();
+			bResult = IsInGameThread() && GetIsInAsyncLoadingTick();
 		}
 		return bResult;
 	}
 
 	/** Returns true if async loading is suspended */
-	FORCEINLINE bool IsAsyncLoadingSuspended()
+	FORCEINLINE bool IsAsyncLoadingSuspended() override
 	{
 		FPlatformMisc::MemoryBarrier();
 		return IsLoadingSuspended.GetValue() != 0;
@@ -363,7 +401,7 @@ public:
 	}
 
 	/** Returns the number of async packages that are currently being processed */
-	FORCEINLINE int32 GetAsyncPackagesCount()
+	FORCEINLINE int32 GetNumAsyncPackages() override
 	{
 		FPlatformMisc::MemoryBarrier();
 		return ExistingAsyncPackagesCounter.GetValue();
@@ -411,22 +449,22 @@ public:
 	/**
 	* [GAME THREAD] Cancels streaming
 	*/
-	void CancelAsyncLoading();
+	void CancelLoading() override;
 
 	/**
 	* [GAME THREAD] Stops the async loading thread and blocks until the thread has exited.
 	*/
-	void Kill();
+	void ShutdownLoading() override;
 
 	/**
 	* [GAME THREAD] Suspends async loading thread
 	*/
-	void SuspendLoading();
+	void SuspendLoading() override;
 
 	/**
 	* [GAME THREAD] Resumes async loading thread
 	*/
-	void ResumeLoading();
+	void ResumeLoading() override;
 
 	/**
 	* [ASYNC/GAME THREAD] Queues a package for streaming.
@@ -478,14 +516,14 @@ public:
 	EAsyncPackageState::Type TickAsyncThread(bool bUseTimeLimit, bool bUseFullTimeLimit, float TimeLimit, bool& bDidSomething, FFlushTree* FlushTree = nullptr);
 
 	/** Initializes async loading thread */
-	void InitializeAsyncThread();
+	void InitializeLoading() override;
 
 	/** 
 	 * [GAME THREAD] Gets the load percentage of the specified package
 	 * @param PackageName Name of the package to return async load percentage for
 	 * @return Percentage (0-100) of the async package load or -1 of package has not been found
 	 */
-	float GetAsyncLoadPercentage(const FName& PackageName);
+	float GetAsyncLoadPercentage(const FName& PackageName) override;
 
 	/** 
 	 * [ASYNC/GAME THREAD] Checks if a request ID already is added to the loading queue
@@ -582,7 +620,42 @@ private:
 
 	/** Cancels async loading internally */
 	void CancelAsyncLoadingInternal();
+	void FinalizeCancelAsyncLoadingInternal();
+
+	/** Event graph for EDL */
+	FEventLoadGraph EventGraph;	
+	/** ELD precache handler */
+	FPrecacheCallbackHandler* PrecacheHandler;
+	/** This Async Loading Thread Index (future use) */
+	int32 AsyncLoadingThreadIndex;
 
 	/** Number of times we re-entered the async loading tick, mostly used by singlethreaded ticking. Debug purposes only. */
-	int32 AsyncLoadingTickCounter;
+	static FThreadSafeCounter AsyncLoadingTickCounter;
+	static int32 CurrentAsyncLoadingTickThreadIndex;
+
+public:
+
+	/** Gets the EDL event graph */
+	FEventLoadGraph& GetEventGraph()
+	{
+		return EventGraph;
+	}
+
+	/** Gets the EDL precache handler */
+	FPrecacheCallbackHandler& GetPrecacheHandler()
+	{
+		return *PrecacheHandler;
+	}
+
+	/** Gets this ALT index (future use) */
+	int32 GetThreadIndex() const
+	{
+		return AsyncLoadingThreadIndex;
+	}
+
+#if !UE_BUILD_SHIPPING && !UE_BUILD_TEST
+	FThreadSafeCounter RecursionNotAllowed;
+#endif
+
+	static FAsyncLoadingThread* Instance;
 };
diff --git a/Engine/Source/Runtime/CoreUObject/Private/Serialization/AsyncPackageLoader.cpp b/Engine/Source/Runtime/CoreUObject/Private/Serialization/AsyncPackageLoader.cpp
new file mode 100644
--- /dev/null
+++ b/Engine/Source/Runtime/CoreUObject/Private/Serialization/AsyncPackageLoader.cpp
@@ -0,0 +1,679 @@
+// Copyright 1998-2019 Epic Games, Inc. All Rights Reserved.
+
+#include "Serialization/AsyncPackageLoader.h"
+#include "Serialization/AsyncLoadingThread.h"
+#include "Serialization/AsyncLoading2.h"
+#include "UObject/GCObject.h"
+#include "UObject/LinkerLoad.h"
+#include "Misc/CoreDelegates.h"
+
+volatile int32 GIsLoaderCreated;
+TUniquePtr<IAsyncPackageLoader> GPackageLoader;
+
+const FName PrestreamPackageClassNameLoad = FName("PrestreamPackage");
+
+struct FEDLBootObjectState
+{
+	ENotifyRegistrationType NotifyRegistrationType;
+	ENotifyRegistrationPhase LastNotifyRegistrationPhase;
+	UObject *(*Register)();
+	bool bDynamic;
+};
+
+struct FEDLBootWaitingPackage
+{
+	FGCObject* Package;
+	FPackageIndex Import;
+};
+
+struct FEDLBootNotificationManager
+	: public IEDLBootNotificationManager
+{
+	TMap<FName, FEDLBootObjectState> PathToState;
+	TMultiMap<FName, FEDLBootWaitingPackage> PathToWaitingPackageNodes;
+	TArray<FName> PathsToFire;
+	TArray<UClass*> CDORecursiveStack;
+	TArray<UClass*> CDORecursives;
+	FCriticalSection EDLBootNotificationManagerLock;
+
+	// return true if you are waiting for this compiled in object
+	bool AddWaitingPackage(FGCObject* Pkg, FName PackageName, FName ObjectName, FPackageIndex Import) override
+	{
+		if (PackageName == GLongCoreUObjectPackageName)
+		{
+			return false; // We assume nothing in coreuobject ever loads assets in a constructor
+		}
+		FScopeLock Lock(&EDLBootNotificationManagerLock);
+		check(GIsInitialLoad);
+		check(Import.IsImport()); // compiled in exports make no sense
+		FString ObjectNameString = ObjectName.ToString();
+		FName LongFName(*(PackageName.ToString() / ObjectNameString));
+		check(LongFName != NAME_None);
+		FName WaitName = LongFName;
+		FEDLBootObjectState* ExistingState = PathToState.Find(LongFName);
+		if (!ExistingState)
+		{
+			//if (ObjectName.ToString().EndsWith(HEADER_GENERATED_DELEGATE_SIGNATURE_SUFFIX))
+			// there are also some arg structs and other things that are just part of the package with no registration
+			{
+				ExistingState = PathToState.Find(PackageName);
+				WaitName = PackageName;
+			}
+			if (!ExistingState)
+			{
+				UE_LOG(LogStreaming, Fatal, TEXT("Compiled in export %s not found; it was never registered."), *LongFName.ToString());
+				return false;
+			}
+		}
+		if (ExistingState->LastNotifyRegistrationPhase == ENotifyRegistrationPhase::NRP_Finished)
+		{
+			return false;
+		}
+		FEDLBootWaitingPackage WaitingPackage;
+		WaitingPackage.Package = Pkg;
+		WaitingPackage.Import = Import;
+
+		PathToWaitingPackageNodes.Add(WaitName, WaitingPackage);
+
+		return true;
+	}
+
+	void NotifyRegistrationEvent(const TCHAR* PackageName, const TCHAR* Name, ENotifyRegistrationType NotifyRegistrationType, ENotifyRegistrationPhase NotifyRegistrationPhase, UObject *(*InRegister)(), bool InbDynamic)
+	{
+		if (!GIsInitialLoad)
+		{
+			return;
+		}
+		static FName LongCoreUObjectPackageName(TEXT("/Script/CoreUObject")); // can't use the global here because it may not be initialized yet
+		FName PackageFName(PackageName);
+		if (PackageFName == LongCoreUObjectPackageName)
+		{
+			return; // We assume nothing in coreuobject ever loads assets in a constructor
+		}
+
+		FScopeLock Lock(&EDLBootNotificationManagerLock);
+
+		FName LongFName(*(FString(PackageName) / Name));
+
+		//FPlatformMisc::LowLevelOutputDebugStringf(TEXT("NotifyRegistrationEvent %s %d %d\r\n"), *LongFName.ToString(), int32(NotifyRegistrationType), int32(NotifyRegistrationPhase));
+
+		// some things, like delegate signatures, are not registered; rather they are part of the package singleton, so we track the package state as being the max of any member of that package
+		FEDLBootObjectState* ExistingPackageState = PathToState.Find(PackageFName);
+		FEDLBootObjectState* ExistingState = PathToState.Find(LongFName);
+
+		if (!ExistingState)
+		{
+			if (NotifyRegistrationPhase != ENotifyRegistrationPhase::NRP_Added)
+			{
+				UE_LOG(LogStreaming, Fatal, TEXT("Attempt to process %s before it has been added."), *LongFName.ToString());
+			}
+			FEDLBootObjectState NewState;
+			NewState.LastNotifyRegistrationPhase = NotifyRegistrationPhase;
+			NewState.NotifyRegistrationType = NotifyRegistrationType;
+			NewState.Register = InRegister;
+			NewState.bDynamic = InbDynamic;
+			PathToState.Add(LongFName, NewState);
+
+			if (!ExistingPackageState)
+			{
+				NewState.NotifyRegistrationType = ENotifyRegistrationType::NRT_Package;
+				PathToState.Add(PackageFName, NewState);
+			}
+		}
+		else
+		{
+			if (int32(ExistingState->LastNotifyRegistrationPhase) + 1 != int32(NotifyRegistrationPhase))
+			{
+				UE_CLOG(GEventDrivenLoaderEnabled, LogStreaming, Fatal, TEXT("Invalid state transition %d %d with %s when it has already been processed."), int32(ExistingState->LastNotifyRegistrationPhase), int32(NotifyRegistrationPhase), *LongFName.ToString());
+			}
+			if (ExistingState->NotifyRegistrationType != (NotifyRegistrationType))
+			{
+				UE_CLOG(GEventDrivenLoaderEnabled, LogStreaming, Fatal, TEXT("Multiple types %d %d with %s when it has already been processed."), int32(ExistingState->NotifyRegistrationType), int32(NotifyRegistrationType), *LongFName.ToString());
+			}
+			ExistingState->LastNotifyRegistrationPhase = NotifyRegistrationPhase;
+			if (NotifyRegistrationPhase == ENotifyRegistrationPhase::NRP_Finished)
+			{
+				ExistingState->Register = nullptr; // we don't need to do this in ConstructWaitingBootObjects
+				PathsToFire.Add(LongFName);
+			}
+			check(ExistingPackageState); // if we have an existing state for the thing, we should also have a 
+			if (ExistingPackageState && int32(NotifyRegistrationPhase) > int32(ExistingPackageState->LastNotifyRegistrationPhase))
+			{
+				ExistingPackageState->LastNotifyRegistrationPhase = NotifyRegistrationPhase;
+				if (NotifyRegistrationPhase == ENotifyRegistrationPhase::NRP_Finished)
+				{
+					//FPlatformMisc::LowLevelOutputDebugStringf(TEXT("Fired package %s %d %d\r\n"), *PackageFName.ToString(), int32(NotifyRegistrationType), int32(NotifyRegistrationPhase));
+					PathsToFire.Add(PackageFName);
+				}
+			}
+		}
+	}
+
+	void NotifyRegistrationComplete()
+	{
+#if USE_EVENT_DRIVEN_ASYNC_LOAD_AT_BOOT_TIME
+		FireCompletedCompiledInImports(true);
+		FlushAsyncLoading();
+		GPackageLoader->StartThread();
+#endif
+#if !HACK_HEADER_GENERATOR
+		check(!GIsInitialLoad && IsInGameThread());
+		FScopeLock Lock(&EDLBootNotificationManagerLock);
+		for (auto& Pair : PathToState)
+		{
+			if (Pair.Value.LastNotifyRegistrationPhase != ENotifyRegistrationPhase::NRP_Finished && !Pair.Value.bDynamic)
+			{
+#if USE_EVENT_DRIVEN_ASYNC_LOAD_AT_BOOT_TIME
+				UE_CLOG(GEventDrivenLoaderEnabled, LogStreaming, Fatal, TEXT("%s (%d) was not complete (%d) after registration was complete."), *Pair.Key.ToString(), int32(Pair.Value.NotifyRegistrationType), int32(Pair.Value.LastNotifyRegistrationPhase));
+#else
+				UE_LOG(LogStreaming, Warning, TEXT("%s was not complete (%d) after registration was complete."), *Pair.Key.ToString(), int32(Pair.Value.LastNotifyRegistrationPhase));
+#endif
+			}
+		}
+		if (PathToWaitingPackageNodes.Num())
+		{
+			UE_LOG(LogStreaming, Fatal, TEXT("Initial load is complete, but we still have %d waiting packages."), PathToWaitingPackageNodes.Num());
+		}
+		if (GEventDrivenLoaderEnabled && PathsToFire.Num() && USE_EVENT_DRIVEN_ASYNC_LOAD_AT_BOOT_TIME)
+		{
+			for (FName Path : PathsToFire)
+			{
+				UE_LOG(LogStreaming, Error, TEXT("%s was not fired."), *Path.ToString());
+			}
+			UE_LOG(LogStreaming, Fatal, TEXT("Initial load is complete, but we still have %d imports to fire (listed above)."), PathsToFire.Num());
+		}
+#endif
+		PathToState.Empty();
+		PathsToFire.Empty();
+	}
+
+	bool ConstructWaitingBootObjects() override
+	{
+		static struct FFixedBootOrder
+		{
+			TArray<FName> Array;
+			FFixedBootOrder()
+			{
+				// look for any packages that we want to force preload at startup
+				FConfigSection* BootObjects = GConfig->GetSectionPrivate(TEXT("/Script/Engine.StreamingSettings"), false, true, GEngineIni);
+				if (BootObjects)
+				{
+					// go through list and add to the array
+					for (FConfigSectionMap::TIterator It(*BootObjects); It; ++It)
+					{
+						if (It.Key() == TEXT("FixedBootOrder"))
+						{
+							// add this package to the list to be fully loaded later
+							Array.Add(FName(*It.Value().GetValue()));
+						}
+					}
+				}
+			}
+
+		} FixedBootOrder;
+
+		check(GIsInitialLoad && IsInGameThread());
+		UObject *(*BootObjectRegister)() = nullptr;
+		UObject *(*BootPackageObjectRegister)() = nullptr;
+		FName WaitingPackage;
+		bool bIsCDO = false;
+
+		while (FixedBootOrder.Array.Num())
+		{
+			FName ThisItem = FixedBootOrder.Array.Pop();
+			FScopeLock Lock(&EDLBootNotificationManagerLock);
+			FEDLBootObjectState* ExistingState = PathToState.Find(ThisItem);
+
+			if (!ExistingState)
+			{
+				UE_LOG(LogStreaming, Fatal, TEXT("%s was listed as a fixed load order but was not found,"), *ThisItem.ToString());
+			}
+			else if (!ExistingState->Register)
+			{
+				UE_LOG(LogStreaming, Log, TEXT("%s was listed as a fixed load order but was already processed"), *ThisItem.ToString());
+			}
+			else
+			{
+				//FPlatformMisc::LowLevelOutputDebugStringf(TEXT("Booting Fixed %s %d\r\n"), *ThisItem.ToString(), int32(ExistingState->NotifyRegistrationType));
+				BootObjectRegister = ExistingState->Register;
+				ExistingState->Register = nullptr; // we don't need to do this more than once
+				bIsCDO = ExistingState->NotifyRegistrationType == ENotifyRegistrationType::NRT_ClassCDO;
+				break;
+			}
+		}
+
+		if (!BootObjectRegister)
+		{
+			FScopeLock Lock(&EDLBootNotificationManagerLock);
+			for (auto& Pair : PathToWaitingPackageNodes)
+			{
+				FEDLBootObjectState* ExistingState = PathToState.Find(Pair.Key);
+				if (ExistingState)
+				{
+					if (ExistingState->Register)
+					{
+						//FPlatformMisc::LowLevelOutputDebugStringf(TEXT("Booting %s %d\r\n"), *Pair.Key.ToString(), int32(ExistingState->NotifyRegistrationType));
+						BootObjectRegister = ExistingState->Register;
+						ExistingState->Register = nullptr; // we don't need to do this more than once
+						bIsCDO = ExistingState->NotifyRegistrationType == ENotifyRegistrationType::NRT_ClassCDO;
+						break;
+					}
+				}
+			}
+		}
+		if (BootObjectRegister)
+		{
+			UObject* BootObject = BootObjectRegister();
+			check(BootObject);
+			UObjectForceRegistration(BootObject);
+			if (bIsCDO)
+			{
+				UClass* Class = CastChecked<UClass>(BootObject);
+				bool bAnyParentOnStack = false;
+				UClass* Super = Class;
+				while (Super)
+				{
+					if (CDORecursiveStack.Contains(Super))
+					{
+						bAnyParentOnStack = true;
+						break;
+					}
+					Super = Super->GetSuperClass();
+				}
+
+				if (!bAnyParentOnStack)
+				{
+					CDORecursiveStack.Push(Class);
+					//FPlatformMisc::LowLevelOutputDebugStringf(TEXT("Create CDO %s\r\n"), *BootObject->GetName());
+					Class->GetDefaultObject();
+					verify(CDORecursiveStack.Pop() == Class);
+				}
+				else
+				{
+					//FPlatformMisc::LowLevelOutputDebugStringf(TEXT("Recursive Deferred %s\r\n"), *BootObject->GetName());
+					CDORecursives.Add(Class);
+				}
+			}
+			return true;
+		}
+		if (CDORecursives.Num())
+		{
+			UClass* OkToRun = nullptr;
+			for (UClass* Class : CDORecursives)
+			{
+				bool bAnyParentOnStack = false;
+				UClass* Super = Class;
+				while (Super)
+				{
+					if (CDORecursiveStack.Contains(Super))
+					{
+						bAnyParentOnStack = true;
+						break;
+					}
+					Super = Super->GetSuperClass();
+				}
+				if (!bAnyParentOnStack)
+				{
+					OkToRun = Class;
+					break;
+				}
+			}
+			if (OkToRun)
+			{
+				//FPlatformMisc::LowLevelOutputDebugStringf(TEXT("CDORecursives %s\r\n"), *OkToRun->GetName());
+				CDORecursives.Remove(OkToRun);
+				CDORecursiveStack.Push(OkToRun);
+				OkToRun->GetDefaultObject();
+				verify(CDORecursiveStack.Pop() == OkToRun);
+			}
+			else
+			{
+				FPlatformProcess::Sleep(.001f);
+			}
+			return true; // even if we didn't do anything we need to return true to avoid checking for cycles
+		}
+		return false;
+	}
+
+	bool IsWaitingForSomething() override
+	{
+		FScopeLock Lock(&EDLBootNotificationManagerLock);
+		return PathToWaitingPackageNodes.Num() > 0;
+	}
+
+	bool IsObjComplete(UObject* Obj)
+	{
+		static FName LongCoreUObjectPackageName(TEXT("/Script/CoreUObject")); // can't use the global here because it may not be initialized yet
+		FName PackageName = Obj->GetOutermost()->GetFName();
+		if (PackageName == LongCoreUObjectPackageName)
+		{
+			return true; // We assume nothing in coreuobject ever loads assets in a constructor, therefore it can be considered complete
+		}
+		FScopeLock Lock(&EDLBootNotificationManagerLock);
+		FName LongFName(*(PackageName.ToString() / Obj->GetName()));
+
+		FEDLBootObjectState* ExistingState = PathToState.Find(LongFName);
+
+		if (!ExistingState || ExistingState->LastNotifyRegistrationPhase == ENotifyRegistrationPhase::NRP_Finished)
+		{
+			return true;
+		}
+		return false;
+	}
+
+	bool FireCompletedCompiledInImports(bool bFinalRun) override
+	{
+#if USE_EVENT_DRIVEN_ASYNC_LOAD_AT_BOOT_TIME
+		FScopeLock Lock(&EDLBootNotificationManagerLock);
+		check(bFinalRun || GIsInitialLoad);
+		bool bResult = !!PathsToFire.Num();
+		for (FName LongName : PathsToFire)
+		{
+			for (auto It = PathToWaitingPackageNodes.CreateKeyIterator(LongName); It; ++It)
+			{
+				FEDLBootWaitingPackage& WaitingPackage = It.Value();
+				GPackageLoader->FireCompletedCompiledInImport(WaitingPackage.Package, WaitingPackage.Import);
+			}
+			PathToWaitingPackageNodes.Remove(LongName);
+		}
+		PathsToFire.Empty();
+		return bResult;
+#else
+		return false;
+#endif
+	}
+};
+
+static FEDLBootNotificationManager& GetGEDLBootNotificationManager()
+{
+	static FEDLBootNotificationManager Singleton;
+	return Singleton;
+}
+
+bool IsEventDrivenLoaderEnabledInCookedBuilds()
+{
+	static struct FEventDrivenLoaderEnabledInCookedBuildsInit
+	{
+		bool bEventDrivenLoaderEnabled;
+		FEventDrivenLoaderEnabledInCookedBuildsInit()
+			: bEventDrivenLoaderEnabled(false)
+		{
+			SetEventDrivenLoaderEnabled();
+		}
+
+		void SetEventDrivenLoaderEnabled()
+		{
+			check(GConfig || IsEngineExitRequested());
+			if (GConfig)
+			{
+				GConfig->GetBool(TEXT("/Script/Engine.StreamingSettings"), TEXT("s.EventDrivenLoaderEnabled"), bEventDrivenLoaderEnabled, GEngineIni);
+#if !UE_BUILD_SHIPPING
+				if (FParse::Param(FCommandLine::Get(), TEXT("NOEDL")))
+				{
+					bEventDrivenLoaderEnabled = false;
+				}
+#endif
+			}
+		}
+	} EventDrivenLoaderEnabledInCookedBuilds;
+
+#if WITH_EDITOR	
+	// when building from the UE4 Editor, s.EventDrivenLoaderEnabled can be changed from Project Settings, so we need to test it at every call
+	if (GIsEditor)
+	{
+		EventDrivenLoaderEnabledInCookedBuilds.SetEventDrivenLoaderEnabled();
+	}
+#endif
+	return EventDrivenLoaderEnabledInCookedBuilds.bEventDrivenLoaderEnabled;
+}
+
+bool IsEventDrivenLoaderEnabled()
+{
+	static struct FEventDrivenLoaderEnabledInit
+	{
+		FEventDrivenLoaderEnabledInit()
+		{
+			GEventDrivenLoaderEnabled = IsEventDrivenLoaderEnabledInCookedBuilds() && FPlatformProperties::RequiresCookedData();
+		}
+	} EventDrivenLoaderEnabledInit;
+	return GEventDrivenLoaderEnabled;
+}
+
+bool IsFullyLoadedObj(UObject* Obj)
+{
+	if (!Obj)
+	{
+		return false;
+	}
+	if (Obj->HasAllFlags(RF_WasLoaded | RF_LoadCompleted)
+		|| Obj->IsA(UPackage::StaticClass())) // packages are never really loaded, so if it exists, it is loaded
+	{
+		return true;
+	}
+	if (Obj->HasAnyFlags(RF_WasLoaded | RF_NeedLoad | RF_WillBeLoaded))
+	{
+		return false;
+	}
+	if (GIsInitialLoad && Obj->GetOutermost()->HasAnyPackageFlags(PKG_CompiledIn))
+	{
+		return GetGEDLBootNotificationManager().IsObjComplete(Obj);
+	}
+	//native blueprint 
+	UDynamicClass* UD = Cast<UDynamicClass>(Obj);
+	if (!UD)
+	{
+		return true;
+	}
+
+	if (GEventDrivenLoaderEnabled && EVENT_DRIVEN_ASYNC_LOAD_ACTIVE_AT_RUNTIME)
+	{
+		if (0 != (UD->ClassFlags & CLASS_Constructed))
+		{
+			return true;
+		}
+	}
+	else
+	{
+		if (UD->GetDefaultObject(false))
+		{
+			UE_CLOG(!UD->HasAnyClassFlags(CLASS_TokenStreamAssembled), LogStreaming, Fatal, TEXT("Class %s is fully loaded, but does not have its token stream assembled."), *UD->GetFullName());
+			return true;
+		}
+	}
+	return false;
+}
+
+bool IsNativeCodePackage(UPackage* Package)
+{
+	if (!Package || !Package->HasAnyPackageFlags(PKG_CompiledIn))
+	{
+		return false;
+	}
+
+	// Make sure it isn't a dynamically loaded one, this check is slower
+	return !GetConvertedDynamicPackageNameToTypeName().Contains(Package->GetFName());
+}
+
+/** Checks if the object can have PostLoad called on the Async Loading Thread */
+bool CanPostLoadOnAsyncLoadingThread(UObject* Object)
+{
+	if (Object->IsPostLoadThreadSafe())
+	{
+		bool bCanPostLoad = true;
+		// All outers should also be safe to call PostLoad on ALT
+		for (UObject* Outer = Object->GetOuter(); Outer && bCanPostLoad; Outer = Outer->GetOuter())
+		{
+			bCanPostLoad = !Outer->HasAnyFlags(RF_NeedPostLoad) || Outer->IsPostLoadThreadSafe();
+		}
+		return bCanPostLoad;
+	}
+	return false;
+}
+
+IAsyncPackageLoader& GetAsyncPackageLoader()
+{
+	check(GPackageLoader.Get());
+	return *GPackageLoader;
+}
+
+void InitAsyncThread()
+{
+	if (FParse::Param(FCommandLine::Get(), TEXT("zenloader")))
+	{
+		GPackageLoader = MakeUnique<FAsyncLoadingThread2>(GetGEDLBootNotificationManager());
+	}
+	else
+	{
+		GPackageLoader = MakeUnique<FAsyncLoadingThread>(/** ThreadIndex = */ 0, GetGEDLBootNotificationManager());
+	}
+
+	FPlatformAtomics::InterlockedIncrement(&GIsLoaderCreated);
+
+	FCoreDelegates::OnSyncLoadPackage.AddStatic([](const FString&) { GSyncLoadCount++; });
+ 
+	GPackageLoader->InitializeLoading();
+}
+
+void ShutdownAsyncThread()
+{
+	if (GPackageLoader)
+	{
+		GPackageLoader->ShutdownLoading();
+		GPackageLoader.Reset(nullptr);
+	}
+}
+
+bool IsInAsyncLoadingThreadCoreUObjectInternal()
+{
+	if (GPackageLoader)
+	{
+		return GPackageLoader->IsInAsyncLoadThread();
+	}
+	else
+	{
+		return false;
+	}
+}
+
+void FlushAsyncLoading(int32 PackageID /* = INDEX_NONE */)
+{
+#if defined(WITH_CODE_GUARD_HANDLER) && WITH_CODE_GUARD_HANDLER
+	void CheckImageIntegrityAtRuntime();
+	CheckImageIntegrityAtRuntime();
+#endif
+	checkf(IsInGameThread(), TEXT("Unable to FlushAsyncLoading from any thread other than the game thread."));
+	if (GPackageLoader)
+	{
+		GPackageLoader->FlushLoading(PackageID);
+	}
+}
+
+EAsyncPackageState::Type ProcessAsyncLoadingUntilComplete(TFunctionRef<bool()> CompletionPredicate, float TimeLimit)
+{
+	return GetAsyncPackageLoader().ProcessLoadingUntilComplete(CompletionPredicate, TimeLimit);
+}
+
+int32 GetNumAsyncPackages()
+{
+	return GetAsyncPackageLoader().GetNumAsyncPackages();
+}
+
+EAsyncPackageState::Type ProcessAsyncLoading(bool bUseTimeLimit, bool bUseFullTimeLimit, float TimeLimit)
+{
+	return GetAsyncPackageLoader().ProcessLoading(bUseTimeLimit, bUseFullTimeLimit, TimeLimit);
+}
+
+bool IsAsyncLoadingCoreUObjectInternal()
+{
+	// GIsInitialLoad guards the async loading thread from being created too early
+	return GetAsyncPackageLoader().IsAsyncLoadingPackages();
+}
+
+bool IsAsyncLoadingMultithreadedCoreUObjectInternal()
+{
+	// GIsInitialLoad guards the async loading thread from being created too early
+	return GetAsyncPackageLoader().IsMultithreaded();
+}
+
+void SuspendAsyncLoadingInternal()
+{
+	check(IsInGameThread() && !IsInSlateThread());
+	GetAsyncPackageLoader().SuspendLoading();
+}
+
+void ResumeAsyncLoadingInternal()
+{
+	check(IsInGameThread() && !IsInSlateThread());
+	GetAsyncPackageLoader().ResumeLoading();
+}
+
+bool IsAsyncLoadingSuspendedInternal()
+{
+	return GetAsyncPackageLoader().IsAsyncLoadingSuspended();
+}
+
+int32 LoadPackageAsync(const FString& InName, const FGuid* InGuid /*= nullptr*/, const TCHAR* InPackageToLoadFrom /*= nullptr*/, FLoadPackageAsyncDelegate InCompletionDelegate /*= FLoadPackageAsyncDelegate()*/, EPackageFlags InPackageFlags /*= PKG_None*/, int32 InPIEInstanceID /*= INDEX_NONE*/, int32 InPackagePriority /*= 0*/)
+{
+	return GetAsyncPackageLoader().LoadPackage(InName, InGuid, InPackageToLoadFrom, InCompletionDelegate, InPackageFlags, InPIEInstanceID, InPackagePriority);
+}
+
+int32 LoadPackageAsync(const FString& PackageName, FLoadPackageAsyncDelegate CompletionDelegate, int32 InPackagePriority /*= 0*/, EPackageFlags InPackageFlags /*= PKG_None*/, int32 InPIEInstanceID /*= INDEX_NONE*/)
+{
+	const FGuid* Guid = nullptr;
+	const TCHAR* PackageToLoadFrom = nullptr;
+	return LoadPackageAsync(PackageName, Guid, PackageToLoadFrom, CompletionDelegate, InPackageFlags, InPIEInstanceID, InPackagePriority );
+}
+
+void CancelAsyncLoading()
+{
+	// Cancelling async loading while loading is suspend will result in infinite stall
+	UE_CLOG(GetAsyncPackageLoader().IsAsyncLoadingSuspended(), LogStreaming, Fatal, TEXT("Cannot Cancel Async Loading while async loading is suspended."));
+	GetAsyncPackageLoader().CancelLoading();
+
+	if (!IsEngineExitRequested())
+	{
+		CollectGarbage(GARBAGE_COLLECTION_KEEPFLAGS, true);
+	}
+
+	const EInternalObjectFlags AsyncFlags = EInternalObjectFlags::Async | EInternalObjectFlags::AsyncLoading;
+	for (int32 ObjectIndex = 0; ObjectIndex < GUObjectArray.GetObjectArrayNum(); ++ObjectIndex)
+	{
+		FUObjectItem* ObjectItem = &GUObjectArray.GetObjectItemArrayUnsafe()[ObjectIndex];
+		if (UObject* Obj = static_cast<UObject*>(ObjectItem->Object))
+		{
+			check(!Obj->HasAnyInternalFlags(AsyncFlags));
+		}
+	}
+}
+
+float GetAsyncLoadPercentage(const FName& PackageName)
+{
+	return GetAsyncPackageLoader().GetAsyncLoadPercentage(PackageName);
+}
+
+void NotifyRegistrationEvent(const TCHAR* PackageName, const TCHAR* Name, ENotifyRegistrationType NotifyRegistrationType, ENotifyRegistrationPhase NotifyRegistrationPhase, UObject *(*InRegister)(), bool InbDynamic)
+{
+	GetGEDLBootNotificationManager().NotifyRegistrationEvent(PackageName, Name, NotifyRegistrationType, NotifyRegistrationPhase, InRegister, InbDynamic);
+}
+
+void NotifyRegistrationComplete()
+{
+	GetGEDLBootNotificationManager().NotifyRegistrationComplete();
+}
+
+void NotifyConstructedDuringAsyncLoading(UObject* Object, bool bSubObject)
+{
+	GetAsyncPackageLoader().NotifyConstructedDuringAsyncLoading(Object, bSubObject);
+}
+
+double GFlushAsyncLoadingTime = 0.0;
+uint32 GFlushAsyncLoadingCount = 0;
+uint32 GSyncLoadCount = 0;
+
+void ResetAsyncLoadingStats()
+{
+	check(IsInGameThread());
+	GFlushAsyncLoadingTime = 0.0;
+	GFlushAsyncLoadingCount = 0;
+	GSyncLoadCount = 0;
+}
diff --git a/Engine/Source/Runtime/CoreUObject/Private/Serialization/BulkData.cpp b/Engine/Source/Runtime/CoreUObject/Private/Serialization/BulkData.cpp
--- a/Engine/Source/Runtime/CoreUObject/Private/Serialization/BulkData.cpp
+++ b/Engine/Source/Runtime/CoreUObject/Private/Serialization/BulkData.cpp
@@ -19,6 +19,7 @@
 #include "Async/MappedFileHandle.h"
 #include "HAL/PlatformFilemanager.h"
 #include "UObject/UObjectThreadContext.h"
+#include "ProfilingDebugging/LoadTimeTracker.h"
 
 /*-----------------------------------------------------------------------------
 	Constructors and operators
@@ -432,7 +433,7 @@ bool FUntypedBulkData::IsBulkDataLoaded() const
 	return !!BulkData;
 }
 
-bool FUntypedBulkData::IsAsyncLoadingComplete() const 
+bool FUntypedBulkData::IsAsyncLoadingComplete() const
 {
 	return SerializeFuture.IsValid() == false || SerializeFuture.WaitFor(FTimespan::Zero());
 }
@@ -734,6 +735,11 @@ void FUntypedBulkData::SetBulkDataFlags( uint32 BulkDataFlagsToSet )
 	BulkDataFlags |= BulkDataFlagsToSet;
 }
 
+void FUntypedBulkData::ResetBulkDataFlags(uint32 BulkDataFlagsToSet)
+{
+	BulkDataFlags = BulkDataFlagsToSet;
+}
+
 /**
 * Gets the current bulk data flags.
 *
@@ -855,6 +861,8 @@ void FUntypedBulkData::Serialize( FArchive& Ar, UObject* Owner, int32 Idx, bool
 {
 	DECLARE_SCOPE_CYCLE_COUNTER(TEXT("FUntypedBulkData::Serialize"), STAT_UBD_Serialize, STATGROUP_Memory);
 
+	SCOPED_LOADTIMER(BulkData_Serialize);
+
 	check( LockStatus == LOCKSTATUS_Unlocked );
 
 	check(!bAttemptFileMapping || Ar.IsLoading()); // makes no sense to map unless we are loading
@@ -997,8 +1005,11 @@ void FUntypedBulkData::Serialize( FArchive& Ar, UObject* Owner, int32 Idx, bool
 				check(Package.IsValid());
 				auto Linker = FLinkerLoad::FindExistingLinkerForPackage(Package.Get());
 				check(Linker);
-				Filename = Linker->Filename;		
+				Filename = Linker->Filename;
+
+				check(Owner->GetLinker() == Linker);
 #endif // WITH_EDITOR
+
 				if (bPayloadInline)
 				{
 					if (ShouldStreamBulkData())
@@ -1386,6 +1397,7 @@ void FUntypedBulkData::StoreCompressedOnDisk(ECompressionFlags CompressionFlags)
 {
 	StoreCompressedOnDisk(FCompression::GetCompressionFormatFromDeprecatedFlags(CompressionFlags));
 }
+
 void FUntypedBulkData::StoreCompressedOnDisk( FName CompressionFormat )
 {
 	if( CompressionFormat != GetDecompressionFormat() )
@@ -1470,6 +1482,8 @@ void FUntypedBulkData::SerializeElements(FArchive& Ar, void* Data)
  */
 void FUntypedBulkData::SerializeBulkData( FArchive& Ar, void* Data )
 {
+	SCOPED_LOADTIMER(BulkData_SerializeBulkData);
+
 	// skip serializing of unused data
 	if( BulkDataFlags & BULKDATA_Unused )
 	{
@@ -1555,6 +1569,127 @@ void FUntypedBulkData::SerializeBulkData( FArchive& Ar, void* Data )
 	}
 }
 
+FBulkDataIORequest* FUntypedBulkData::CreateStreamingRequest(EAsyncIOPriorityAndFlags Priority, FAsyncFileCallBack* CompleteCallback, uint8* UserSuppliedMemory) const
+{
+	const int64 DataSize = GetBulkDataSize();
+
+	return CreateStreamingRequest(0, DataSize, Priority, CompleteCallback, UserSuppliedMemory);
+}
+
+FBulkDataIORequest* FUntypedBulkData::CreateStreamingRequest(int64 OffsetInBulkData, int64 BytesToRead, EAsyncIOPriorityAndFlags Priority, FAsyncFileCallBack* CompleteCallback, uint8* UserSuppliedMemory) const
+{
+	check(Filename.IsEmpty() == false);
+
+	UE_CLOG(IsStoredCompressedOnDisk(), LogSerialization, Fatal, TEXT("Package level compression is no longer supported (%s)."), *Filename);
+	UE_CLOG(GetBulkDataSize() <= 0, LogSerialization, Error, TEXT("(%s) has invalid bulk data size."), *Filename);
+
+	IAsyncReadFileHandle* IORequestHandle = FPlatformFileManager::Get().GetPlatformFile().OpenAsyncRead(*Filename);
+	check(IORequestHandle); // this generally cannot fail because it is async
+
+	if (IORequestHandle == nullptr)
+	{
+		return nullptr;
+	}
+
+	const int64 OffsetInFile = GetBulkDataOffsetInFile() + OffsetInBulkData;
+
+	IAsyncReadRequest* ReadRequest = IORequestHandle->ReadRequest(OffsetInFile, BytesToRead, Priority, CompleteCallback, UserSuppliedMemory);
+	if (ReadRequest != nullptr)
+	{
+		return new FBulkDataIORequest(IORequestHandle, ReadRequest, BytesToRead);
+	}
+	else
+	{
+		delete IORequestHandle;
+		return nullptr;
+	}
+}
+
+#if USE_BULKDATA_STREAMING_TOKEN 
+
+FBulkDataStreamingToken FUntypedBulkData::CreateStreamingToken() const
+{
+	// Checks since we are casting signed 64bit values to unsigned 32bit
+	check(GetBulkDataOffsetInFile() >= 0);
+	check(GetBulkDataOffsetInFile() < TNumericLimits<uint32>::Max());
+	check(GetBulkDataSize() >= 0);
+	check(GetBulkDataSize() < TNumericLimits<uint32>::Max());
+
+	return FBulkDataStreamingToken((uint32)GetBulkDataOffsetInFile(), (uint32)GetBulkDataSize());
+}
+
+FBulkDataIORequest* FUntypedBulkData::CreateStreamingRequestForRange(const FString& Filename, const FBulkDataStreamingToken& Start, const FBulkDataStreamingToken& End, EAsyncIOPriorityAndFlags Priority, FAsyncFileCallBack* CompleteCallback)
+{
+	check(Filename.IsEmpty() == false);
+	check(Start.IsValid());
+	check(End.IsValid());
+
+	IAsyncReadFileHandle* IORequestHandle = FPlatformFileManager::Get().GetPlatformFile().OpenAsyncRead(*Filename);
+	check(IORequestHandle); // this generally cannot fail because it is async
+
+	if (IORequestHandle == nullptr)
+	{
+		return nullptr;
+	}
+
+	const int64 ReadOffset = Start.GetOffset();
+	const int64 ReadSize = (End.GetOffset() + End.GetSize()) - ReadOffset;
+
+	check(ReadSize > 0);
+
+	IAsyncReadRequest* ReadRequest = IORequestHandle->ReadRequest(ReadOffset, ReadSize, Priority, CompleteCallback, nullptr);
+	if (ReadRequest != nullptr)
+	{
+		return new FBulkDataIORequest(IORequestHandle, ReadRequest, ReadSize);
+	}
+	else
+	{
+		delete IORequestHandle;
+		return nullptr;
+	}
+}
+
+#endif
+
+FBulkDataIORequest::~FBulkDataIORequest()
+{
+	delete ReadRequest;
+	delete FileHandle;
+}
+
+bool FBulkDataIORequest::PollCompletion() const
+{
+	return ReadRequest->PollCompletion();
+}
+
+bool FBulkDataIORequest::WaitCompletion(float TimeLimitSeconds) const
+{
+	return ReadRequest->WaitCompletion(TimeLimitSeconds);
+}
+
+
+uint8* FBulkDataIORequest::GetReadResults() const
+{
+	return ReadRequest->GetReadResults();
+}
+
+int64 FBulkDataIORequest::GetSize() const
+{
+	if (ReadRequest->PollCompletion())
+	{
+		return Size;
+	}
+	else
+	{
+		return -1;
+	}
+}
+
+void FBulkDataIORequest::Cancel() const
+{
+	ReadRequest->Cancel();
+}
+
 /**
  * Loads the bulk data if it is not already loaded.
  */
@@ -1648,6 +1783,7 @@ void FUntypedBulkData::LoadDataIntoMemory( void* Dest )
 	}
 
 	if (!BulkDataArchive)
+
 	{
 		BulkDataArchive = AttachedAr;
 	}
@@ -1746,7 +1882,7 @@ void FUntypedBulkData::LoadDataIntoMemory( void* Dest )
  *
  * @return Size in bytes of single element
  */
-int32 FByteBulkData::GetElementSize() const
+int32 FByteBulkDataOld::GetElementSize() const
 {
 	return sizeof(uint8);
 }
@@ -1758,7 +1894,7 @@ int32 FByteBulkData::GetElementSize() const
  * @param Data			Base pointer to data
  * @param ElementIndex	Element index to serialize
  */
-void FByteBulkData::SerializeElement( FArchive& Ar, void* Data, int64 ElementIndex )
+void FByteBulkDataOld::SerializeElement( FArchive& Ar, void* Data, int64 ElementIndex )
 {
 	uint8& ByteData = *((uint8*)Data + ElementIndex);
 	Ar << ByteData;
@@ -1773,7 +1909,7 @@ void FByteBulkData::SerializeElement( FArchive& Ar, void* Data, int64 ElementInd
  *
  * @return Size in bytes of single element
  */
-int32 FWordBulkData::GetElementSize() const
+int32 FWordBulkDataOld::GetElementSize() const
 {
 	return sizeof(uint16);
 }
@@ -1785,7 +1921,7 @@ int32 FWordBulkData::GetElementSize() const
  * @param Data			Base pointer to data
  * @param ElementIndex	Element index to serialize
  */
-void FWordBulkData::SerializeElement( FArchive& Ar, void* Data, int64 ElementIndex )
+void FWordBulkDataOld::SerializeElement( FArchive& Ar, void* Data, int64 ElementIndex )
 {
 	uint16& WordData = *((uint16*)Data + ElementIndex);
 	Ar << WordData;
@@ -1800,7 +1936,7 @@ void FWordBulkData::SerializeElement( FArchive& Ar, void* Data, int64 ElementInd
  *
  * @return Size in bytes of single element
  */
-int32 FIntBulkData::GetElementSize() const
+int32 FIntBulkDataOld::GetElementSize() const
 {
 	return sizeof(int32);
 }
@@ -1812,7 +1948,7 @@ int32 FIntBulkData::GetElementSize() const
  * @param Data			Base pointer to data
  * @param ElementIndex	Element index to serialize
  */
-void FIntBulkData::SerializeElement( FArchive& Ar, void* Data, int64 ElementIndex )
+void FIntBulkDataOld::SerializeElement( FArchive& Ar, void* Data, int64 ElementIndex )
 {
 	int32& IntData = *((int32*)Data + ElementIndex);
 	Ar << IntData;
@@ -1827,7 +1963,7 @@ void FIntBulkData::SerializeElement( FArchive& Ar, void* Data, int64 ElementInde
  *
  * @return Size in bytes of single element
  */
-int32 FFloatBulkData::GetElementSize() const
+int32 FFloatBulkDataOld::GetElementSize() const
 {
 	return sizeof(float);
 }
@@ -1839,12 +1975,24 @@ int32 FFloatBulkData::GetElementSize() const
  * @param Data			Base pointer to data
  * @param ElementIndex	Element index to serialize
  */
-void FFloatBulkData::SerializeElement( FArchive& Ar, void* Data, int64 ElementIndex )
+void FFloatBulkDataOld::SerializeElement( FArchive& Ar, void* Data, int64 ElementIndex )
 {
 	float& FloatData = *((float*)Data + ElementIndex);
 	Ar << FloatData;
 }
 
+FBulkDataIORequest* BulkDataUtils::CreateStreamingRequestForRange(const FBulkDataInterface& Start, const FBulkDataInterface& End, EAsyncIOPriorityAndFlags Priority, FAsyncFileCallBack* CompleteCallback)
+{
+	check(Start.GetFilename() == End.GetFilename());
+
+	const int64 ReadOffset = Start.GetBulkDataOffsetInFile();
+	const int64 ReadSize = (End.GetBulkDataOffsetInFile() + End.GetBulkDataSize()) - ReadOffset;
+
+	check(ReadSize > 0);
+
+	return Start.CreateStreamingRequest(0, ReadSize, Priority, CompleteCallback, nullptr);
+}
+
 void FFormatContainer::Serialize(FArchive& Ar, UObject* Owner, const TArray<FName>* FormatsToSave, bool bSingleUse, uint32 InAlignment, bool bInline, bool bMapped)
 {
 	if (Ar.IsLoading())
@@ -1856,8 +2004,10 @@ void FFormatContainer::Serialize(FArchive& Ar, UObject* Owner, const TArray<FNam
 			FName Name;
 			Ar << Name;
 			FByteBulkData& Bulk = GetFormat(Name);
+#if !USE_NEW_BULKDATA
 			Bulk.SetBulkDataAlignment(InAlignment);
-			Bulk.Serialize(Ar, Owner);
+#endif
+			Bulk.Serialize(Ar, Owner, INDEX_NONE, false);
 		}
 	}
 	else
@@ -1905,7 +2055,7 @@ void FFormatContainer::Serialize(FArchive& Ar, UObject* Owner, const TArray<FNam
 				{
 					Bulk->SetBulkDataFlags(BULKDATA_SingleUse);
 				}
-				Bulk->Serialize(Ar, Owner);
+				Bulk->Serialize(Ar, Owner, INDEX_NONE, false);
 				Bulk->ClearBulkDataFlags(0xFFFFFFFF);
 				Bulk->SetBulkDataFlags(OldBulkDataFlags);
 			}
@@ -1927,4 +2077,3 @@ void FFormatContainer::SerializeAttemptMappedLoad(FArchive& Ar, UObject* Owner)
 		Bulk.Serialize(Ar, Owner, -1, true);
 	}
 }
-
diff --git a/Engine/Source/Runtime/CoreUObject/Private/Serialization/BulkData2.cpp b/Engine/Source/Runtime/CoreUObject/Private/Serialization/BulkData2.cpp
new file mode 100644
--- /dev/null
+++ b/Engine/Source/Runtime/CoreUObject/Private/Serialization/BulkData2.cpp
@@ -0,0 +1,655 @@
+// Copyright 1998-2019 Epic Games, Inc. All Rights Reserved.
+
+#include "Serialization/BulkData2.h"
+
+#include "HAL/FileManager.h"
+#include "HAL/PlatformFilemanager.h"
+#include "Misc/Paths.h"
+#include "ProfilingDebugging/LoadTimeTracker.h"
+#include "Serialization/BulkData.h"
+#include "UObject/LinkerLoad.h"
+#include "UObject/Object.h"
+
+namespace
+{
+	// TODO: Maybe expose this and start using everywhere?
+	static constexpr const TCHAR* InlinedExt = TEXT(".uexp");			// Inlined
+	static constexpr const TCHAR* DefaultExt = TEXT(".ubulk");			// Stored in a separate file
+	static constexpr const TCHAR* MemoryMappedExt = TEXT(".m.ubulk");	// Stored in a separate file aligned for memory mapping
+	static constexpr const TCHAR* OptionalExt = TEXT(".uptnl");			// Stored in a separate file that is optional
+}
+
+// TODO: The code in the FileTokenSystem namespace is a temporary system so that FBulkDataBase can hold
+// all of it's info about where the data is on disk in a single 8byte value. This can all be removed when
+// we switch this over to the new packing system.
+namespace FileTokenSystem
+{
+	struct Data
+	{
+		FString Filename;
+		int64 BulkDataSize;
+		int64 BulkDataOffsetInFile;
+	};
+
+	TMap<FBulkDataBase::FileToken, Data> TokenDataMap;
+	uint64 TokenCounter = 0;
+	FRWLock TokenLock;
+
+	FBulkDataBase::FileToken RegisterFileToken( const FString& Filename, uint64 BulkDataSize, uint64 BulkDataOffsetInFile )
+	{
+		FWriteScopeLock LockForScope(TokenLock);
+
+		Data& Data = TokenDataMap.Add(++TokenCounter);
+		Data.Filename = Filename;
+		Data.BulkDataSize = BulkDataSize;
+		Data.BulkDataOffsetInFile = BulkDataOffsetInFile;
+
+		return TokenCounter;
+	}
+
+	void UnregisterFileToken(FBulkDataBase::FileToken ID)
+	{
+		if (ID != FBulkDataBase::InvalidToken)
+		{
+			FWriteScopeLock LockForScope(TokenLock);
+			TokenDataMap.Remove(ID);
+		}
+	}
+
+	FBulkDataBase::FileToken CopyFileToken(FBulkDataBase::FileToken ID)
+	{
+		if (ID != FBulkDataBase::InvalidToken)
+		{
+			FWriteScopeLock LockForScope(TokenLock);
+
+			Data* Data = TokenDataMap.Find(ID);
+			check(Data);
+
+			TokenDataMap.Add(++TokenCounter, *Data);
+
+			return TokenCounter;		
+		}
+		else
+		{
+			return FBulkDataBase::InvalidToken;
+		}
+	}
+
+	Data GetFileData(FBulkDataBase::FileToken ID)
+	{
+		if (ID == FBulkDataBase::InvalidToken)
+		{
+			return Data();
+		}
+
+		FReadScopeLock LockForScope(TokenLock);
+		Data* Data = TokenDataMap.Find(ID);
+		check(Data);
+
+		return *Data;
+	}
+
+	FString GetFilename(FBulkDataBase::FileToken ID)
+	{
+		if (ID == FBulkDataBase::InvalidToken)
+		{
+			return FString();
+		}
+
+		FReadScopeLock LockForScope(TokenLock);
+		Data* Data = TokenDataMap.Find(ID);
+		check(Data);
+
+		return Data->Filename;
+	}
+
+	uint64 GetBulkDataSize(FBulkDataBase::FileToken ID)
+	{
+		if (ID == FBulkDataBase::InvalidToken)
+		{
+			return 0;
+		}
+		
+		FReadScopeLock LockForScope(TokenLock);
+		Data* Data = TokenDataMap.Find(ID);
+		check(Data);
+
+		return Data != nullptr ? Data->BulkDataSize : 0;
+	}
+
+	uint64 GetBulkDataOffset(FBulkDataBase::FileToken ID)
+	{
+		if (ID == FBulkDataBase::InvalidToken)
+		{
+			return 0;
+		}
+
+		FReadScopeLock LockForScope(TokenLock);
+		Data* Data = TokenDataMap.Find(ID);
+		check(Data);
+
+		return Data != nullptr ? Data->BulkDataOffsetInFile : 0;
+	}
+}
+
+FBulkDataBase::FBulkDataBase(FBulkDataBase&& Other)
+	: Token(Other.Token)
+	, DataBuffer(Other.DataBuffer)
+	, BulkDataFlags(Other.BulkDataFlags)
+{
+	check(Other.LockStatus == LOCKSTATUS_Unlocked); // Make sure that the other object wasn't inuse
+
+	Other.Token = InvalidToken; // Prevent the other object from unregistering the token
+	Other.DataBuffer = nullptr; // Prevent the other object from deleting our data
+}
+
+FBulkDataBase& FBulkDataBase::operator=(const FBulkDataBase& Other)
+{
+	DECLARE_SCOPE_CYCLE_COUNTER(TEXT("FBulkDataBase::operator="), STAT_UBD_Constructor, STATGROUP_Memory);
+
+	check(LockStatus == LOCKSTATUS_Unlocked);
+	check(Other.LockStatus == LOCKSTATUS_Unlocked);
+
+	RemoveBulkData();
+
+	Token = FileTokenSystem::CopyFileToken(Other.Token);
+
+	// Copy token
+	BulkDataFlags = Other.BulkDataFlags;
+
+	if (Other.DataBuffer != nullptr)
+	{
+		const int64 DataSize = Other.GetBulkDataSize();
+
+		AllocateData(DataSize);
+		FMemory::Memcpy(DataBuffer, DataBuffer, DataSize);
+	}
+
+	return *this;
+}
+
+FBulkDataBase::~FBulkDataBase()
+{
+	check(LockStatus == LOCKSTATUS_Unlocked);
+
+	FreeData();
+	FileTokenSystem::UnregisterFileToken(Token);
+}
+
+void FBulkDataBase::Serialize(FArchive& Ar, UObject* Owner, int32 /*Index*/, bool bAttemptFileMapping, int32 ElementSize)
+{
+	DECLARE_SCOPE_CYCLE_COUNTER(TEXT("FBulkDataBase::Serialize"), STAT_UBD_Serialize, STATGROUP_Memory);
+
+	SCOPED_LOADTIMER(BulkData_Serialize);
+
+#if WITH_EDITOR == 0 && WITH_EDITORONLY_DATA == 0
+	check(Ar.IsLoading());				// Only support loading from cooked data!
+	check(!GIsEditor);					// The editor path is not supported
+	check(GEventDrivenLoaderEnabled);	// We are assuming the EDL path is enabled ( TODO: Might need to remove this check)
+	check(LockStatus == LOCKSTATUS_Unlocked);
+
+	if (Ar.IsPersistent() && !Ar.IsObjectReferenceCollector() && !Ar.ShouldSkipBulkData())
+	{
+		Ar << BulkDataFlags;
+
+		int64 ElementCount = 0;
+		int64 BulkDataSizeOnDisk = 0;
+		int64 BulkDataSize = 0;
+		int64 BulkDataOffsetInFile = 0;
+
+		if (BulkDataFlags & BULKDATA_Size64Bit)
+		{
+			Ar << ElementCount;
+			Ar << BulkDataSizeOnDisk;
+		}
+		else
+		{
+			int32 Temp32ByteValue;
+
+			Ar << Temp32ByteValue;
+			ElementCount = Temp32ByteValue;
+
+			Ar << Temp32ByteValue;
+			BulkDataSizeOnDisk = Temp32ByteValue;
+		}
+
+		BulkDataSize = ElementCount * ElementSize;
+
+		Ar << BulkDataOffsetInFile;
+
+		// fix up the file offset, but only if not stored inline
+		if (Owner != nullptr && Owner->GetLinker() != nullptr && !IsInlined())
+		{
+			BulkDataOffsetInFile += Owner->GetLinker()->Summary.BulkDataStartOffset;
+		}
+
+		FArchive* CacheableArchive = Ar.GetCacheableArchive();
+		if (Ar.IsAllowingLazyLoading() && Owner != nullptr && CacheableArchive != nullptr)
+		{
+			UPackage* Package = Owner->GetOutermost();
+			check(Package != nullptr);
+
+			FLinkerLoad* Linker = FLinkerLoad::FindExistingLinkerForPackage(Package);
+			check(Linker != nullptr);
+
+			FString Filename = ConvertFilenameFromFlags(Linker->Filename);
+
+			if (IsInlined())
+			{
+				// Inline data is already in the archive so serialize it immediately
+				AllocateData(BulkDataSize);
+				SerializeBulkData(Ar, DataBuffer, BulkDataSize);
+			}
+			else
+			{
+				if (IsDuplicateNonOptional())
+				{
+					FString OptionalFilename = FPaths::ChangeExtension(Filename, OptionalExt);
+					if (IFileManager::Get().FileExists(*Filename))
+					{
+						SerializeDuplicateData(Ar, Owner, BulkDataFlags, BulkDataSizeOnDisk, BulkDataOffsetInFile);
+						Filename = OptionalFilename;
+					}
+					else
+					{
+						// Skip over the optional data info (can't do a seek because we need to load the flags to work
+						// out if we read things as 32bit or 64bit)
+						uint32 DummyValue32;
+						int64 DummyValue64;
+						SerializeDuplicateData(Ar, Owner, DummyValue32, DummyValue64, DummyValue64);
+
+						check(Filename.EndsWith(DefaultExt));
+					}			
+				}			
+			}
+
+			// TODO: Replace with the proper storage solution
+			Token = FileTokenSystem::RegisterFileToken(Filename, BulkDataSize, BulkDataOffsetInFile);
+		}
+		else
+		{
+			BULKDATA_NOT_IMPLEMENTED_FOR_RUNTIME;
+		}
+	}
+#else
+	check(false); // Only implemented for cooked builds!
+#endif
+}
+
+void* FBulkDataBase::Lock(uint32 LockFlags)
+{
+	check(LockStatus == LOCKSTATUS_Unlocked);
+	
+	ForceBulkDataResident(); // Will load 
+
+	if (LockFlags & LOCK_READ_WRITE)
+	{
+		LockStatus = LOCKSTATUS_ReadWriteLock;
+	}
+	else if (LockFlags & LOCK_READ_ONLY)
+	{
+		LockStatus = LOCKSTATUS_ReadOnlyLock;
+	}
+	else
+	{
+		UE_LOG(LogSerialization, Fatal, TEXT("Unknown lock flag %i"), LockFlags);
+	}
+
+	return DataBuffer;
+}
+
+const void* FBulkDataBase::LockReadOnly() const
+{
+	check(LockStatus == LOCKSTATUS_Unlocked);
+	return DataBuffer;
+}
+
+void FBulkDataBase::Unlock()
+{
+	check(LockStatus != LOCKSTATUS_Unlocked);
+
+	LockStatus = LOCKSTATUS_Unlocked;
+
+	// Free pointer if we're guaranteed to only to access the data once.
+	if (IsSingleUse())
+	{
+		FreeData();
+	}
+}
+
+bool FBulkDataBase::IsLocked() const
+{ 
+	return LockStatus != LOCKSTATUS_Unlocked;
+}
+
+void* FBulkDataBase::Realloc(int64 InElementCount)
+{
+	BULKDATA_NOT_IMPLEMENTED_FOR_RUNTIME;
+	return nullptr;
+}
+
+void FBulkDataBase::GetCopy(void** DstBuffer, bool bDiscardInternalCopy)
+{
+	DECLARE_SCOPE_CYCLE_COUNTER(TEXT("FUntypedBulkData::GetCopy"), STAT_UBD_GetCopy, STATGROUP_Memory);
+
+	check(LockStatus == LOCKSTATUS_Unlocked);
+	check(DstBuffer);
+
+	if (*DstBuffer != nullptr)
+	{
+		// TODO: Might be worth changing the API so that we can validate that the buffer is large enough?
+		if (IsBulkDataLoaded())
+		{
+			FMemory::Memcpy(*DstBuffer, DataBuffer, GetBulkDataSize());
+
+			if (bDiscardInternalCopy && (CanLoadFromDisk() || IsSingleUse()))
+			{
+				FreeData();
+			}
+		}
+		else
+		{
+			LoadDataDirectly(DstBuffer);
+		}
+	}
+	else
+	{
+		if (IsBulkDataLoaded())
+		{
+			if (bDiscardInternalCopy && (CanLoadFromDisk() || IsSingleUse()))
+			{
+				// Since we were going to discard the data anyway we can just hand over ownership to the caller
+				::Swap(*DstBuffer, DataBuffer);
+			}
+			else
+			{
+				int64 BulkDataSize = GetBulkDataSize();
+
+				*DstBuffer = FMemory::Malloc(BulkDataSize, 0);
+				FMemory::Memcpy(*DstBuffer, DataBuffer, GetBulkDataSize());
+			}
+		}
+		else
+		{
+			*DstBuffer = FMemory::Malloc(GetBulkDataSize(), 0);
+			LoadDataDirectly(DstBuffer);
+		}
+	}
+}
+
+void  FBulkDataBase::SetBulkDataFlags(uint32 BulkDataFlagsToSet)
+{
+	check(!CanLoadFromDisk());	// We only want to allow the editing of flags if the BulkData
+								// was dynamically created at runtime, not loaded off disk
+
+	BulkDataFlags |= BulkDataFlagsToSet;
+}
+
+void  FBulkDataBase::ResetBulkDataFlags(uint32 BulkDataFlagsToSet)
+{
+	check(!CanLoadFromDisk());	// We only want to allow the editing of flags if the BulkData
+								// was dynamically created at runtime, not loaded off disk
+
+	BulkDataFlags = BulkDataFlagsToSet;
+}
+
+void  FBulkDataBase::ClearBulkDataFlags(uint32 BulkDataFlagsToClear)
+{ 
+	check(!CanLoadFromDisk());	// We only want to allow the editing of flags if the BulkData
+								// was dynamically created at runtime, not loaded off disk
+
+	BulkDataFlags &= ~BulkDataFlagsToClear;
+}
+
+int64 FBulkDataBase::GetBulkDataSize() const
+{
+	return FileTokenSystem::GetBulkDataSize(Token);
+}
+
+bool FBulkDataBase::IsStoredCompressedOnDisk() const
+{
+	return (GetBulkDataFlags() & BULKDATA_SerializeCompressed) != 0;
+}
+
+FName FBulkDataBase::GetDecompressionFormat() const
+{
+	return (BulkDataFlags & BULKDATA_SerializeCompressedZLIB) ? NAME_Zlib : NAME_None;
+}
+
+bool FBulkDataBase::IsAvailableForUse() const 
+{ 
+	return (GetBulkDataFlags() & BULKDATA_Unused) != 0;
+}
+
+bool FBulkDataBase::IsDuplicateNonOptional() const
+{
+	return (GetBulkDataFlags() & BULKDATA_DuplicateNonOptionalPayload) != 0;
+}
+
+bool FBulkDataBase::IsOptional() const
+{ 
+	return (GetBulkDataFlags() & BULKDATA_OptionalPayload) != 0;
+}
+
+bool FBulkDataBase::IsInlined() const
+{
+	return	(GetBulkDataFlags() & BULKDATA_PayloadAtEndOfFile) == 0;
+}
+
+bool FBulkDataBase::InSeperateFile() const
+{
+	return	(GetBulkDataFlags() & BULKDATA_PayloadInSeperateFile) != 0;
+}
+
+bool FBulkDataBase::IsSingleUse() const
+{
+	return (BulkDataFlags & BULKDATA_SingleUse) != 0;
+}
+
+bool FBulkDataBase::IsMemoryMapped() const
+{
+	return (BulkDataFlags & BULKDATA_MemoryMappedPayload) != 0;
+}
+
+FBulkDataIORequest* FBulkDataBase::CreateStreamingRequest(EAsyncIOPriorityAndFlags Priority, FAsyncFileCallBack* CompleteCallback, uint8* UserSuppliedMemory) const
+{
+	const int64 DataSize = GetBulkDataSize();
+
+	return CreateStreamingRequest(0, DataSize, Priority, CompleteCallback, UserSuppliedMemory);
+}
+
+FBulkDataIORequest* FBulkDataBase::CreateStreamingRequest(int64 OffsetInBulkData, int64 BytesToRead, EAsyncIOPriorityAndFlags Priority, FAsyncFileCallBack* CompleteCallback, uint8* UserSuppliedMemory) const
+{
+#if 1
+	FileTokenSystem::Data FileData = FileTokenSystem::GetFileData(Token);
+
+	check(FileData.Filename.IsEmpty() == false);
+
+	UE_CLOG(IsStoredCompressedOnDisk(), LogSerialization, Fatal, TEXT("Package level compression is no longer supported (%s)."), *FileData.Filename);
+	UE_CLOG(FileData.BulkDataSize <= 0, LogSerialization, Error, TEXT("(%s) has invalid bulk data size."), *FileData.Filename);
+
+	IAsyncReadFileHandle* IORequestHandle = FPlatformFileManager::Get().GetPlatformFile().OpenAsyncRead(*FileData.Filename);
+	check(IORequestHandle); // this generally cannot fail because it is async
+
+	if (IORequestHandle == nullptr)
+	{
+		return nullptr;
+	}
+
+	const int64 OffsetInFile = FileData.BulkDataOffsetInFile + OffsetInBulkData;
+
+	IAsyncReadRequest* ReadRequest = IORequestHandle->ReadRequest(OffsetInFile, BytesToRead, Priority, CompleteCallback, UserSuppliedMemory);
+	if (ReadRequest != nullptr)
+	{
+		return new FBulkDataIORequest(IORequestHandle, ReadRequest, BytesToRead);
+	}
+	else
+	{
+		delete IORequestHandle;
+		return nullptr;
+	}
+#else
+	return nullptr;
+#endif
+}
+
+void FBulkDataBase::ForceBulkDataResident()
+{
+	if (!IsBulkDataLoaded())
+	{
+		LoadDataDirectly(&DataBuffer);
+	}
+}
+
+FOwnedBulkDataPtr* FBulkDataBase::StealFileMapping()
+{
+	check(LockStatus == LOCKSTATUS_Unlocked);
+
+	FOwnedBulkDataPtr* Result = new FOwnedBulkDataPtr(DataBuffer);
+
+	// We are giving ownership of the data to the caller so we just null out the pointer.
+	DataBuffer = nullptr;
+
+	return Result;
+}
+
+void FBulkDataBase::RemoveBulkData()
+{
+	check(LockStatus == LOCKSTATUS_Unlocked);
+
+	FreeData();
+
+	FileTokenSystem::UnregisterFileToken(Token);
+	Token = InvalidToken;
+
+	BulkDataFlags = 0;
+
+}
+
+int64 FBulkDataBase::GetBulkDataOffsetInFile() const
+{
+	return FileTokenSystem::GetBulkDataOffset(Token);
+}
+
+FString FBulkDataBase::GetFilename() const
+{
+	return FileTokenSystem::GetFilename(Token);
+}
+
+void FBulkDataBase::LoadDataDirectly(void** DstBuffer)
+{
+	DECLARE_SCOPE_CYCLE_COUNTER(TEXT("FBulkDataBase::LoadDataDirectly"), STAT_UBD_LoadDataDirectly, STATGROUP_Memory);
+	check(CanLoadFromDisk());
+
+	FileTokenSystem::Data FileData = FileTokenSystem::GetFileData(Token);
+
+	// In the older code path if the bulkdata was not in a separate file we could just serialize from the package instead, but the new system
+	// does not keep a reference to the pack so we cannot do that.
+	// TODO: Maybe add a check for this condition?
+
+	check(!IsInlined());	// TODO: Currently inline data will have the wrong BulkDataOffsetInFile (need to subtract the size of the original
+							// .uasset or .umap file from it but we don't have access to that filepath at this point and doing it during
+							// ::Serialize is too slow!)
+
+	// If the data is inlined then we already loaded is during ::Serialize, this warning should help track cases where data is being discarded then re-requested.
+	UE_CLOG(IsInlined(), LogSerialization, Warning, TEXT("Reloading inlined bulk data directly from disk, consider not discarding it in the first place. Filename: '%s'."), *FileData.Filename);
+
+	FArchive* Ar = IFileManager::Get().CreateFileReader(*FileData.Filename, FILEREAD_Silent);
+	checkf(Ar != nullptr, TEXT("Failed to open the file to load bulk data from. Filename: '%s'."), *FileData.Filename);
+
+	// Seek to the beginning of the bulk data in the file.
+	Ar->Seek(FileData.BulkDataOffsetInFile);
+
+	if (*DstBuffer == nullptr)
+	{
+		*DstBuffer = FMemory::Malloc(FileData.BulkDataSize, 0);
+	}
+
+	SerializeBulkData(*Ar, *DstBuffer, FileData.BulkDataSize);
+
+	delete Ar;
+}
+
+void FBulkDataBase::SerializeDuplicateData(FArchive& Ar, UObject* Owner, uint32& OutBulkDataFlags, int64& OutBulkDataSizeOnDisk, int64& OutBulkDataOffsetInFile)
+{
+	Ar << OutBulkDataFlags;
+
+	if (OutBulkDataFlags & BULKDATA_Size64Bit)
+	{
+		Ar << OutBulkDataSizeOnDisk;
+	}
+	else
+	{
+		int32 Temp32ByteValue;
+		Ar << Temp32ByteValue;
+
+		OutBulkDataSizeOnDisk = Temp32ByteValue;
+	}
+
+	Ar << OutBulkDataOffsetInFile;
+
+	// fix up the file offset, but only if not stored inline
+	if (Owner != nullptr && Owner->GetLinker() != nullptr)
+	{
+		OutBulkDataOffsetInFile += Owner->GetLinker()->Summary.BulkDataStartOffset;
+	}
+}
+
+void FBulkDataBase::SerializeBulkData(FArchive& Ar, void* DstBuffer, int64 DataLength)
+{	
+	check(Ar.IsLoading()); // Currently only support loading
+
+	if (IsAvailableForUse()) // skip serializing of unused data
+	{
+		return;
+	}
+
+	// Skip serialization for bulk data of zero length
+	if (DataLength == 0)
+	{
+		return;
+	}
+
+	check(DstBuffer != nullptr);
+
+	if (IsStoredCompressedOnDisk())
+	{
+		Ar.SerializeCompressed(DstBuffer, DataLength, GetDecompressionFormat(), COMPRESS_NoFlags, false);
+	}
+	// Uncompressed/ regular serialization.
+	else
+	{
+		Ar.Serialize(DstBuffer, DataLength);
+	}
+}
+
+void FBulkDataBase::AllocateData(SIZE_T SizeInBytes)
+{
+	DataBuffer = FMemory::Realloc(DataBuffer, SizeInBytes, DEFAULT_ALIGNMENT);
+}
+
+void FBulkDataBase::FreeData()
+{
+	FMemory::Free(DataBuffer);
+	DataBuffer = nullptr;
+}
+
+FString FBulkDataBase::ConvertFilenameFromFlags(const FString& Filename)
+{
+	if (IsInlined())
+	{
+		return FPaths::ChangeExtension(Filename, InlinedExt);
+	}
+	else if (IsOptional())
+	{
+		return FPaths::ChangeExtension(Filename, OptionalExt);
+	}
+	else if (IsMemoryMapped())
+	{
+		return FPaths::ChangeExtension(Filename, MemoryMappedExt);
+	}
+	else
+	{
+		return FPaths::ChangeExtension(Filename, DefaultExt);
+	}
+}
diff --git a/Engine/Source/Runtime/CoreUObject/Private/UObject/LinkerManager.h b/Engine/Source/Runtime/CoreUObject/Private/UObject/LinkerManager.h
--- a/Engine/Source/Runtime/CoreUObject/Private/UObject/LinkerManager.h
+++ b/Engine/Source/Runtime/CoreUObject/Private/UObject/LinkerManager.h
@@ -91,9 +91,16 @@ public:
 	}
 
 #if !UE_BUILD_SHIPPING && !UE_BUILD_TEST
-	FORCEINLINE TArray<FLinkerLoad*>& GetLiveLinkers()
+	FORCEINLINE void AddLiveLinker(FLinkerLoad* Linker)
 	{
-		return LiveLinkers;
+		FScopeLock LiveLinkersLock(&LiveLinkersCritical);
+		LiveLinkers.Add(Linker);
+	}
+
+	FORCEINLINE void RemoveLiveLinker(FLinkerLoad* Linker)
+	{
+		FScopeLock LiveLinkersLock(&LiveLinkersCritical);
+		LiveLinkers.Remove(Linker);
 	}
 #endif
 
@@ -130,6 +137,7 @@ private:
 #endif
 #if !UE_BUILD_SHIPPING && !UE_BUILD_TEST
 	/** List of all the existing linker loaders **/
+	FCriticalSection LiveLinkersCritical;
 	TArray<FLinkerLoad*> LiveLinkers;
 #endif
 	
diff --git a/Engine/Source/Runtime/CoreUObject/Private/UObject/LinkerSave.cpp b/Engine/Source/Runtime/CoreUObject/Private/UObject/LinkerSave.cpp
--- a/Engine/Source/Runtime/CoreUObject/Private/UObject/LinkerSave.cpp
+++ b/Engine/Source/Runtime/CoreUObject/Private/UObject/LinkerSave.cpp
@@ -4,6 +4,7 @@
 #include "HAL/FileManager.h"
 #include "Serialization/LargeMemoryWriter.h"
 #include "UObject/Package.h"
+#include "UObject/SavePackage.h"
 #include "UObject/Class.h"
 #include "Templates/Casts.h"
 #include "UObject/LazyObjectPtr.h"
@@ -20,9 +21,10 @@
 /** A mapping of package name to generated script SHA keys */
 TMap<FString, TArray<uint8> > FLinkerSave::PackagesToScriptSHAMap;
 
-FLinkerSave::FLinkerSave(UPackage* InParent, const TCHAR* InFilename, bool bForceByteSwapping, bool bInSaveUnversioned)
-:	FLinker(ELinkerType::Save, InParent, InFilename )
+FLinkerSave::FLinkerSave(FPackageHeaderSaver& InPackageHeaderSaver, UPackage* InParent, const TCHAR* InFilename, bool bForceByteSwapping, bool bInSaveUnversioned)
+:	FLinker(ELinkerType::Save, InParent, InFilename)
 ,	Saver(nullptr)
+,	HeaderSaver(InPackageHeaderSaver)
 {
 	if (FPlatformProperties::HasEditorOnlyData())
 	{
@@ -72,9 +74,10 @@ FLinkerSave::FLinkerSave(UPackage* InParent, const TCHAR* InFilename, bool bForc
 }
 
 
-FLinkerSave::FLinkerSave(UPackage* InParent, FArchive *InSaver, bool bForceByteSwapping, bool bInSaveUnversioned)
+FLinkerSave::FLinkerSave(FPackageHeaderSaver& InPackageHeaderSaver, UPackage* InParent, FArchive *InSaver, bool bForceByteSwapping, bool bInSaveUnversioned)
 : FLinker(ELinkerType::Save, InParent, TEXT("$$Memory$$"))
 , Saver(nullptr)
+, HeaderSaver(InPackageHeaderSaver)
 {
 	if (FPlatformProperties::HasEditorOnlyData())
 	{
@@ -124,9 +127,10 @@ FLinkerSave::FLinkerSave(UPackage* InParent, FArchive *InSaver, bool bForceByteS
 	}
 }
 
-FLinkerSave::FLinkerSave(UPackage* InParent, bool bForceByteSwapping, bool bInSaveUnversioned )
-:	FLinker(ELinkerType::Save, InParent,TEXT("$$Memory$$"))
+FLinkerSave::FLinkerSave(FPackageHeaderSaver& InPackageHeaderSaver, UPackage* InParent, bool bForceByteSwapping, bool bInSaveUnversioned )
+:	FLinker(ELinkerType::Save, InParent, TEXT("$$Memory$$"))
 ,	Saver(nullptr)
+,	HeaderSaver(InPackageHeaderSaver)
 {
 	if (FPlatformProperties::HasEditorOnlyData())
 	{
@@ -191,18 +195,6 @@ FLinkerSave::~FLinkerSave()
 	CloseAndDestroySaver();
 }
 
-int32 FLinkerSave::MapName(FNameEntryId Id) const
-{
-	const int32* IndexPtr = NameIndices.Find(Id);
-
-	if (IndexPtr)
-	{
-		return *IndexPtr;
-	}
-
-	return INDEX_NONE;
-}
-
 FPackageIndex FLinkerSave::MapObject( const UObject* Object ) const
 {
 	if (Object)
@@ -277,7 +269,7 @@ FString FLinkerSave::GetArchiveName() const
 
 FArchive& FLinkerSave::operator<<( FName& InName )
 {
-	int32 Save = MapName(InName.GetDisplayIndex());
+	int32 Save = HeaderSaver.NameMapSaver.MapName(InName);
 
 	check(GetSerializeContext());
 	ensureMsgf(Save != INDEX_NONE, TEXT("Name \"%s\" is not mapped when saving %s (object: %s, property: %s)"), 
