Copying //UE4/Dev-Core to Main (//UE4/Main)

diff --git a/Engine/Source/Runtime/CoreUObject/Private/UObject/SavePackage.cpp b/Engine/Source/Runtime/CoreUObject/Private/UObject/SavePackage.cpp
--- a/Engine/Source/Runtime/CoreUObject/Private/UObject/SavePackage.cpp
+++ b/Engine/Source/Runtime/CoreUObject/Private/UObject/SavePackage.cpp
@@ -1,5 +1,6 @@
 // Copyright 1998-2019 Epic Games, Inc. All Rights Reserved.
 
+#include "UObject/SavePackage.h"
 #include "CoreMinimal.h"
 #include "Misc/MessageDialog.h"
 #include "HAL/FileManager.h"
@@ -60,7 +61,6 @@
 
 DEFINE_LOG_CATEGORY_STATIC(LogSavePackage, Log, All);
 
-static const int32 MAX_MERGED_COMPRESSION_CHUNKSIZE = 1024 * 1024;
 static const FName WorldClassName = FName("World");
 static const FName PrestreamPackageClassName = FName("PrestreamPackage");
 static FCriticalSection InitializeCoreClassesCritSec;
@@ -68,12 +68,10 @@ static FCriticalSection InitializeCoreClassesCritSec;
 #define VALIDATE_INITIALIZECORECLASSES 0
 #define EXPORT_SORTING_DETAILED_LOGGING 0
 
-// Uncomment this code to measure UObject::Serialize time taken per UClass type during save.
-// This code tracks each type in a hash, so is too heavyweight to leave on in general,
-// but can be really useful for telling what classes are causing the most save time.
-#define ENABLE_PACKAGE_CLASS_SERIALIZATION_TIMES 0
-// uncomment this code to measure UObject::PreSave time taken per uclass type during save
-#define ENABLE_TAGEXPORTS_CLASS_PRESAVE_TIMES 0
+static void SaveThumbnails(UPackage* InOuter, FLinkerSave* Linker, FStructuredArchive::FSlot Slot);
+static void SaveAssetRegistryData(UPackage* InOuter, FLinkerSave* Linker, FStructuredArchive::FSlot Slot);
+static void SaveWorldLevelInfo(UPackage* InOuter, FLinkerSave* Linker, FStructuredArchive::FRecord Record);
+static EObjectMark GetExcludedObjectMarksForTargetPlatform(const class ITargetPlatform* TargetPlatform, const bool bIsCooking);
 
 #if ENABLE_COOK_STATS
 #include "ProfilingDebugging/ScopedTimers.h"
@@ -95,9 +93,6 @@ namespace SavePackageStats
 	static double MBWritten = 0.0;
 	TMap<FName, FArchiveDiffStats> PackageDiffStats;
 	static int32 NumberOfDifferentPackages = 0;
-	static TMap<FName, FCookStatsManager::TKeyValuePair<double, uint32>> PackageClassSerializeTimes;
-	static TMap<FName, FCookStatsManager::TKeyValuePair<double, uint32>> TagExportSerializeTimes;
-	static TMap<FName, FCookStatsManager::TKeyValuePair<double, uint32>> ClassPreSaveTimes;
 	static FCookStatsManager::FAutoRegisterCallback RegisterCookStats([](FCookStatsManager::AddStatFuncRef AddStat)
 	{
 		// Don't use FCookStatsManager::CreateKeyValueArray because there's just too many arguments. Don't need to overburden the compiler here.
@@ -181,127 +176,6 @@ namespace SavePackageStats
 		#undef ADD_COOK_STAT		
 		
 		const FString TotalString = TEXT("Total");
-
-		if (PackageClassSerializeTimes.Num() > 0)
-		{
-			// Sort the class serialize times in reverse order.
-			typedef FCookStatsManager::TKeyValuePair<FName, FCookStatsManager::TKeyValuePair<double, uint32>> FClassSerializeTimeData;
-			TArray<FClassSerializeTimeData> SerializeTimesArray;
-			SerializeTimesArray.Empty(PackageClassSerializeTimes.Num());
-			for (const TPair<FName, FCookStatsManager::TKeyValuePair<double, uint32>>& KV : PackageClassSerializeTimes)
-			{
-				SerializeTimesArray.Emplace(FCookStatsManager::MakePair(KV.Key, FCookStatsManager::MakePair(KV.Value.Key, KV.Value.Value)));
-			}
-			SerializeTimesArray.Sort([](const FClassSerializeTimeData& LHS, const FClassSerializeTimeData& RHS)
-			{
-				return LHS.Value.Key > RHS.Value.Key;
-			});
-
-			// always print at least the top n, but not anything < 0.1% of total save time.
-			int ClassesLogged = 0;
-			for (const FClassSerializeTimeData& KV : SerializeTimesArray)
-			{
-				// since we're sorted on size already. Just find the first one below the threshold and stop there
-				if (ClassesLogged >= 10 && KV.Value.Key < 0.001 * SavePackageTimeSec)
-				{
-					break;
-				}
-				const FString ClassName = KV.Key.ToString();
-				AddStat(TEXT("Package.Serialize"), FCookStatsManager::CreateKeyValueArray(
-					TEXT("Class"), ClassName, 
-					TEXT("TimeSec"), KV.Value.Key, 
-					TEXT("Calls"), KV.Value.Value));
-				ClassesLogged++;
-			}
-		}
-
-		if (TagExportSerializeTimes.Num() > 0)
-		{
-			// Sort the class serialize times in reverse order.
-			typedef FCookStatsManager::TKeyValuePair<FName, FCookStatsManager::TKeyValuePair<double, uint32>> FClassSerializeTimeData;
-			TArray<FClassSerializeTimeData> SerializeTimesArray;
-			SerializeTimesArray.Empty(TagExportSerializeTimes.Num());
-			for (const TPair<FName, FCookStatsManager::TKeyValuePair<double, uint32>>& KV : TagExportSerializeTimes)
-			{
-				SerializeTimesArray.Emplace(FCookStatsManager::MakePair(KV.Key, FCookStatsManager::MakePair(KV.Value.Key, KV.Value.Value)));
-			}
-			SerializeTimesArray.Sort([](const FClassSerializeTimeData& LHS, const FClassSerializeTimeData& RHS)
-			{
-				return LHS.Value.Key > RHS.Value.Key;
-			});
-
-			double TotalSerializeTime = 0.0;
-			int TotalSerializeCalls = 0;
-
-			// always print at least the top n, but not anything < 0.1% of total save time.
-			int ClassesLogged = 0;
-			for (const FClassSerializeTimeData& KV : SerializeTimesArray)
-			{
-				TotalSerializeTime += KV.Value.Key;
-				TotalSerializeCalls += KV.Value.Value;
-
-				// since we're sorted on size already. Just find the first one below the threshold and stop there
-				if (ClassesLogged >= 10 && KV.Value.Key < 0.001 * SavePackageTimeSec)
-				{
-					break;
-				}
-				const FString ClassName = KV.Key.ToString();
-				AddStat(TEXT("Package.TagExportSerialize"), FCookStatsManager::CreateKeyValueArray(
-					TEXT("Class"), ClassName,
-					TEXT("TimeSec"), KV.Value.Key,
-					TEXT("Calls"), KV.Value.Value));
-				ClassesLogged++;
-			}
-
-			AddStat(TEXT("Package.TagExportSerialize"), FCookStatsManager::CreateKeyValueArray(
-				TEXT("Class"), TotalString,
-				TEXT("TimeSec"), TotalSerializeTime,
-				TEXT("Calls"), TotalSerializeCalls));
-		}
-
-		if (ClassPreSaveTimes.Num() > 0)
-		{
-			// Sort the class serialize times in reverse order.
-			typedef FCookStatsManager::TKeyValuePair<FName, FCookStatsManager::TKeyValuePair<double, uint32>> FClassSerializeTimeData;
-			TArray<FClassSerializeTimeData> SerializeTimesArray;
-			SerializeTimesArray.Empty(ClassPreSaveTimes.Num());
-			for (const TPair<FName, FCookStatsManager::TKeyValuePair<double, uint32>>& KV : ClassPreSaveTimes)
-			{
-				SerializeTimesArray.Emplace(FCookStatsManager::MakePair(KV.Key, FCookStatsManager::MakePair(KV.Value.Key, KV.Value.Value)));
-			}
-			SerializeTimesArray.Sort([](const FClassSerializeTimeData& LHS, const FClassSerializeTimeData& RHS)
-			{
-				return LHS.Value.Key > RHS.Value.Key;
-			});
-
-			// always print at least the top n, but not anything < 0.1% of total save time.
-			// even if we don't print them then add up the time spent inside presave so we can at least account for it
-			double TotalPreSaveTime = 0.0;
-			int TotalPreSaveCalls = 0;
-			int ClassesLogged = 0;
-			for (const FClassSerializeTimeData& KV : SerializeTimesArray)
-			{
-				TotalPreSaveTime += KV.Value.Key;
-				TotalPreSaveCalls += KV.Value.Value;
-
-				// since we're sorted on size already. Just find the first one below the threshold and stop there
-				if (ClassesLogged >= 10 && KV.Value.Key < 0.001 * SavePackageTimeSec)
-				{
-					break;
-				}
-				const FString ClassName = KV.Key.ToString();
-				AddStat(TEXT("Package.PreSave"), FCookStatsManager::CreateKeyValueArray(
-					TEXT("Class"), ClassName,
-					TEXT("TimeSec"), KV.Value.Key,
-					TEXT("Calls"), KV.Value.Value));
-				ClassesLogged++;
-			}
-			
-			AddStat(TEXT("Package.PreSave"), FCookStatsManager::CreateKeyValueArray(
-				TEXT("Class"), TotalString,
-				TEXT("TimeSec"), TotalPreSaveTime,
-				TEXT("Calls"), TotalPreSaveCalls));
-		}
 	});
 }
 #endif
@@ -450,6 +324,28 @@ void UPackage::WaitForAsyncFileWrites()
 	}
 }
 
+static void WriteToFile(const FString& Filename, const uint8* InDataPtr, int64 InDataSize)
+{
+	IFileManager& FileManager = IFileManager::Get();
+
+	if (FArchive* Ar = FileManager.CreateFileWriter(*Filename))
+	{
+		Ar->Serialize(/* grrr */ const_cast<uint8*>(InDataPtr), InDataSize);
+		delete Ar;
+
+		if (FileManager.FileSize(*Filename) != InDataSize)
+		{
+			FileManager.Delete(*Filename);
+
+			UE_LOG(LogSavePackage, Fatal, TEXT("Could not save to %s!"), *Filename);
+		}
+	}
+	else
+	{
+		UE_LOG(LogSavePackage, Fatal, TEXT("Could not write to %s!"), *Filename);
+	}
+}
+
 struct FLargeMemoryDelete
 {
 	void operator()(uint8* Ptr) const
@@ -463,89 +359,27 @@ struct FLargeMemoryDelete
 
 typedef TUniquePtr<uint8, FLargeMemoryDelete> FLargeMemoryPtr;
 
-void AsyncWriteFile(FLargeMemoryPtr Data, const int64 DataSize, const TCHAR* Filename, const FDateTime& TimeStamp, bool bUseTempFilename = true)
+static void AsyncWriteFile(FLargeMemoryPtr Data, const int64 DataSize, const TCHAR* Filename)
 {
 	class FAsyncWriteWorker : public FNonAbandonableTask
 	{
 	public:
-		/** Filename To write to**/
-		FString Filename;
-		/** Should we write to a temp file then move it */
-		bool bUseTempFilename;
-		/** Data for the file. Will be freed after write **/
+		FString			Filename;
 		FLargeMemoryPtr Data;
-		/** Size of data */
-		const int64 DataSize;
-		/** Timestamp to give the file. MinValue if shouldn't be modified */
-		FDateTime FinalTimeStamp;
-
-		/** Constructor
-		*/
-		FAsyncWriteWorker(const TCHAR* InFilename, FLargeMemoryPtr InData, const int64 InDataSize, const FDateTime& InTimeStamp, bool inbUseTempFilename)
+		const int64		DataSize;
+
+		FAsyncWriteWorker(const TCHAR* InFilename, FLargeMemoryPtr InData, const int64 InDataSize)
 			: Filename(InFilename)
-			, bUseTempFilename(inbUseTempFilename)
 			, Data(MoveTemp(InData))
 			, DataSize(InDataSize)
-			, FinalTimeStamp(InTimeStamp)
 		{
+			check(InDataSize);
 		}
 		
 		/** Write the file  */
 		void DoWork()
 		{
-			check(DataSize);
-			FString TempFilename; 
-
-			if ( bUseTempFilename )
-			{
-				TempFilename = FPaths::GetBaseFilename(Filename, false);
-				TempFilename += TEXT(".t");
-			}
-			else
-			{
-				TempFilename = Filename;
-			}
-			
-			// Open a file writer for saving
-			FArchive* Ar = IFileManager::Get().CreateFileWriter(*TempFilename);
-			if (Ar)
-			{
-				Ar->Serialize(Data.Get(), DataSize);
-				delete Ar;
-				
-				// Clean-up the memory as soon as we save the file to reduce the memory footprint.
-				Data.Reset();
-
-				if (IFileManager::Get().FileSize(*TempFilename) == DataSize)
-				{
-					if ( bUseTempFilename )
-					{
-						if( !IFileManager::Get().Move(*Filename, *TempFilename, true, true, false, false))
-						{
-							UE_LOG(LogSavePackage, Fatal, TEXT("Could not move to %s."),*Filename);
-						}
-
-						// if everything worked, this is not necessary, but we will make every effort to avoid leaving junk in the cache
-						if (FPaths::FileExists(TempFilename))
-						{
-							IFileManager::Get().Delete(*TempFilename);
-						}
-					}
-
-					if (FinalTimeStamp != FDateTime::MinValue())
-					{
-						IFileManager::Get().SetTimeStamp(*Filename, FinalTimeStamp);
-					}
-				}
-				else
-				{
-					UE_LOG(LogSavePackage, Fatal, TEXT("Could not save to %s!"),*TempFilename);
-				}
-			}
-			else
-			{
-				UE_LOG(LogSavePackage, Fatal, TEXT("Could not write to %s!"),*TempFilename);
-			}
+			WriteToFile(Filename, Data.Get(), DataSize);
 			
 			OutstandingAsyncWrites.Decrement();
 		}
@@ -554,150 +388,40 @@ void AsyncWriteFile(FLargeMemoryPtr Data, const int64 DataSize, const TCHAR* Fil
 		{
 			RETURN_QUICK_DECLARE_CYCLE_STAT(FAsyncWriteWorker, STATGROUP_ThreadPoolAsyncTasks);
 		}
-
 	};
 
 	OutstandingAsyncWrites.Increment();
-	(new FAutoDeleteAsyncTask<FAsyncWriteWorker>(Filename, MoveTemp(Data), DataSize, TimeStamp, bUseTempFilename))->StartBackgroundTask();
+	(new FAutoDeleteAsyncTask<FAsyncWriteWorker>(Filename, MoveTemp(Data), DataSize))->StartBackgroundTask();
 }
 
-void AsyncWriteFileWithSplitExports(FLargeMemoryPtr Data, const int64 DataSize, const int64 HeaderSize, const TCHAR* Filename, const FDateTime& TimeStamp, bool bUseTempFilename = true)
+static void AsyncWriteFileWithSplitExports(FLargeMemoryPtr Data, const int64 DataSize, const int64 HeaderSize, const TCHAR* Filename)
 {
 	class FAsyncWriteWorkerWithSplitExports : public FNonAbandonableTask
 	{
 	public:
-		/** Filename To write to**/
-		FString Filename;
-		/** Should we write to a temp file then move it */
-		bool bUseTempFilename;
-		/** Data for the file. Will be freed after write **/
+		FString			Filename;
 		FLargeMemoryPtr Data;
-		/** Size of data */
-		const int64 DataSize;
-		/** Size of data */
-		const int64 HeaderSize;
-		/** Timestamp to give the file. MinValue if shouldn't be modified */
-		FDateTime FinalTimeStamp;
-
-		/** Constructor
-		*/
-		FAsyncWriteWorkerWithSplitExports(const TCHAR* InFilename, FLargeMemoryPtr InData, const int64 InDataSize, const int64 InHeaderSize, const FDateTime& InTimeStamp, bool inbUseTempFilename)
+		const int64		DataSize;
+		const int64		HeaderSize;
+
+		FAsyncWriteWorkerWithSplitExports(const TCHAR* InFilename, FLargeMemoryPtr InData, const int64 InDataSize, const int64 InHeaderSize)
 			: Filename(InFilename)
-			, bUseTempFilename(inbUseTempFilename)
 			, Data(MoveTemp(InData))
 			, DataSize(InDataSize)
 			, HeaderSize(InHeaderSize)
-			, FinalTimeStamp(InTimeStamp)
 		{
 			check(InDataSize);
 		}
 
-		/** Write the file  */
 		void DoWork()
 		{
-			check(DataSize);
-			FString BaseFilename;
-			FString TempFilename;
-			FString FilenameExports;
-			FString TempFilenameExports;
+			// Write .uasset file
+			WriteToFile(Filename, Data.Get(), HeaderSize);
 
-			BaseFilename = FPaths::GetBaseFilename(Filename, false);
-			FilenameExports = BaseFilename;
-			FilenameExports += TEXT(".uexp");
-			if (bUseTempFilename)
-			{
-				TempFilename = BaseFilename;
-				TempFilename += TEXT(".t");
-				TempFilenameExports = BaseFilename;
-				TempFilenameExports += TEXT(".e");
-			}
-			else
-			{
-				TempFilename = Filename;
-				TempFilenameExports = FilenameExports;
-			}
-
-			{
-				FArchive* Ar = IFileManager::Get().CreateFileWriter(*TempFilename);
-				if (Ar)
-				{
-					Ar->Serialize(Data.Get(), HeaderSize);
-					delete Ar;
-					Ar = nullptr;
-
-					if (IFileManager::Get().FileSize(*TempFilename) == HeaderSize)
-					{
-						if (bUseTempFilename)
-						{
-							if (!IFileManager::Get().Move(*Filename, *TempFilename, true, true, false, false))
-							{
-								UE_LOG(LogSavePackage, Fatal, TEXT("Could not move to %s."), *Filename);
-							}
-
-							// if everything worked, this is not necessary, but we will make every effort to avoid leaving junk in the cache
-							if (FPaths::FileExists(TempFilename))
-							{
-								IFileManager::Get().Delete(*TempFilename);
-							}
-						}
+			// Write .uexp file
+			const FString FilenameExports = FPaths::ChangeExtension(Filename, TEXT(".uexp"));
+			WriteToFile(FilenameExports, Data.Get() + HeaderSize, DataSize - HeaderSize);
 
-						if (FinalTimeStamp != FDateTime::MinValue())
-						{
-							IFileManager::Get().SetTimeStamp(*Filename, FinalTimeStamp);
-						}
-					}
-					else
-					{
-						UE_LOG(LogSavePackage, Fatal, TEXT("Could not save to %s!"), *TempFilename);
-					}
-				}
-				else
-				{
-					UE_LOG(LogSavePackage, Fatal, TEXT("Could not write to %s!"), *TempFilename);
-				}
-			}
-			{
-				FArchive* Ar = IFileManager::Get().CreateFileWriter(*TempFilenameExports);
-				if (Ar)
-				{
-					Ar->Serialize(Data.Get() + HeaderSize, DataSize - HeaderSize);
-					delete Ar;
-					Ar = nullptr;
-
-					// Clean-up the memory as soon as we save the file to reduce the memory footprint.
-					Data.Reset();
-
-					if (IFileManager::Get().FileSize(*TempFilenameExports) == DataSize - HeaderSize)
-					{
-						if (bUseTempFilename)
-						{
-							if (!IFileManager::Get().Move(*FilenameExports, *TempFilenameExports, true, true, false, false))
-							{
-								UE_LOG(LogSavePackage, Fatal, TEXT("Could not move to %s."), *Filename);
-							}
-
-							// if everything worked, this is not necessary, but we will make every effort to avoid leaving junk in the cache
-							if (FPaths::FileExists(TempFilenameExports))
-							{
-								IFileManager::Get().Delete(*TempFilenameExports);
-							}
-						}
-
-						if (FinalTimeStamp != FDateTime::MinValue())
-						{
-							IFileManager::Get().SetTimeStamp(*FilenameExports, FinalTimeStamp);
-						}
-					}
-					else
-					{
-						UE_LOG(LogSavePackage, Fatal, TEXT("Could not save to %s!"), *TempFilenameExports);
-					}
-				}
-				else
-				{
-					UE_LOG(LogSavePackage, Fatal, TEXT("Could not write to %s!"), *TempFilenameExports);
-				}
-			}
 			OutstandingAsyncWrites.Decrement();
 		}
 
@@ -705,76 +429,165 @@ void AsyncWriteFileWithSplitExports(FLargeMemoryPtr Data, const int64 DataSize,
 		{
 			RETURN_QUICK_DECLARE_CYCLE_STAT(FAsyncWriteWorkerWithSplitExports, STATGROUP_ThreadPoolAsyncTasks);
 		}
-
 	};
 
 	OutstandingAsyncWrites.Increment();
-	(new FAutoDeleteAsyncTask<FAsyncWriteWorkerWithSplitExports>(Filename, MoveTemp(Data), DataSize, HeaderSize, TimeStamp, bUseTempFilename))->StartBackgroundTask();
+	(new FAutoDeleteAsyncTask<FAsyncWriteWorkerWithSplitExports>(Filename, MoveTemp(Data), DataSize, HeaderSize))->StartBackgroundTask();
 }
 
-/** 
- * Helper object for all of the cross object state that is needed only while saving a package.
- */
-class FSavePackageState
+void FSinglePackageNameMapSaver::MarkNameAsReferenced(FNameEntryId Name)
 {
-public:
-	/** 
-	 * Mark a FName as referenced.
-	 * @param Name name to mark as referenced
-	 */
-	void MarkNameAsReferenced(const FName& Name)
+	ReferencedNames.Add(Name);
+}
+
+int32 FSinglePackageNameMapSaver::MapName(FName Name) const
+{
+	const FNameEntryId Id = Name.GetDisplayIndex();
+	const int32* Index = NameIndices.Find(Id);
+	return Index ? *Index : INDEX_NONE;
+}
+
+bool FSinglePackageNameMapSaver::NameExistsInCurrentPackage(FNameEntryId ComparisonId) const
+{
+	for (FNameEntryId DisplayId : ReferencedNames)
 	{
-		ReferencedNames.Add(Name.GetDisplayIndex());
+		if (FName::GetComparisonIdFromDisplayId(DisplayId) == ComparisonId)
+		{
+			return true;
+		}
 	}
+	return false;
+}
+
+void FSinglePackageNameMapSaver::MarkNameAsReferenced(FName Name)
+{
+	ReferencedNames.Add(Name.GetDisplayIndex());
+}
 
-	void MarkNameAsReferenced(FNameEntryId Name)
+
+FPackageStoreNameMapSaver::FPackageStoreNameMapSaver(const TCHAR* InFilename)
+: Filename(InFilename) {}
+
+void FPackageStoreNameMapSaver::BeginPackage()
+{
+	PackageReferencedNames.Reset();
+}
+
+void FPackageStoreNameMapSaver::MarkNameAsReferenced(FName Name)
+{
+	const FNameEntryId Id = Name.GetDisplayIndex();
+	int32& Index = NameIndices.FindOrAdd(Id);
+	if (Index == 0)
 	{
-		ReferencedNames.Add(Name);
+		Index = NameIndices.Num();
+		NameMap.Add(Id);
 	}
+	PackageReferencedNames.Add(Id);
+}
 
-#if WITH_EDITOR
-	/**
-	 * Names are cached before we replace objects for imports. So the names of replacements must be ealier.
-	 */
-	void AddReplacementsNames(UObject* Obj, const ITargetPlatform* TargetPlatform, const bool bIsCooking)
+int32 FPackageStoreNameMapSaver::MapName(FName Name) const
+{
+	const FNameEntryId Id = Name.GetDisplayIndex();
+	// debug counts
 	{
-		if (bIsCooking && TargetPlatform)
+		const int32 Number = Name.GetNumber();
+		TTuple<int32,int32,int32>& Counts = DebugNameCounts.FindOrAdd(Id);
+
+		if (Number == 0)
+		{
+			++Counts.Get<0>();
+		}
+		else
 		{
-			if (const IBlueprintNativeCodeGenCore* Coordinator = IBlueprintNativeCodeGenCore::Get())
+			++Counts.Get<1>();
+			if (Number > Counts.Get<2>())
 			{
-				const FCompilerNativizationOptions& NativizationOptions = Coordinator->GetNativizationOptionsForPlatform(TargetPlatform);
-				if (const UClass* ReplObjClass = Coordinator->FindReplacedClassForObject(Obj, NativizationOptions))
-				{
-					MarkNameAsReferenced(ReplObjClass->GetFName());
-				}
+				Counts.Get<2>() = Number;
+			}
+		}
+
+	}
 
-				FName ReplacedName;
-				Coordinator->FindReplacedNameAndOuter(Obj, ReplacedName, NativizationOptions); //TODO: should we care about replaced outer ?
-				if (ReplacedName != NAME_None)
+	const int32* Index = NameIndices.Find(Id);
+	return Index ? *Index - 1 : INDEX_NONE;
+}
+
+bool FPackageStoreNameMapSaver::NameExistsInCurrentPackage(FNameEntryId ComparisonId) const
+{
+	for (FNameEntryId DisplayId : PackageReferencedNames)
+	{
+		if (FName::GetComparisonIdFromDisplayId(DisplayId) == ComparisonId)
+		{
+			return true;
+		}
+	}
+	return false;
+}
+
+void FPackageStoreNameMapSaver::End()
+{
+	// real megafile.unamemap
+	{
+		TUniquePtr<FArchive> BinArchive(IFileManager::Get().CreateFileWriter(*Filename));
+		int32 NameCount = NameMap.Num();
+		*BinArchive << NameCount;
+		for (int32 I = 0; I < NameCount; ++I)
+		{
+			FName::GetEntry(NameMap[I])->Write(*BinArchive);
+		}
+	}
+	// debug megafile.unamemap.csv
+	{
+		FString CsvFilePath = Filename;
+		CsvFilePath.Append(TEXT(".csv"));
+		TUniquePtr<FArchive> CsvArchive(IFileManager::Get().CreateFileWriter(*CsvFilePath));
+		if (CsvArchive)
+		{
+			TCHAR Name[FName::StringBufferSize];
+			ANSICHAR Line[MAX_SPRINTF + FName::StringBufferSize];
+			ANSICHAR Header[] = "Length\tMaxNumber\tNumberCount\tBaseCount\tTotalCount\tFName\n";
+			CsvArchive->Serialize(Header, sizeof(Header) - 1);
+			for (auto& Counts : DebugNameCounts)
+			{
+				const int32 NameLen = FName::CreateFromDisplayId(Counts.Key, 0).ToString(Name);
+				FCStringAnsi::Sprintf(Line, "%d\t%d\t%d\t%d\t%d\t",
+					NameLen, Counts.Value.Get<2>(), Counts.Value.Get<1>(), Counts.Value.Get<0>(), Counts.Value.Get<0>() + Counts.Value.Get<1>());
+				ANSICHAR* L = Line + FCStringAnsi::Strlen(Line);
+				const TCHAR* N = Name;
+				while (*N)
 				{
-					MarkNameAsReferenced(ReplacedName);
+					*L++ = CharCast<ANSICHAR,TCHAR>(*N++);
 				}
+				*L++ = '\n';
+				CsvArchive.Get()->Serialize(Line, L - Line);
 			}
 		}
 	}
-#endif //WITH_EDITOR
+}
 
-	/** 
-	 * Add the marked names to a linker
-	 * @param Linker linker to save the marked names to
-	 */
-	void UpdateLinkerWithMarkedNames(FLinkerSave* Linker)
+#if WITH_EDITOR
+static void AddReplacementsNames(INameMapSaver& NameMapSaver, UObject* Obj, const ITargetPlatform* TargetPlatform)
+{
+	if (TargetPlatform)
 	{
-		Linker->NameMap.Reserve(Linker->NameMap.Num() + ReferencedNames.Num());
-		for (FNameEntryId Name : ReferencedNames)
+		if (const IBlueprintNativeCodeGenCore* Coordinator = IBlueprintNativeCodeGenCore::Get())
 		{
-			Linker->NameMap.Add(Name);
+			const FCompilerNativizationOptions& NativizationOptions = Coordinator->GetNativizationOptionsForPlatform(TargetPlatform);
+			if (const UClass* ReplObjClass = Coordinator->FindReplacedClassForObject(Obj, NativizationOptions))
+			{
+				NameMapSaver.MarkNameAsReferenced(ReplObjClass->GetFName());
+			}
+
+			FName ReplacedName;
+			Coordinator->FindReplacedNameAndOuter(Obj, ReplacedName, NativizationOptions); //TODO: should we care about replaced outer ?
+			if (ReplacedName != NAME_None)
+			{
+				NameMapSaver.MarkNameAsReferenced(ReplacedName);
+			}
 		}
 	}
-
-private:
-	TSet<FNameEntryId> ReferencedNames;
-};
+}
+#endif //WITH_EDITOR
 
 bool IsEditorOnlyObject(const UObject* InObject, bool bCheckRecursive, bool bCheckMarks)
 {
@@ -866,10 +679,14 @@ static void ConditionallyExcludeObjectForTarget(UObject* Obj, EObjectMark Exclud
 		MarksToModify = (EObjectMark)(MarksToModify | (ObjToCheckMarks & MarkMask));
 	};
 
+	// MarksToProcess is a superset of marks retrieved from UPackage::GetExcludedObjectMarksForTargetPlatform
+	const uint32 MarksToProcess = OBJECTMARK_EditorOnly | OBJECTMARK_NotForClient | OBJECTMARK_NotForServer | OBJECTMARK_KeepForTargetPlatform;
+	check((ExcludedObjectMarks & ~MarksToProcess) == 0);
+
 	EObjectMark CurrentMarks = OBJECTMARK_NOMARKS;
-	InheritMarks(CurrentMarks, Obj, OBJECTMARK_EditorOnly | OBJECTMARK_NotForClient | OBJECTMARK_NotForServer);
+	InheritMarks(CurrentMarks, Obj, MarksToProcess);
 
-	if ((CurrentMarks & ExcludedObjectMarks) != 0)
+	if ((CurrentMarks & MarksToProcess) != 0)
 	{
 		// Already marked
 		return;
@@ -955,6 +772,12 @@ static void ConditionallyExcludeObjectForTarget(UObject* Obj, EObjectMark Exclud
 		NewMarks = (EObjectMark)(NewMarks | OBJECTMARK_EditorOnly);
 	}
 
+	// If not excluded after a full set of tests, it is implicitly a keep
+	if (NewMarks == 0)
+	{
+		NewMarks = OBJECTMARK_KeepForTargetPlatform;
+	}
+
 	// If our marks are different than original, set them on the object
 	if (CurrentMarks != NewMarks)
 	{
@@ -1078,7 +901,7 @@ FArchive& FArchiveSaveTagExports::operator<<(UObject*& Obj)
 	}
 
 	// Check outer chain for any exlcuded object marks
-	const EObjectMark ExcludedObjectMarks = UPackage::GetExcludedObjectMarksForTargetPlatform(CookingTarget(), IsCooking());
+	const EObjectMark ExcludedObjectMarks = GetExcludedObjectMarksForTargetPlatform(CookingTarget(), IsCooking());
 	ConditionallyExcludeObjectForTarget(Obj, ExcludedObjectMarks, CookingTarget(), IsCooking());
 
 	if (!Obj->HasAnyMarks((EObjectMark)(ExcludedObjectMarks)))
@@ -1177,11 +1000,6 @@ void FArchiveSaveTagExports::ProcessTaggedObjects()
 			}
 			// In the CDO case the above would serialize most of the references, including transient properties
 			// but we still want to serialize the object using the normal path to collect all custom versions it might be using.
-#if ENABLE_PACKAGE_CLASS_SERIALIZATION_TIMES
-			auto& TimingInfo = SavePackageStats::TagExportSerializeTimes.FindOrAdd(Obj->GetClass()->GetFName());
-			TimingInfo.Value++;
-			FScopedDurationTimer SerializeTimer(TimingInfo.Key);
-#endif
 			Obj->Serialize(*this);
 		}
 
@@ -1197,7 +1015,7 @@ class FArchiveSaveTagImports : public FArchiveUObject
 {
 public:
 	FLinkerSave* Linker;
-	FSavePackageState& SavePackageState;
+	INameMapSaver& NameMapSaver;
 	TArray<UObject*> Dependencies;
 	TArray<UObject*> NativeDependencies;
 	TArray<UObject*> OtherImports;
@@ -1222,9 +1040,9 @@ public:
 		}
 	};
 
-	FArchiveSaveTagImports(FLinkerSave* InLinker, FSavePackageState& InSavePackageState)
+	FArchiveSaveTagImports(FLinkerSave* InLinker, INameMapSaver& InNameMapSaver)
 		: Linker(InLinker)
-		, SavePackageState(InSavePackageState)
+		, NameMapSaver(InNameMapSaver)
 		, bIgnoreDependencies(false)
 	{
 		check(Linker);
@@ -1290,7 +1108,7 @@ FArchive& FArchiveSaveTagImports::operator<<( UObject*& Obj )
 	// Check transient and pending kill flags for outers
 	CheckObjectPriorToSave(*this, Obj, nullptr);
 
-	const EObjectMark ExcludedObjectMarks = UPackage::GetExcludedObjectMarksForTargetPlatform( CookingTarget(), IsCooking() );
+	const EObjectMark ExcludedObjectMarks = GetExcludedObjectMarksForTargetPlatform( CookingTarget(), IsCooking() );
 	ConditionallyExcludeObjectForTarget(Obj, ExcludedObjectMarks, CookingTarget(), IsCooking());
 	bool bExcludePackageFromCook = Obj && FCoreUObjectDelegates::ShouldCookPackageForPlatform.IsBound() ? !FCoreUObjectDelegates::ShouldCookPackageForPlatform.Execute(Obj->GetOutermost(), CookingTarget()) : false;
 
@@ -1363,7 +1181,7 @@ FArchive& FArchiveSaveTagImports::operator<<( UObject*& Obj )
 						}
 					}
 #if WITH_EDITOR
-					SavePackageState.AddReplacementsNames(Obj, CookingTarget(), IsCooking());
+					AddReplacementsNames(NameMapSaver, Obj, CookingTarget());
 #endif //WITH_EDITOR
 				}
 
@@ -1424,7 +1242,7 @@ FArchive& FArchiveSaveTagImports::operator<<(FSoftObjectPath& Value)
 			// Don't track if this is a never collect path
 			FString Path = Value.ToString();
 			FName PackageName = FName(*FPackageName::ObjectPathToPackageName(Path));
-			SavePackageState.MarkNameAsReferenced(PackageName);
+			NameMapSaver.MarkNameAsReferenced(PackageName);
 			Linker->SoftPackageReferenceList.AddUnique(PackageName);
 		}
 	}
@@ -1433,7 +1251,7 @@ FArchive& FArchiveSaveTagImports::operator<<(FSoftObjectPath& Value)
 
 FArchive& FArchiveSaveTagImports::operator<<(FName& Name)
 {
-	SavePackageState.MarkNameAsReferenced(Name);
+	NameMapSaver.MarkNameAsReferenced(Name);
 	return *this;
 }
 
@@ -1454,7 +1272,7 @@ void FArchiveSaveTagImports::MarkSearchableName(const UObject* TypeObject, const
 	}
 
 	// Manually mark the name as referenced, in case it got skipped due to delta serialization
-	SavePackageState.MarkNameAsReferenced(ValueName);
+	NameMapSaver.MarkNameAsReferenced(ValueName);
 
 	Linker->SearchableNamesObjectMap.FindOrAdd(TypeObject).AddUnique(ValueName);
 }
@@ -1538,7 +1356,7 @@ public:
 	 * @param	Linker				linker containing the names that need to be sorted
 	 * @param	LinkerToConformTo	optional linker to conform against.
 	 */
-	void SortNames( FLinkerSave* Linker, FLinkerLoad* LinkerToConformTo, FSavePackageState& SavePackageState)
+	void SortNames( FLinkerSave* Linker, FLinkerLoad* LinkerToConformTo, FSinglePackageNameMapSaver& NameMapSaver)
 	{
 		int32 SortStartPosition = 0;
 
@@ -1559,7 +1377,7 @@ public:
 			for ( int32 NameIndex = 0; NameIndex < Linker->NameMap.Num(); NameIndex++ )
 			{
 				FNameEntryId CurrentName = Linker->NameMap[NameIndex];
-				SavePackageState.MarkNameAsReferenced(CurrentName);
+				NameMapSaver.MarkNameAsReferenced(CurrentName);
 			}
 		}
 
@@ -1570,6 +1388,31 @@ public:
 	}
 };
 
+void FSinglePackageNameMapSaver::EndPackage(FLinkerSave& Linker, FLinkerLoad* Conform, FArchive* BinarySaver)
+{
+	// Add names
+	Linker.NameMap.Reserve(Linker.NameMap.Num() + ReferencedNames.Num());
+	for (FNameEntryId Name : ReferencedNames)
+	{
+		Linker.NameMap.Add(Name);
+	}
+
+	// Sort names
+	FObjectNameSortHelper NameSortHelper;
+	NameSortHelper.SortNames(&Linker, Conform, *this);
+
+	// Serialize names and build NameIndicies
+	if (BinarySaver)
+	{
+		Linker.Summary.NameCount = Linker.NameMap.Num();
+		for (int32 i = 0; i < Linker.NameMap.Num(); i++)
+		{
+			FName::GetEntry(Linker.NameMap[i])->Write(Linker);
+			NameIndices.Add(Linker.NameMap[i], i);
+		}
+	}
+}
+
 /**
  * Helper structure to encapsulate sorting a linker's import table according to the of the import table of the package being conformed against.
  */
@@ -1633,6 +1476,8 @@ public:
 			Used.AddZeroed(Orig.Num());
 
 			TMap<FString,int32> OriginalImportIndexes;
+			OriginalImportIndexes.Reserve(Orig.Num());
+			ObjectToFullNameMap.Reserve(Orig.Num());
 			for ( int32 i = 0; i < Orig.Num(); i++ )
 			{
 				FObjectImport& Import = Orig[i];
@@ -1675,12 +1520,13 @@ public:
 		}
 		else
 		{
+			ObjectToFullNameMap.Reserve(Linker->ImportMap.Num());
 			for ( int32 ImportIndex = 0; ImportIndex < Linker->ImportMap.Num(); ImportIndex++ )
 			{
 				const FObjectImport& Import = Linker->ImportMap[ImportIndex];
 				if ( Import.XObject )
 				{
-					ObjectToFullNameMap.Add(Import.XObject, *Import.XObject->GetFullName());
+					ObjectToFullNameMap.Add(Import.XObject, Import.XObject->GetFullName());
 				}
 			}
 		}
@@ -1806,6 +1652,15 @@ public:
 	{
 		bUseFObjectFullName = InbUseFObjectFullName;
 
+		if (bUseFObjectFullName)
+		{
+			ObjectToObjectFullNameMap.Reserve(Linker->ExportMap.Num());
+		}
+		else
+		{
+			ObjectToFullNameMap.Reserve(Linker->ExportMap.Num());
+		}
+
 		int32 SortStartPosition=0;
 		if ( LinkerToConformTo )
 		{
@@ -1817,6 +1672,7 @@ public:
 
 			// Populate object to current index map.
 			TMap<FString,int32> OriginalExportIndexes;
+			OriginalExportIndexes.Reserve(Linker->ExportMap.Num());
 			for( int32 ExportIndex=0; ExportIndex < Linker->ExportMap.Num(); ExportIndex++ )
 			{
 				const FObjectExport& Export = Linker->ExportMap[ExportIndex];
@@ -1920,7 +1776,7 @@ public:
 					}
 					else
 					{
-						ObjectToFullNameMap.Add(Export.Object, *Export.Object->GetFullName());
+						ObjectToFullNameMap.Add(Export.Object, Export.Object->GetFullName());
 					}
 				}
 			}
@@ -2876,11 +2732,6 @@ struct FPackageExportTagger
 		{
 			if ( bRoutePresave )
 			{
-#if ENABLE_TAGEXPORTS_CLASS_PRESAVE_TIMES
-				auto& TimingInfo = SavePackageStats::ClassPreSaveTimes.FindOrAdd(Base->GetClass()->GetFName());
-				TimingInfo.Value++;
-				FScopedDurationTimer SerializeTimer(TimingInfo.Key);
-#endif
 				if (bIsCooking && Base->HasAnyFlags(RF_ClassDefaultObject | RF_ArchetypeObject))
 				{
 					FArchiveObjectCrc32NonEditorProperties CrcArchive;
@@ -2938,11 +2789,6 @@ struct FPackageExportTagger
 				for(int32 Index = 0; Index < TagExpObjects.Num(); Index++)
 				{
 					UObject* Obj = TagExpObjects[Index];
-#if ENABLE_TAGEXPORTS_CLASS_PRESAVE_TIMES
-					auto& TimingInfo = SavePackageStats::ClassPreSaveTimes.FindOrAdd(Obj->GetClass()->GetFName());
-					TimingInfo.Value++;
-					FScopedDurationTimer SerializeTimer(TimingInfo.Key);
-#endif
 					check(Obj->HasAnyMarks(OBJECTMARK_TagExp));
 					//@warning: Objects created from within PreSave will NOT have PreSave called on them!!!
 					if (bIsCooking && Obj->HasAnyFlags(RF_ClassDefaultObject | RF_ArchetypeObject))
@@ -3112,7 +2958,15 @@ static bool ValidateConformCompatibility(UPackage* NewPackage, FLinkerLoad* OldL
 	return !bHadCompatibilityErrors;
 }
 
-EObjectMark UPackage::GetExcludedObjectMarksForTargetPlatform( const class ITargetPlatform* TargetPlatform, const bool bIsCooking )
+/**
+ * Determines the set of object marks that should be excluded for the target platform
+ *
+ * @param TargetPlatform	The platform being saved for
+ * @param bIsCooking		Whether we are cooking or not
+ *
+ * @return Excluded object marks specific for the particular target platform, objects with any of these marks will be rejected from the cook
+ */
+EObjectMark GetExcludedObjectMarksForTargetPlatform( const class ITargetPlatform* TargetPlatform, const bool bIsCooking )
 {
 	EObjectMark ObjectMarks = OBJECTMARK_NOMARKS;
 
@@ -3515,14 +3369,6 @@ void VerifyEDLCookInfo()
 	FEDLCookChecker::Verify();
 }
 
-static void AddArchiveToHash(FBufferArchive& Ar, FMD5& OutHash)
-{
-	if (int64 Size = Ar.TotalSize())
-	{
-		OutHash.Update(Ar.GetData(), Size);
-	}
-}
-
 void AddFileToHash(FString const &Filename, FMD5 &Hash)
 {
 	TArray<uint8> LocalScratch;
@@ -3544,12 +3390,15 @@ void AddFileToHash(FString const &Filename, FMD5 &Hash)
 }
 
 FSavePackageResultStruct UPackage::Save(UPackage* InOuter, UObject* Base, EObjectFlags TopLevelFlags, const TCHAR* Filename,
-	FOutputDevice* Error, FLinkerLoad* Conform, bool bForceByteSwapping, bool bWarnOfLongFilename, uint32 SaveFlags, 
-	const class ITargetPlatform* TargetPlatform, const FDateTime&  FinalTimeStamp, bool bSlowTask, FArchiveDiffMap* InOutDiffMap)
+	FOutputDevice* Error, FLinkerNull* ConformNO, bool bForceByteSwapping, bool bWarnOfLongFilename, uint32 SaveFlags, 
+	const class ITargetPlatform* TargetPlatform, const FDateTime&  FinalTimeStamp, bool bSlowTask, FArchiveDiffMap* InOutDiffMap,
+	FSavePackageContext* SavePackageContext)
 {
 	COOK_STAT(FScopedDurationTimer FuncSaveTimer(SavePackageStats::SavePackageTimeSec));
 	COOK_STAT(SavePackageStats::NumPackagesSaved++);
 
+	FLinkerLoad* Conform = nullptr;
+
 #if WITH_EDITOR
 	TMap<UObject*, UObject*> ReplacedImportOuters;
 #endif //WITH_EDITOR
@@ -3802,7 +3651,12 @@ FSavePackageResultStruct UPackage::Save(UPackage* InOuter, UObject* Base, EObjec
 		// Size of serialized out package in bytes. This is before compression.
 		int32 PackageSize = INDEX_NONE;
 		{
-			FSavePackageState SavePackageState;
+			// TODO: Require a SavePackageContext and move to EditorEngine
+			FSinglePackageNameMapSaver SinglePackageNameMapSaver;
+			FPackageHeaderSaver LocalSaver(SinglePackageNameMapSaver);
+			FPackageHeaderSaver& PackageHeaderSaver = SavePackageContext ? SavePackageContext->HeaderSaver : LocalSaver;
+			INameMapSaver& NameMapSaver = PackageHeaderSaver.NameMapSaver;
+			NameMapSaver.BeginPackage();
 
 			uint32 ComparisonFlags = PPF_DeepCompareInstances;
 
@@ -3887,8 +3741,8 @@ FSavePackageResultStruct UPackage::Save(UPackage* InOuter, UObject* Base, EObjec
 			
 				{
 					COOK_STAT(FScopedDurationTimer SaveTimer(SavePackageStats::TagPackageExportsTimeSec));
-					// Clear OBJECTMARK_TagExp again as we need to redo tagging below.
-					UnMarkAllObjects((EObjectMark)(OBJECTMARK_TagExp|OBJECTMARK_EditorOnly));
+					// Clear all marks (OBJECTMARK_TagExp and exclusion marks) again as we need to redo tagging below.
+					UnMarkAllObjects();
 			
 					// We need to serialize objects yet again to tag objects that were created by PreSave as OBJECTMARK_TagExp.
 					PackageExportTagger.TagPackageExports( ExportTaggerArchive, false );
@@ -3924,14 +3778,14 @@ FSavePackageResultStruct UPackage::Save(UPackage* InOuter, UObject* Base, EObjec
 				TSet<UObject*> DependenciesReferencedByNonRedirectors;
 		
 				/** If true, we are going to save to disk async to save time. */
-				bool bSaveAsync = !!(SaveFlags & SAVE_Async);
+				const bool bSaveAsync = !!(SaveFlags & SAVE_Async);
 
-				bool bSaveUnversioned = !!(SaveFlags & SAVE_Unversioned);
+				const bool bSaveUnversioned = !!(SaveFlags & SAVE_Unversioned);
 
 				TUniquePtr<FLinkerSave> Linker = nullptr;
 				FArchiveFormatterType* Formatter = nullptr;
 				FArchive* TextFormatArchive = nullptr;
-				bool bTextFormat = FString(Filename).EndsWith(FPackageName::GetTextAssetPackageExtension()) || FString(Filename).EndsWith(FPackageName::GetTextMapPackageExtension());
+				const bool bTextFormat = FString(Filename).EndsWith(FPackageName::GetTextAssetPackageExtension()) || FString(Filename).EndsWith(FPackageName::GetTextMapPackageExtension());
 				
 #if WITH_EDITOR
 				FString DiffCookedPackagesPath;
@@ -3956,13 +3810,13 @@ FSavePackageResultStruct UPackage::Save(UPackage* InOuter, UObject* Base, EObjec
 					// The entire package will be serialized to memory and then compared against package on disk.
 					// Each difference will be log with its Serialize call stack trace
 					FArchive* Saver = new FArchiveStackTrace(FindAssetInPackage(InOuter), *InOuter->FileName.ToString(), true, InOutDiffMap);
-					Linker = TUniquePtr<FLinkerSave>(new FLinkerSave(InOuter, Saver, bForceByteSwapping, bSaveUnversioned));
+					Linker = TUniquePtr<FLinkerSave>(new FLinkerSave(PackageHeaderSaver, InOuter, Saver, bForceByteSwapping, bSaveUnversioned));
 				}
 				else if (TargetPlatform != nullptr && (SaveFlags & SAVE_DiffOnly))
 				{
 					// The entire package will be serialized to memory and then compared against package on disk
 					FArchive* Saver = new FArchiveStackTrace(FindAssetInPackage(InOuter), *InOuter->FileName.ToString(), false);
-					Linker = TUniquePtr<FLinkerSave>(new FLinkerSave(InOuter, Saver, bForceByteSwapping, bSaveUnversioned));
+					Linker = TUniquePtr<FLinkerSave>(new FLinkerSave(PackageHeaderSaver, InOuter, Saver, bForceByteSwapping, bSaveUnversioned));
 				}
 				else if ((!!TargetPlatform) && FParse::Value(FCommandLine::Get(), TEXT("DiffCookedPackages="), DiffCookedPackagesPath))
 				{
@@ -3975,26 +3829,28 @@ FSavePackageResultStruct UPackage::Save(UPackage* InOuter, UObject* Base, EObjec
 					
 					FArchive* TestArchive = IFileManager::Get().CreateFileReader(*TestArchiveFilename); 
 					FArchive* Saver = new FDiffSerializeArchive(*InOuter->FileName.ToString(), TestArchive);
-					Linker = TUniquePtr<FLinkerSave>(new FLinkerSave(InOuter, Saver, bForceByteSwapping));
+					Linker = TUniquePtr<FLinkerSave>(new FLinkerSave(PackageHeaderSaver, InOuter, Saver, bForceByteSwapping));
 				}
 				else 
 #endif
 				if (bSaveAsync)
 				{
 					// Allocate the linker with a memory writer, forcing byte swapping if wanted.
-					Linker = TUniquePtr<FLinkerSave>(new FLinkerSave(InOuter, bForceByteSwapping, bSaveUnversioned));
+					Linker = TUniquePtr<FLinkerSave>(new FLinkerSave(PackageHeaderSaver, InOuter, bForceByteSwapping, bSaveUnversioned));
 				}
 				else
 				{
 					// Allocate the linker, forcing byte swapping if wanted.
-					Linker = TUniquePtr<FLinkerSave>(new FLinkerSave(InOuter, *TempFilename, bForceByteSwapping, bSaveUnversioned));
+					Linker = TUniquePtr<FLinkerSave>(new FLinkerSave(PackageHeaderSaver, InOuter, *TempFilename, bForceByteSwapping, bSaveUnversioned));
 				}
 
 #if WITH_TEXT_ARCHIVE_SUPPORT
 				if (bTextFormat)
 				{
 					TextFormatArchive = IFileManager::Get().CreateFileWriter(*TextFormatTempFilename);
-					Formatter = new FJsonArchiveOutputFormatter(*TextFormatArchive);
+					FJsonArchiveOutputFormatter* OutputFormatter = new FJsonArchiveOutputFormatter(*TextFormatArchive);
+					OutputFormatter->SetObjectIndicesMap(&Linker->ObjectIndicesMap);
+					Formatter = OutputFormatter;
 				}
 				else
 #endif
@@ -4050,7 +3906,7 @@ FSavePackageResultStruct UPackage::Save(UPackage* InOuter, UObject* Base, EObjec
 					TArray<UObject*> TagExpObjects;
 					GetObjectsWithAnyMarks(TagExpObjects, OBJECTMARK_TagExp);
 
-					const EObjectMark ExcludedObjectMarks = UPackage::GetExcludedObjectMarksForTargetPlatform(TargetPlatform, Linker->IsCooking());
+					const EObjectMark ExcludedObjectMarks = GetExcludedObjectMarksForTargetPlatform(TargetPlatform, Linker->IsCooking());
 					if (Linker->IsCooking() && ExcludedObjectMarks != OBJECTMARK_NOMARKS)
 					{
 						// Make sure that nothing is in the export table that should have been filtered out
@@ -4114,7 +3970,7 @@ FSavePackageResultStruct UPackage::Save(UPackage* InOuter, UObject* Base, EObjec
 						check(Obj->HasAnyMarks(OBJECTMARK_TagExp));
 
 						// Build list.
-						FArchiveSaveTagImports ImportTagger(Linker.Get(), SavePackageState);
+						FArchiveSaveTagImports ImportTagger(Linker.Get(), NameMapSaver);
 						ImportTagger.SetPortFlags(ComparisonFlags);
 						ImportTagger.SetFilterEditorOnly(FilterEditorOnly);
 						ImportTagger.SetSerializeContext(SaveContext);
@@ -4236,21 +4092,12 @@ FSavePackageResultStruct UPackage::Save(UPackage* InOuter, UObject* Base, EObjec
 							UE_LOG(LogSavePackage, Fatal, TEXT("%s"), *FString::Printf( TEXT("Transient object imported: %s"), *Obj->GetFullName() ) );
 						}
 
-						// add the list of dependencies to the dependency map
-						ObjectDependencies.Add(Obj, ImportTagger.Dependencies);
-						NativeObjectDependencies.Add(Obj, ImportTagger.NativeDependencies);
-
 						if (Obj->GetClass() != UObjectRedirector::StaticClass())
 						{
-							for ( auto DepIt = ImportTagger.Dependencies.CreateConstIterator(); DepIt; ++DepIt )
-							{
-								UObject* DependencyObject = *DepIt;
-								if ( DependencyObject )
-								{
-									DependenciesReferencedByNonRedirectors.Add(DependencyObject);
-								}
-							}
+							DependenciesReferencedByNonRedirectors.Append(ImportTagger.Dependencies);
 						}
+						ObjectDependencies.Add(Obj, MoveTemp(ImportTagger.Dependencies));
+						NativeObjectDependencies.Add(Obj, MoveTemp(ImportTagger.NativeDependencies));
 					}
 				}
 				if (PrestreamPackages.Num())
@@ -4297,20 +4144,20 @@ FSavePackageResultStruct UPackage::Save(UPackage* InOuter, UObject* Base, EObjec
 						UObject* Obj = TagExpImpObjects[Index];
 						check(Obj->HasAnyMarks(EObjectMark(OBJECTMARK_TagExp|OBJECTMARK_TagImp)));
 
-						SavePackageState.MarkNameAsReferenced(Obj->GetFName());
+						NameMapSaver.MarkNameAsReferenced(Obj->GetFName());
 #if WITH_EDITOR
-						SavePackageState.AddReplacementsNames(Obj, TargetPlatform, bIsCooking);
+						AddReplacementsNames(NameMapSaver, Obj, TargetPlatform);
 #endif //WITH_EDITOR
 						if( Obj->GetOuter() )
 						{
-							SavePackageState.MarkNameAsReferenced(Obj->GetOuter()->GetFName());
+							NameMapSaver.MarkNameAsReferenced(Obj->GetOuter()->GetFName());
 						}
 
 						if( Obj->HasAnyMarks(OBJECTMARK_TagImp) )
 						{
-							SavePackageState.MarkNameAsReferenced(Obj->GetClass()->GetFName());
+							NameMapSaver.MarkNameAsReferenced(Obj->GetClass()->GetFName());
 							check(Obj->GetClass()->GetOuter());
-							SavePackageState.MarkNameAsReferenced(Obj->GetClass()->GetOuter()->GetFName());
+							NameMapSaver.MarkNameAsReferenced(Obj->GetClass()->GetOuter()->GetFName());
 						
 							// if a private object was marked by the cooker, it will be in memory on load, and will be found. helps with some objects
 							// from a package in a package being moved into Startup_int.xxx, but not all
@@ -4329,7 +4176,7 @@ FSavePackageResultStruct UPackage::Save(UPackage* InOuter, UObject* Base, EObjec
 							UPackage* ObjPackage = Obj->GetOutermost();
 							if (PrestreamPackages.Contains(ObjPackage))
 							{
-								SavePackageState.MarkNameAsReferenced(PrestreamPackageClassName);
+								NameMapSaver.MarkNameAsReferenced(PrestreamPackageClassName);
 								// These are not errors
 								UE_LOG(LogSavePackage, Display, TEXT("Prestreaming package %s "), *ObjPackage->GetPathName()); //-V595
 								continue;
@@ -4562,60 +4409,26 @@ FSavePackageResultStruct UPackage::Save(UPackage* InOuter, UObject* Base, EObjec
 				SlowTask.EnterProgressFrame();
 
 
-				// Build NameMap.
-				Linker->Summary.NameOffset = Linker->Tell();
-				SavePackageState.UpdateLinkerWithMarkedNames(Linker.Get());
-
 #if WITH_EDITOR
 				if ( GOutputCookingWarnings )
 				{
 					// check the name list for uniqueobjectnamefor cooking
 					static FNameEntryId NAME_UniqueObjectNameForCooking = FName("UniqueObjectNameForCooking").GetComparisonIndex();
-
-					for (FNameEntryId NameInUse : Linker->NameMap)
+					if (NameMapSaver.NameExistsInCurrentPackage(NAME_UniqueObjectNameForCooking))
 					{
-						if (FName::GetComparisonIdFromDisplayId(NameInUse) == NAME_UniqueObjectNameForCooking)
-						{
-							//UObject *Object = FindObject<UObject>( ANY_PACKAGE, *NameInUse.ToString());
-
-							// error
-							// check(Object);
-							UE_LOG(LogSavePackage, Warning, TEXT("Saving object into cooked package %s which was created at cook time"), Filename);
-						}
+						UE_LOG(LogSavePackage, Warning, TEXT("Saving object into cooked package %s which was created at cook time"), Filename);
 					}
 				}
 #endif
 
-				if ( EndSavingIfCancelled( Linker.Get(), TempFilename ) )
-				{ 
-					return ESavePackageResult::Canceled;
-				}
-				SlowTask.EnterProgressFrame();
-
-				// Sort names.
-				FObjectNameSortHelper NameSortHelper;
-				NameSortHelper.SortNames( Linker.Get(), Conform, SavePackageState );
-
-				if ( EndSavingIfCancelled( Linker.Get(), TempFilename ) )
-				{ 
-					return ESavePackageResult::Canceled;
-				}
-				SlowTask.EnterProgressFrame();
-
-				// Save names.
-				if (!bTextFormat)
+				// Build NameMap.
+				Linker->Summary.NameOffset = Linker->Tell();
 				{
-					FStructuredArchive::FStream NameStream = StructuredArchiveRoot.EnterStream(FIELD_NAME_TEXT("Names"));
 #if WITH_EDITOR
 					FArchive::FScopeSetDebugSerializationFlags S(*Linker, DSF_IgnoreDiff, true);
 					FArchiveStackTraceIgnoreScope IgnoreSummaryDiffsScope(DiffSettings.bIgnoreHeaderDiffs);
 #endif
-					Linker->Summary.NameCount = Linker->NameMap.Num();
-					for (int32 i = 0; i < Linker->NameMap.Num(); i++)
-					{
-						FName::GetEntry(Linker->NameMap[i])->Write(NameStream.EnterElement());
-						Linker->NameIndices.Add(Linker->NameMap[i], i);
-					}
+					NameMapSaver.EndPackage(*Linker.Get(), Conform, bTextFormat ? nullptr : Linker->Saver);
 				}
 				if ( EndSavingIfCancelled( Linker.Get(), TempFilename ) )
 				{ 
@@ -4623,7 +4436,7 @@ FSavePackageResultStruct UPackage::Save(UPackage* InOuter, UObject* Base, EObjec
 				}
 				SlowTask.EnterProgressFrame();
 
-				FStructuredArchive::FStream Stream = StructuredArchiveRoot.EnterStream(FIELD_NAME_TEXT("GatherableTextData"));
+				FStructuredArchive::FStream Stream = StructuredArchiveRoot.EnterStream(SA_FIELD_NAME(TEXT("GatherableTextData")));
 				Linker->Summary.GatherableTextDataOffset = 0;
 				Linker->Summary.GatherableTextDataCount = 0;
 				if (!(Linker->Summary.PackageFlags & PKG_FilterEditorOnly))
@@ -4666,7 +4479,7 @@ FSavePackageResultStruct UPackage::Save(UPackage* InOuter, UObject* Base, EObjec
 				{
 					TArray<UObject*> TagImpObjects;
 
-					const EObjectMark ExcludedObjectMarks = UPackage::GetExcludedObjectMarksForTargetPlatform(TargetPlatform, Linker->IsCooking());
+					const EObjectMark ExcludedObjectMarks = GetExcludedObjectMarksForTargetPlatform(TargetPlatform, Linker->IsCooking());
 					GetObjectsWithAnyMarks(TagImpObjects, OBJECTMARK_TagImp);
 
 					if (Linker->IsCooking() && ExcludedObjectMarks != OBJECTMARK_NOMARKS)
@@ -4838,6 +4651,7 @@ FSavePackageResultStruct UPackage::Save(UPackage* InOuter, UObject* Base, EObjec
 					checkf(SrcDepends,TEXT("Couldn't find dependency map for %s"), *Object->GetFullName());
 
 					// go through each object and...
+					DependIndices.Reserve(SrcDepends->Num());
 					for (int32 DependIndex = 0; DependIndex < SrcDepends->Num(); DependIndex++)
 					{
 						UObject* DependentObject = (*SrcDepends)[DependIndex];
@@ -5045,29 +4859,32 @@ FSavePackageResultStruct UPackage::Save(UPackage* InOuter, UObject* Base, EObjec
 				}
 				SlowTask.EnterProgressFrame();
 
-				FStructuredArchive::FStream DependsStream = StructuredArchiveRoot.EnterStream(FIELD_NAME_TEXT("DependsMap"));
-				if (Linker->IsCooking())
+				if (!bTextFormat)
 				{
+					FStructuredArchive::FStream DependsStream = StructuredArchiveRoot.EnterStream(SA_FIELD_NAME(TEXT("DependsMap")));
+					if (Linker->IsCooking())
+					{
 #if WITH_EDITOR
-					FArchiveStackTraceIgnoreScope IgnoreSummaryDiffsScope(DiffSettings.bIgnoreHeaderDiffs);
+						FArchiveStackTraceIgnoreScope IgnoreSummaryDiffsScope(DiffSettings.bIgnoreHeaderDiffs);
 #endif // WITH_EDITOR
-					//@todo optimization, this should just be stripped entirely from cooked packages
-					TArray<FPackageIndex> Depends; // empty array
-					Linker->Summary.DependsOffset = Linker->Tell();
-					for (int32 i = 0; i < Linker->ExportMap.Num(); i++)
-					{
-						DependsStream.EnterElement() << Depends;
+						//@todo optimization, this should just be stripped entirely from cooked packages
+						TArray<FPackageIndex> Depends; // empty array
+						Linker->Summary.DependsOffset = Linker->Tell();
+						for (int32 i = 0; i < Linker->ExportMap.Num(); i++)
+						{
+							DependsStream.EnterElement() << Depends;
+						}
 					}
-				}
-				else
-				{
-					// save depends map (no need for later patching)
-					check(Linker->DependsMap.Num() == Linker->ExportMap.Num());
-					Linker->Summary.DependsOffset = Linker->Tell();
-					for (int32 i = 0; i < Linker->ExportMap.Num(); i++)
+					else
 					{
-						TArray<FPackageIndex>& Depends = Linker->DependsMap[i];
-						DependsStream.EnterElement() << Depends;
+						// save depends map (no need for later patching)
+						check(Linker->DependsMap.Num() == Linker->ExportMap.Num());
+						Linker->Summary.DependsOffset = Linker->Tell();
+						for (int32 i = 0; i < Linker->ExportMap.Num(); i++)
+						{
+							TArray<FPackageIndex>& Depends = Linker->DependsMap[i];
+							DependsStream.EnterElement() << Depends;
+						}
 					}
 				}
 
@@ -5084,20 +4901,21 @@ FSavePackageResultStruct UPackage::Save(UPackage* InOuter, UObject* Base, EObjec
 					// Save soft package references
 					Linker->Summary.SoftPackageReferencesOffset = Linker->Tell();
 					Linker->Summary.SoftPackageReferencesCount = Linker->SoftPackageReferenceList.Num();
+					if (!bTextFormat)
 					{
 #if WITH_EDITOR
 						FArchive::FScopeSetDebugSerializationFlags S(*Linker, DSF_IgnoreDiff, true);
 #endif
-						FStructuredArchive::FStream SoftReferenceStream = StructuredArchiveRoot.EnterStream(FIELD_NAME_TEXT("SoftReferences"));
+						FStructuredArchive::FStream SoftReferenceStream = StructuredArchiveRoot.EnterStream(SA_FIELD_NAME(TEXT("SoftReferences")));
 						for (FName& SoftPackageName : Linker->SoftPackageReferenceList)
 						{
 							SoftReferenceStream.EnterElement() << SoftPackageName;
 						}
-					}
 
-					// Save searchable names map
-					Linker->Summary.SearchableNamesOffset = Linker->Tell();
-					Linker->SerializeSearchableNamesMap(StructuredArchiveRoot.EnterField(FIELD_NAME_TEXT("SearchableNames")));
+						// Save searchable names map
+						Linker->Summary.SearchableNamesOffset = Linker->Tell();
+						Linker->SerializeSearchableNamesMap(StructuredArchiveRoot.EnterField(SA_FIELD_NAME(TEXT("SearchableNames"))));
+					}
 				}
 				else
 				{
@@ -5112,13 +4930,16 @@ FSavePackageResultStruct UPackage::Save(UPackage* InOuter, UObject* Base, EObjec
 #endif // WITH_EDITOR
 
 					// Save thumbnails
-					UPackage::SaveThumbnails(InOuter, Linker.Get(), StructuredArchiveRoot.EnterField(FIELD_NAME_TEXT("Thumbnails")));
+					SaveThumbnails(InOuter, Linker.Get(), StructuredArchiveRoot.EnterField(SA_FIELD_NAME(TEXT("Thumbnails"))));
 
-					// Save asset registry data so the editor can search for information about assets in this package
-					UPackage::SaveAssetRegistryData(InOuter, Linker.Get(), StructuredArchiveRoot.EnterField(FIELD_NAME_TEXT("AssetRegistry")));
+					if (!bTextFormat)
+					{	
+						// Save asset registry data so the editor can search for information about assets in this package
+						SaveAssetRegistryData(InOuter, Linker.Get(), StructuredArchiveRoot.EnterField(SA_FIELD_NAME(TEXT("AssetRegistry"))));
+					}
 
 					// Save level information used by World browser
-					UPackage::SaveWorldLevelInfo(InOuter, Linker.Get(), StructuredArchiveRoot.EnterField(FIELD_NAME_TEXT("WorldLevelInfo")));
+					SaveWorldLevelInfo(InOuter, Linker.Get(), StructuredArchiveRoot);
 				}
 
 
@@ -5212,7 +5033,7 @@ FSavePackageResultStruct UPackage::Save(UPackage* InOuter, UObject* Base, EObjec
 					FArchiveStackTraceIgnoreScope IgnoreSummaryDiffsScope(DiffSettings.bIgnoreHeaderDiffs);
 #endif // WITH_EDITOR
 
-					const EObjectMark ExcludedObjectMarks = UPackage::GetExcludedObjectMarksForTargetPlatform(Linker->CookingTarget(), Linker->IsCooking());
+					const EObjectMark ExcludedObjectMarks = GetExcludedObjectMarksForTargetPlatform(Linker->CookingTarget(), Linker->IsCooking());
 					Linker->Summary.PreloadDependencyCount = 0;
 
 					auto IncludeObjectAsDependency = [&Linker, ExcludedObjectMarks](int32 CallSite, TSet<FPackageIndex>& AddTo, UObject* ToTest, UObject* ForObj, bool bMandatory, bool bOnlyIfInLinkerTable)
@@ -5278,17 +5099,21 @@ FSavePackageResultStruct UPackage::Save(UPackage* InOuter, UObject* Base, EObjec
 						}
 					};
 
-					FStructuredArchive::FStream DepedenciesStream = StructuredArchiveRoot.EnterStream(FIELD_NAME_TEXT("PreloadDependencies"));
+					FStructuredArchive::FStream DepedenciesStream = StructuredArchiveRoot.EnterStream(SA_FIELD_NAME(TEXT("PreloadDependencies")));
 					TArray<UObject*> Subobjects;
 					TArray<UObject*> Deps;
+					TSet<FPackageIndex> SerializationBeforeCreateDependencies;
+					TSet<FPackageIndex> SerializationBeforeSerializationDependencies;
+					TSet<FPackageIndex> CreateBeforeSerializationDependencies;
+					TSet<FPackageIndex> CreateBeforeCreateDependencies;
 					for (int32 i = 0; i < Linker->ExportMap.Num(); i++)
 					{
 						FObjectExport& Export = Linker->ExportMap[i];
 						if (Export.Object)
 						{
 							EDLCookChecker.AddExport(Export.Object);
-							TSet<FPackageIndex> SerializationBeforeCreateDependencies;
 							{
+								SerializationBeforeCreateDependencies.Reset();
 								IncludeIndexAsDependency(SerializationBeforeCreateDependencies, Export.ClassIndex);
 								UObject* CDO = Export.Object->GetArchetype();
 								IncludeObjectAsDependency(1, SerializationBeforeCreateDependencies, CDO, Export.Object, true, false);
@@ -5315,8 +5140,8 @@ FSavePackageResultStruct UPackage::Save(UPackage* InOuter, UObject* Base, EObjec
 									}
 								}
 							}
-							TSet<FPackageIndex> SerializationBeforeSerializationDependencies;
 							{
+								SerializationBeforeSerializationDependencies.Reset();
 								Deps.Reset();
 								Export.Object->GetPreloadDependencies(Deps);
 
@@ -5362,8 +5187,8 @@ FSavePackageResultStruct UPackage::Save(UPackage* InOuter, UObject* Base, EObjec
 								}
 							}
 
-							TSet<FPackageIndex> CreateBeforeSerializationDependencies;
 							{
+								CreateBeforeSerializationDependencies.Reset();
 								UClass* Class = Cast<UClass>(Export.Object);
 								UObject* ClassCDO = Class ? Class->GetDefaultObject() : nullptr;
 								{
@@ -5388,8 +5213,9 @@ FSavePackageResultStruct UPackage::Save(UPackage* InOuter, UObject* Base, EObjec
 									}
 								}
 							}
-							TSet<FPackageIndex> CreateBeforeCreateDependencies;
+
 							{
+								CreateBeforeCreateDependencies.Reset();
 								IncludeIndexAsDependency(CreateBeforeCreateDependencies, Export.OuterIndex);
 								IncludeIndexAsDependency(CreateBeforeCreateDependencies, Export.SuperIndex);
 							}
@@ -5503,6 +5329,8 @@ FSavePackageResultStruct UPackage::Save(UPackage* InOuter, UObject* Base, EObjec
 #endif
 					FScopedSlowTask ExportScope(Linker->ExportMap.Num());
 
+					FStructuredArchive::FRecord ExportsRecord = StructuredArchiveRoot.EnterRecord(SA_FIELD_NAME(TEXT("Exports")));
+
 					// Save exports.
 					int32 LastExportSaveStep = 0;
 					for( int32 i=0; i<Linker->ExportMap.Num(); i++ )
@@ -5514,16 +5342,21 @@ FSavePackageResultStruct UPackage::Save(UPackage* InOuter, UObject* Base, EObjec
 						ExportScope.EnterProgressFrame();
 
 						FObjectExport& Export = Linker->ExportMap[i];
-						if( Export.Object )
+						if (Export.Object)
 						{
 							// Save the object data.
 							Export.SerialOffset = Linker->Tell();
 							Linker->CurrentlySavingExport = FPackageIndex::FromExport(i);
 							// UE_LOG(LogSavePackage, Log, TEXT("export %s for %s"), *Export.Object->GetFullName(), *Linker->CookingTarget()->PlatformName());
 
-							//FString ObjectName = Export.Object->GetPathName(InOuter);
 							FString ObjectName = Export.Object->GetPathName(InOuter);
-							FStructuredArchive::FSlot ExportSlot = StructuredArchiveRoot.EnterField(FIELD_NAME(*ObjectName));
+							FStructuredArchive::FSlot ExportSlot = ExportsRecord.EnterField(SA_FIELD_NAME(*ObjectName));
+
+							if (bTextFormat)
+							{
+								FObjectTextExport ObjectTextExport(Export, InOuter);
+								ExportSlot << ObjectTextExport;
+							}
 
 #if WITH_EDITOR
 							bool bSupportsText = UClass::IsSafeToSerializeToStructuredArchives(Export.Object->GetClass());
@@ -5540,16 +5373,12 @@ FSavePackageResultStruct UPackage::Save(UPackage* InOuter, UObject* Base, EObjec
 								else
 								{
 									FArchiveUObjectFromStructuredArchive Adapter(ExportSlot);
-									Export.Object->GetClass()->SerializeDefaultObject(Export.Object, Adapter);
+									Export.Object->GetClass()->SerializeDefaultObject(Export.Object, Adapter.GetArchive());
+									Adapter.Close();
 								}
 							}
 							else
 							{
-#if ENABLE_PACKAGE_CLASS_SERIALIZATION_TIMES
-								auto& TimingInfo = SavePackageStats::PackageClassSerializeTimes.FindOrAdd(Export.Object->GetClass()->GetFName());
-								TimingInfo.Value++;
-								FScopedDurationTimer SerializeTimer(TimingInfo.Key);
-#endif
 								TGuardValue<UObject*> GuardSerializedObject(SaveContext->SerializedObject, Export.Object);
 
 								if (bSupportsText)
@@ -5560,7 +5389,8 @@ FSavePackageResultStruct UPackage::Save(UPackage* InOuter, UObject* Base, EObjec
 								else
 								{
 									FArchiveUObjectFromStructuredArchive Adapter(ExportSlot);
-									Export.Object->Serialize(Adapter);
+									Export.Object->Serialize(Adapter.GetArchive());
+									Adapter.Close();
 								}
 
 #if WITH_EDITOR
@@ -5604,7 +5434,7 @@ FSavePackageResultStruct UPackage::Save(UPackage* InOuter, UObject* Base, EObjec
 
 				// now we write all the bulkdata that is supposed to be at the end of the package
 				// and fix up the offset
-				int64 StartOfBulkDataArea = Linker->Tell();
+				const int64 StartOfBulkDataArea = Linker->Tell();
 				Linker->Summary.BulkDataStartOffset = StartOfBulkDataArea;
 
 				check(!bTextFormat || Linker->BulkDataToAppend.Num() == 0);
@@ -5615,57 +5445,37 @@ FSavePackageResultStruct UPackage::Save(UPackage* InOuter, UObject* Base, EObjec
 
 					FScopedSlowTask BulkDataFeedback(Linker->BulkDataToAppend.Num());
 
-					FArchive* BulkArchive = nullptr;
-					FArchive* OptionalBulkArchive = nullptr;
-					FArchive* MappedBulkArchive = nullptr;
+					TUniquePtr<FLargeMemoryWriter> BulkArchive;
+					TUniquePtr<FLargeMemoryWriter> OptionalBulkArchive;
+					TUniquePtr<FLargeMemoryWriter> MappedBulkArchive;
+
 					uint32 ExtraBulkDataFlags = 0;
 
-					static struct FUseSeperateBulkDataFiles
+					static const struct FUseSeparateBulkDataFiles
 					{
-						bool bEnable;
-						FUseSeperateBulkDataFiles()
+						bool bEnable = false;
+
+						FUseSeparateBulkDataFiles()
 						{
-							if (!GConfig->GetBool(TEXT("Core.System"), TEXT("UseSeperateBulkDataFiles"), bEnable, GEngineIni))
-							{
-								bEnable = false;
-							}
+							GConfig->GetBool(TEXT("Core.System"), TEXT("UseSeperateBulkDataFiles"), /* out */ bEnable, GEngineIni);
+
 							if (IsEventDrivenLoaderEnabledInCookedBuilds())
 							{
 								// Always split bulk data when splitting cooked files
 								bEnable = true;
 							}
 						}
-					} ShouldUseSeperateBulkDataFiles;
-
-					const bool bShouldUseSeparateBulkFile = ShouldUseSeperateBulkDataFiles.bEnable && Linker->IsCooking();
-
-					const FString BulkFilename = FPaths::ChangeExtension(Filename, TEXT(".ubulk"));
-
-					const static TCHAR* OptionalBulkFileExtension = TEXT("uptnl");
-
-					const FString OptionalBulkFilename = FPaths::ChangeExtension(Filename, OptionalBulkFileExtension);
+					} ShouldUseSeparateBulkDataFiles;
 
-					const static TCHAR* MappedBulkFileExtension = TEXT(".m.ubulk");
+					const bool bShouldUseSeparateBulkFile = ShouldUseSeparateBulkDataFiles.bEnable && Linker->IsCooking();
 
-					const FString MappedBulkFilename = FPaths::ChangeExtension(Filename, MappedBulkFileExtension);
-
-					bool bBufferedBulkArchives = false;
 					if (bShouldUseSeparateBulkFile)
 					{
-						ExtraBulkDataFlags = BULKDATA_PayloadInSeperateFile;
-						if (bSaveAsync || bDiffing)
-						{
-							bBufferedBulkArchives = true;
-							BulkArchive = new FBufferArchive(true);
-							OptionalBulkArchive = new FBufferArchive(true);
-							MappedBulkArchive = new FBufferArchive(true);
-						}
-						else
-						{
-							BulkArchive = IFileManager::Get().CreateFileWriter(*BulkFilename);
-							OptionalBulkArchive = IFileManager::Get().CreateFileWriter(*OptionalBulkFilename);
-							MappedBulkArchive = IFileManager::Get().CreateFileWriter(*MappedBulkFilename);
-						}
+						ExtraBulkDataFlags	= BULKDATA_PayloadInSeperateFile;
+
+						BulkArchive.Reset(new FLargeMemoryWriter(0, /* IsPersistent */ true));
+						OptionalBulkArchive.Reset(new FLargeMemoryWriter(0, /* IsPersistent */ true));
+						MappedBulkArchive.Reset(new FLargeMemoryWriter(0, /* IsPersistent */ true));
 					}
 
 					bool bAlignBulkData = false;
@@ -5677,16 +5487,14 @@ FSavePackageResultStruct UPackage::Save(UPackage* InOuter, UObject* Base, EObjec
 						BulkDataAlignment = TargetPlatform->GetMemoryMappingAlignment();
 					}
 
-					for (int32 i=0; i < Linker->BulkDataToAppend.Num(); ++i)
+					for (FLinkerSave::FBulkDataStorageInfo& BulkDataStorageInfo : Linker->BulkDataToAppend)
 					{
 						BulkDataFeedback.EnterProgressFrame();
 
-						FLinkerSave::FBulkDataStorageInfo& BulkDataStorageInfo = Linker->BulkDataToAppend[i];
-
 						// Set bulk data flags to what they were during initial serialization (they might have changed after that)
 						const uint32 OldBulkDataFlags = BulkDataStorageInfo.BulkData->GetBulkDataFlags();
 						uint32 ModifiedBulkDataFlags = BulkDataStorageInfo.BulkDataFlags | ExtraBulkDataFlags;
-						bool bBulkItemIsOptional = (ModifiedBulkDataFlags & BULKDATA_OptionalPayload) != 0;
+						const bool bBulkItemIsOptional = (ModifiedBulkDataFlags & BULKDATA_OptionalPayload) != 0;
 						bool bBulkItemIsMapped = bAlignBulkData && ((ModifiedBulkDataFlags & BULKDATA_MemoryMappedPayload) != 0);
 
 						if (bBulkItemIsMapped && bBulkItemIsOptional)
@@ -5700,44 +5508,66 @@ FSavePackageResultStruct UPackage::Save(UPackage* InOuter, UObject* Base, EObjec
 						BulkDataStorageInfo.BulkData->SetBulkDataFlags(ModifiedBulkDataFlags);
 
 						FArchive* TargetArchive = Linker.Get();
-						if ( bShouldUseSeparateBulkFile )
+
+						if (bShouldUseSeparateBulkFile)
 						{
-							check( OptionalBulkArchive && BulkArchive );
-							TargetArchive = bBulkItemIsOptional ? OptionalBulkArchive : (bBulkItemIsMapped ? MappedBulkArchive : BulkArchive);
+							check(MappedBulkArchive && OptionalBulkArchive && BulkArchive);
+
+							if (bBulkItemIsOptional)
+							{
+								TargetArchive = OptionalBulkArchive.Get();
+							}
+							else if (bBulkItemIsMapped)
+							{
+								TargetArchive = MappedBulkArchive.Get();
+							}
+							else
+							{
+								TargetArchive = BulkArchive.Get();
+							}
 						}
 
-						int64 BulkStartOffset = TargetArchive->Tell();
+						// Pad archive for proper alignment for memory mapping
 
-						if (bBulkItemIsMapped && BulkDataAlignment > 0 && !IsAligned(BulkStartOffset, BulkDataAlignment))
+						if (bBulkItemIsMapped && BulkDataAlignment > 0)
 						{
-							int64 AlignedOffset = Align(BulkStartOffset, BulkDataAlignment);
-							int64 Padding = AlignedOffset - BulkStartOffset;
-							check(Padding > 0);
-							uint64 Zero64 = 0;
-							while (Padding >= 8)
-							{
-								*TargetArchive << Zero64;
-								Padding -= 8;
-							}
-							uint8 Zero8 = 0;
-							while (Padding > 0)
+							const int64 BulkStartOffset = TargetArchive->Tell();
+
+							if (!IsAligned(BulkStartOffset, BulkDataAlignment))
 							{
-								*TargetArchive << Zero8;
-								Padding--;
+								const int64 AlignedOffset = Align(BulkStartOffset, BulkDataAlignment);
+
+								int64 Padding = AlignedOffset - BulkStartOffset;
+								check(Padding > 0);
+
+								uint64 Zero64 = 0;
+								while(Padding >= 8)
+								{
+									*TargetArchive << Zero64;
+									Padding -= 8;
+								}
+
+								uint8 Zero8 = 0;
+								while(Padding > 0)
+								{
+									*TargetArchive << Zero8;
+									Padding--;
+								}
+
+								check(TargetArchive->Tell() == AlignedOffset);
 							}
-							BulkStartOffset = TargetArchive->Tell();
-							check(BulkStartOffset == AlignedOffset);
 						}
 
+						const int64 BulkStartOffset = TargetArchive->Tell();
 
 						int64 StoredBulkStartOffset = BulkStartOffset - StartOfBulkDataArea;
 
 						BulkDataStorageInfo.BulkData->SerializeBulkData(*TargetArchive, BulkDataStorageInfo.BulkData->Lock(LOCK_READ_ONLY));
 
 						int64 BulkEndOffset = TargetArchive->Tell();
-						int64 LinkerEndOffset = Linker->Tell();
+						const int64 LinkerEndOffset = Linker->Tell();
 
-						int64 SizeOnDisk = (int32)(BulkEndOffset - BulkStartOffset);
+						int64 SizeOnDisk = BulkEndOffset - BulkStartOffset;
 				
 						Linker->Seek(BulkDataStorageInfo.BulkDataFlagsPos);
 						*Linker << ModifiedBulkDataFlags;
@@ -5767,50 +5597,59 @@ FSavePackageResultStruct UPackage::Save(UPackage* InOuter, UObject* Base, EObjec
 
 					if (BulkArchive)
 					{
-						check( OptionalBulkArchive);
+						check(OptionalBulkArchive);
+						check(MappedBulkArchive);
 
-						const bool bWriteBulkToDisk = bSaveAsync && !bDiffing; 
-						auto FinalizeBulkDataFile = [bWriteBulkToDisk](FArchive* Archive, const FString& ArchiveFilename)
+						const bool bWriteBulkToDisk = !bDiffing;
+
+						if (SavePackageContext && bWriteBulkToDisk)
+						{
+							auto ToIoBuffer = [](FLargeMemoryWriter* Writer)
 							{
-								Archive->Close();
-								if (bWriteBulkToDisk)
-								{
-									if (int64 DataSize = Archive->TotalSize())
-									{
-										// TODO - update the BulkBuffer code to write into a FLargeMemoryWriter so we can 
-										// take ownership of the data
-										uint8* Data = new uint8[DataSize];
-										FMemory::Memcpy(Data, ((FBufferArchive*)(Archive))->GetData(), DataSize);
-										FLargeMemoryPtr DataPtr = FLargeMemoryPtr(Data);
-		
-										AsyncWriteFile(MoveTemp(DataPtr), DataSize, *ArchiveFilename, FDateTime::MinValue(), false);
-									}
-								}
-								delete Archive;
+								const int64 TotalSize = Writer->TotalSize();
+								return FIoBuffer(FIoBuffer::AssumeOwnership, Writer->ReleaseOwnership(), TotalSize);
 							};
 
-						if (bComputeHash && bBufferedBulkArchives)
-						{
-							AddArchiveToHash(*(FBufferArchive*)BulkArchive, CookedPackageHash);
-							AddArchiveToHash(*(FBufferArchive*)OptionalBulkArchive, CookedPackageHash);
-							AddArchiveToHash(*(FBufferArchive*)MappedBulkArchive, CookedPackageHash);
-						}
-						
-						TotalPackageSizeUncompressed += BulkArchive->TotalSize();
-						TotalPackageSizeUncompressed += OptionalBulkArchive->TotalSize();		
-						TotalPackageSizeUncompressed += MappedBulkArchive->TotalSize();
+							FPackageStoreWriter::FBulkDataInfo BulkInfo;
+							BulkInfo.PackageName = InOuter->GetFName();
+							BulkInfo.LooseFilePath = Filename;
+							BulkInfo.BulkdataType = FPackageStoreWriter::FBulkDataInfo::Standard;
 
-						FinalizeBulkDataFile(BulkArchive, BulkFilename);
-						FinalizeBulkDataFile(OptionalBulkArchive, OptionalBulkFilename);
-						FinalizeBulkDataFile(MappedBulkArchive, MappedBulkFilename);
+							SavePackageContext->PackageStoreWriter.WriteBulkdata(BulkInfo, ToIoBuffer(BulkArchive.Get()));
 
-						if (bComputeHash && !bBufferedBulkArchives)
+							BulkInfo.BulkdataType = FPackageStoreWriter::FBulkDataInfo::Optional;
+							SavePackageContext->PackageStoreWriter.WriteBulkdata(BulkInfo, ToIoBuffer(OptionalBulkArchive.Get()));
+
+							BulkInfo.BulkdataType = FPackageStoreWriter::FBulkDataInfo::Mmap;
+							SavePackageContext->PackageStoreWriter.WriteBulkdata(BulkInfo, ToIoBuffer(MappedBulkArchive.Get()));
+						}
+						else
 						{
-							check(!bWriteBulkToDisk);
+							auto WriteBulkData = [&](FLargeMemoryWriter* Archive, const TCHAR* BulkFileExtension)
+								{
+									if (const int64 DataSize = Archive->TotalSize())
+									{
+										if (bComputeHash)
+										{
+											CookedPackageHash.Update(Archive->GetData(), DataSize);
+										}
+
+										TotalPackageSizeUncompressed += DataSize;
+
+										if (bWriteBulkToDisk)
+										{
+											FLargeMemoryPtr DataPtr(Archive->ReleaseOwnership());
+
+											const FString ArchiveFilename = FPaths::ChangeExtension(Filename, BulkFileExtension);
+
+											AsyncWriteFile(MoveTemp(DataPtr), DataSize, *ArchiveFilename);
+										}
+									}
+								};
 
-							AddFileToHash(BulkFilename, CookedPackageHash);
-							AddFileToHash(OptionalBulkFilename, CookedPackageHash);
-							AddFileToHash(MappedBulkFilename, CookedPackageHash);
+							WriteBulkData(BulkArchive.Get(),			TEXT(".ubulk"));		// Regular separate bulk data file
+							WriteBulkData(OptionalBulkArchive.Get(),	TEXT(".uptnl"));		// Optional bulk data
+							WriteBulkData(MappedBulkArchive.Get(),		TEXT(".m.ubulk"));		// Memory-mapped bulk data
 						}
 					}
 				}
@@ -5825,15 +5664,17 @@ FSavePackageResultStruct UPackage::Save(UPackage* InOuter, UObject* Base, EObjec
 					{
 						const int64 Size = Writer.TotalSize();
 						TotalPackageSizeUncompressed += Size;
+
 						if (bComputeHash)
 						{
 							CookedPackageHash.Update(Writer.GetData(), Size);
 						}
+
 						if (bWriteFileToDisk)
 						{
-							FLargeMemoryPtr DataPtr(Writer.GetData());
-							Writer.ReleaseOwnership();
-							AsyncWriteFile(MoveTemp(DataPtr), Size, *Writer.GetArchiveName(), FDateTime::MinValue(), false);
+							FLargeMemoryPtr DataPtr(Writer.ReleaseOwnership());
+
+							AsyncWriteFile(MoveTemp(DataPtr), Size, *Writer.GetArchiveName());
 						}
 					}
 					AdditionalFilesFromExports.Empty();
@@ -5860,89 +5701,86 @@ FSavePackageResultStruct UPackage::Save(UPackage* InOuter, UObject* Base, EObjec
 
 					if (!bTextFormat)
 					{
-					Linker->Seek(Linker->Summary.ImportOffset);
-					}
+						Linker->Seek(Linker->Summary.ImportOffset);
 
-					int32 NumImports = Linker->ImportMap.Num();
-					FStructuredArchive::FStream ImportTableStream = StructuredArchiveRoot.EnterStream(FIELD_NAME_TEXT("ImportTable"));
+						int32 NumImports = Linker->ImportMap.Num();
+						FStructuredArchive::FStream ImportTableStream = StructuredArchiveRoot.EnterStream(SA_FIELD_NAME(TEXT("ImportTable")));
 
-					for (int32 i = 0; i < Linker->ImportMap.Num(); i++)
-					{
-						FObjectImport& Import = Linker->ImportMap[i];
-						if (Import.XObject)
+						for (int32 i = 0; i < Linker->ImportMap.Num(); i++)
 						{
-							// Set the package index.
-							if (Import.XObject->GetOuter())
+							FObjectImport& Import = Linker->ImportMap[i];
+							if (Import.XObject)
 							{
-								if (Import.XObject->GetOuter()->IsIn(InOuter))
+								// Set the package index.
+								if (Import.XObject->GetOuter())
 								{
-									if (!Import.XObject->HasAllFlags(RF_Transient) || !Import.XObject->IsNative())
+									if (Import.XObject->GetOuter()->IsIn(InOuter))
 									{
-										UE_LOG(LogSavePackage, Warning, TEXT("Bad Object=%s"), *Import.XObject->GetFullName());
+										if (!Import.XObject->HasAllFlags(RF_Transient) || !Import.XObject->IsNative())
+										{
+											UE_LOG(LogSavePackage, Warning, TEXT("Bad Object=%s"), *Import.XObject->GetFullName());
+										}
+										else
+										{
+											// if an object is marked RF_Transient and native, it is either an intrinsic class or
+											// a property of an intrinsic class.  Only properties of intrinsic classes will have
+											// an Outer that passes the check for "GetOuter()->IsIn(InOuter)" (thus ending up in this
+											// block of code).  Just verify that the Outer for this property is also marked RF_Transient and Native
+											check(Import.XObject->GetOuter()->HasAllFlags(RF_Transient) && Import.XObject->GetOuter()->IsNative());
+										}
 									}
-									else
+									check(!Import.XObject->GetOuter()->IsIn(InOuter) || Import.XObject->HasAllFlags(RF_Transient) || Import.XObject->IsNative());
+#if WITH_EDITOR
+									UObject** ReplacedOuter = ReplacedImportOuters.Find(Import.XObject);
+									if (ReplacedOuter && *ReplacedOuter)
 									{
-										// if an object is marked RF_Transient and native, it is either an intrinsic class or
-										// a property of an intrinsic class.  Only properties of intrinsic classes will have
-										// an Outer that passes the check for "GetOuter()->IsIn(InOuter)" (thus ending up in this
-										// block of code).  Just verify that the Outer for this property is also marked RF_Transient and Native
-										check(Import.XObject->GetOuter()->HasAllFlags(RF_Transient) && Import.XObject->GetOuter()->IsNative());
+										Import.OuterIndex = Linker->MapObject(*ReplacedOuter);
+										ensure(Import.OuterIndex != FPackageIndex());
 									}
-								}
-								check(!Import.XObject->GetOuter()->IsIn(InOuter) || Import.XObject->HasAllFlags(RF_Transient) || Import.XObject->IsNative());
-#if WITH_EDITOR
-								UObject** ReplacedOuter = ReplacedImportOuters.Find(Import.XObject);
-								if (ReplacedOuter && *ReplacedOuter)
-								{
-									Import.OuterIndex = Linker->MapObject(*ReplacedOuter);
-									ensure(Import.OuterIndex != FPackageIndex());
-								}
-								else
+									else
 #endif
-								{
-									Import.OuterIndex = Linker->MapObject(Import.XObject->GetOuter());
-								}
+									{
+										Import.OuterIndex = Linker->MapObject(Import.XObject->GetOuter());
+									}
 
-								if (Linker->IsCooking() && IsEventDrivenLoaderEnabledInCookedBuilds())
-								{
-									// Only package imports are allowed to have no outer
-									ensureMsgf(Import.OuterIndex != FPackageIndex() || Import.ClassName == NAME_Package, TEXT("Import %s has no valid outer when cooking!"), *Import.XObject->GetPathName());
+									if (Linker->IsCooking() && IsEventDrivenLoaderEnabledInCookedBuilds())
+									{
+										// Only package imports are allowed to have no outer
+										ensureMsgf(Import.OuterIndex != FPackageIndex() || Import.ClassName == NAME_Package, TEXT("Import %s has no valid outer when cooking!"), *Import.XObject->GetPathName());
+									}
 								}
 							}
-						}
-						else
-						{
-							checkf(Conform != nullptr, TEXT("NULL XObject for import %i - Object: %s Class: %s"), i, *Import.ObjectName.ToString(), *Import.ClassName.ToString());
-						}
+							else
+							{
+								checkf(Conform != nullptr, TEXT("NULL XObject for import %i - Object: %s Class: %s"), i, *Import.ObjectName.ToString(), *Import.ClassName.ToString());
+							}
 
-						// Save it.
-						ImportTableStream.EnterElement() << Import;
+							// Save it.
+							ImportTableStream.EnterElement() << Import;
+						}
 					}
 				}
 
 				// Save the export map.
 				if (!bTextFormat)
 				{
-				check( Linker->Tell() == OffsetAfterImportMap );
+					check( Linker->Tell() == OffsetAfterImportMap );
 					Linker->Seek(Linker->Summary.ExportOffset);
-				}
 
-				int32 NumExports = Linker->ExportMap.Num();
-				FStructuredArchive::FStream ExportTableStream = StructuredArchiveRoot.EnterStream(FIELD_NAME_TEXT("ExportTable"));
-				{
+					int32 NumExports = Linker->ExportMap.Num();
+					FStructuredArchive::FStream ExportTableStream = StructuredArchiveRoot.EnterStream(SA_FIELD_NAME(TEXT("ExportTable")));
+					{
 #if WITH_EDITOR
-					FArchive::FScopeSetDebugSerializationFlags S(*Linker, DSF_IgnoreDiff, true);
-					FArchiveStackTraceIgnoreScope IgnoreSummaryDiffsScope(DiffSettings.bIgnoreHeaderDiffs);
+						FArchive::FScopeSetDebugSerializationFlags S(*Linker, DSF_IgnoreDiff, true);
+						FArchiveStackTraceIgnoreScope IgnoreSummaryDiffsScope(DiffSettings.bIgnoreHeaderDiffs);
 #endif
-					for (int32 i = 0; i < Linker->ExportMap.Num(); i++)
-					{
-						FObjectExport& Export = Linker->ExportMap[i];
-						ExportTableStream.EnterElement() << Export;
+						for (int32 i = 0; i < Linker->ExportMap.Num(); i++)
+						{
+							FObjectExport& Export = Linker->ExportMap[i];
+							ExportTableStream.EnterElement() << Export;
+						}
 					}
-				}
 
-				if (!bTextFormat)
-				{
 					check( Linker->Tell() == OffsetAfterExportMap );
 				}
 
@@ -6000,7 +5838,7 @@ FSavePackageResultStruct UPackage::Save(UPackage* InOuter, UObject* Base, EObjec
 #if WITH_EDITOR
 					FArchiveStackTraceIgnoreScope IgnoreSummaryDiffsScope(DiffSettings.bIgnoreHeaderDiffs);
 #endif // WITH_EDITOR
-					StructuredArchiveRoot.EnterField(FIELD_NAME_TEXT("Summary")) << Linker->Summary;
+					StructuredArchiveRoot.EnterField(SA_FIELD_NAME(TEXT("Summary"))) << Linker->Summary;
 				}
 
 				if (!bTextFormat)
@@ -6042,7 +5880,6 @@ FSavePackageResultStruct UPackage::Save(UPackage* InOuter, UObject* Base, EObjec
 
 				if( Success == true )
 				{
-					// Compress the temporarily file to destination.
 					if (bSaveAsync)
 					{						
 						FString NewPathToSave = NewPath;
@@ -6091,27 +5928,54 @@ FSavePackageResultStruct UPackage::Save(UPackage* InOuter, UObject* Base, EObjec
 						{
 							UE_LOG(LogSavePackage, Verbose, TEXT("Async saving from memory to '%s'"), *NewPathToSave);
 
-							// Detach archive used for memory saving.
-							FLargeMemoryWriter* Writer = (FLargeMemoryWriter*)(Linker->Saver);
-							int64 DataSize = Writer->TotalSize();
+							FLargeMemoryWriter* Writer = static_cast<FLargeMemoryWriter*>(Linker->Saver);
+							const int64 DataSize = Writer->TotalSize();
 
-							COOK_STAT(FScopedDurationTimer SaveTimer(SavePackageStats::AsyncWriteTimeSec));
 							TotalPackageSizeUncompressed += DataSize;
 
 							if (bComputeHash)
 							{
-								CookedPackageHash.Update( Writer->GetData(), Writer->TotalSize() );
+								CookedPackageHash.Update(Writer->GetData(), DataSize);
 							}
 
-							FLargeMemoryPtr DataPtr(Writer->GetData());
-							Writer->ReleaseOwnership();
 							if (IsEventDrivenLoaderEnabledInCookedBuilds() && Linker->IsCooking())
 							{
-								AsyncWriteFileWithSplitExports(MoveTemp(DataPtr), DataSize, Linker->Summary.TotalHeaderSize, *NewPathToSave, FinalTimeStamp);
+								if (SavePackageContext)
+								{
+									FIoBuffer IoBuffer(FIoBuffer::AssumeOwnership, Writer->ReleaseOwnership(), DataSize);
+
+									const int32 HeaderSize = Linker->Summary.TotalHeaderSize;
+
+									FPackageStoreWriter::HeaderInfo HeaderInfo;
+									HeaderInfo.PackageName		= InOuter->GetFName();
+									HeaderInfo.LooseFilePath	= Filename;
+
+									SavePackageContext->PackageStoreWriter.WriteHeader(HeaderInfo, FIoBuffer(IoBuffer.Data(), HeaderSize, IoBuffer));
+
+									FPackageStoreWriter::ExportsInfo ExportsInfo;
+									ExportsInfo.LooseFilePath	= Filename;
+									ExportsInfo.PackageName		= InOuter->GetFName();
+
+									const uint8* ExportsData = IoBuffer.Data() + HeaderSize;
+									const int32 ExportCount = Linker->ExportMap.Num();
+
+									ExportsInfo.Exports.Reserve(ExportCount);
+
+									for (const FObjectExport& Export : Linker->ExportMap)
+									{
+										ExportsInfo.Exports.Add(FIoBuffer(IoBuffer.Data() + Export.SerialOffset, Export.SerialSize, IoBuffer));
+									}
+
+									SavePackageContext->PackageStoreWriter.WriteExports(ExportsInfo, FIoBuffer(ExportsData, DataSize - HeaderSize, IoBuffer));
+								}
+								else
+								{
+									AsyncWriteFileWithSplitExports(FLargeMemoryPtr(Writer->ReleaseOwnership()), DataSize, Linker->Summary.TotalHeaderSize, *NewPathToSave);
+								}
 							}
 							else
 							{
-								AsyncWriteFile(MoveTemp(DataPtr), DataSize, *NewPathToSave, FinalTimeStamp);
+								AsyncWriteFile(FLargeMemoryPtr(Writer->ReleaseOwnership()), DataSize, *NewPathToSave);
 							}
 						}
 						Linker->CloseAndDestroySaver();
@@ -6286,7 +6150,7 @@ FSavePackageResultStruct UPackage::Save(UPackage* InOuter, UObject* Base, EObjec
 }
 
 bool UPackage::SavePackage(UPackage* InOuter, UObject* Base, EObjectFlags TopLevelFlags, const TCHAR* Filename,
-	FOutputDevice* Error, FLinkerLoad* Conform, bool bForceByteSwapping, bool bWarnOfLongFilename, uint32 SaveFlags,
+	FOutputDevice* Error, FLinkerNull* Conform, bool bForceByteSwapping, bool bWarnOfLongFilename, uint32 SaveFlags,
 	const class ITargetPlatform* TargetPlatform, const FDateTime&  FinalTimeStamp, bool bSlowTask)
 {
 	const FSavePackageResultStruct Result = Save(InOuter, Base, TopLevelFlags, Filename, Error, Conform, bForceByteSwapping,
@@ -6301,7 +6165,7 @@ bool UPackage::SavePackage(UPackage* InOuter, UObject* Base, EObjectFlags TopLev
  * @param	Linker							linker we're currently saving with
  * @param	Slot							structed archive slot we are saving too (temporary)
  */
-void UPackage::SaveThumbnails(UPackage* InOuter, FLinkerSave* Linker, FStructuredArchive::FSlot Slot)
+static void SaveThumbnails(UPackage* InOuter, FLinkerSave* Linker, FStructuredArchive::FSlot Slot)
 {
 	FStructuredArchive::FRecord Record = Slot.EnterRecord();
 
@@ -6355,7 +6219,7 @@ void UPackage::SaveThumbnails(UPackage* InOuter, FLinkerSave* Linker, FStructure
 		if( ObjectsWithThumbnails.Num() > 0 )
 		{
 			// Save out the image data for the thumbnails
-			FStructuredArchive::FStream ThumbnailStream = Record.EnterStream(FIELD_NAME_TEXT("Thumbnails"));
+			FStructuredArchive::FStream ThumbnailStream = Record.EnterStream(SA_FIELD_NAME(TEXT("Thumbnails")));
 
 			for( int32 CurObjectIndex = 0; CurObjectIndex < ObjectsWithThumbnails.Num(); ++CurObjectIndex )
 			{
@@ -6376,7 +6240,7 @@ void UPackage::SaveThumbnails(UPackage* InOuter, FLinkerSave* Linker, FStructure
 
 				// Save number of thumbnails
 				int32 ThumbnailCount = ObjectsWithThumbnails.Num();
-				FStructuredArchive::FArray IndexArray = Record.EnterField(FIELD_NAME_TEXT("Index")).EnterArray(ThumbnailCount);
+				FStructuredArchive::FArray IndexArray = Record.EnterField(SA_FIELD_NAME(TEXT("Index"))).EnterArray(ThumbnailCount);
 
 				// Store a list of object names along with the offset in the file where the thumbnail is stored
 				for( int32 CurObjectIndex = 0; CurObjectIndex < ObjectsWithThumbnails.Num(); ++CurObjectIndex )
@@ -6400,9 +6264,9 @@ void UPackage::SaveThumbnails(UPackage* InOuter, FLinkerSave* Linker, FStructure
 					int32 FileOffset = CurObjectThumb.FileOffset;
 
 					IndexArray.EnterElement().EnterRecord()
-						<< NAMED_FIELD(ObjectClassName)
-						<< NAMED_FIELD(ObjectPathWithoutPackageName)
-						<< NAMED_FIELD(FileOffset);
+						<< SA_VALUE(TEXT("ObjectClassName"), ObjectClassName)
+						<< SA_VALUE(TEXT("ObjectPathWithoutPackageName"), ObjectPathWithoutPackageName)
+						<< SA_VALUE(TEXT("FileOffset"), FileOffset);
 				}
 			}
 		}
@@ -6416,7 +6280,7 @@ void UPackage::SaveThumbnails(UPackage* InOuter, FLinkerSave* Linker, FStructure
 #endif
 }
 
-void UPackage::SaveAssetRegistryData(UPackage* InOuter, FLinkerSave* Linker, FStructuredArchive::FSlot Slot)
+static void SaveAssetRegistryData(UPackage* InOuter, FLinkerSave* Linker, FStructuredArchive::FSlot Slot)
 {
 	// Make a copy of the tag map
 	TArray<UObject*> AssetObjects;
@@ -6450,13 +6314,13 @@ void UPackage::SaveAssetRegistryData(UPackage* InOuter, FLinkerSave* Linker, FSt
 		FString ObjectPath = Object->GetPathName(Object->GetOutermost());
 		FString ObjectClassName = Object->GetClass()->GetName();
 		
-		TArray<FAssetRegistryTag> SourceTags;
+		TArray<UObject::FAssetRegistryTag> SourceTags;
 		Object->GetAssetRegistryTags(SourceTags);
 
-		TArray<FAssetRegistryTag> Tags;
-		for (FAssetRegistryTag& SourceTag : SourceTags)
+		TArray<UObject::FAssetRegistryTag> Tags;
+		for (UObject::FAssetRegistryTag& SourceTag : SourceTags)
 		{
-			FAssetRegistryTag* Existing = Tags.FindByPredicate([SourceTag](const FAssetRegistryTag& InTag) { return InTag.Name == SourceTag.Name; });
+			UObject::FAssetRegistryTag* Existing = Tags.FindByPredicate([SourceTag](const UObject::FAssetRegistryTag& InTag) { return InTag.Name == SourceTag.Name; });
 			if (Existing)
 			{
 				Existing->Value = SourceTag.Value;
@@ -6470,11 +6334,11 @@ void UPackage::SaveAssetRegistryData(UPackage* InOuter, FLinkerSave* Linker, FSt
 		int32 TagCount = Tags.Num();
 
 		FStructuredArchive::FRecord AssetRecord = AssetArray.EnterElement().EnterRecord();
-		AssetRecord << NAMED_ITEM("Path", ObjectPath) << NAMED_ITEM("Class", ObjectClassName);
+		AssetRecord << SA_VALUE(TEXT("Path"), ObjectPath) << SA_VALUE(TEXT("Class"), ObjectClassName);
 
-		FStructuredArchive::FMap TagMap = AssetRecord.EnterField(FIELD_NAME_TEXT("Tags")).EnterMap(TagCount);
+		FStructuredArchive::FMap TagMap = AssetRecord.EnterField(SA_FIELD_NAME(TEXT("Tags"))).EnterMap(TagCount);
 				
-		for (TArray<FAssetRegistryTag>::TConstIterator TagIter(Tags); TagIter; ++TagIter)
+		for (TArray<UObject::FAssetRegistryTag>::TConstIterator TagIter(Tags); TagIter; ++TagIter)
 		{
 			FString Key = TagIter->Name.ToString();
 			FString Value = TagIter->Value;
@@ -6484,23 +6348,18 @@ void UPackage::SaveAssetRegistryData(UPackage* InOuter, FLinkerSave* Linker, FSt
 	}
 }
 
-void UPackage::SaveWorldLevelInfo(UPackage* InOuter, FLinkerSave* Linker, FStructuredArchive::FSlot Slot)
+static void SaveWorldLevelInfo(UPackage* InOuter, FLinkerSave* Linker, FStructuredArchive::FRecord Record)
 {
 	Linker->Summary.WorldTileInfoDataOffset = 0;
 	
 	if(InOuter->WorldTileInfo.IsValid())
 	{
 		Linker->Summary.WorldTileInfoDataOffset = Linker->Tell();
-		Slot << *(InOuter->WorldTileInfo);
-	}
-	else
-	{
-		// Make empty record, because structured archive API doesn't allow us to leave a slot empty
-		Slot.EnterRecord();
+		Record << SA_VALUE(TEXT("WorldLevelInfo"), *(InOuter->WorldTileInfo));
 	}
 }
 
-bool UPackage::IsEmptyPackage (UPackage* Package, const UObject* LastReferencer)
+bool UPackage::IsEmptyPackage(UPackage* Package, const UObject* LastReferencer)
 {
 	// Don't count null or volatile packages as empty, just let them be NULL or get GCed
 	if ( Package != nullptr )
@@ -6548,3 +6407,55 @@ bool UPackage::IsEmptyPackage (UPackage* Package, const UObject* LastReferencer)
 	// Invalid package
 	return false;
 }
+
+//////////////////////////////////////////////////////////////////////////
+// TODO: this should go elsewhere, this file is big enough as it is already
+//
+
+FPackageStoreWriter::FPackageStoreWriter() = default;
+FPackageStoreWriter::~FPackageStoreWriter() = default;
+
+FLooseFileWriter::FLooseFileWriter() = default;
+FLooseFileWriter::~FLooseFileWriter() = default;
+
+void FLooseFileWriter::WriteHeader(const FLooseFileWriter::HeaderInfo& Info, const FIoBuffer& HeaderData)
+{
+	WriteToFile(Info.LooseFilePath, HeaderData.Data(), HeaderData.DataSize());
+}
+
+void FLooseFileWriter::WriteExports(const FLooseFileWriter::ExportsInfo& Info, const FIoBuffer& ExportsData)
+{
+	const FString ArchiveFilename = FPaths::ChangeExtension(Info.LooseFilePath, TEXT(".uexp"));
+
+	WriteToFile(ArchiveFilename, ExportsData.Data(), ExportsData.DataSize());
+}
+
+void FLooseFileWriter::WriteBulkdata(const FLooseFileWriter::FBulkDataInfo& Info, const FIoBuffer& BulkData)
+{
+	if (BulkData.DataSize() == 0)
+		return;
+
+	const TCHAR* BulkFileExtension = nullptr;
+
+	switch (Info.BulkdataType)
+	{
+		case FPackageStoreWriter::FBulkDataInfo::Standard:
+			BulkFileExtension = TEXT(".ubulk");		// Regular separate bulk data file
+			break;
+
+		case FPackageStoreWriter::FBulkDataInfo::Mmap:
+			BulkFileExtension = TEXT(".m.ubulk");	// Memory-mapped bulk data
+			break;
+
+		case FPackageStoreWriter::FBulkDataInfo::Optional:
+			BulkFileExtension = TEXT(".uptnl");		// Optional bulk data
+			break;
+
+		default:
+			check(false);
+	}
+
+	const FString ArchiveFilename = FPaths::ChangeExtension(Info.LooseFilePath, BulkFileExtension);
+
+	WriteToFile(ArchiveFilename, BulkData.Data(), BulkData.DataSize());
+}
diff --git a/Engine/Source/Runtime/CoreUObject/Private/UObject/UObjectGlobals.cpp b/Engine/Source/Runtime/CoreUObject/Private/UObject/UObjectGlobals.cpp
--- a/Engine/Source/Runtime/CoreUObject/Private/UObject/UObjectGlobals.cpp
+++ b/Engine/Source/Runtime/CoreUObject/Private/UObject/UObjectGlobals.cpp
@@ -52,6 +52,7 @@
 #include "UObject/TextProperty.h"
 #include "UObject/MetaData.h"
 #include "HAL/LowLevelMemTracker.h"
+#include "HAL/LowLevelMemStats.h"
 #include "Misc/CoreDelegates.h"
 #include "ProfilingDebugging/CsvProfiler.h"
 
@@ -1040,7 +1041,7 @@ public:
 
 		
 
-		while ( Tick(0.0, false, false) == FLinkerLoad::LINKER_TimedOut ) 
+		while ( Tick(0.0, false, false, nullptr) == FLinkerLoad::LINKER_TimedOut ) 
 		{ 
 		}
 
