Copying //UE4/Dev-Core to Main (//UE4/Main)

diff --git a/Engine/Source/Runtime/CoreUObject/Public/Serialization/BulkData.h b/Engine/Source/Runtime/CoreUObject/Public/Serialization/BulkData.h
--- a/Engine/Source/Runtime/CoreUObject/Public/Serialization/BulkData.h
+++ b/Engine/Source/Runtime/CoreUObject/Public/Serialization/BulkData.h
@@ -6,6 +6,18 @@
 #include "Containers/SortedMap.h"
 #include "UObject/WeakObjectPtr.h"
 #include "Async/Future.h"
+#include "Async/AsyncFileHandle.h"
+#include "GenericPlatform/GenericPlatformFile.h"
+#include "BulkData2.h"
+
+#if WITH_EDITOR == 0 && WITH_EDITORONLY_DATA == 0 //Runtime
+	#define USE_NEW_BULKDATA 0
+#else
+	#define USE_NEW_BULKDATA 0
+#endif
+
+// Enable the following to use the more compact FBulkDataStreamingToken in places where it is implemented
+#define USE_BULKDATA_STREAMING_TOKEN !USE_NEW_BULKDATA
 
 /**
  * Flags serialized with the bulk data.
@@ -129,6 +141,91 @@ private:
 	IMappedFileRegion* MappedRegion;
 };
 
+/**
+ * Represents an IO request from the FUntypedBulkData streaming API.
+ *
+ * It functions pretty much the same as IAsyncReadRequest expect that it also holds
+ * the file handle as well. 
+ *
+ * 
+ */
+struct COREUOBJECT_API FBulkDataIORequest
+{
+	FBulkDataIORequest(IAsyncReadFileHandle* InFileHandle, IAsyncReadRequest* InReadRequest, int64 BytesToRead)
+		: FileHandle(InFileHandle)
+		, ReadRequest(InReadRequest)
+		, Size(BytesToRead)
+	{
+
+	}
+
+	~FBulkDataIORequest();
+
+	bool PollCompletion() const;
+	bool WaitCompletion( float TimeLimitSeconds = 0.0f ) const;
+
+	uint8* GetReadResults() const;
+	int64 GetSize() const;
+
+	bool operator == (const IAsyncReadRequest* InReadRequest)
+	{
+		return ReadRequest == InReadRequest;
+	}
+
+	void Cancel() const;
+
+private:
+	IAsyncReadFileHandle* FileHandle;
+	IAsyncReadRequest* ReadRequest;
+
+	int64 Size;
+};
+
+#if USE_BULKDATA_STREAMING_TOKEN 
+/**
+ * Some areas of code currently find FUntypedBulkData too bloated for use so were storing smaller parts of it
+ * and performing the file IO manually. This is an attempt to wrap that functionality in a single place to make
+ * it easier to replace in the future.
+ * This structure will represent the area of a file that a FUntypedBulkData points to.
+ * The structure does not contain info on which file it is however (as that would increase the data
+ * size by an unacceptable amount) and it is assumed that this knowledge is being kept at a higher level.
+ *
+ * This is intended only to exist for a short time to aid the transition to a new system.
+ */
+struct COREUOBJECT_API FBulkDataStreamingToken
+{
+public:
+	static constexpr uint32 InvalidOffset = TNumericLimits<uint32>::Max();
+
+public:
+	FBulkDataStreamingToken() {}
+	FBulkDataStreamingToken(uint32 InOffsetInFile, uint32 InBulkDataSize)
+		: OffsetInFile(InOffsetInFile)
+		, BulkDataSize(InBulkDataSize)
+	{
+	}
+
+	bool IsValid() const
+	{ 
+		return OffsetInFile != InvalidOffset && BulkDataSize > 0;
+	}
+
+	uint32 GetOffset() const
+	{
+		return OffsetInFile;
+	}
+
+	uint32 GetSize() const
+	{
+		return BulkDataSize;
+	}
+
+private:
+	uint32 OffsetInFile{ InvalidOffset };
+	uint32 BulkDataSize{ 0 };
+};
+#endif
+
 /**
  * @documentation @todo documentation
  */
@@ -375,12 +472,48 @@ public:
 	bool IsAvailableForUse() const;
 
 	/**
-	 * Sets the passed in bulk data flags.
+	* Returns whether this bulk data represents optional data or not
+	* @return true if BULKDATA_OptionalPayload is set
+	*/
+	bool IsOptional() const
+	{
+		return (GetBulkDataFlags() & BULKDATA_OptionalPayload) != 0;
+	}
+
+	/**
+	* Returns whether this bulk data is currently stored inline or not
+	* @return true if BULKDATA_PayloadAtEndOfFile is not set
+	*/
+	bool IsInlined() const
+	{
+		return   (GetBulkDataFlags() & BULKDATA_PayloadAtEndOfFile) == 0;
+	}
+
+	/**
+	* Returns whether this bulk data is currently stored in it's own file or not
+	* @return true if BULKDATA_PayloadInSeperateFile is not set
+	*/
+	bool InSeperateFile() const
+	{
+		return	(GetBulkDataFlags() & BULKDATA_PayloadInSeperateFile) != 0;
+	}
+
+	/**
+	 * Enables the given flags without affecting any previously set flags.
 	 *
 	 * @param BulkDataFlagsToSet	Bulk data flags to set
 	 */
 	void SetBulkDataFlags( uint32 BulkDataFlagsToSet );
 
+	/**
+	 * Enable the given flags and disable all other flags.
+	 * This can be used with ::GetBulkDataFlags() to effectively
+	 * reset the flags to a previous state.
+	 *
+	 * @param BulkDataFlagsToSet	Bulk data flags to set
+	 */
+	void ResetBulkDataFlags(uint32 BulkDataFlagsToSet);
+
 	/**
 	* Gets the current bulk data flags.
 	*
@@ -516,10 +649,64 @@ public:
 	 *
 	 * @param	Ar					Archive to serialize with
 	 * @param	Data				Memory to serialize either to or from
-	 */
+	 */ 
 	void SerializeBulkData( FArchive& Ar, void* Data );
 
+	/*-----------------------------------------------------------------------------
+		Async Streaming Interface.
+	-----------------------------------------------------------------------------*/
 	
+	/**
+	* Create an async read request for the bulk data.
+	* This version will load the entire data range that the FUntypedBulkData represents.
+	*
+	* @param Priority			Priority and flags of the request. If this includes AIOP_FLAG_PRECACHE, then memory will never be returned. The request should always be canceled and waited for, even for a precache request.
+	* @param CompleteCallback	Called from an arbitrary thread when the request is complete. Can be nullptr, if non-null, must remain valid until it is called. It will always be called.
+	* @param UserSuppliedMemory A pointer to memory for the IO request to be written to, it is up to the caller to make sure that it is large enough. If the pointer is null then the system will allocate memory instead.
+	* @return					A request for the read. This is owned by the caller and must be deleted by the caller.
+	**/
+	FBulkDataIORequest* CreateStreamingRequest(EAsyncIOPriorityAndFlags Priority, FAsyncFileCallBack* CompleteCallback, uint8* UserSuppliedMemory) const;
+
+	/**
+	* Create an async read request for the bulk data.
+	* This version allows the user to request a subset of the data that the FUntypedBulkData represents.
+	*
+	* @param OffsetInBulkData	Offset into the bulk data to start reading from.
+	* @param BytesToRead		The number of bytes to read. If this request is AIOP_Preache, the size can be anything, even MAX_int64, otherwise the size and offset must be fully contained in the file.
+	* @param Priority			Priority and flags of the request. If this includes AIOP_FLAG_PRECACHE, then memory will never be returned. The request should always be canceled and waited for, even for a precache request.
+	* @param CompleteCallback	Called from an arbitrary thread when the request is complete. Can be nullptr, if non-null, must remain valid until it is called. It will always be called.
+	* @param UserSuppliedMemory A pointer to memory for the IO request to be written to, it is up to the caller to make sure that it is large enough. If the pointer is null then the system will allocate memory instead.
+	* @return					A request for the read. This is owned by the caller and must be deleted by the caller.
+	**/
+	FBulkDataIORequest* CreateStreamingRequest(int64 OffsetInBulkData, int64 BytesToRead, EAsyncIOPriorityAndFlags Priority, FAsyncFileCallBack* CompleteCallback, uint8* UserSuppliedMemory) const;
+
+#if USE_BULKDATA_STREAMING_TOKEN 
+
+	/**
+	* Creates a FBulkDataStreamingToken representing the area of the file that the FUntypedBulkData represents. See the declaration of 
+	* FBulkDataStreamingToken for further details.
+	* 
+	* @return	A FBulkDataStreamingToken valid for the FUntypedBulkData.
+	**/
+	FBulkDataStreamingToken CreateStreamingToken() const;
+
+	/**
+	* Create an async read request for a range of bulk data streaming tokens
+	* The request will read all data between the two given streaming tokens objects. They must both represent areas of data in the file!
+	* There is no way to validate this and it is up to the caller to make sure that it is correct.
+	* The memory to be read into will be automatically allocated the size of which can be retrieved by calling FBulkDataIORequest::GetSize()
+	*
+	* @param Filename			The file to read from.
+	* @param Start				The bulk data to start reading from.
+	* @param End				The bulk data to finish reading from.
+	* @param Priority			Priority and flags of the request. If this includes AIOP_FLAG_PRECACHE, then memory will never be returned. The request should always be canceled and waited for, even for a precache request.
+	* @param CompleteCallback	Called from an arbitrary thread when the request is complete. Can be nullptr, if non-null, must remain valid until it is called. It will always be called.
+	* @return					A request for the read. This is owned by the caller and must be deleted by the caller.
+	**/
+	static FBulkDataIORequest* CreateStreamingRequestForRange(const FString& Filename, const FBulkDataStreamingToken& Start, const FBulkDataStreamingToken& End, EAsyncIOPriorityAndFlags Priority, FAsyncFileCallBack* CompleteCallback);
+
+#endif
+
 	/*-----------------------------------------------------------------------------
 		Class specific virtuals.
 	-----------------------------------------------------------------------------*/
@@ -654,7 +841,7 @@ protected:
 	uint8 version of bulk data.
 -----------------------------------------------------------------------------*/
 
-struct COREUOBJECT_API FByteBulkData : public FUntypedBulkData
+struct COREUOBJECT_API FByteBulkDataOld : public FUntypedBulkData
 {
 	/**
 	 * Returns size in bytes of single element.
@@ -677,7 +864,7 @@ struct COREUOBJECT_API FByteBulkData : public FUntypedBulkData
 	WORD version of bulk data.
 -----------------------------------------------------------------------------*/
 
-struct COREUOBJECT_API FWordBulkData : public FUntypedBulkData
+struct COREUOBJECT_API FWordBulkDataOld : public FUntypedBulkData
 {
 	/**
 	 * Returns size in bytes of single element.
@@ -700,7 +887,7 @@ struct COREUOBJECT_API FWordBulkData : public FUntypedBulkData
 	int32 version of bulk data.
 -----------------------------------------------------------------------------*/
 
-struct COREUOBJECT_API FIntBulkData : public FUntypedBulkData
+struct COREUOBJECT_API FIntBulkDataOld : public FUntypedBulkData
 {
 	/**
 	 * Returns size in bytes of single element.
@@ -723,7 +910,7 @@ struct COREUOBJECT_API FIntBulkData : public FUntypedBulkData
 	float version of bulk data.
 -----------------------------------------------------------------------------*/
 
-struct COREUOBJECT_API FFloatBulkData : public FUntypedBulkData
+struct COREUOBJECT_API FFloatBulkDataOld : public FUntypedBulkData
 {
 	/**
 	 * Returns size in bytes of single element.
@@ -742,6 +929,37 @@ struct COREUOBJECT_API FFloatBulkData : public FUntypedBulkData
 	virtual void SerializeElement( FArchive& Ar, void* Data, int64 ElementIndex );
 };
 
+// Switch between the old and new types based on USE_NEW_BULKDATA
+#if !USE_NEW_BULKDATA
+using FBulkDataInterface	= FUntypedBulkData;
+using FByteBulkData			= FByteBulkDataOld;
+using FWordBulkData			= FWordBulkDataOld;
+using FIntBulkData			= FIntBulkDataOld;
+using FFloatBulkData		= FFloatBulkDataOld;
+#else
+using FBulkDataInterface	= FBulkDataBase;
+using FByteBulkData			= FByteBulkData2;
+using FWordBulkData			= FWordBulkData2;
+using FIntBulkData			= FIntBulkData2;
+using FFloatBulkData		= FFloatBulkData2;
+#endif
+
+namespace BulkDataUtils
+{
+	/**
+	* Create an async read request for a range of bulk data
+	* The request will read all data between the two given BulkData objects. They must both represent areas of data in the same file!
+	* The memory to be read into will be automatically allocated the size of which can be retrieved by calling FBulkDataIORequest::GetSize()
+	*
+	* @param Start				The bulk data to start reading from.
+	* @param End				The bulk data to finish reading from.
+	* @param Priority			Priority and flags of the request. If this includes AIOP_FLAG_PRECACHE, then memory will never be returned. The request should always be canceled and waited for, even for a precache request.
+	* @param CompleteCallback	Called from an arbitrary thread when the request is complete. Can be nullptr, if non-null, must remain valid until it is called. It will always be called.
+	* @return					A request for the read. This is owned by the caller and must be deleted by the caller.
+	**/
+	COREUOBJECT_API FBulkDataIORequest* CreateStreamingRequestForRange(const FBulkDataInterface& Start, const FBulkDataInterface& End, EAsyncIOPriorityAndFlags Priority, FAsyncFileCallBack* CompleteCallback);
+}
+
 class FFormatContainer
 {
 	friend class UBodySetup;
@@ -778,4 +996,3 @@ public:
 	COREUOBJECT_API void Serialize(FArchive& Ar, UObject* Owner, const TArray<FName>* FormatsToSave = nullptr, bool bSingleUse = true, uint32 InAlignment = DEFAULT_ALIGNMENT, bool bInline = true, bool bMapped = false);
 	COREUOBJECT_API void SerializeAttemptMappedLoad(FArchive& Ar, UObject* Owner);
 };
-
diff --git a/Engine/Source/Runtime/CoreUObject/Public/Serialization/BulkData2.h b/Engine/Source/Runtime/CoreUObject/Public/Serialization/BulkData2.h
new file mode 100644
--- /dev/null
+++ b/Engine/Source/Runtime/CoreUObject/Public/Serialization/BulkData2.h
@@ -0,0 +1,145 @@
+// Copyright 1998-2019 Epic Games, Inc. All Rights Reserved.
+
+#pragma once
+
+#include "CoreMinimal.h"
+#include "Async/AsyncFileHandle.h"
+
+// Any place marked with BULKDATA_NOT_IMPLEMENTED_FOR_RUNTIME marks a method that we do not support
+// but needs to exist in order for the code to compile. This pretty much means code that is editor only
+// but is not compiled out at runtime etc.
+// It can also be found inside of methods to show code paths that the old path implements but I have
+// not been able to find any way to actually execute/test.
+#define BULKDATA_NOT_IMPLEMENTED_FOR_RUNTIME PLATFORM_BREAK()
+
+struct FBulkDataIORequest;
+struct FOwnedBulkDataPtr;
+
+/**
+ * @documentation @todo documentation
+ */
+class COREUOBJECT_API FBulkDataBase
+{
+public:
+	using FileToken = uint64;
+	static constexpr FileToken InvalidToken = 0;
+
+	FBulkDataBase(const FBulkDataBase& Other) { *this = Other; }
+	FBulkDataBase(FBulkDataBase&& Other);
+	FBulkDataBase& operator=(const FBulkDataBase& Other);
+
+	// Implemented:
+	FBulkDataBase() = default;
+	~FBulkDataBase();
+protected:
+	// Unimplemented:
+	void Serialize(FArchive& Ar, UObject* Owner, int32 Index, bool bAttemptFileMapping, int32 ElementSize );
+
+public:
+	// Unimplemented:
+	void* Lock(uint32 LockFlags);
+	const void* LockReadOnly() const;
+	void Unlock();
+	bool IsLocked() const;
+	
+	void* Realloc(int64 InElementCount);
+	void GetCopy(void** Dest, bool bDiscardInternalCopy);
+
+	int64 GetBulkDataSize() const;
+	
+	void SetBulkDataFlags(uint32 BulkDataFlagsToSet);
+	void ResetBulkDataFlags(uint32 BulkDataFlagsToSet);
+	void ClearBulkDataFlags(uint32 BulkDataFlagsToClear);
+	uint32 GetBulkDataFlags() const { return BulkDataFlags; }
+
+	bool CanLoadFromDisk() const { return Token != InvalidToken; }
+
+	bool IsStoredCompressedOnDisk() const;
+	FName GetDecompressionFormat() const;
+
+	bool IsBulkDataLoaded() const { return DataBuffer != nullptr; }
+
+	// TODO: The flag tests could be inline if we fixed the header dependency issues (the flags are defined in Bulkdata.h at the moment)
+	bool IsAvailableForUse() const;
+	bool IsDuplicateNonOptional() const;
+	bool IsOptional() const;
+	bool IsInlined() const;
+	bool InSeperateFile() const;
+	bool IsSingleUse() const;
+	bool IsMemoryMapped() const;
+
+	FBulkDataIORequest* CreateStreamingRequest(EAsyncIOPriorityAndFlags Priority, FAsyncFileCallBack* CompleteCallback, uint8* UserSuppliedMemory) const;
+	FBulkDataIORequest* CreateStreamingRequest(int64 OffsetInBulkData, int64 BytesToRead, EAsyncIOPriorityAndFlags Priority, FAsyncFileCallBack* CompleteCallback, uint8* UserSuppliedMemory) const;
+
+	void RemoveBulkData();
+
+	bool IsAsyncLoadingComplete() const { return true;  }
+
+	// Added for compatibility with the older BulkData system
+	int64 GetBulkDataOffsetInFile() const;
+	FString GetFilename() const;
+
+public:
+	// The following methods are for compatibility with SoundWave.cpp which assumes memory mapping.
+	void ForceBulkDataResident(); // Is closer to MakeSureBulkDataIsLoaded in the old system but kept the name due to existing use
+	FOwnedBulkDataPtr* StealFileMapping();
+
+private:
+	void LoadDataDirectly(void** DstBuffer);
+
+	void SerializeDuplicateData(FArchive& Ar, UObject* Owner, uint32& OutBulkDataFlags, int64& OutBulkDataSizeOnDisk, int64& OutBulkDataOffsetInFile);
+	void SerializeBulkData(FArchive& Ar, void* DstBuffer, int64 DataLength);
+
+	void AllocateData(SIZE_T SizeInBytes);
+	void FreeData();
+
+	FString ConvertFilenameFromFlags(const FString& Filename);
+
+	FileToken Token = InvalidToken; // TODO: Maybe store as a union of Token/InlineData size so that inline files don't need to register?
+	void* DataBuffer = nullptr;
+	uint32 BulkDataFlags = 0;
+	mutable uint8 LockStatus = 0; // Mutable so that the read only lock can be const
+};
+
+/**
+ * @documentation @todo documentation
+ */
+template<typename ElementType>
+class COREUOBJECT_API FUntypedBulkData2 : public FBulkDataBase
+{
+	// In the older Bulkdata system the data was being loaded as if it was POD with the option to opt out
+	// but nothing actually opted out. This check should help catch if any non-POD data was actually being
+	// used or not.
+	static_assert(TIsPODType<ElementType>::Value, "FUntypedBulkData2 is limited to POD types!");
+public:
+	void Serialize(FArchive& Ar, UObject* Owner, int32 Index, bool bAttemptFileMapping)
+	{
+		FBulkDataBase::Serialize(Ar, Owner, Index, bAttemptFileMapping, sizeof(ElementType));
+	}
+
+	int32 GetElementSize() const
+	{
+		return sizeof(ElementType);
+	}
+
+	ElementType* Lock(uint32 LockFlags)
+	{
+		return (ElementType*)FBulkDataBase::Lock(LockFlags);
+	}
+
+	const ElementType* LockReadOnly() const
+	{
+		return (const ElementType*)FBulkDataBase::LockReadOnly();
+	}
+
+	ElementType* Realloc(int64 InElementCount)
+	{
+		return (ElementType*)FBulkDataBase::Realloc(InElementCount);
+	}
+};
+
+// Commonly used types
+using FByteBulkData2 = FUntypedBulkData2<uint8>;
+using FWordBulkData2 = FUntypedBulkData2<uint16>;
+using FIntBulkData2 = FUntypedBulkData2<int32>;
+using FFloatBulkData2 = FUntypedBulkData2<float>;
diff --git a/Engine/Source/Runtime/CoreUObject/Public/UObject/LinkerLoad.h b/Engine/Source/Runtime/CoreUObject/Public/UObject/LinkerLoad.h
--- a/Engine/Source/Runtime/CoreUObject/Public/UObject/LinkerLoad.h
+++ b/Engine/Source/Runtime/CoreUObject/Public/UObject/LinkerLoad.h
@@ -82,6 +82,7 @@ class FLinkerLoad
 	friend class UObject;
 	friend class UPackageMap;
 	friend struct FAsyncPackage;
+	friend struct FAsyncPackage2;
 	friend struct FResolvingExportTracker;
 protected:
 	/** Linker loading status. */
@@ -142,14 +143,15 @@ private:
 	FStructuredArchive* StructuredArchive;
 	FArchiveFormatterType* StructuredArchiveFormatter;
 	TOptional<FStructuredArchive::FRecord> StructuredArchiveRootRecord;
-
-	/** A map of full object path name to package index. Used with text assets to resolve incoming string names to an export */
-	TMap<FName, FPackageIndex> ObjectNameToPackageImportIndex;
-	TMap<FName, FPackageIndex> ObjectNameToPackageExportIndex;
+	TArray<FStructuredArchiveChildReader*> ExportReaders;
 
 	/** The archive that actually reads the raw data from disk.																*/
 	FArchive*				Loader;
 
+	int32* LocalImportIndices = nullptr;
+	UObject** GlobalImportObjects = nullptr;
+	const TArray<FNameEntryId>* ActiveNameMap = &NameMap;
+
 protected:
 
 	void SetLoader(FArchive* InLoader);
@@ -178,14 +180,15 @@ public:
 	}
 
 	/** The async package associated with this linker */
-	struct FAsyncPackage* AsyncRoot;
+	class FGCObject* AsyncRoot;
 #if WITH_EDITOR
 	/** Bulk data that does not need to be loaded when the linker is loaded.												*/
 	TArray<FUntypedBulkData*> BulkDataLoaders;
 #endif // WITH_EDITOR
 
 	/** Hash table for exports.																								*/
-	int32						ExportHash[256];
+	static constexpr int32 ExportHashCount = 256;
+	TUniquePtr<int32[]> ExportHash;
 
 	/**
 	* List of imports and exports that must be serialized before other exports...all packed together, see FirstExportDependency
@@ -277,6 +280,8 @@ private:
 
 	/** Whether we already serialized the package file summary.																*/
 	bool					bHasSerializedPackageFileSummary;
+	/** Whether we have already reconstructed the import/export tables for a text asset */
+	bool					bHasReconstructedImportAndExportMap;
 	/** Whether we already serialized preload dependencies.																*/
 	bool					bHasSerializedPreloadDependencies;
 	/** Whether we already fixed up import map.																				*/
@@ -818,10 +823,10 @@ private:
 		int32 Number = 0;
 		Ar << Number;
 
-		if (NameMap.IsValidIndex(NameIndex))
+		if (ActiveNameMap->IsValidIndex(NameIndex))
 		{
 			// if the name wasn't loaded (because it wasn't valid in this context)
-			FNameEntryId MappedName = NameMap[NameIndex];
+			FNameEntryId MappedName = (*ActiveNameMap)[NameIndex];
 
 			// simply create the name from the NameMap's name and the serialized instance number
 			Name = FName::CreateFromDisplayId(MappedName, Number);
@@ -878,7 +883,7 @@ protected: // Daniel L: Made this protected so I can override the constructor an
 	 * 
 	 * @return	true if linker has finished creation, false if it is still in flight
 	 */
-	ELinkerStatus Tick( float InTimeLimit, bool bInUseTimeLimit, bool bInUseFullTimeLimit);
+	ELinkerStatus Tick( float InTimeLimit, bool bInUseTimeLimit, bool bInUseFullTimeLimit, TMap<TPair<FName, FPackageIndex>, FPackageIndex>* ObjectNameWithOuterToExportMap);
 
 	/**
 	 * Private constructor, passing arguments through from CreateLinker.
@@ -911,6 +916,11 @@ private:
 	 */
 	ELinkerStatus SerializePackageFileSummary();
 
+	/**
+	 * Updates the linker, loader and root package with data from the package file summary.
+	 * */
+	ELinkerStatus UpdateFromPackageFileSummary();
+
 	/**
 	 * Serializes the name map.
 	 */
@@ -931,6 +941,13 @@ private:
 	 */
 	ELinkerStatus SerializeExportMap();
 
+#if WITH_TEXT_ARCHIVE_SUPPORT
+	/**
+	 * Create an import and export table when loading a text asset.
+	 */
+	ELinkerStatus ReconstructImportAndExportMap();
+#endif
+
 	ELinkerStatus SerializeDependsMap();
 
 	ELinkerStatus SerializePreloadDependencies();
@@ -938,6 +955,12 @@ private:
 	/** Sets the basic linker archive info */
 	void ResetStatusInfo();
 
+	/** For a given full object path, find or create the associated import or export table record. Used when loading text assets which only store object paths */
+	FPackageIndex FindOrCreateImportOrExport(const FString& InFullPath);
+
+	/** For the given object and class info, find or create an associated import record. Used when loading text assets which only store object paths */
+	FPackageIndex FindOrCreateImport(const FName InObjectName, const FName InClassName, const FName InClassPackageName);
+
 public:
 	/**
 	 * Serializes the gatherable text data container.
@@ -1181,12 +1204,22 @@ private:
 	/**
 	 * Finalizes linker creation, adding linker to loaders array and potentially verifying imports.
 	 */
-	ELinkerStatus FinalizeCreation();
+	ELinkerStatus FinalizeCreation(TMap<TPair<FName, FPackageIndex>, FPackageIndex>* ObjectNameWithOuterToExportMap);
 
 	//
 	// FLinkerLoad creation helpers END
 	//
 
+private:
+
+#if WITH_TEXT_ARCHIVE_SUPPORT
+	// Cache of the export names in a text asset. Allows us to enter those export slots by index rather than needing to reconstruct the name
+	TArray<FName> OriginalExportNames;
+
+	// Function to get a slot for a given export
+	FStructuredArchiveSlot GetExportSlot(FPackageIndex InExportIndex);
+#endif
+
 public:
 
 	//~ FArchive interface
diff --git a/Engine/Source/Runtime/CoreUObject/Public/UObject/LinkerSave.h b/Engine/Source/Runtime/CoreUObject/Public/UObject/LinkerSave.h
--- a/Engine/Source/Runtime/CoreUObject/Public/UObject/LinkerSave.h
+++ b/Engine/Source/Runtime/CoreUObject/Public/UObject/LinkerSave.h
@@ -10,6 +10,7 @@
 #include "Templates/RefCounting.h"
 
 struct FUntypedBulkData;
+class FPackageHeaderSaver;
 
 /*----------------------------------------------------------------------------
 	FLinkerSave.
@@ -42,8 +43,7 @@ public:
 	/** List of Searchable Names, by object containing them. This gets turned into package indices later */
 	TMap<const UObject *, TArray<FName> > SearchableNamesObjectMap;
 
-	/** Index array - location of the name in the NameMap array for each FName is stored in the NameIndices array using the FName's Index */
-	TMap<FNameEntryId, int32> NameIndices;
+	FPackageHeaderSaver& HeaderSaver;
 
 	/** Save context associated with this linker */
 	TRefCountPtr<FUObjectSerializeContext> SaveContext;
@@ -68,14 +68,11 @@ public:
 	COREUOBJECT_API static TMap<FString, TArray<uint8> > PackagesToScriptSHAMap;
 
 	/** Constructor for file writer */
-	FLinkerSave(UPackage* InParent, const TCHAR* InFilename, bool bForceByteSwapping, bool bInSaveUnversioned = false );
+	FLinkerSave(FPackageHeaderSaver& InHeaderSaver, UPackage* InParent, const TCHAR* InFilename, bool bForceByteSwapping, bool bInSaveUnversioned = false );
 	/** Constructor for memory writer */
-	FLinkerSave(UPackage* InParent, bool bForceByteSwapping, bool bInSaveUnversioned = false );
+	FLinkerSave(FPackageHeaderSaver& InHeaderSaver, UPackage* InParent, bool bForceByteSwapping, bool bInSaveUnversioned = false );
 	/** Constructor for custom savers. The linker assumes ownership of the custom saver. */
-	FLinkerSave(UPackage* InParent, FArchive *InSaver, bool bForceByteSwapping, bool bInSaveUnversioned = false);
-
-	/** Returns the appropriate name index for the source name, or 0 if not found in NameIndices */
-	int32 MapName( FNameEntryId Name) const;
+	FLinkerSave(FPackageHeaderSaver& InHeaderSaver, UPackage* InParent, FArchive *InSaver, bool bForceByteSwapping, bool bInSaveUnversioned = false);
 
 	/** Returns the appropriate package index for the source object, or default value if not found in ObjectIndicesMap */
 	FPackageIndex MapObject(const UObject* Object) const;
@@ -95,7 +92,6 @@ public:
 	virtual void PopDebugDataString() override { Saver->PopDebugDataString(); }
 #endif
 
-
 	virtual FString GetArchiveName() const override;
 
 	/**
diff --git a/Engine/Source/Runtime/CoreUObject/Public/UObject/PackageFileSummary.h b/Engine/Source/Runtime/CoreUObject/Public/UObject/PackageFileSummary.h
--- a/Engine/Source/Runtime/CoreUObject/Public/UObject/PackageFileSummary.h
+++ b/Engine/Source/Runtime/CoreUObject/Public/UObject/PackageFileSummary.h
@@ -231,10 +231,7 @@ public:
 		return FileVersionLicenseeUE4;
 	}
 
-	const FCustomVersionContainer& GetCustomVersionContainer() const
-	{
-		return CustomVersionContainer;
-	}
+	COREUOBJECT_API const FCustomVersionContainer& GetCustomVersionContainer() const;
 
 	void SetCustomVersionContainer(const FCustomVersionContainer& InContainer);
 
diff --git a/Engine/Source/Runtime/CoreUObject/Public/UObject/SavePackage.h b/Engine/Source/Runtime/CoreUObject/Public/UObject/SavePackage.h
new file mode 100644
--- /dev/null
+++ b/Engine/Source/Runtime/CoreUObject/Public/UObject/SavePackage.h
@@ -0,0 +1,148 @@
+// Copyright 1998-2019 Epic Games, Inc. All Rights Reserved.
+
+#pragma once
+
+#include "Containers/Array.h"
+#include "Containers/Set.h"
+#include "Containers/Map.h"
+#include "IO/IoDispatcher.h"
+#include "UObject/NameTypes.h"
+
+class FArchive;
+class FLinkerLoad;
+class FLinkerSave;
+
+class INameMapSaver
+{
+public:
+	virtual ~INameMapSaver() = default;
+
+	virtual void Begin() = 0;
+	virtual void End() = 0;
+	virtual void BeginPackage() = 0;
+	virtual void EndPackage(FLinkerSave& Linker, FLinkerLoad* Conform, FArchive* BinarySaver) = 0;
+
+	virtual void MarkNameAsReferenced(FName Name) = 0;
+	virtual int32 MapName(FName Name) const = 0;
+	virtual bool NameExistsInCurrentPackage(FNameEntryId ComparisonId) const = 0;
+};
+
+class FSinglePackageNameMapSaver : public INameMapSaver
+{
+public:
+	virtual void Begin() override {};
+	virtual void End() override {};
+	virtual void BeginPackage() override {};
+	virtual void EndPackage(FLinkerSave& Linker, FLinkerLoad* Conform, FArchive* BinarySaver) override;
+
+	virtual void MarkNameAsReferenced(FName Name) override;
+	virtual int32 MapName(FName Name) const override;
+	virtual bool NameExistsInCurrentPackage(FNameEntryId ComparisonId) const override;
+
+	void MarkNameAsReferenced(FNameEntryId Name);
+
+private:
+	TSet<FNameEntryId> ReferencedNames;
+	TMap<FNameEntryId, int32> NameIndices;
+};
+
+class FPackageStoreNameMapSaver : public INameMapSaver
+{
+public:
+	COREUOBJECT_API FPackageStoreNameMapSaver(const TCHAR* InFilename);
+
+	COREUOBJECT_API virtual void Begin() override {};
+	COREUOBJECT_API virtual void End() override;
+	virtual void BeginPackage();
+	virtual void EndPackage(FLinkerSave& Linker, FLinkerLoad* Conform, FArchive* BinarySaver) override {}
+
+	virtual void MarkNameAsReferenced(FName Name) override;
+	virtual int32 MapName(FName Name) const override;
+	virtual bool NameExistsInCurrentPackage(FNameEntryId ComparisonId) const override;
+
+private:
+	TMap<FNameEntryId, int32> NameIndices;
+	TArray<FNameEntryId> NameMap;
+	mutable TMap<FNameEntryId, TTuple<int32,int32,int32>> DebugNameCounts; // <Number0Count,OtherNumberCount,MaxNumber>
+	TSet<FNameEntryId> PackageReferencedNames;
+	const FString Filename;
+};
+
+class FPackageHeaderSaver
+{
+public:
+	FPackageHeaderSaver(INameMapSaver& InNameMapSaver) : NameMapSaver(InNameMapSaver) {}
+	INameMapSaver& NameMapSaver;
+};
+
+class FPackageStoreWriter
+{
+public:
+	COREUOBJECT_API			FPackageStoreWriter();
+	COREUOBJECT_API virtual ~FPackageStoreWriter();
+
+	struct HeaderInfo
+	{
+		FName	PackageName;
+		FString	LooseFilePath;
+	};
+
+	/** Write 'uasset' data
+	  */
+	virtual void WriteHeader(const HeaderInfo& Info, const FIoBuffer& HeaderData) = 0;
+
+	struct ExportsInfo
+	{
+		FName	PackageName;
+		FString	LooseFilePath;
+
+		TArray<FIoBuffer> Exports;
+	};
+
+	/** Write 'uexp' data
+	  */
+	virtual void WriteExports(const ExportsInfo& Info, const FIoBuffer& ExportsData) = 0;
+
+	struct FBulkDataInfo
+	{
+		enum EType 
+		{
+			Standard,
+			Mmap,
+			Optional
+		};
+
+		FName	PackageName;
+		EType	BulkdataType = Standard;
+		FString	LooseFilePath;
+	};
+
+	/** Write 'ubulk' data
+	  */
+	virtual void WriteBulkdata(const FBulkDataInfo& Info, const FIoBuffer& BulkData) = 0;
+};
+
+class FLooseFileWriter : public FPackageStoreWriter
+{
+public:
+	COREUOBJECT_API FLooseFileWriter();
+	COREUOBJECT_API ~FLooseFileWriter();
+
+	COREUOBJECT_API virtual void WriteHeader(const HeaderInfo& Info, const FIoBuffer& HeaderData) override;
+	COREUOBJECT_API virtual void WriteExports(const ExportsInfo& Info, const FIoBuffer& ExportsData) override;
+	COREUOBJECT_API virtual void WriteBulkdata(const FBulkDataInfo& Info, const FIoBuffer& BulkData) override;
+
+private:
+};
+
+class FSavePackageContext
+{
+public:
+	FSavePackageContext(FPackageHeaderSaver& InHeaderSaver, FPackageStoreWriter& InPackageStoreWriter) 
+	: HeaderSaver(InHeaderSaver), PackageStoreWriter(InPackageStoreWriter) 
+	{
+	}
+
+	FPackageHeaderSaver& HeaderSaver;
+	FPackageStoreWriter& PackageStoreWriter;
+};
diff --git a/Engine/Source/Runtime/CoreUObject/Public/UObject/UObjectGlobals.h b/Engine/Source/Runtime/CoreUObject/Public/UObject/UObjectGlobals.h
--- a/Engine/Source/Runtime/CoreUObject/Public/UObject/UObjectGlobals.h
+++ b/Engine/Source/Runtime/CoreUObject/Public/UObject/UObjectGlobals.h
@@ -36,6 +36,7 @@ DECLARE_DWORD_COUNTER_STAT_EXTERN(TEXT("FindObject"),STAT_FindObject,STATGROUP_O
 DECLARE_DWORD_COUNTER_STAT_EXTERN(TEXT("FindObjectFast"),STAT_FindObjectFast,STATGROUP_ObjectVerbose, );
 
 #define	INVALID_OBJECT	(UObject*)-1
+#define PERF_TRACK_DETAILED_ASYNC_STATS (0)
 
 
 // Private system wide variables.
diff --git a/Engine/Source/Runtime/CoreUObject/Public/UObject/UObjectThreadContext.h b/Engine/Source/Runtime/CoreUObject/Public/UObject/UObjectThreadContext.h
--- a/Engine/Source/Runtime/CoreUObject/Public/UObject/UObjectThreadContext.h
+++ b/Engine/Source/Runtime/CoreUObject/Public/UObject/UObjectThreadContext.h
@@ -77,7 +77,7 @@ public:
 	/* Object that is currently being constructed with ObjectInitializer */
 	UObject* ConstructedObject;
 	/** Async Package currently processing objects */
-	struct FAsyncPackage* AsyncPackage;
+	class FGCObject* AsyncPackage;
 
 #if !(UE_BUILD_SHIPPING || UE_BUILD_TEST)
 	/** Stack to ensure that PostInitProperties is routed through Super:: calls. **/
diff --git a/Engine/Source/Runtime/Engine/Private/Actor.cpp b/Engine/Source/Runtime/Engine/Private/Actor.cpp
--- a/Engine/Source/Runtime/Engine/Private/Actor.cpp
+++ b/Engine/Source/Runtime/Engine/Private/Actor.cpp
@@ -1,7 +1,6 @@
 // Copyright 1998-2019 Epic Games, Inc. All Rights Reserved.
 
 #include "GameFramework/Actor.h"
-#include "Serialization/AsyncLoading.h"
 #include "EngineDefines.h"
 #include "EngineStats.h"
 #include "EngineGlobals.h"
diff --git a/Engine/Source/Runtime/Engine/Private/Animation/AnimationStreaming.cpp b/Engine/Source/Runtime/Engine/Private/Animation/AnimationStreaming.cpp
--- a/Engine/Source/Runtime/Engine/Private/Animation/AnimationStreaming.cpp
+++ b/Engine/Source/Runtime/Engine/Private/Animation/AnimationStreaming.cpp
@@ -35,7 +35,6 @@ void FLoadedAnimationChunk::CleanUpIORequest()
 
 FStreamingAnimationData::FStreamingAnimationData()
 	: StreamableAnim(NULL)
-	, IORequestHandle(nullptr)
 	, AnimationStreamingManager(nullptr)
 {
 	ResetRequestedChunks();
@@ -43,7 +42,7 @@ FStreamingAnimationData::FStreamingAnimationData()
 
 FStreamingAnimationData::~FStreamingAnimationData()
 {
-	check(IORequestHandle == nullptr);
+
 }
 
 void FStreamingAnimationData::FreeResources()
@@ -63,18 +62,10 @@ void FStreamingAnimationData::FreeResources()
 	{
 		FreeLoadedChunk(LoadedChunk);
 	}
-
-	if (IORequestHandle)
-	{
-		delete IORequestHandle;
-		IORequestHandle = nullptr;
-	}
 }
 
 bool FStreamingAnimationData::Initialize(UAnimStreamable* InStreamableAnim, FAnimationStreamingManager* InAnimationStreamingManager)
 {
-	check(!IORequestHandle);
-
 	check(InStreamableAnim && InStreamableAnim->HasRunningPlatformData());
 
 	FStreamableAnimPlatformData& RunningAnimPlatformData = InStreamableAnim->GetRunningPlatformData();
@@ -219,16 +210,8 @@ void FStreamingAnimationData::BeginPendingRequests(const TArray<uint32>& Indices
 		if(!ExistingCompressedData)
 		{
 			UE_CLOG(ChunkStorage.CompressedAnimData != nullptr, LogAnimation, Fatal, TEXT("Existing compressed data for streaming animation chunk."));
-			UE_CLOG(Chunk.BulkData.GetFilename().Len()==0, LogAnimation, Fatal, TEXT("Streaming animation chunk has no file name."));
-			UE_CLOG(Chunk.BulkData.IsStoredCompressedOnDisk(), LogAnimation, Fatal, TEXT("Package level compression is not supported for streaming animation."));
 			UE_CLOG(ChunkStorage.IORequest, LogAnimation, Fatal, TEXT("Streaming animation chunk already has IORequest."));
 			
-			if (!IORequestHandle)
-			{
-				IORequestHandle = FPlatformFileManager::Get().GetPlatformFile().OpenAsyncRead(*Chunk.BulkData.GetFilename());
-				UE_CLOG(!IORequestHandle, LogAnimation, Fatal, TEXT("Failed to get IORequest Handle for streaming animation.")); // this generally cannot fail because it is async
-			}
-
 			int64 ChunkSize = Chunk.BulkData.GetBulkDataSize();
 			ChunkStorage.RequestStart = FPlatformTime::Seconds();
 			UE_LOG(LogAnimation, Warning, TEXT("Anim Streaming Request Started %s Chunk:%i At:%.3f\n"), *StreamableAnim->GetName(), ChunkIndex, ChunkStorage.RequestStart);
@@ -239,7 +222,7 @@ void FStreamingAnimationData::BeginPendingRequests(const TArray<uint32>& Indices
 			};
 
 			UE_LOG(LogAnimation, Warning, TEXT("Loading Stream Anim %s Chunk:%i Length: %.3f Offset:%i Size:%i File:%s\n"), *StreamableAnim->GetName(), ChunkIndex, Chunk.SequenceLength, Chunk.BulkData.GetBulkDataOffsetInFile(), Chunk.BulkData.GetBulkDataSize(), *Chunk.BulkData.GetFilename());
-			ChunkStorage.IORequest = IORequestHandle->ReadRequest(Chunk.BulkData.GetBulkDataOffsetInFile(), Chunk.BulkData.GetBulkDataSize(), AsyncIOPriority, &AsyncFileCallBack);
+			ChunkStorage.IORequest = Chunk.BulkData.CreateStreamingRequest(AsyncIOPriority, &AsyncFileCallBack, nullptr);
 			if (!ChunkStorage.IORequest)
 			{
 				UE_LOG(LogAnimation, Error, TEXT("Animation streaming read request failed."));
diff --git a/Engine/Source/Runtime/Engine/Private/Atmosphere/Atmosphere.h b/Engine/Source/Runtime/Engine/Private/Atmosphere/Atmosphere.h
--- a/Engine/Source/Runtime/Engine/Private/Atmosphere/Atmosphere.h
+++ b/Engine/Source/Runtime/Engine/Private/Atmosphere/Atmosphere.h
@@ -111,7 +111,7 @@ public:
 		}
 
 		int32 TotalSize = SizeX * SizeY * SizeZ * DataSize;
-		if (InTextureData.GetElementCount() == TotalSize)
+		if (InTextureData.GetBulkDataSize() == TotalSize)
 		{
 			// Grab a copy of the static volume data.
 			InTextureData.GetCopy(&TextureData, false);
diff --git a/Engine/Source/Runtime/Engine/Private/Atmosphere/AtmosphericFogComponent.cpp b/Engine/Source/Runtime/Engine/Private/Atmosphere/AtmosphericFogComponent.cpp
--- a/Engine/Source/Runtime/Engine/Private/Atmosphere/AtmosphericFogComponent.cpp
+++ b/Engine/Source/Runtime/Engine/Private/Atmosphere/AtmosphericFogComponent.cpp
@@ -195,7 +195,7 @@ static TAutoConsoleVariable<int32> CVarAtmosphereRender(
 // On CPU
 void UAtmosphericFogComponent::InitResource()
 {
-	bool bNewAtmosphere = CVarAtmosphereRender.GetValueOnGameThread() != 0;
+	bool bNewAtmosphere = CVarAtmosphereRender.GetValueOnAnyThread() != 0;
 
 	if(!bNewAtmosphere)
 	{
@@ -205,19 +205,19 @@ void UAtmosphericFogComponent::InitResource()
 	if (PrecomputeCounter >= EValid)
 	{
 		// A little inefficient for thread-safe
-		if (TransmittanceData.GetElementCount() && TransmittanceResource == NULL)
+		if (TransmittanceData.GetBulkDataSize() && TransmittanceResource == NULL)
 		{
 			TransmittanceResource = new FAtmosphereTextureResource( PrecomputeParams, TransmittanceData, FAtmosphereTextureResource::E_Transmittance ); 
 			BeginInitResource( TransmittanceResource );
 		}
 
-		if (IrradianceData.GetElementCount() && IrradianceResource == NULL)
+		if (IrradianceData.GetBulkDataSize() && IrradianceResource == NULL)
 		{
 			IrradianceResource = new FAtmosphereTextureResource( PrecomputeParams, IrradianceData, FAtmosphereTextureResource::E_Irradiance ); 
 			BeginInitResource( IrradianceResource );
 		}
 
-		if (InscatterData.GetElementCount() && InscatterResource == NULL)
+		if (InscatterData.GetBulkDataSize() > 0 && InscatterResource == NULL)
 		{
 			InscatterResource = new FAtmosphereTextureResource( PrecomputeParams, InscatterData, FAtmosphereTextureResource::E_Inscatter ); 
 			BeginInitResource( InscatterResource );
@@ -644,18 +644,18 @@ void UAtmosphericFogComponent::Serialize(FArchive& Ar)
 
 	if (Ar.UE4Ver() >= VER_UE4_ATMOSPHERIC_FOG_CACHE_DATA)
 	{
-		TransmittanceData.Serialize(Ar, this);
-		IrradianceData.Serialize(Ar, this);
+		TransmittanceData.Serialize(Ar, this, INDEX_NONE, false);
+		IrradianceData.Serialize(Ar, this, INDEX_NONE, false);
 	}
 
-	InscatterData.Serialize(Ar,this);
+	InscatterData.Serialize(Ar,this, INDEX_NONE, false);
 
 	if (Ar.IsLoading())
 	{
 		int32 CounterVal;
 		Ar << CounterVal;
 		// Precomputation was not successful, just ignore it
-		if (CounterVal < EValid || !TransmittanceData.GetElementCount())
+		if (CounterVal < EValid || TransmittanceData.GetBulkDataSize() == 0)
 		{
 			CounterVal = EInvalid;
 		}
@@ -684,7 +684,7 @@ TStructOnScope<FActorComponentInstanceData> UAtmosphericFogComponent::GetCompone
 {
 	TStructOnScope<FActorComponentInstanceData> InstanceData;
 
-	if (TransmittanceData.GetElementCount() && IrradianceData.GetElementCount() && InscatterData.GetElementCount() && PrecomputeCounter == EValid)
+	if (TransmittanceData.GetBulkDataSize() > 0 && IrradianceData.GetBulkDataSize() > 0 && InscatterData.GetBulkDataSize() > 0 && PrecomputeCounter == EValid)
 	{
 		// Allocate new struct for holding light map data
 		InstanceData.InitializeAs<FAtmospherePrecomputeInstanceData>(this);
diff --git a/Engine/Source/Runtime/Engine/Private/AudioDerivedData.cpp b/Engine/Source/Runtime/Engine/Private/AudioDerivedData.cpp
--- a/Engine/Source/Runtime/Engine/Private/AudioDerivedData.cpp
+++ b/Engine/Source/Runtime/Engine/Private/AudioDerivedData.cpp
@@ -831,7 +831,7 @@ int32 FStreamedAudioPlatformData::GetChunkFromDDC(int32 ChunkIndex, uint8** OutC
 		if (OutChunkData)
 		{
 			ChunkDataSize = Chunk.BulkData.GetBulkDataSize();
-			Chunk.BulkData.GetCopy((void**)OutChunkData);
+			Chunk.BulkData.GetCopy((void**)OutChunkData, true);
 		}
 	}
 #endif // #if WITH_EDITORONLY_DATA
diff --git a/Engine/Source/Runtime/Engine/Private/AudioStreaming.cpp b/Engine/Source/Runtime/Engine/Private/AudioStreaming.cpp
--- a/Engine/Source/Runtime/Engine/Private/AudioStreaming.cpp
+++ b/Engine/Source/Runtime/Engine/Private/AudioStreaming.cpp
@@ -99,14 +99,13 @@ void FAsyncStreamDerivedChunkWorker::DoWork()
 
 FStreamingWaveData::FStreamingWaveData()
 	: SoundWave(NULL)
-	, IORequestHandle(nullptr)
 	, AudioStreamingManager(nullptr)
 {
 }
 
 FStreamingWaveData::~FStreamingWaveData()
 {
-	check(IORequestHandle == nullptr);
+
 }
 
 void FStreamingWaveData::FreeResources()
@@ -126,17 +125,10 @@ void FStreamingWaveData::FreeResources()
 	{
 		FreeLoadedChunk(LoadedChunk);
 	}
-	if (IORequestHandle)
-	{
-		delete IORequestHandle;
-		IORequestHandle = nullptr;
-	}
 }
 
 bool FStreamingWaveData::Initialize(USoundWave* InSoundWave, FLegacyAudioStreamingManager* InAudioStreamingManager)
 {
-	check(!IORequestHandle);
-
 	if (!InSoundWave || !InSoundWave->RunningPlatformData->Chunks.Num())
 	{
 #if WITH_EDITOR
@@ -219,9 +211,10 @@ bool FStreamingWaveData::UpdateStreamingStatus()
 			// could maybe iterate over the things we know are done, but I couldn't tell if that was IndicesToLoad or not.
 			for (FLoadedAudioChunk& LoadedChunk : LoadedChunks)
 			{
-				if (LoadedChunk.IORequest && LoadedChunk.IORequest->PollCompletion())
+				if (LoadedChunk.IORequest != nullptr)
 				{
 					LoadedChunk.IORequest->WaitCompletion();
+
 					delete LoadedChunk.IORequest;
 					LoadedChunk.IORequest = nullptr;
 				}
@@ -337,7 +330,6 @@ void FStreamingWaveData::BeginPendingRequests(const TArray<uint32>& IndicesToLoa
 #if WITH_EDITORONLY_DATA
 			if (Chunk.DerivedDataKey.IsEmpty() == false)
 			{
-				ChunkStorage->MemorySize = ChunkSize;
 				ChunkStorage->Data = static_cast<uint8*>(FMemory::Malloc(ChunkSize));
 				INC_DWORD_STAT_BY(STAT_AudioMemorySize, ChunkSize);
 				INC_DWORD_STAT_BY(STAT_AudioMemory, ChunkSize);
@@ -359,16 +351,10 @@ void FStreamingWaveData::BeginPendingRequests(const TArray<uint32>& IndicesToLoa
 			else
 #endif // #if WITH_EDITORONLY_DATA
 			{
-				check(Chunk.BulkData.GetFilename().Len());
-				UE_CLOG(Chunk.BulkData.IsStoredCompressedOnDisk(), LogAudio, Fatal, TEXT("Package level compression is no longer supported."));
-				check(!ChunkStorage->IORequest);
-				if (!IORequestHandle)
-				{
-					IORequestHandle = FPlatformFileManager::Get().GetPlatformFile().OpenAsyncRead(*Chunk.BulkData.GetFilename());
-					check(IORequestHandle); // this generally cannot fail because it is async
-				}
-				check(Chunk.BulkData.GetBulkDataSize() == ChunkStorage->DataSize);
-
+				check(!ChunkStorage->IORequest); // Make sure we do not already have a request
+				check(!ChunkStorage->Data); // Make sure we do not already have data
+				check(Chunk.BulkData.GetBulkDataSize() == ChunkStorage->DataSize); // Make sure that the bulkdata size matches
+				
 				FAsyncFileCallBack AsyncFileCallBack =
 					[this, LoadedChunkStorageIndex](bool bWasCancelled, IAsyncReadRequest* Req)
 				{
@@ -377,8 +363,8 @@ void FStreamingWaveData::BeginPendingRequests(const TArray<uint32>& IndicesToLoa
 					PendingChunkChangeRequestStatus.Decrement();
 				};
 
-				check(!ChunkStorage->Data);
-				ChunkStorage->IORequest = IORequestHandle->ReadRequest(Chunk.BulkData.GetBulkDataOffsetInFile(), ChunkStorage->DataSize, AsyncIOPriority, &AsyncFileCallBack);
+				ChunkStorage->IORequest = Chunk.BulkData.CreateStreamingRequest(AsyncIOPriority, &AsyncFileCallBack, nullptr);
+
 				if (!ChunkStorage->IORequest)
 				{
 					UE_LOG(LogAudio, Error, TEXT("Audio streaming read request failed."));
@@ -410,6 +396,7 @@ bool FStreamingWaveData::BlockTillAllRequestsFinished(float TimeLimit)
 			if (LoadedChunk.IORequest)
 			{
 				LoadedChunk.IORequest->WaitCompletion();
+
 				delete LoadedChunk.IORequest;
 				LoadedChunk.IORequest = nullptr;
 			}
@@ -471,6 +458,7 @@ void FStreamingWaveData::FreeLoadedChunk(FLoadedAudioChunk& LoadedChunk)
 	{
 		LoadedChunk.IORequest->Cancel();
 		LoadedChunk.IORequest->WaitCompletion();
+
 		delete LoadedChunk.IORequest;
 		LoadedChunk.IORequest = nullptr;
 
@@ -488,7 +476,6 @@ void FStreamingWaveData::FreeLoadedChunk(FLoadedAudioChunk& LoadedChunk)
 	LoadedChunk.Data = NULL;
 	LoadedChunk.AudioDataSize = 0;
 	LoadedChunk.DataSize = 0;
-	LoadedChunk.MemorySize = 0;
 	LoadedChunk.Index = 0;
 }
 
diff --git a/Engine/Source/Runtime/Engine/Private/AudioStreamingCache.cpp b/Engine/Source/Runtime/Engine/Private/AudioStreamingCache.cpp
--- a/Engine/Source/Runtime/Engine/Private/AudioStreamingCache.cpp
+++ b/Engine/Source/Runtime/Engine/Private/AudioStreamingCache.cpp
@@ -889,18 +889,11 @@ void FAudioChunkCache::KickOffAsyncLoad(FCacheElement* CacheElement, const FChun
 	else
 #endif // #if WITH_EDITORONLY_DATA
 	{
-		check(Chunk.BulkData.GetFilename().Len());
-		UE_CLOG(Chunk.BulkData.IsStoredCompressedOnDisk(), LogAudio, Fatal, TEXT("Package level compression is no longer supported."));
-
 		if (CacheElement->IsLoadInProgress())
 		{
 			CacheElement->WaitForAsyncLoadCompletion(true);
 		}
 
-		// TODO: Figure out how to cache this file handle somewhere, so that we are not reopening a file on every individual load operation.
-		CacheElement->FileHandle.Reset(FPlatformFileManager::Get().GetPlatformFile().OpenAsyncRead(*Chunk.BulkData.GetFilename()));
-		check(CacheElement->FileHandle.IsValid());
-
 		// Sanity check our bulk data against
 		const int32 ChunkBulkDataSize = Chunk.BulkData.GetBulkDataSize();
 		check(ChunkDataSize <= ChunkBulkDataSize);
@@ -931,7 +924,7 @@ void FAudioChunkCache::KickOffAsyncLoad(FCacheElement* CacheElement, const FChun
 		CacheElement->DebugInfo.TimeLoadStarted = FPlatformTime::Seconds();
 #endif
 
-		CacheElement->ReadRequest.Reset(CacheElement->FileHandle->ReadRequest(Chunk.BulkData.GetBulkDataOffsetInFile(), ChunkDataSize, AsyncIOPriority, &AsyncFileCallBack, CacheElement->ChunkData.GetData()));
+		CacheElement->ReadRequest.Reset(Chunk.BulkData.CreateStreamingRequest( 0, ChunkDataSize, AsyncIOPriority, &AsyncFileCallBack, CacheElement->ChunkData.GetData()));
 		if (!CacheElement->ReadRequest.IsValid())
 		{
 			UE_LOG(LogAudio, Error, TEXT("Chunk load in audio LRU cache failed."));
diff --git a/Engine/Source/Runtime/Engine/Private/ImageUtils.cpp b/Engine/Source/Runtime/Engine/Private/ImageUtils.cpp
--- a/Engine/Source/Runtime/Engine/Private/ImageUtils.cpp
+++ b/Engine/Source/Runtime/Engine/Private/ImageUtils.cpp
@@ -329,7 +329,7 @@ UTexture2D* FImageUtils::CreateCheckerboardTexture(FColor ColorOne, FColor Color
 	UTexture2D* CheckerboardTexture = UTexture2D::CreateTransient(CheckerSize, CheckerSize, PF_B8G8R8A8);
 
 	// Lock the checkerboard texture so it can be modified
-	FColor* MipData = static_cast<FColor*>( CheckerboardTexture->PlatformData->Mips[0].BulkData.Lock(LOCK_READ_WRITE) );
+	FColor* MipData = reinterpret_cast<FColor*>( CheckerboardTexture->PlatformData->Mips[0].BulkData.Lock(LOCK_READ_WRITE) );
 
 	// Fill in the colors in a checkerboard pattern
 	for ( int32 RowNum = 0; RowNum < CheckerSize; ++RowNum )
@@ -445,7 +445,7 @@ public:
 		//Put first mip data into usable array
 		if (bReadSuccess)
 		{
-			uint32 TotalSize = Texture->PlatformData->Mips[0].BulkData.GetElementCount();
+			const uint32 TotalSize = Texture->PlatformData->Mips[0].BulkData.GetBulkDataSize();
 			RawData.AddZeroed(TotalSize);
 			FMemory::Memcpy(RawData.GetData(), RawData2[0], TotalSize);
 		}
diff --git a/Engine/Source/Runtime/Engine/Private/Level.cpp b/Engine/Source/Runtime/Engine/Private/Level.cpp
--- a/Engine/Source/Runtime/Engine/Private/Level.cpp
+++ b/Engine/Source/Runtime/Engine/Private/Level.cpp
@@ -7,10 +7,8 @@ Level.cpp: Level-related functions
 #include "Engine/Level.h"
 #include "Misc/ScopedSlowTask.h"
 #include "UObject/RenderingObjectVersion.h"
-#include "Templates/ScopedPointer.h"
 #include "Templates/UnrealTemplate.h"
 #include "UObject/Package.h"
-#include "Serialization/AsyncLoading.h"
 #include "EngineStats.h"
 #include "Engine/Blueprint.h"
 #include "GameFramework/Actor.h"
diff --git a/Engine/Source/Runtime/Engine/Private/LightMap.cpp b/Engine/Source/Runtime/Engine/Private/LightMap.cpp
--- a/Engine/Source/Runtime/Engine/Private/LightMap.cpp
+++ b/Engine/Source/Runtime/Engine/Private/LightMap.cpp
@@ -3105,12 +3105,18 @@ void FLegacyLightMap1D::Serialize(FArchive& Ar)
 	check(Ar.IsLoading());
 
 	UObject* Owner;
+
+#if !USE_NEW_BULKDATA
 	TQuantizedLightSampleBulkData<FQuantizedDirectionalLightSample> DirectionalSamples;
 	TQuantizedLightSampleBulkData<FQuantizedSimpleLightSample> SimpleSamples;
+#else
+	FUntypedBulkData2<FQuantizedDirectionalLightSample> DirectionalSamples;
+	FUntypedBulkData2<FQuantizedSimpleLightSample> SimpleSamples;
+#endif
 
 	Ar << Owner;
 
-	DirectionalSamples.Serialize( Ar, Owner );
+	DirectionalSamples.Serialize( Ar, Owner, INDEX_NONE, false );
 
 	for (int32 ElementIndex = 0; ElementIndex < 5; ElementIndex++)
 	{
@@ -3118,7 +3124,7 @@ void FLegacyLightMap1D::Serialize(FArchive& Ar)
 		Ar << Dummy;
 	}
 
-	SimpleSamples.Serialize( Ar, Owner );
+	SimpleSamples.Serialize( Ar, Owner, INDEX_NONE, false );
 }
 
 /*-----------------------------------------------------------------------------
diff --git a/Engine/Source/Runtime/Engine/Private/PhysicsEngine/BodySetup.cpp b/Engine/Source/Runtime/Engine/Private/PhysicsEngine/BodySetup.cpp
--- a/Engine/Source/Runtime/Engine/Private/PhysicsEngine/BodySetup.cpp
+++ b/Engine/Source/Runtime/Engine/Private/PhysicsEngine/BodySetup.cpp
@@ -1427,7 +1427,7 @@ void UBodySetup::GetResourceSizeEx(FResourceSizeEx& CumulativeResourceSize)
 	if (CookedFormatData.Contains(FPlatformProperties::GetPhysicsFormat()))
 	{
 		const FByteBulkData& FmtData = CookedFormatData.GetFormat(FPlatformProperties::GetPhysicsFormat());
-		CumulativeResourceSize.AddDedicatedSystemMemoryBytes(FmtData.GetElementSize() * FmtData.GetElementCount());
+		CumulativeResourceSize.AddDedicatedSystemMemoryBytes(FmtData.GetBulkDataSize());
 	}
 	
 	// Count any UV info
diff --git a/Engine/Source/Runtime/Engine/Private/PhysicsEngine/PhysXSupport.h b/Engine/Source/Runtime/Engine/Private/PhysicsEngine/PhysXSupport.h
--- a/Engine/Source/Runtime/Engine/Private/PhysicsEngine/PhysXSupport.h
+++ b/Engine/Source/Runtime/Engine/Private/PhysicsEngine/PhysXSupport.h
@@ -14,10 +14,10 @@
 #include "Physics/PhysicsFiltering.h"
 #include "PhysXPublic.h"
 #include "PhysXSupportCore.h"
+#include "Serialization/BulkData.h"
 
 class UBodySetup;
 class UPhysicalMaterial;
-struct FByteBulkData;
 struct FCollisionShape;
 class FPhysScene_PhysX;
 
diff --git a/Engine/Source/Runtime/Engine/Private/SkeletalMeshLODRenderData.cpp b/Engine/Source/Runtime/Engine/Private/SkeletalMeshLODRenderData.cpp
--- a/Engine/Source/Runtime/Engine/Private/SkeletalMeshLODRenderData.cpp
+++ b/Engine/Source/Runtime/Engine/Private/SkeletalMeshLODRenderData.cpp
@@ -1021,20 +1021,21 @@ void FSkeletalMeshLODRenderData::Serialize(FArchive& Ar, UObject* Owner, int32 I
 			else
 #endif
 			{
-				FByteBulkData TmpBulkData;
-				TmpBulkData.Serialize(Ar, Owner, Idx);
+#if USE_BULKDATA_STREAMING_TOKEN
+				FByteBulkData StreamingBulkData;
+				StreamingBulkData.Serialize(Ar, Owner, Idx, false);
 
-				bIsLODOptional = !!(TmpBulkData.GetBulkDataFlags() & BULKDATA_OptionalPayload);
-
-				int64 Tmp = TmpBulkData.GetBulkDataOffsetInFile();
-				check(Tmp >= 0 && Tmp <= 0xffffffffll);
-				OffsetInFile = static_cast<uint32>(Tmp);
-				Tmp = TmpBulkData.GetBulkDataSize();
-				check(Tmp >= 0 && Tmp <= 0xffffffffll);
-				BulkDataSize = static_cast<uint32>(Tmp);
+				BulkDataStreamingToken = StreamingBulkData.CreateStreamingToken();
+#else
+				StreamingBulkData.Serialize(Ar, Owner, Idx, false);
+#endif
+				bIsLODOptional = !!(StreamingBulkData.GetBulkDataFlags() & BULKDATA_OptionalPayload);
 
-				bDiscardBulkData = !BulkDataSize;
-				BuffersSize = bDiscardBulkData ? 0 : BuffersSize;
+				if (StreamingBulkData.GetBulkDataSize() == 0)
+				{
+					bDiscardBulkData = true;
+					BuffersSize = 0;
+				}
 			}
 
 			if (!bDiscardBulkData)
diff --git a/Engine/Source/Runtime/Engine/Private/SoundWave.cpp b/Engine/Source/Runtime/Engine/Private/SoundWave.cpp
--- a/Engine/Source/Runtime/Engine/Private/SoundWave.cpp
+++ b/Engine/Source/Runtime/Engine/Private/SoundWave.cpp
@@ -126,13 +126,16 @@ void FStreamedAudioChunk::Serialize(FArchive& Ar, UObject* Owner, int32 ChunkInd
 	Ar << bCooked;
 
 	// ChunkIndex 0 is always inline payload, all other chunks are streamed.
-	if (ChunkIndex == 0)
+	if (Ar.IsSaving())
 	{
-		BulkData.SetBulkDataFlags(BULKDATA_ForceInlinePayload);
-	}
-	else
-	{
-		BulkData.SetBulkDataFlags(BULKDATA_Force_NOT_InlinePayload);
+		if (ChunkIndex == 0)
+		{
+			BulkData.SetBulkDataFlags(BULKDATA_ForceInlinePayload);
+		}
+		else
+		{
+			BulkData.SetBulkDataFlags(BULKDATA_Force_NOT_InlinePayload);
+		}
 	}
 
 	// streaming doesn't use memory mapped IO
@@ -404,7 +407,7 @@ void USoundWave::Serialize( FArchive& Ar )
 	else
 	{
 		// only save the raw data for non-cooked packages
-		RawData.Serialize( Ar, this );
+		RawData.Serialize(Ar, this, INDEX_NONE, false);
 	}
 
 	Ar << CompressedDataGuid;
diff --git a/Engine/Source/Runtime/Engine/Private/StaticMesh.cpp b/Engine/Source/Runtime/Engine/Private/StaticMesh.cpp
--- a/Engine/Source/Runtime/Engine/Private/StaticMesh.cpp
+++ b/Engine/Source/Runtime/Engine/Private/StaticMesh.cpp
@@ -634,6 +634,7 @@ void FStaticMeshLODResources::Serialize(FArchive& Ar, UObject* Owner, int32 Inde
 		}
 		else if (FPlatformProperties::RequiresCookedData() || Ar.IsCooking() || bUsingCookedData)
 		{
+			uint32 BulkDataSize = 0;
 #if WITH_EDITOR
 			if (Ar.IsSaving())
 			{
@@ -669,14 +670,18 @@ void FStaticMeshLODResources::Serialize(FArchive& Ar, UObject* Owner, int32 Inde
 			else
 #endif
 			{
+#if USE_BULKDATA_STREAMING_TOKEN
 				FByteBulkData TmpBulkData;
-				TmpBulkData.Serialize(Ar, Owner, Index);
+				TmpBulkData.Serialize(Ar, Owner, Index, false);
 				bIsOptionalLOD = !!(TmpBulkData.GetBulkDataFlags() & BULKDATA_OptionalPayload);
-				int64 Tmp = TmpBulkData.GetBulkDataOffsetInFile();
-				check(Tmp >= 0 && Tmp <= 0xffffffff);
-				OffsetInFile = static_cast<uint32>(Tmp);
-				check(TmpBulkData.GetBulkDataSize() >= 0);
-				BulkDataSize = static_cast<uint32>(TmpBulkData.GetBulkDataSize());
+
+				BulkDataStreamingToken = TmpBulkData.CreateStreamingToken();			
+				BulkDataSize = BulkDataStreamingToken.GetSize();
+#else
+				StreamingBulkData.Serialize(Ar, Owner, Index, false);
+				bIsOptionalLOD = !!(StreamingBulkData.GetBulkDataFlags() & BULKDATA_OptionalPayload);
+				BulkDataSize = (uint32)StreamingBulkData.GetBulkDataSize();
+#endif
 
 				// Streaming CPU data in editor build isn't supported yet because tools and utils need access
 				if (bUsingCookedData && BulkDataSize > 0)
@@ -693,7 +698,7 @@ void FStaticMeshLODResources::Serialize(FArchive& Ar, UObject* Owner, int32 Inde
 			Ar << TmpBuffersSize;
 			BuffersSize = TmpBuffersSize.CalcBuffersSize();
 
-			if (Ar.IsLoading() && bIsOptionalLOD && !BulkDataSize)
+			if (Ar.IsLoading() && bIsOptionalLOD)
 			{
 				ClearAvailabilityInfo();
 			}
@@ -2310,7 +2315,7 @@ void FStaticMeshRenderData::Cache(UStaticMesh* Owner, const FStaticMeshLODSettin
 
 
 	{
-		TRACE_CPUPROFILER_EVENT_SCOPE_TEXT(TEXT("StaticMesh_Cache"));
+		TRACE_CPUPROFILER_EVENT_SCOPE(StaticMesh_Cache);
 
 		COOK_STAT(auto Timer = StaticMeshCookStats::UsageStats.TimeSyncWork());
 		int32 T0 = FPlatformTime::Cycles();
diff --git a/Engine/Source/Runtime/Engine/Private/Streaming/SkeletalMeshUpdate.cpp b/Engine/Source/Runtime/Engine/Private/Streaming/SkeletalMeshUpdate.cpp
--- a/Engine/Source/Runtime/Engine/Private/Streaming/SkeletalMeshUpdate.cpp
+++ b/Engine/Source/Runtime/Engine/Private/Streaming/SkeletalMeshUpdate.cpp
@@ -328,7 +328,6 @@ void FSkeletalMeshStreamIn_IO::FCancelIORequestsTask::DoWork()
 
 FSkeletalMeshStreamIn_IO::FSkeletalMeshStreamIn_IO(USkeletalMesh* InMesh, int32 InRequestedMips, bool bHighPrio)
 	: FSkeletalMeshStreamIn(InMesh, InRequestedMips)
-	, IOFileHandle(nullptr)
 	, IORequest(nullptr)
 	, bHighPrioIORequest(bHighPrio)
 {}
@@ -339,8 +338,7 @@ void FSkeletalMeshStreamIn_IO::Abort()
 	{
 		FSkeletalMeshStreamIn::Abort();
 
-		// IO requests can only exist in the lifetime of IOFileHandle.
-		if (IOFileHandle)
+		if (IORequest != nullptr)
 		{
 			// Prevent the update from being considered done before this is finished.
 			// By checking that it was not already cancelled, we make sure this doesn't get called twice.
@@ -394,50 +392,40 @@ void FSkeletalMeshStreamIn_IO::SetIORequest(const FContext& Context, const FStri
 	{
 		return;
 	}
-	check(!IOFileHandle && PendingFirstMip < CurrentFirstLODIdx);
-	IOFileHandle = FPlatformFileManager::Get().GetPlatformFile().OpenAsyncRead(*IOFilename);
+	check(!IORequest && PendingFirstMip < CurrentFirstLODIdx);
+
 	FSkeletalMeshRenderData* RenderData = Context.RenderData;
 	FString debugName;
 	if (Context.Mesh)
 	{
 		debugName = Context.Mesh->GetName();
 	}
-	if (IOFileHandle && RenderData)
+
+	if (RenderData != nullptr)
 	{
 		SetAsyncFileCallback(Context);
 
-		int64 ReadOffset;
-		int64 ReadSize;
-		{
-			FSkeletalMeshLODRenderData& FirstLOD = RenderData->LODRenderData[PendingFirstMip];
-			FSkeletalMeshLODRenderData& LastLOD = RenderData->LODRenderData[CurrentFirstLODIdx - 1];
-			ReadOffset = FirstLOD.OffsetInFile;
-			ReadSize = LastLOD.OffsetInFile + LastLOD.BulkDataSize - ReadOffset;
-#if DO_CHECK
-			int64 TestOffset = ReadOffset;
-			for (int32 LODIdx = PendingFirstMip; LODIdx < CurrentFirstLODIdx; ++LODIdx)
-			{
-				FSkeletalMeshLODRenderData& LODResource = RenderData->LODRenderData[LODIdx];
-				const int64 CurOffset = (int64)LODResource.OffsetInFile;
-				check(ReadOffset <= CurOffset);
-				check(TestOffset == CurOffset);
-				TestOffset += LODResource.BulkDataSize;
-			}
-			check(TestOffset == ReadOffset + ReadSize);
-#endif
-		}
+		const FSkeletalMeshLODRenderData& FirstLOD = RenderData->LODRenderData[PendingFirstMip];
+		const FSkeletalMeshLODRenderData& LastLOD = RenderData->LODRenderData[CurrentFirstLODIdx - 1];
 
 		// Increment as we push the request. If a request complete immediately, then it will call the callback
 		// but that won't do anything because the tick would not try to acquire the lock since it is already locked.
 		TaskSynchronization.Increment();
-		LODData.Empty(ReadSize);
-		LODData.AddUninitialized(ReadSize);
-		IORequest = IOFileHandle->ReadRequest(
-			ReadOffset,
-			ReadSize,
+
+#if USE_BULKDATA_STREAMING_TOKEN
+		IORequest = FUntypedBulkData::CreateStreamingRequestForRange(
+			IOFilename,
+			FirstLOD.BulkDataStreamingToken,
+			LastLOD.BulkDataStreamingToken,
+			bHighPrioIORequest ? AIOP_BelowNormal : AIOP_Low,
+			&AsyncFileCallback);
+#else
+		IORequest = BulkDataUtils::CreateStreamingRequestForRange(
+			FirstLOD.StreamingBulkData,
+			LastLOD.StreamingBulkData,
 			bHighPrioIORequest ? AIOP_BelowNormal : AIOP_Low,
-			&AsyncFileCallback,
-			LODData.GetData());
+			&AsyncFileCallback);
+#endif
 	}
 	else
 	{
@@ -447,22 +435,16 @@ void FSkeletalMeshStreamIn_IO::SetIORequest(const FContext& Context, const FStri
 
 void FSkeletalMeshStreamIn_IO::ClearIORequest(const FContext& Context)
 {
-	if (IOFileHandle)
+	if (IORequest != nullptr)
 	{
-		if (IORequest)
+		// If clearing requests not yet completed, cancel and wait.
+		if (!IORequest->PollCompletion())
 		{
-			// If clearing requests not yet completed, cancel and wait.
-			if (!IORequest->PollCompletion())
-			{
-				IORequest->Cancel();
-				IORequest->WaitCompletion();
-			}
-			delete IORequest;
-			IORequest = nullptr;
+			IORequest->Cancel();
+			IORequest->WaitCompletion();
 		}
-
-		delete IOFileHandle;
-		IOFileHandle = nullptr;
+		delete IORequest;
+		IORequest = nullptr;
 	}
 }
 
@@ -476,7 +458,10 @@ void FSkeletalMeshStreamIn_IO::SerializeLODData(const FContext& Context)
 	if (!IsCancelled() && Mesh && RenderData)
 	{
 		check(PendingFirstMip < CurrentFirstLODIdx && CurrentFirstLODIdx == RenderData->CurrentFirstLODIdx);
-		FMemoryReader Ar(LODData, true);
+		check(IORequest->GetSize() >= 0 && IORequest->GetSize() <= TNumericLimits<uint32>::Max());
+
+		TArrayView<uint8> Data(IORequest->GetReadResults(), IORequest->GetSize());
+		FMemoryReaderView Ar(Data, true);
 		for (int32 LODIdx = PendingFirstMip; LODIdx < CurrentFirstLODIdx; ++LODIdx)
 		{
 			FSkeletalMeshLODRenderData& LODResource = RenderData->LODRenderData[LODIdx];
@@ -485,7 +470,8 @@ void FSkeletalMeshStreamIn_IO::SerializeLODData(const FContext& Context)
 			constexpr uint8 DummyStripFlags = 0;
 			LODResource.SerializeStreamedData(Ar, Mesh, LODIdx, DummyStripFlags, bNeedsCPUAccess, bForceKeepCPUResources);
 		}
-		LODData.Empty();
+
+		FMemory::Free(Data.GetData()); // Free the memory we took ownership of via IORequest->GetReadResults()
 	}
 }
 
diff --git a/Engine/Source/Runtime/Engine/Private/Streaming/SkeletalMeshUpdate.h b/Engine/Source/Runtime/Engine/Private/Streaming/SkeletalMeshUpdate.h
--- a/Engine/Source/Runtime/Engine/Private/Streaming/SkeletalMeshUpdate.h
+++ b/Engine/Source/Runtime/Engine/Private/Streaming/SkeletalMeshUpdate.h
@@ -192,9 +192,7 @@ protected:
 	/** Called by FAsyncCancelIORequestsTask to cancel inflight IO request if any */
 	void CancelIORequest();
 
-	TArray<uint8> LODData;
-	IAsyncReadFileHandle* IOFileHandle;
-	IAsyncReadRequest* IORequest;
+	struct FBulkDataIORequest* IORequest;
 	FAsyncFileCallBack AsyncFileCallback;
 	bool bHighPrioIORequest;
 };
diff --git a/Engine/Source/Runtime/Engine/Private/Streaming/StaticMeshUpdate.cpp b/Engine/Source/Runtime/Engine/Private/Streaming/StaticMeshUpdate.cpp
--- a/Engine/Source/Runtime/Engine/Private/Streaming/StaticMeshUpdate.cpp
+++ b/Engine/Source/Runtime/Engine/Private/Streaming/StaticMeshUpdate.cpp
@@ -307,7 +307,6 @@ void FStaticMeshStreamIn_IO::FCancelIORequestsTask::DoWork()
 
 FStaticMeshStreamIn_IO::FStaticMeshStreamIn_IO(UStaticMesh* InMesh, int32 InRequestedMips, bool bHighPrio)
 	: FStaticMeshStreamIn(InMesh, InRequestedMips)
-	, IOFileHandle(nullptr)
 	, IORequest(nullptr)
 	, bHighPrioIORequest(bHighPrio)
 {}
@@ -318,8 +317,7 @@ void FStaticMeshStreamIn_IO::Abort()
 	{
 		FStaticMeshStreamIn::Abort();
 
-		// IO requests can only exist in the lifetime of IOFileHandle.
-		if (IOFileHandle)
+		if (IORequest != nullptr)
 		{
 			// Prevent the update from being considered done before this is finished.
 			// By checking that it was not already cancelled, we make sure this doesn't get called twice.
@@ -373,71 +371,54 @@ void FStaticMeshStreamIn_IO::SetIORequest(const FContext& Context, const FString
 	{
 		return;
 	}
-	check(!IOFileHandle && PendingFirstMip < CurrentFirstLODIdx);
-	IOFileHandle = FPlatformFileManager::Get().GetPlatformFile().OpenAsyncRead(*IOFilename);
+
+	check(!IORequest && PendingFirstMip < CurrentFirstLODIdx);
+
 	FStaticMeshRenderData* RenderData = Context.RenderData;
-	if (IOFileHandle && RenderData)
+	if (RenderData != nullptr)
 	{
 		SetAsyncFileCallback(Context);
 
-		int64 ReadOffset;
-		int64 ReadSize;
-		{
-			FStaticMeshLODResources& FirstLOD = RenderData->LODResources[PendingFirstMip];
-			FStaticMeshLODResources& LastLOD = RenderData->LODResources[CurrentFirstLODIdx - 1];
-			ReadOffset = FirstLOD.OffsetInFile;
-			ReadSize = LastLOD.OffsetInFile + LastLOD.BulkDataSize - ReadOffset;
-#if DO_CHECK
-			int64 TestOffset = ReadOffset;
-			for (int32 LODIdx = PendingFirstMip; LODIdx < CurrentFirstLODIdx; ++LODIdx)
-			{
-				FStaticMeshLODResources& LODResource = RenderData->LODResources[LODIdx];
-				const int64 CurOffset = (int64)LODResource.OffsetInFile;
-				check(ReadOffset <= CurOffset);
-				check(TestOffset == CurOffset);
-				TestOffset += LODResource.BulkDataSize;
-			}
-			check(TestOffset == ReadOffset + ReadSize);
-#endif
-		}
+		const FStaticMeshLODResources& FirstLOD = RenderData->LODResources[PendingFirstMip];
+		const FStaticMeshLODResources& LastLOD = RenderData->LODResources[CurrentFirstLODIdx - 1];
 
 		// Increment as we push the request. If a request complete immediately, then it will call the callback
 		// but that won't do anything because the tick would not try to acquire the lock since it is already locked.
 		TaskSynchronization.Increment();
-		LODData.Empty(ReadSize);
-		LODData.AddUninitialized(ReadSize);
 
-		IORequest = IOFileHandle->ReadRequest(
-			ReadOffset,
-			ReadSize,
+#if USE_BULKDATA_STREAMING_TOKEN
+		IORequest = FUntypedBulkData::CreateStreamingRequestForRange(
+			IOFilename,
+			FirstLOD.BulkDataStreamingToken,
+			LastLOD.BulkDataStreamingToken,
 			bHighPrioIORequest ? AIOP_BelowNormal : AIOP_Low,
-			&AsyncFileCallback,
-			LODData.GetData());
+			&AsyncFileCallback);
+#else
+		IORequest = BulkDataUtils::CreateStreamingRequestForRange(
+			FirstLOD.StreamingBulkData,
+			LastLOD.StreamingBulkData,
+			bHighPrioIORequest ? AIOP_BelowNormal : AIOP_Low,
+			&AsyncFileCallback);
+#endif
 	}
 	else
 	{
 		MarkAsCancelled();
 	}
-}
+} 
 
 void FStaticMeshStreamIn_IO::ClearIORequest(const FContext& Context)
 {
-	if (IOFileHandle)
+	if (IORequest != nullptr)
 	{
-		if (IORequest)
+		// If clearing requests not yet completed, cancel and wait.
+		if (!IORequest->PollCompletion())
 		{
-			// If clearing requests not yet completed, cancel and wait.
-			if (!IORequest->PollCompletion())
-			{
-				IORequest->Cancel();
-				IORequest->WaitCompletion();
-			}
-			delete IORequest;
-			IORequest = nullptr;
+			IORequest->Cancel();
+			IORequest->WaitCompletion();
 		}
-
-		delete IOFileHandle;
-		IOFileHandle = nullptr;
+		delete IORequest;
+		IORequest = nullptr;
 	}
 }
 
@@ -451,7 +432,11 @@ void FStaticMeshStreamIn_IO::SerializeLODData(const FContext& Context)
 	if (!IsCancelled() && Mesh && RenderData)
 	{
 		check(PendingFirstMip < CurrentFirstLODIdx && CurrentFirstLODIdx == RenderData->CurrentFirstLODIdx);
-		FMemoryReader Ar(LODData, true);
+		check(IORequest->GetSize() >= 0 && IORequest->GetSize() <= TNumericLimits<uint32>::Max());
+
+		TArrayView<uint8> Data(IORequest->GetReadResults(), IORequest->GetSize());
+
+		FMemoryReaderView Ar(Data, true);
 		for (int32 LODIdx = PendingFirstMip; LODIdx < CurrentFirstLODIdx; ++LODIdx)
 		{
 			FStaticMeshLODResources& LODResource = RenderData->LODResources[LODIdx];
@@ -460,7 +445,8 @@ void FStaticMeshStreamIn_IO::SerializeLODData(const FContext& Context)
 			LODResource.SerializeBuffers(Ar, Mesh, DummyStripFlags, DummyBuffersSize);
 			check(DummyBuffersSize.CalcBuffersSize() == LODResource.BuffersSize);
 		}
-		LODData.Empty();
+		
+		FMemory::Free(Data.GetData());// Free the memory we took ownership of via IORequest->GetReadResults()
 	}
 }
 
diff --git a/Engine/Source/Runtime/Engine/Private/Streaming/StaticMeshUpdate.h b/Engine/Source/Runtime/Engine/Private/Streaming/StaticMeshUpdate.h
--- a/Engine/Source/Runtime/Engine/Private/Streaming/StaticMeshUpdate.h
+++ b/Engine/Source/Runtime/Engine/Private/Streaming/StaticMeshUpdate.h
@@ -183,9 +183,7 @@ protected:
 	/** Called by FAsyncCancelIORequestsTask to cancel inflight IO request if any */
 	void CancelIORequest();
 
-	TArray<uint8> LODData;
-	IAsyncReadFileHandle* IOFileHandle;
-	IAsyncReadRequest* IORequest;
+	struct FBulkDataIORequest* IORequest;
 	FAsyncFileCallBack AsyncFileCallback;
 	bool bHighPrioIORequest;
 };
diff --git a/Engine/Source/Runtime/Engine/Private/Streaming/Texture2DStreamIn_IO.h b/Engine/Source/Runtime/Engine/Private/Streaming/Texture2DStreamIn_IO.h
--- a/Engine/Source/Runtime/Engine/Private/Streaming/Texture2DStreamIn_IO.h
+++ b/Engine/Source/Runtime/Engine/Private/Streaming/Texture2DStreamIn_IO.h
@@ -9,6 +9,8 @@ Texture2DStreamIn.h: Stream in helper for 2D textures using texture streaming fi
 #include "CoreMinimal.h"
 #include "Texture2DStreamIn.h"
 
+struct FBulkDataIORequest;
+
 // Base StreamIn framework exposing MipData
 class FTexture2DStreamIn_IO : public FTexture2DStreamIn
 {
@@ -40,6 +42,9 @@ protected:
 
 private:
 
+	// Poll if any of the mips currently have an active IO request
+	bool HasPendingIORequests();
+
 	class FCancelIORequestsTask : public FNonAbandonableTask
 	{
 	public:
@@ -58,15 +63,16 @@ private:
 
 
 	// Request for loading into each mip.
-	IAsyncReadRequest* IORequests[MAX_TEXTURE_MIP_COUNT];
+	FBulkDataIORequest* IORequests[MAX_TEXTURE_MIP_COUNT];
 
 	bool bPrioritizedIORequest;
 
 	// Async handle.
+#if TEXTURE2DMIPMAP_USE_COMPACT_BULKDATA
 	FString IOFilename;
-	int64 IOFileOffset;
+#endif
 
-	IAsyncReadFileHandle* IOFileHandle;
+	int64 IOFileOffset;
 	FAsyncFileCallBack AsyncFileCallBack;
 };
 
diff --git a/Engine/Source/Runtime/Engine/Private/Texture2D.cpp b/Engine/Source/Runtime/Engine/Private/Texture2D.cpp
--- a/Engine/Source/Runtime/Engine/Private/Texture2D.cpp
+++ b/Engine/Source/Runtime/Engine/Private/Texture2D.cpp
@@ -183,7 +183,7 @@ private:
 /*-----------------------------------------------------------------------------
 	FTexture2DMipMap
 -----------------------------------------------------------------------------*/
-
+#if TEXTURE2DMIPMAP_USE_COMPACT_BULKDATA
 FTexture2DMipMap::FCompactByteBulkData::FCompactByteBulkData()
 {
 	Reset();
@@ -239,14 +239,14 @@ typename FTexture2DMipMap::FCompactByteBulkData& FTexture2DMipMap::FCompactByteB
 	return *this;
 }
 
-void FTexture2DMipMap::FCompactByteBulkData::Serialize(FArchive& Ar, UObject* Owner, int32 MipIdx)
+void FTexture2DMipMap::FCompactByteBulkData::Serialize(FArchive& Ar, UObject* Owner, int32 MipIdx, bool /*bAttemptFileMapping*/)
 {
 	check(Ar.IsLoading());
 	FMemory::Free(TexelData);
 	TexelData = nullptr;
 
 	FByteBulkData TmpBulkData;
-	TmpBulkData.Serialize(Ar, Owner, MipIdx);
+	TmpBulkData.Serialize(Ar, Owner, MipIdx, false);
 
 	const int64 Tmp = TmpBulkData.GetBulkDataOffsetInFile();
 	check(Tmp >= 0 && Tmp <= 0xffffffffll);
@@ -256,7 +256,7 @@ void FTexture2DMipMap::FCompactByteBulkData::Serialize(FArchive& Ar, UObject* Ow
 
 	if (IsInlined())
 	{
-		TmpBulkData.GetCopy((void**)&TexelData);
+		TmpBulkData.GetCopy((void**)&TexelData, true);
 	}
 
 	if (!MipIdx && !IsInlined())
@@ -348,12 +348,43 @@ void FTexture2DMipMap::FCompactByteBulkData::GetCopy(void** Dest, bool bDiscardI
 	}
 }
 
+FBulkDataIORequest* FTexture2DMipMap::FCompactByteBulkData::CreateStreamingRequest(const FString& Filename, int64 OffsetInBulkData, int64 BytesToRead, EAsyncIOPriorityAndFlags Priority, FAsyncFileCallBack* CompleteCallback, uint8* UserSuppliedMemory) const
+{
+	check(Filename.IsEmpty() == false);
+
+	UE_CLOG(IsStoredCompressedOnDisk(), LogSerialization, Fatal, TEXT("Package level compression is no longer supported (%s)."), *Filename);
+	UE_CLOG(GetBulkDataSize() <= 0, LogSerialization, Error, TEXT("(%s) has invalid bulk data size."), *Filename);
+
+	IAsyncReadFileHandle* IORequestHandle = FPlatformFileManager::Get().GetPlatformFile().OpenAsyncRead(*Filename);
+	check(IORequestHandle); // this generally cannot fail because it is async
+
+	if (IORequestHandle == nullptr)
+	{
+		return nullptr;
+	}
+
+	const int64 FinalOffsetInFile = GetBulkDataOffsetInFile() + OffsetInBulkData;
+
+	IAsyncReadRequest* ReadRequest = IORequestHandle->ReadRequest(FinalOffsetInFile, BytesToRead, Priority, CompleteCallback, UserSuppliedMemory);
+	if (ReadRequest != nullptr)
+	{
+		return new FBulkDataIORequest(IORequestHandle, ReadRequest, BytesToRead);
+	}
+	else
+	{
+		delete IORequestHandle;
+		return nullptr;
+	}
+}
+
+#endif // #if TEXTURE2DMIPMAP_USE_COMPACT_BULKDATA
+
 void FTexture2DMipMap::Serialize(FArchive& Ar, UObject* Owner, int32 MipIdx)
 {
 	bool bCooked = Ar.IsCooking();
 	Ar << bCooked;
 
-	BulkData.Serialize(Ar, Owner, MipIdx);
+	BulkData.Serialize(Ar, Owner, MipIdx, false);
 	Ar << SizeX;
 	Ar << SizeY;
 	Ar << SizeZ;
@@ -453,7 +484,7 @@ int32 UTexture2D::CalcNumOptionalMips() const
 		int32 NumOptionalMips = 0;
 		for (int32 MipIndex = 0; MipIndex < PlatformData->Mips.Num(); ++MipIndex)
 		{
-			if (PlatformData->Mips[MipIndex].BulkData.GetBulkDataFlags() & BULKDATA_OptionalPayload)
+			if (PlatformData->Mips[MipIndex].BulkData.IsOptional())
 			{
 				++NumOptionalMips;
 			}
@@ -478,7 +509,7 @@ bool UTexture2D::GetMipDataFilename(const int32 MipIndex, FString& OutBulkDataFi
 			OutBulkDataFilename = PlatformData->Mips[MipIndex].BulkData.GetFilename();
 #else
 			OutBulkDataFilename = PlatformData->CachedPackageFileName;
-			bool UseOptionalBulkDataFileName = PlatformData->Mips[MipIndex].BulkData.GetBulkDataFlags() & BULKDATA_OptionalPayload;
+			const bool UseOptionalBulkDataFileName = PlatformData->Mips[MipIndex].BulkData.IsOptional();
 			OutBulkDataFilename = FPaths::ChangeExtension(OutBulkDataFilename, UseOptionalBulkDataFileName ? TEXT(".uptnl") : TEXT(".ubulk"));
 #endif
 			return true;
@@ -2328,7 +2359,7 @@ FIncomingTextureArrayDataEntry::FIncomingTextureArrayDataEntry(UTexture2D* InTex
 			MipData[MipIndex].SizeX = Mip.SizeX;
 			MipData[MipIndex].SizeY = Mip.SizeY;
 			
-			const int32 MipDataSize = Mip.BulkData.GetElementCount() * Mip.BulkData.GetElementSize();
+			const int32 MipDataSize = Mip.BulkData.GetBulkDataSize();
 			MipData[MipIndex].Data.Empty(MipDataSize);
 			MipData[MipIndex].Data.AddUninitialized(MipDataSize);
 			// Get copy of data, potentially loading array or using already loaded version.
diff --git a/Engine/Source/Runtime/Engine/Private/VT/UploadingVirtualTexture.cpp b/Engine/Source/Runtime/Engine/Private/VT/UploadingVirtualTexture.cpp
--- a/Engine/Source/Runtime/Engine/Private/VT/UploadingVirtualTexture.cpp
+++ b/Engine/Source/Runtime/Engine/Private/VT/UploadingVirtualTexture.cpp
@@ -306,7 +306,7 @@ FVTDataAndStatus FUploadingVirtualTexture::ReadData(FGraphEventArray& OutComplet
 {
 	FVirtualTextureDataChunk& Chunk = Data->Chunks[ChunkIndex];
 	FByteBulkData& BulkData = Chunk.BulkData;
-	const TCHAR* ChunkFileName = nullptr;
+	FString ChunkFileName;
 	int64 ChunkOffsetInFile = 0;
 
 #if WITH_EDITOR
@@ -316,7 +316,7 @@ FVTDataAndStatus FUploadingVirtualTexture::ReadData(FGraphEventArray& OutComplet
 	if (BulkData.GetFilename().IsEmpty() == false)
 	{
 		ensure(Size <= (size_t)BulkData.GetBulkDataSize());
-		ChunkFileName = *BulkData.GetFilename();
+		ChunkFileName = BulkData.GetFilename();
 		ChunkOffsetInFile = BulkData.GetBulkDataOffsetInFile();
 	}
 	// It could be that the bulkdata has no file associated yet (aka lightmaps have been build but not saved to disk yet) but still contains valid data
@@ -346,16 +346,16 @@ FVTDataAndStatus FUploadingVirtualTexture::ReadData(FGraphEventArray& OutComplet
 		{
 			return EVTRequestPageStatus::Saturated;
 		}
-		ChunkFileName = *ChunkFileNameDCC;
+		ChunkFileName = ChunkFileNameDCC;
 	}
 #else // WITH_EDITOR
-	ChunkFileName = *BulkData.GetFilename();
+	ChunkFileName = BulkData.GetFilename();
 	ChunkOffsetInFile = BulkData.GetBulkDataOffsetInFile();
 	if (BulkData.GetBulkDataSize() == 0)
 	{
 		if (!InvalidChunks[ChunkIndex])
 		{
-			UE_LOG(LogConsoleResponse, Display, TEXT("BulkData for chunk %d in file '%s' is empty."), ChunkIndex, ChunkFileName);
+			UE_LOG(LogConsoleResponse, Display, TEXT("BulkData for chunk %d in file '%s' is empty."), ChunkIndex, *ChunkFileName);
 			InvalidChunks[ChunkIndex] = true;
 		}
 		return EVTRequestPageStatus::Invalid;
@@ -365,13 +365,13 @@ FVTDataAndStatus FUploadingVirtualTexture::ReadData(FGraphEventArray& OutComplet
 	TUniquePtr<IFileCacheHandle>& Handle = HandlePerChunk[ChunkIndex];
 	if (!Handle)
 	{
-		Handle.Reset(IFileCacheHandle::CreateFileCacheHandle(ChunkFileName));
+		Handle.Reset(IFileCacheHandle::CreateFileCacheHandle(*ChunkFileName));
 		// Don't expect CreateFileCacheHandle() to fail, async files should never fail to open
 		if (!ensure(Handle.IsValid()))
 		{
 			if (!InvalidChunks[ChunkIndex])
 			{
-				UE_LOG(LogConsoleResponse, Display, TEXT("Could not create a file cache for '%s'."), ChunkFileName);
+				UE_LOG(LogConsoleResponse, Display, TEXT("Could not create a file cache for '%s'."), *ChunkFileName);
 				InvalidChunks[ChunkIndex] = true;
 			}
 			return EVTRequestPageStatus::Invalid;
diff --git a/Engine/Source/Runtime/Engine/Private/VT/VirtualTexture.cpp b/Engine/Source/Runtime/Engine/Private/VT/VirtualTexture.cpp
--- a/Engine/Source/Runtime/Engine/Private/VT/VirtualTexture.cpp
+++ b/Engine/Source/Runtime/Engine/Private/VT/VirtualTexture.cpp
@@ -141,7 +141,7 @@ void UVirtualTexture::Serialize(FArchive& Ar)
 		}
 
 		FByteBulkData BulkData;
-		BulkData.Serialize(Ar, this, ChunkId);
+		BulkData.Serialize(Ar, this, ChunkId, false);
 
 #if WITH_EDITORONLY_DATA
 		if (!bCooked && Version >= FVirtualTextureBuiltDataCustomVersion::SplitDDC)
diff --git a/Engine/Source/Runtime/Engine/Private/VT/VirtualTextureBuiltData.cpp b/Engine/Source/Runtime/Engine/Private/VT/VirtualTextureBuiltData.cpp
--- a/Engine/Source/Runtime/Engine/Private/VT/VirtualTextureBuiltData.cpp
+++ b/Engine/Source/Runtime/Engine/Private/VT/VirtualTextureBuiltData.cpp
@@ -163,7 +163,7 @@ void FVirtualTextureBuiltData::Serialize(FArchive& Ar, UObject* Owner, int32 Fir
 			Ar << Chunk.CodecPayloadOffset[LayerIndex];
 		}
 
-		Chunk.BulkData.Serialize(Ar, Owner, SerialzeChunkId);
+		Chunk.BulkData.Serialize(Ar, Owner, SerialzeChunkId, false);
 
 #if WITH_EDITORONLY_DATA
 		if (!bCooked)
diff --git a/Engine/Source/Runtime/Engine/Private/VectorField.cpp b/Engine/Source/Runtime/Engine/Private/VectorField.cpp
--- a/Engine/Source/Runtime/Engine/Private/VectorField.cpp
+++ b/Engine/Source/Runtime/Engine/Private/VectorField.cpp
@@ -407,7 +407,7 @@ void UVectorFieldStatic::Serialize(FArchive& Ar)
 		SourceData.SetBulkDataFlags(BULKDATA_ForceInlinePayload | BULKDATA_SingleUse);
 	}
 
-	SourceData.Serialize(Ar,this);
+	SourceData.Serialize(Ar, this, INDEX_NONE, false);
 }
 
 void UVectorFieldStatic::PostLoad()
diff --git a/Engine/Source/Runtime/Engine/Private/World.cpp b/Engine/Source/Runtime/Engine/Private/World.cpp
--- a/Engine/Source/Runtime/Engine/Private/World.cpp
+++ b/Engine/Source/Runtime/Engine/Private/World.cpp
@@ -19,7 +19,6 @@
 #include "UObject/UObjectAnnotation.h"
 #include "UObject/ReferenceChainSearch.h"
 #include "Misc/PackageName.h"
-#include "Serialization/AsyncLoading.h"
 #include "GameMapsSettings.h"
 #include "TimerManager.h"
 #include "Materials/MaterialInterface.h"
diff --git a/Engine/Source/Runtime/Engine/Public/Animation/AnimationStreaming.h b/Engine/Source/Runtime/Engine/Public/Animation/AnimationStreaming.h
--- a/Engine/Source/Runtime/Engine/Public/Animation/AnimationStreaming.h
+++ b/Engine/Source/Runtime/Engine/Public/Animation/AnimationStreaming.h
@@ -23,7 +23,7 @@ struct FLoadedAnimationChunk
 {
 	TAtomic<FCompressedAnimSequence*> CompressedAnimData;
 
-	class IAsyncReadRequest* IORequest;
+	struct FBulkDataIORequest* IORequest;
 	double RequestStart;
 
 	uint32	Index;
@@ -116,8 +116,6 @@ public:
 
 	mutable FCriticalSection LoadedChunksCritcalSection;
 
-	class IAsyncReadFileHandle* IORequestHandle;
-
 	/** Indices of chunks that are currently loaded */
 	TArray<uint32>	LoadedChunkIndices;
 
diff --git a/Engine/Source/Runtime/Engine/Public/AudioStreaming.h b/Engine/Source/Runtime/Engine/Public/AudioStreaming.h
--- a/Engine/Source/Runtime/Engine/Public/AudioStreaming.h
+++ b/Engine/Source/Runtime/Engine/Public/AudioStreaming.h
@@ -102,8 +102,7 @@ struct FWaveRequest
 struct FLoadedAudioChunk
 {
 	uint8*	Data;
-	class IAsyncReadRequest* IORequest;
-	int32	MemorySize; 
+	struct FBulkDataIORequest* IORequest;
 	int32	DataSize;
 	int32	AudioDataSize;
 	uint32	Index;
@@ -111,7 +110,6 @@ struct FLoadedAudioChunk
 	FLoadedAudioChunk()
 		: Data(nullptr)
 		, IORequest(nullptr)
-		, MemorySize(0)
 		, DataSize(0)
 		, AudioDataSize(0)
 		, Index(0)
@@ -121,6 +119,7 @@ struct FLoadedAudioChunk
 	~FLoadedAudioChunk()
 	{
 		checkf(Data == nullptr, TEXT("Audio chunk Data ptr not null (%p), DataSize: %d"), Data, DataSize);
+		checkf(IORequest == nullptr, TEXT("Audio chunk IORequest ptr not null (%p)"), IORequest );
 	}
 };
 
@@ -208,8 +207,6 @@ public:
 	/* Contains pointers to Chunks of audio data that have been streamed in */
 	TArray<FLoadedAudioChunk> LoadedChunks;
 
-	class IAsyncReadFileHandle* IORequestHandle;
-
 	/** Indices of chunks that are currently loaded */
 	TArray<uint32>	LoadedChunkIndices;
 
diff --git a/Engine/Source/Runtime/Engine/Public/AudioStreamingCache.h b/Engine/Source/Runtime/Engine/Public/AudioStreamingCache.h
--- a/Engine/Source/Runtime/Engine/Public/AudioStreamingCache.h
+++ b/Engine/Source/Runtime/Engine/Public/AudioStreamingCache.h
@@ -146,9 +146,8 @@ private:
 		TUniquePtr<FAsyncStreamDerivedChunkTask> DDCTask;
 #endif
 
-		TUniquePtr<IAsyncReadFileHandle> FileHandle;
 		// Handle to our async read request operation.
-		TUniquePtr<IAsyncReadRequest> ReadRequest;
+		TUniquePtr<FBulkDataIORequest> ReadRequest;
 
 #if DEBUG_STREAM_CACHE
 		FCacheElementDebugInfo DebugInfo;
diff --git a/Engine/Source/Runtime/Engine/Public/LightMap.h b/Engine/Source/Runtime/Engine/Public/LightMap.h
--- a/Engine/Source/Runtime/Engine/Public/LightMap.h
+++ b/Engine/Source/Runtime/Engine/Public/LightMap.h
@@ -395,6 +395,10 @@ struct FQuantizedSimpleLightSample
 	FColor	Coefficients[NUM_LQ_LIGHTMAP_COEF];
 };
 
+// These act like a POD as FColor is also set to act like POD
+template <> struct TIsPODType<FQuantizedDirectionalLightSample> { enum { Value = true }; };
+template <> struct TIsPODType<FQuantizedSimpleLightSample> { enum { Value = true }; };
+
 /**
  * Bulk data array of FQuantizedLightSamples
  */
diff --git a/Engine/Source/Runtime/Engine/Public/Rendering/SkeletalMeshLODRenderData.h b/Engine/Source/Runtime/Engine/Public/Rendering/SkeletalMeshLODRenderData.h
--- a/Engine/Source/Runtime/Engine/Public/Rendering/SkeletalMeshLODRenderData.h
+++ b/Engine/Source/Runtime/Engine/Public/Rendering/SkeletalMeshLODRenderData.h
@@ -130,10 +130,11 @@ public:
 
 	uint32 BuffersSize;
 
-	/** Offset in the .bulk file if this LOD is streamed */
-	uint32 OffsetInFile;
-	/** Size of serialized resource data in bytes */
-	uint32 BulkDataSize;
+#if USE_BULKDATA_STREAMING_TOKEN
+	FBulkDataStreamingToken BulkDataStreamingToken;
+#else
+	FByteBulkData StreamingBulkData;
+#endif
 
 	/** Whether buffers of this LOD is inlined (i.e. stored in .uexp instead of .ubulk) */
 	uint32 bStreamedDataInlined : 1;
@@ -164,8 +165,6 @@ public:
 	/** Constructor (default) */
 	FSkeletalMeshLODRenderData()
 		: BuffersSize(0)
-		, OffsetInFile(0xffffffff)
-		, BulkDataSize(0)
 		, bStreamedDataInlined(false)
 		, bIsLODOptional(false)
 	{
diff --git a/Engine/Source/Runtime/Engine/Public/StaticMeshResources.h b/Engine/Source/Runtime/Engine/Public/StaticMeshResources.h
--- a/Engine/Source/Runtime/Engine/Public/StaticMeshResources.h
+++ b/Engine/Source/Runtime/Engine/Public/StaticMeshResources.h
@@ -11,7 +11,6 @@
 #include "Misc/Guid.h"
 #include "Engine/EngineTypes.h"
 #include "UObject/UObjectIterator.h"
-#include "Templates/ScopedPointer.h"
 #include "Materials/MaterialInterface.h"
 #include "RenderResource.h"
 #include "PackedNormal.h"
@@ -383,10 +382,11 @@ struct FStaticMeshLODResources
 	/** Sum of all vertex and index buffer sizes. Calculated in SerializeBuffers */
 	uint32 BuffersSize;
 
-	/** Offset in the .bulk file if this LOD is streamed */
-	uint32 OffsetInFile;
-	/** Size of serialized buffer data in bytes */
-	uint32 BulkDataSize;
+#if USE_BULKDATA_STREAMING_TOKEN
+	FBulkDataStreamingToken BulkDataStreamingToken;
+#else
+	FByteBulkData StreamingBulkData;
+#endif
 
 #if STATS
 	uint32 StaticMeshIndexMemory;
diff --git a/Engine/Source/Runtime/Engine/Public/TextureResource.h b/Engine/Source/Runtime/Engine/Public/TextureResource.h
--- a/Engine/Source/Runtime/Engine/Public/TextureResource.h
+++ b/Engine/Source/Runtime/Engine/Public/TextureResource.h
@@ -11,7 +11,6 @@
 #include "Containers/IndirectArray.h"
 #include "Stats/Stats.h"
 #include "Containers/List.h"
-#include "Templates/ScopedPointer.h"
 #include "Async/AsyncWork.h"
 #include "Async/AsyncFileHandle.h"
 #include "RHI.h"
@@ -33,13 +32,16 @@ class IVirtualTexture;
 #ifndef FORCE_ENABLE_TEXTURE_STREAMING
 #define FORCE_ENABLE_TEXTURE_STREAMING 0
 #endif
-#define TEXTURE2DMIPMAP_USE_COMPACT_BULKDATA (!(WITH_EDITORONLY_DATA) && !(UE_SERVER) && FORCE_ENABLE_TEXTURE_STREAMING && PLATFORM_SUPPORTS_TEXTURE_STREAMING)
+
+#define TEXTURE2DMIPMAP_USE_COMPACT_BULKDATA (!(WITH_EDITORONLY_DATA) && !(UE_SERVER) && FORCE_ENABLE_TEXTURE_STREAMING && PLATFORM_SUPPORTS_TEXTURE_STREAMING && !USE_NEW_BULKDATA)
+
 
 /**
  * A 2D texture mip-map.
  */
 struct FTexture2DMipMap
 {
+#if TEXTURE2DMIPMAP_USE_COMPACT_BULKDATA
 	class FCompactByteBulkData
 	{
 	public:
@@ -53,7 +55,7 @@ struct FTexture2DMipMap
 		FCompactByteBulkData& operator=(const FCompactByteBulkData& Other);
 		FCompactByteBulkData& operator=(FCompactByteBulkData&& Other);
 
-		void Serialize(FArchive& Ar, UObject* Owner, int32 MipIdx);
+		void Serialize(FArchive& Ar, UObject* Owner, int32 MipIdx, bool bAttemptFileMapping);
 
 		uint32 GetBulkDataOffsetInFile() const { return OffsetInFile; }
 		uint32 GetBulkDataSize() const { return BulkDataSize; }
@@ -61,9 +63,13 @@ struct FTexture2DMipMap
 		uint32 GetElementCount() const { return BulkDataSize; }
 		uint32 GetElementSize() const { return sizeof(uint8); }
 		void SetBulkDataFlags(uint32 Flags) { BulkDataFlags |= Flags; }
+		void ResetBulkDataFlags(uint32 Flags) { BulkDataFlags = Flags; }
 		void ClearBulkDataFlags(uint32 FlagsToClear) { BulkDataFlags &= ~FlagsToClear; }
 		bool CanLoadFromDisk() const { return !IsInlined(); }
 		bool IsAvailableForUse() const { return !(BulkDataFlags & BULKDATA_Unused); }
+		bool IsOptional() const { return (BulkDataFlags & BULKDATA_OptionalPayload) != 0; }
+		bool IsInlined() const { return (GetBulkDataFlags() & BULKDATA_PayloadInSeperateFile) == 0 && (GetBulkDataFlags() & BULKDATA_PayloadAtEndOfFile) == 0; }
+		bool InSeperateFile() const { return (GetBulkDataFlags() & BULKDATA_PayloadInSeperateFile) != 0; }
 		bool IsBulkDataLoaded() const { return IsInlined(); }
 		bool IsAsyncLoadingComplete() const { return true; }
 		bool IsStoredCompressedOnDisk() const { return !!(BulkDataFlags & BULKDATA_SerializeCompressed); }
@@ -73,6 +79,8 @@ struct FTexture2DMipMap
 		void* Realloc(int32 NumBytes);
 		void GetCopy(void** Dest, bool bDiscardInternalCopy = true);
 
+		FBulkDataIORequest* CreateStreamingRequest(const FString& Filename, int64 OffsetInBulkData, int64 BytesToRead, EAsyncIOPriorityAndFlags Priority, FAsyncFileCallBack* CompleteCallback, uint8* UserSuppliedMemory) const;
+
 		/**
 		 * FCompactByteBulkData doesn't support GetFilename. Use FTexturePlatformData::CachedPackageFileName
 		 * or UTexture2D::GetMipDataFilename instead
@@ -90,8 +98,11 @@ struct FTexture2DMipMap
 		uint8* TexelData;
 
 		void Reset();
-		bool IsInlined() const { return !(BulkDataFlags & BULKDATA_Force_NOT_InlinePayload); }
+		
 	};
+#else
+	struct FCompactByteBulkData {};
+#endif
 
 	/** Width of the mip-map. */
 	int32 SizeX;
@@ -100,6 +111,7 @@ struct FTexture2DMipMap
 	/** Depth of the mip-map. */
 	int32 SizeZ;
 	/** Bulk data if stored in the package. */
+
 	typename TChooseClass<TEXTURE2DMIPMAP_USE_COMPACT_BULKDATA, FCompactByteBulkData, FByteBulkData>::Result BulkData;
 
 	/** Default constructor. */
diff --git a/Engine/Source/Runtime/NavigationSystem/Private/NavCollision.cpp b/Engine/Source/Runtime/NavigationSystem/Private/NavCollision.cpp
--- a/Engine/Source/Runtime/NavigationSystem/Private/NavCollision.cpp
+++ b/Engine/Source/Runtime/NavigationSystem/Private/NavCollision.cpp
@@ -570,7 +570,7 @@ void UNavCollision::GetResourceSizeEx(FResourceSizeEx& CumulativeResourceSize)
 	if (CookedFormatData.Contains(NAVCOLLISION_FORMAT))
 	{
 		const FByteBulkData& FmtData = CookedFormatData.GetFormat(NAVCOLLISION_FORMAT);
-		CumulativeResourceSize.AddDedicatedSystemMemoryBytes(FmtData.GetElementSize() * FmtData.GetElementCount());
+		CumulativeResourceSize.AddDedicatedSystemMemoryBytes(FmtData.GetBulkDataSize());
 	}
 }
 
diff --git a/Engine/Source/Runtime/SlateCore/Private/Fonts/FontBulkData.cpp b/Engine/Source/Runtime/SlateCore/Private/Fonts/FontBulkData.cpp
--- a/Engine/Source/Runtime/SlateCore/Private/Fonts/FontBulkData.cpp
+++ b/Engine/Source/Runtime/SlateCore/Private/Fonts/FontBulkData.cpp
@@ -10,7 +10,7 @@ DECLARE_MEMORY_STAT(TEXT("Font BulkData Memory"), STAT_SlateBulkFontDataMemory,
 
 UFontBulkData::UFontBulkData()
 {
-	BulkData.SetBulkDataFlags(BULKDATA_SerializeCompressed|BULKDATA_SerializeCompressedBitWindow);
+	
 }
 
 void UFontBulkData::Initialize(const FString& InFontFilename)
@@ -72,7 +72,7 @@ const void* UFontBulkData::Lock(int32& OutFontDataSizeBytes) const
 
 }
 
-void UFontBulkData::Unlock() const
+void UFontBulkData::Unlock()
 {
 	bool bWasLoaded = BulkData.IsBulkDataLoaded();
 	int32 BulkDataSize = BulkData.GetBulkDataSize(); 
@@ -94,13 +94,23 @@ int32 UFontBulkData::GetBulkDataSize() const
 
 int32 UFontBulkData::GetBulkDataSizeOnDisk() const
 {
+#if !USE_NEW_BULKDATA
 	return BulkData.GetBulkDataSizeOnDisk();
+#else
+	return 0;
+#endif
 }
 
 void UFontBulkData::Serialize(FArchive& Ar)
 {
 	Super::Serialize(Ar);
-	BulkData.Serialize(Ar, this);
+
+	if (Ar.IsSaving())
+	{
+		BulkData.SetBulkDataFlags(BULKDATA_SerializeCompressed | BULKDATA_SerializeCompressedBitWindow);
+	}
+	
+	BulkData.Serialize(Ar, this, INDEX_NONE, false);
 
 	if( !GIsEditor && Ar.IsLoading() )
 	{
diff --git a/Engine/Source/Runtime/SlateCore/Public/Fonts/CompositeFont.h b/Engine/Source/Runtime/SlateCore/Public/Fonts/CompositeFont.h
--- a/Engine/Source/Runtime/SlateCore/Public/Fonts/CompositeFont.h
+++ b/Engine/Source/Runtime/SlateCore/Public/Fonts/CompositeFont.h
@@ -262,7 +262,7 @@ private:
 	 * This can be removed once we no longer support loading packages older than FEditorObjectVersion::AddedFontFaceAssets (as can UFontBulkData itself).
 	 */
 	UPROPERTY()
-	const UFontBulkData* BulkDataPtr_DEPRECATED;
+	UFontBulkData* BulkDataPtr_DEPRECATED;
 
 	/**
 	 * Legacy font data v1. This used to be where font data was stored prior to font bulk data.
diff --git a/Engine/Source/Runtime/SlateCore/Public/Fonts/FontBulkData.h b/Engine/Source/Runtime/SlateCore/Public/Fonts/FontBulkData.h
--- a/Engine/Source/Runtime/SlateCore/Public/Fonts/FontBulkData.h
+++ b/Engine/Source/Runtime/SlateCore/Public/Fonts/FontBulkData.h
@@ -27,12 +27,13 @@ public:
 	const void* Lock(int32& OutFontDataSizeBytes) const;
 
 	/** Unlock the bulk font data, after which point the pointer returned by Lock no longer is valid */
-	void Unlock() const;
+	void Unlock();
 
 	/** Returns the size of the bulk data in bytes */
 	int32 GetBulkDataSize() const;
 
 	/** Returns the size of the bulk data on disk. This can differ from GetBulkDataSize if BULKDATA_SerializeCompressed is set */
+	UE_DEPRECATED(4.24, "This method should no longer be used")
 	int32 GetBulkDataSizeOnDisk() const;
 
 	// UObject interface
