Mark IsRooted const
#rnx

diff --git a/Engine/Source/Runtime/CoreUObject/Public/UObject/UObjectBaseUtility.h b/Engine/Source/Runtime/CoreUObject/Public/UObject/UObjectBaseUtility.h
--- a/Engine/Source/Runtime/CoreUObject/Public/UObject/UObjectBaseUtility.h
+++ b/Engine/Source/Runtime/CoreUObject/Public/UObject/UObjectBaseUtility.h
@@ -202,7 +202,7 @@ public:
 	 *
 	 * @return true if the object was explicitly added as part of the root set.
 	 */
-	FORCEINLINE bool IsRooted()
+	FORCEINLINE bool IsRooted() const
 	{
 		return GUObjectArray.IndexToObject(InternalIndex)->IsRootSet();
 	}
