Minor changes (bumped uobject count for UHT, better error messages) QOL changes from Fury branch

diff --git a/Engine/Programs/UnrealHeaderTool/Config/DefaultEngine.ini b/Engine/Programs/UnrealHeaderTool/Config/DefaultEngine.ini
--- a/Engine/Programs/UnrealHeaderTool/Config/DefaultEngine.ini
+++ b/Engine/Programs/UnrealHeaderTool/Config/DefaultEngine.ini
@@ -3,6 +3,9 @@
 +ScriptSupportedModules=Engine
 +ScriptSupportedModules=ScriptPlugin
 
+[/Script/Engine.GarbageCollectionSettings]
+gc.MaxObjectsInProgram=262144
+
 [UnrealHeaderTool]
 +StructsWithNoPrefix="uint64"
 +StructsWithNoPrefix="uint32"
@@ -27,3 +30,4 @@
 +TypeRedirects=(OldType="FStringClassReference", NewType="FSoftClassPath")
 +TypeRedirects=(OldType="TAssetPtr", NewType="TSoftObjectPtr")
 +TypeRedirects=(OldType="TAssetSubclassOf", NewType="TSoftClassPtr")
+
diff --git a/Engine/Source/Developer/DesktopPlatform/Private/PlatformInfo.cpp b/Engine/Source/Developer/DesktopPlatform/Private/PlatformInfo.cpp
--- a/Engine/Source/Developer/DesktopPlatform/Private/PlatformInfo.cpp
+++ b/Engine/Source/Developer/DesktopPlatform/Private/PlatformInfo.cpp
@@ -202,6 +202,9 @@ const FPlatformInfo* FindPlatformInfo(const FName& InPlatformName)
 			return &PlatformInfo;
 		}
 	}
+
+	UE_LOG(LogDesktopPlatform, Warning, TEXT("Unable to find platform info for '%s'"), *InPlatformName.ToString());
+
 	return nullptr;
 }
 
diff --git a/Engine/Source/Runtime/CoreUObject/Public/UObject/UObjectArray.h b/Engine/Source/Runtime/CoreUObject/Public/UObject/UObjectArray.h
--- a/Engine/Source/Runtime/CoreUObject/Public/UObject/UObjectArray.h
+++ b/Engine/Source/Runtime/CoreUObject/Public/UObject/UObjectArray.h
@@ -232,7 +232,7 @@ public:
 	int32 AddSingle() TSAN_SAFE
 	{
 		int32 Result = NumElements;
-		checkf(NumElements + 1 <= MaxElements, TEXT("Maximum number of UObjects (%d) exceeded, make sure you update MaxObjectsInGame/MaxObjectsInEditor in project settings."), MaxElements);
+		checkf(NumElements + 1 <= MaxElements, TEXT("Maximum number of UObjects (%d) exceeded, make sure you update MaxObjectsInGame/MaxObjectsInEditor/MaxObjectsInProgram in project settings."), MaxElements);
 		check(Result == NumElements);
 		++NumElements;
 		FPlatformMisc::MemoryBarrier();
@@ -243,7 +243,7 @@ public:
 	int32 AddRange(int32 Count) TSAN_SAFE
 	{
 		int32 Result = NumElements + Count - 1;
-		checkf(NumElements + Count <= MaxElements, TEXT("Maximum number of UObjects (%d) exceeded, make sure you update MaxObjectsInGame/MaxObjectsInEditor in project settings."), MaxElements);
+		checkf(NumElements + Count <= MaxElements, TEXT("Maximum number of UObjects (%d) exceeded, make sure you update MaxObjectsInGame/MaxObjectsInEditor/MaxObjectsInProgram in project settings."), MaxElements);
 		check(Result == (NumElements + Count - 1));
 		NumElements += Count;
 		FPlatformMisc::MemoryBarrier();
@@ -508,7 +508,7 @@ public:
 	int32 AddRange(int32 NumToAdd) TSAN_SAFE
 	{
 		int32 Result = NumElements;
-		checkf(Result + NumToAdd <= MaxElements, TEXT("Maximum number of UObjects (%d) exceeded, make sure you update MaxObjectsInGame/MaxObjectsInEditor in project settings."), MaxElements);
+		checkf(Result + NumToAdd <= MaxElements, TEXT("Maximum number of UObjects (%d) exceeded, make sure you update MaxObjectsInGame/MaxObjectsInEditor/MaxObjectsInProgram in project settings."), MaxElements);
 		ExpandChunksToIndex(Result + NumToAdd - 1);
 		NumElements += NumToAdd;
 		return Result;
