Revert PackageFileSummary CustomVersion changes accidentally copied up from dev-core.
Preparation for fixing UE-80965.

diff --git a/Engine/Source/Runtime/CoreUObject/Private/UObject/PackageFileSummary.cpp b/Engine/Source/Runtime/CoreUObject/Private/UObject/PackageFileSummary.cpp
--- a/Engine/Source/Runtime/CoreUObject/Private/UObject/PackageFileSummary.cpp
+++ b/Engine/Source/Runtime/CoreUObject/Private/UObject/PackageFileSummary.cpp
@@ -128,6 +128,7 @@ void operator<<(FStructuredArchive::FSlot Slot, FPackageFileSummary& Sum)
 					Sum.bUnversioned = true;
 					Sum.FileVersionUE4 = GPackageFileUE4Version;
 					Sum.FileVersionLicenseeUE4 = GPackageFileLicenseeUE4Version;
+					Sum.CustomVersionContainer = FCustomVersionContainer::GetRegistered();
 				}
 			}
 			else
@@ -373,14 +374,6 @@ FArchive& operator<<( FArchive& Ar, FPackageFileSummary& Sum )
 	return Ar;
 }
 
-const FCustomVersionContainer& FPackageFileSummary::GetCustomVersionContainer() const
-{
-	return
-		(bUnversioned && CustomVersionContainer.GetAllVersions().Num() == 0) ?
-		FCustomVersionContainer::GetRegistered() :
-		CustomVersionContainer;
-}
-
 void FPackageFileSummary::SetCustomVersionContainer(const FCustomVersionContainer& InContainer)
 {
 	CustomVersionContainer = InContainer;
diff --git a/Engine/Source/Runtime/CoreUObject/Public/UObject/PackageFileSummary.h b/Engine/Source/Runtime/CoreUObject/Public/UObject/PackageFileSummary.h
--- a/Engine/Source/Runtime/CoreUObject/Public/UObject/PackageFileSummary.h
+++ b/Engine/Source/Runtime/CoreUObject/Public/UObject/PackageFileSummary.h
@@ -231,7 +231,10 @@ public:
 		return FileVersionLicenseeUE4;
 	}
 
-	COREUOBJECT_API const FCustomVersionContainer& GetCustomVersionContainer() const;
+	const FCustomVersionContainer& GetCustomVersionContainer() const
+	{
+		return CustomVersionContainer;
+	}
 
 	void SetCustomVersionContainer(const FCustomVersionContainer& InContainer);
 
