Removing some code in-preparation of thread-safety fix

We shouldn't care if *all* custom versions are present in any particular package that is being loaded. The whole point of custom versions is to allow different systems to versions and upgrade data independently. Everything that should matter is if all of the relevant custom versions are up to date or not.

diff --git a/Engine/Source/Runtime/CoreUObject/Private/UObject/LinkerLoad.cpp b/Engine/Source/Runtime/CoreUObject/Private/UObject/LinkerLoad.cpp
--- a/Engine/Source/Runtime/CoreUObject/Private/UObject/LinkerLoad.cpp
+++ b/Engine/Source/Runtime/CoreUObject/Private/UObject/LinkerLoad.cpp
@@ -1322,9 +1322,6 @@ FLinkerLoad::ELinkerStatus FLinkerLoad::UpdateFromPackageFileSummary()
 				bAllSavedVersionsMatch = false;
 			}
 		}
-
-		const bool bSameNumberOfVersions = (PackageCustomVersions.Num() == LatestCustomVersions.GetAllVersions().Num());
-		bCustomVersionIsLatest = bSameNumberOfVersions && bAllSavedVersionsMatch;
 	}
 
 	const FCustomVersionContainer& SummaryVersions = Summary.GetCustomVersionContainer();
