Fixing a server crash when an asset file exists but can't be read

diff --git a/Engine/Source/Runtime/CoreUObject/Private/Serialization/AsyncLoading.cpp b/Engine/Source/Runtime/CoreUObject/Private/Serialization/AsyncLoading.cpp
--- a/Engine/Source/Runtime/CoreUObject/Private/Serialization/AsyncLoading.cpp
+++ b/Engine/Source/Runtime/CoreUObject/Private/Serialization/AsyncLoading.cpp
@@ -1464,7 +1464,11 @@ void FAsyncPackage::Event_FinishLinker()
 {
 	FScopedAsyncPackageEvent Scope(this);
 	SCOPED_LOADTIMER(Package_FinishLinker);
-	EAsyncPackageState::Type Result = FinishLinker();
+	EAsyncPackageState::Type Result = EAsyncPackageState::Complete;
+	if (!bLoadHasFailed)
+	{
+		Result = FinishLinker();
+	}
 	if (Result == EAsyncPackageState::TimeOut && !bLoadHasFailed)
 	{
 		AsyncLoadingThread.QueueEvent_FinishLinker(FWeakAsyncPackagePtr(this), FAsyncLoadEvent::EventSystemPriority_MAX);
@@ -5774,6 +5778,11 @@ EAsyncPackageState::Type FAsyncPackage::CreateLinker()
 						check(!Linker->GetAsyncLoader());
 						AsyncLoadingThread.GetPrecacheHandler().SummaryComplete(WeakPtr);
 					}
+					else if (!Linker->Loader)
+					{
+						AsyncLoadingThread.GetPrecacheHandler().SummaryComplete(WeakPtr);
+						bLoadHasFailed = true;
+					}
 				}
 			}
 			else
