Merging CL# 9893571 as it should not have been deadended
#rnx

diff --git a/Engine/Plugins/Experimental/PythonScriptPlugin/Source/PythonScriptPlugin/Private/PyConversion.cpp b/Engine/Plugins/Experimental/PythonScriptPlugin/Source/PythonScriptPlugin/Private/PyConversion.cpp
--- a/Engine/Plugins/Experimental/PythonScriptPlugin/Source/PythonScriptPlugin/Private/PyConversion.cpp
+++ b/Engine/Plugins/Experimental/PythonScriptPlugin/Source/PythonScriptPlugin/Private/PyConversion.cpp
@@ -670,7 +670,7 @@ PyObject* PythonizeEnumEntry(const int64 Val, const UEnum* EnumType, const ESetE
 	return Obj;
 }
 
-FPyConversionResult NativizeProperty(PyObject* PyObj, const UProperty* Prop, void* ValueAddr, const FPyWrapperOwnerContext& InChangeOwner, const ESetErrorState SetErrorState)
+FPyConversionResult NativizeProperty(PyObject* PyObj, const UProperty* Prop, void* ValueAddr, const FPropertyAccessChangeNotify* InChangeNotify, const ESetErrorState SetErrorState)
 {
 #define PYCONVERSION_PROPERTY_RETURN(RESULT) \
 	PYCONVERSION_RETURN(RESULT, TEXT("NativizeProperty"), *FString::Printf(TEXT("Cannot nativize '%s' as '%s' (%s)"), *PyUtil::GetFriendlyTypename(PyObj), *Prop->GetName(), *Prop->GetClass()->GetName()))
@@ -680,18 +680,21 @@ FPyConversionResult NativizeProperty(PyObject* PyObj, const UProperty* Prop, voi
 		FPyWrapperFixedArrayPtr PyFixedArray = FPyWrapperFixedArrayPtr::StealReference(FPyWrapperFixedArray::CastPyObject(PyObj, &PyWrapperFixedArrayType, Prop));
 		if (PyFixedArray)
 		{
-			const int32 ArrSize = FMath::Min(Prop->ArrayDim, PyFixedArray->ArrayProp->ArrayDim);
-			for (int32 ArrIndex = 0; ArrIndex < ArrSize; ++ArrIndex)
+			EmitPropertyChangeNotifications(InChangeNotify, [&]()
 			{
-				Prop->CopySingleValue(static_cast<uint8*>(ValueAddr) + (Prop->ElementSize * ArrIndex), FPyWrapperFixedArray::GetItemPtr(PyFixedArray, ArrIndex));
-			}
+				const int32 ArrSize = FMath::Min(Prop->ArrayDim, PyFixedArray->ArrayProp->ArrayDim);
+				for (int32 ArrIndex = 0; ArrIndex < ArrSize; ++ArrIndex)
+				{
+					Prop->CopySingleValue(static_cast<uint8*>(ValueAddr) + (Prop->ElementSize * ArrIndex), FPyWrapperFixedArray::GetItemPtr(PyFixedArray, ArrIndex));
+				}
+			});
 			return FPyConversionResult::Success();
 		}
 
 		PYCONVERSION_PROPERTY_RETURN(FPyConversionResult::Failure());
 	}
 
-	return NativizeProperty_Direct(PyObj, Prop, ValueAddr, InChangeOwner, SetErrorState);
+	return NativizeProperty_Direct(PyObj, Prop, ValueAddr, InChangeNotify, SetErrorState);
 
 #undef PYCONVERSION_PROPERTY_RETURN
 }
@@ -707,7 +710,7 @@ FPyConversionResult PythonizeProperty(const UProperty* Prop, const void* ValueAd
 	return PythonizeProperty_Direct(Prop, ValueAddr, OutPyObj, ConversionMethod, OwnerPyObj, SetErrorState);
 }
 
-FPyConversionResult NativizeProperty_Direct(PyObject* PyObj, const UProperty* Prop, void* ValueAddr, const FPyWrapperOwnerContext& InChangeOwner, const ESetErrorState SetErrorState)
+FPyConversionResult NativizeProperty_Direct(PyObject* PyObj, const UProperty* Prop, void* ValueAddr, const FPropertyAccessChangeNotify* InChangeNotify, const ESetErrorState SetErrorState)
 {
 #define PYCONVERSION_PROPERTY_RETURN(RESULT) \
 	PYCONVERSION_RETURN(RESULT, TEXT("NativizeProperty"), *FString::Printf(TEXT("Cannot nativize '%s' as '%s' (%s)"), *PyUtil::GetFriendlyTypename(PyObj), *Prop->GetName(), *Prop->GetClass()->GetName()))
@@ -722,7 +725,7 @@ FPyConversionResult NativizeProperty_Direct(PyObject* PyObj, const UProperty* Pr
 			auto OldValue = CastProp->GetPropertyValue(ValueAddr);					\
 			if (OldValue != NewValue)												\
 			{																		\
-				EmitPropertyChangeNotifications(InChangeOwner, [&]()				\
+				EmitPropertyChangeNotifications(InChangeNotify, [&]()				\
 				{																	\
 					CastProp->SetPropertyValue(ValueAddr, NewValue);				\
 				});																	\
@@ -741,7 +744,7 @@ FPyConversionResult NativizeProperty_Direct(PyObject* PyObj, const UProperty* Pr
 			auto* ValuePtr = static_cast<PROPTYPE::TCppType*>(ValueAddr);			\
 			if (!CastProp->Identical(ValuePtr, &NewValue, PPF_None))				\
 			{																		\
-				EmitPropertyChangeNotifications(InChangeOwner, [&]()				\
+				EmitPropertyChangeNotifications(InChangeNotify, [&]()				\
 				{																	\
 					*ValuePtr = MoveTemp(NewValue);									\
 				});																	\
@@ -793,7 +796,7 @@ FPyConversionResult NativizeProperty_Direct(PyObject* PyObj, const UProperty* Pr
 			auto* ValuePtr = static_cast<uint8*>(ValueAddr);
 			if (*ValuePtr != NewValue)
 			{
-				EmitPropertyChangeNotifications(InChangeOwner, [&]()
+				EmitPropertyChangeNotifications(InChangeNotify, [&]()
 				{
 					*ValuePtr = NewValue;
 				});
@@ -823,7 +826,7 @@ FPyConversionResult NativizeProperty_Direct(PyObject* PyObj, const UProperty* Pr
 				const int64 OldValue = EnumInternalProp->GetSignedIntPropertyValue(ValueAddr);
 				if (OldValue != NewValue)
 				{
-					EmitPropertyChangeNotifications(InChangeOwner, [&]()
+					EmitPropertyChangeNotifications(InChangeNotify, [&]()
 					{
 						EnumInternalProp->SetIntPropertyValue(ValueAddr, NewValue);
 					});
@@ -843,7 +846,7 @@ FPyConversionResult NativizeProperty_Direct(PyObject* PyObj, const UProperty* Pr
 			UObject* OldValue = CastProp->GetObjectPropertyValue(ValueAddr);
 			if (OldValue != NewValue)
 			{
-				EmitPropertyChangeNotifications(InChangeOwner, [&]()
+				EmitPropertyChangeNotifications(InChangeNotify, [&]()
 				{
 					CastProp->SetObjectPropertyValue(ValueAddr, NewValue);
 				});
@@ -861,7 +864,7 @@ FPyConversionResult NativizeProperty_Direct(PyObject* PyObj, const UProperty* Pr
 			UObject* OldValue = CastProp->GetObjectPropertyValue(ValueAddr);
 			if (OldValue != NewValue)
 			{
-				EmitPropertyChangeNotifications(InChangeOwner, [&]()
+				EmitPropertyChangeNotifications(InChangeNotify, [&]()
 				{
 					CastProp->SetObjectPropertyValue(ValueAddr, NewValue);
 				});
@@ -879,7 +882,7 @@ FPyConversionResult NativizeProperty_Direct(PyObject* PyObj, const UProperty* Pr
 			UObject* OldValue = CastProp->GetObjectPropertyValue(ValueAddr);
 			if (OldValue != NewValue)
 			{
-				EmitPropertyChangeNotifications(InChangeOwner, [&]()
+				EmitPropertyChangeNotifications(InChangeNotify, [&]()
 				{
 					CastProp->SetObjectPropertyValue(ValueAddr, NewValue);
 				});
@@ -897,7 +900,7 @@ FPyConversionResult NativizeProperty_Direct(PyObject* PyObj, const UProperty* Pr
 			UObject* OldValue = CastProp->GetPropertyValue(ValueAddr).GetObject();
 			if (OldValue != NewValue)
 			{
-				EmitPropertyChangeNotifications(InChangeOwner, [&]()
+				EmitPropertyChangeNotifications(InChangeNotify, [&]()
 				{
 					CastProp->SetPropertyValue(ValueAddr, FScriptInterface(NewValue, NewValue ? NewValue->GetInterfaceAddress(CastProp->InterfaceClass) : nullptr));
 				});
@@ -915,7 +918,7 @@ FPyConversionResult NativizeProperty_Direct(PyObject* PyObj, const UProperty* Pr
 		{
 			if (!CastProp->Identical(ValueAddr, PyStruct->StructInstance, PPF_None))
 			{
-				EmitPropertyChangeNotifications(InChangeOwner, [&]()
+				EmitPropertyChangeNotifications(InChangeNotify, [&]()
 				{
 					CastProp->Struct->CopyScriptStruct(ValueAddr, PyStruct->StructInstance);
 				});
@@ -941,7 +944,7 @@ FPyConversionResult NativizeProperty_Direct(PyObject* PyObj, const UProperty* Pr
 		{
 			if (!CastProp->Identical(ValueAddr, PyDelegate->DelegateInstance, PPF_None))
 			{
-				EmitPropertyChangeNotifications(InChangeOwner, [&]()
+				EmitPropertyChangeNotifications(InChangeNotify, [&]()
 				{
 					CastProp->SetPropertyValue(ValueAddr, *PyDelegate->DelegateInstance);
 				});
@@ -959,7 +962,7 @@ FPyConversionResult NativizeProperty_Direct(PyObject* PyObj, const UProperty* Pr
 		{
 			if (!CastProp->Identical(ValueAddr, PyDelegate->DelegateInstance, PPF_None))
 			{
-				EmitPropertyChangeNotifications(InChangeOwner, [&]()
+				EmitPropertyChangeNotifications(InChangeNotify, [&]()
 				{
 					CastProp->SetMulticastDelegate(ValueAddr, *PyDelegate->DelegateInstance);
 				});
@@ -976,9 +979,9 @@ FPyConversionResult NativizeProperty_Direct(PyObject* PyObj, const UProperty* Pr
 		{
 			if (!CastProp->Identical(ValueAddr, PyArray->ArrayInstance, PPF_None))
 			{
-				EmitPropertyChangeNotifications(InChangeOwner, [&]()
+				EmitPropertyChangeNotifications(InChangeNotify, [&]()
 				{
-					CastProp->CopyCompleteValue(ValueAddr, PyArray->ArrayInstance);
+					CastProp->CopySingleValue(ValueAddr, PyArray->ArrayInstance);
 				});
 			}
 		}
@@ -993,9 +996,9 @@ FPyConversionResult NativizeProperty_Direct(PyObject* PyObj, const UProperty* Pr
 		{
 			if (!CastProp->Identical(ValueAddr, PySet->SetInstance, PPF_None))
 			{
-				EmitPropertyChangeNotifications(InChangeOwner, [&]()
+				EmitPropertyChangeNotifications(InChangeNotify, [&]()
 				{
-					CastProp->CopyCompleteValue(ValueAddr, PySet->SetInstance);
+					CastProp->CopySingleValue(ValueAddr, PySet->SetInstance);
 				});
 			}
 		}
@@ -1010,9 +1013,9 @@ FPyConversionResult NativizeProperty_Direct(PyObject* PyObj, const UProperty* Pr
 		{
 			if (!CastProp->Identical(ValueAddr, PyMap->MapInstance, PPF_None))
 			{
-				EmitPropertyChangeNotifications(InChangeOwner, [&]()
+				EmitPropertyChangeNotifications(InChangeNotify, [&]()
 				{
-					CastProp->CopyCompleteValue(ValueAddr, PyMap->MapInstance);
+					CastProp->CopySingleValue(ValueAddr, PyMap->MapInstance);
 				});
 			}
 		}
@@ -1161,10 +1164,10 @@ FPyConversionResult PythonizeProperty_Direct(const UProperty* Prop, const void*
 #undef PYCONVERSION_PROPERTY_RETURN
 }
 
-FPyConversionResult NativizeProperty_InContainer(PyObject* PyObj, const UProperty* Prop, void* BaseAddr, const int32 ArrayIndex, const FPyWrapperOwnerContext& InChangeOwner, const ESetErrorState SetErrorState)
+FPyConversionResult NativizeProperty_InContainer(PyObject* PyObj, const UProperty* Prop, void* BaseAddr, const int32 ArrayIndex, const FPropertyAccessChangeNotify* InChangeNotify, const ESetErrorState SetErrorState)
 {
 	check(ArrayIndex < Prop->ArrayDim);
-	return NativizeProperty(PyObj, Prop, Prop->ContainerPtrToValuePtr<void>(BaseAddr, ArrayIndex), InChangeOwner, SetErrorState);
+	return NativizeProperty(PyObj, Prop, Prop->ContainerPtrToValuePtr<void>(BaseAddr, ArrayIndex), InChangeNotify, SetErrorState);
 }
 
 FPyConversionResult PythonizeProperty_InContainer(const UProperty* Prop, const void* BaseAddr, const int32 ArrayIndex, PyObject*& OutPyObj, const EPyConversionMethod ConversionMethod, PyObject* OwnerPyObj, const ESetErrorState SetErrorState)
@@ -1173,83 +1176,11 @@ FPyConversionResult PythonizeProperty_InContainer(const UProperty* Prop, const v
 	return PythonizeProperty(Prop, Prop->ContainerPtrToValuePtr<void>(BaseAddr, ArrayIndex), OutPyObj, ConversionMethod, OwnerPyObj, SetErrorState);
 }
 
-void EmitPropertyChangeNotifications(const FPyWrapperOwnerContext& InChangeOwner, const TFunctionRef<void()>& InDoChangeFunc)
+void EmitPropertyChangeNotifications(const FPropertyAccessChangeNotify* InChangeNotify, const TFunctionRef<void()>& InDoChangeFunc)
 {
-#if WITH_EDITOR
-	auto BuildPropertyChain = [&InChangeOwner](FEditPropertyChain& OutPropertyChain) -> UObject*
-	{
-		auto AppendOwnerPropertyToChain = [&OutPropertyChain](const FPyWrapperOwnerContext& InOwnerContext) -> bool
-		{
-			const UProperty* LeafProp = nullptr;
-			if (PyObject_IsInstance(InOwnerContext.GetOwnerObject(), (PyObject*)&PyWrapperObjectType) == 1 || PyObject_IsInstance(InOwnerContext.GetOwnerObject(), (PyObject*)&PyWrapperStructType) == 1)
-			{
-				LeafProp = InOwnerContext.GetOwnerProperty();
-			}
-
-			if (LeafProp)
-			{
-				OutPropertyChain.AddHead((UProperty*)LeafProp);
-				return true;
-			}
-
-			return false;
-		};
-
-		FPyWrapperOwnerContext OwnerContext = InChangeOwner;
-		while (OwnerContext.HasOwner() && AppendOwnerPropertyToChain(OwnerContext))
-		{
-			PyObject* PyObj = OwnerContext.GetOwnerObject();
-
-			if (PyObj == InChangeOwner.GetOwnerObject())
-			{
-				OutPropertyChain.SetActivePropertyNode(OutPropertyChain.GetHead()->GetValue());
-			}
-
-			if (PyObject_IsInstance(PyObj, (PyObject*)&PyWrapperObjectType) == 1)
-			{
-				// Found an object, this is the end of the chain
-				OutPropertyChain.SetActiveMemberPropertyNode(OutPropertyChain.GetHead()->GetValue());
-				return ((FPyWrapperObject*)PyObj)->ObjectInstance;
-			}
-
-			if (PyObject_IsInstance(PyObj, (PyObject*)&PyWrapperStructType) == 1)
-			{
-				// Found a struct, recurse up the chain
-				OwnerContext = ((FPyWrapperStruct*)PyObj)->OwnerContext;
-				continue;
-			}
-
-			// Unknown object type - just bail
-			break;
-		}
-
-		return nullptr;
-	};
-
-	// Build the property chain we should notify of the change
-	FEditPropertyChain PropertyChain;
-	UObject* ObjectToNotify = BuildPropertyChain(PropertyChain);
-
-	// Notify that a change is about to occur
-	if (ObjectToNotify)
-	{
-		ObjectToNotify->PreEditChange(PropertyChain);
-	}
-
-	// Perform the change
-	InDoChangeFunc();
-
-	// Notify that the change has occurred
-	if (ObjectToNotify)
-	{
-		FPropertyChangedEvent PropertyEvent(PropertyChain.GetActiveNode()->GetValue());
-		PropertyEvent.SetActiveMemberProperty(PropertyChain.GetActiveMemberNode()->GetValue());
-		FPropertyChangedChainEvent PropertyChainEvent(PropertyChain, PropertyEvent);
-		ObjectToNotify->PostEditChangeChainProperty(PropertyChainEvent);
-	}
-#else	// WITH_EDITOR
+	PropertyAccessUtil::EmitPreChangeNotify(InChangeNotify);
 	InDoChangeFunc();
-#endif	// WITH_EDITOR
+	PropertyAccessUtil::EmitPostChangeNotify(InChangeNotify);
 }
 
 }
diff --git a/Engine/Plugins/Experimental/PythonScriptPlugin/Source/PythonScriptPlugin/Private/PyConversion.h b/Engine/Plugins/Experimental/PythonScriptPlugin/Source/PythonScriptPlugin/Private/PyConversion.h
--- a/Engine/Plugins/Experimental/PythonScriptPlugin/Source/PythonScriptPlugin/Private/PyConversion.h
+++ b/Engine/Plugins/Experimental/PythonScriptPlugin/Source/PythonScriptPlugin/Private/PyConversion.h
@@ -218,15 +218,15 @@ namespace PyConversion
 	}
 
 	/** Conversion for property instances (including fixed arrays) - ValueAddr should point to the property data */
-	FPyConversionResult NativizeProperty(PyObject* PyObj, const UProperty* Prop, void* ValueAddr, const FPyWrapperOwnerContext& InChangeOwner = FPyWrapperOwnerContext(), const ESetErrorState SetErrorState = ESetErrorState::Yes);
+	FPyConversionResult NativizeProperty(PyObject* PyObj, const UProperty* Prop, void* ValueAddr, const FPropertyAccessChangeNotify* InChangeNotify = nullptr, const ESetErrorState SetErrorState = ESetErrorState::Yes);
 	FPyConversionResult PythonizeProperty(const UProperty* Prop, const void* ValueAddr, PyObject*& OutPyObj, const EPyConversionMethod ConversionMethod = EPyConversionMethod::Copy, PyObject* OwnerPyObj = nullptr, const ESetErrorState SetErrorState = ESetErrorState::Yes);
 
 	/** Conversion for single property instances - ValueAddr should point to the property data */
-	FPyConversionResult NativizeProperty_Direct(PyObject* PyObj, const UProperty* Prop, void* ValueAddr, const FPyWrapperOwnerContext& InChangeOwner = FPyWrapperOwnerContext(), const ESetErrorState SetErrorState = ESetErrorState::Yes);
+	FPyConversionResult NativizeProperty_Direct(PyObject* PyObj, const UProperty* Prop, void* ValueAddr, const FPropertyAccessChangeNotify* InChangeNotify = nullptr, const ESetErrorState SetErrorState = ESetErrorState::Yes);
 	FPyConversionResult PythonizeProperty_Direct(const UProperty* Prop, const void* ValueAddr, PyObject*& OutPyObj, const EPyConversionMethod ConversionMethod = EPyConversionMethod::Copy, PyObject* OwnerPyObj = nullptr, const ESetErrorState SetErrorState = ESetErrorState::Yes);
 
 	/** Conversion for property instances within a structure (including fixed arrays) - BaseAddr should point to the structure data */
-	FPyConversionResult NativizeProperty_InContainer(PyObject* PyObj, const UProperty* Prop, void* BaseAddr, const int32 ArrayIndex, const FPyWrapperOwnerContext& InChangeOwner = FPyWrapperOwnerContext(), const ESetErrorState SetErrorState = ESetErrorState::Yes);
+	FPyConversionResult NativizeProperty_InContainer(PyObject* PyObj, const UProperty* Prop, void* BaseAddr, const int32 ArrayIndex, const FPropertyAccessChangeNotify* InChangeNotify = nullptr, const ESetErrorState SetErrorState = ESetErrorState::Yes);
 	FPyConversionResult PythonizeProperty_InContainer(const UProperty* Prop, const void* BaseAddr, const int32 ArrayIndex, PyObject*& OutPyObj, const EPyConversionMethod ConversionMethod = EPyConversionMethod::Copy, PyObject* OwnerPyObj = nullptr, const ESetErrorState SetErrorState = ESetErrorState::Yes);
 
 	/**
@@ -234,7 +234,7 @@ namespace PyConversion
 	 * This function should be called when you know the value will actually change (or know you want to emit the notifications for it changing) and will do 
 	 * the pre-change notify, invoke the passed delegate to perform the change, then do the post-change notify
 	 */
-	void EmitPropertyChangeNotifications(const FPyWrapperOwnerContext& InChangeOwner, const TFunctionRef<void()>& InDoChangeFunc);
+	void EmitPropertyChangeNotifications(const FPropertyAccessChangeNotify* InChangeNotify, const TFunctionRef<void()>& InDoChangeFunc);
 }
 
 #endif	// WITH_PYTHON
diff --git a/Engine/Plugins/Experimental/PythonScriptPlugin/Source/PythonScriptPlugin/Private/PyGenUtil.cpp b/Engine/Plugins/Experimental/PythonScriptPlugin/Source/PythonScriptPlugin/Private/PyGenUtil.cpp
--- a/Engine/Plugins/Experimental/PythonScriptPlugin/Source/PythonScriptPlugin/Private/PyGenUtil.cpp
+++ b/Engine/Plugins/Experimental/PythonScriptPlugin/Source/PythonScriptPlugin/Private/PyGenUtil.cpp
@@ -504,8 +504,8 @@ FGeneratedWrappedPropertyDoc::FGeneratedWrappedPropertyDoc(const UProperty* InPr
 	PythonPropName = GetPropertyPythonName(InProp);
 
 	const FString PropTooltip = GetFieldTooltip(InProp);
-	DocString = PythonizePropertyTooltip(PropTooltip, InProp);
-	EditorDocString = PythonizePropertyTooltip(PropTooltip, InProp, CPF_EditConst);
+	DocString = PythonizePropertyTooltip(PropTooltip, InProp, PropertyAccessUtil::RuntimeReadOnlyFlags);
+	EditorDocString = PythonizePropertyTooltip(PropTooltip, InProp, PropertyAccessUtil::EditorReadOnlyFlags);
 }
 
 bool FGeneratedWrappedPropertyDoc::SortPredicate(const FGeneratedWrappedPropertyDoc& InOne, const FGeneratedWrappedPropertyDoc& InTwo)
@@ -706,7 +706,7 @@ FPythonizeTooltipContext::FPythonizeTooltipContext(const UProperty* InProp, cons
 FPythonizeTooltipContext::FPythonizeTooltipContext(const UFunction* InFunc, const TSet<FName>& InParamsToIgnore)
 	: Prop(nullptr)
 	, Func(InFunc)
-	, ReadOnlyFlags(CPF_BlueprintReadOnly | CPF_EditConst)
+	, ReadOnlyFlags(PropertyAccessUtil::RuntimeReadOnlyFlags)
 	, ParamsToIgnore(InParamsToIgnore)
 {
 	if (Func)
@@ -1158,7 +1158,7 @@ bool InvokePythonCallableFromUnrealFunctionThunk(FPyObjectPtr InSelf, PyObject*
 	return true;
 }
 
-PyObject* GetPropertyValue(const UStruct* InStruct, void* InStructData, const FGeneratedWrappedProperty& InPropDef, const char *InAttributeName, PyObject* InOwnerPyObject, const TCHAR* InErrorCtxt)
+PyObject* GetPropertyValue(const UStruct* InStruct, const void* InStructData, const FGeneratedWrappedProperty& InPropDef, const char *InAttributeName, PyObject* InOwnerPyObject, const TCHAR* InErrorCtxt)
 {
 	// Has this property been deprecated?
 	if (InStruct && InPropDef.Prop && InPropDef.DeprecationMessage.IsSet())
@@ -1183,7 +1183,7 @@ PyObject* GetPropertyValue(const UStruct* InStruct, void* InStructData, const FG
 	return PyUtil::GetPropertyValue(InStruct, InStructData, InPropDef.Prop, InAttributeName, InOwnerPyObject, InErrorCtxt);
 }
 
-int SetPropertyValue(const UStruct* InStruct, void* InStructData, PyObject* InValue, const FGeneratedWrappedProperty& InPropDef, const char *InAttributeName, const FPyWrapperOwnerContext& InChangeOwner, const uint64 InReadOnlyFlags, const bool InOwnerIsTemplate, const TCHAR* InErrorCtxt)
+int SetPropertyValue(const UStruct* InStruct, void* InStructData, PyObject* InValue, const FGeneratedWrappedProperty& InPropDef, const char *InAttributeName, const FPropertyAccessChangeNotify* InChangeNotify, const uint64 InReadOnlyFlags, const bool InOwnerIsTemplate, const TCHAR* InErrorCtxt)
 {
 	// Has this property been deprecated?
 	if (InStruct && InPropDef.Prop && InPropDef.DeprecationMessage.IsSet())
@@ -1205,7 +1205,7 @@ int SetPropertyValue(const UStruct* InStruct, void* InStructData, PyObject* InVa
 		}
 	}
 
-	return PyUtil::SetPropertyValue(InStruct, InStructData, InValue, InPropDef.Prop, InAttributeName, InChangeOwner, InReadOnlyFlags, InOwnerIsTemplate, InErrorCtxt);
+	return PyUtil::SetPropertyValue(InStruct, InStructData, InValue, InPropDef.Prop, InAttributeName, InChangeNotify, InReadOnlyFlags, InOwnerIsTemplate, InErrorCtxt);
 }
 
 FString BuildFunctionDocString(const UFunction* InFunc, const FString& InFuncPythonName, const TArray<FGeneratedWrappedMethodParameter>& InInputParams, const TArray<FGeneratedWrappedMethodParameter>& InOutputParams, const bool* InStaticOverride)
diff --git a/Engine/Plugins/Experimental/PythonScriptPlugin/Source/PythonScriptPlugin/Private/PyGenUtil.h b/Engine/Plugins/Experimental/PythonScriptPlugin/Source/PythonScriptPlugin/Private/PyGenUtil.h
--- a/Engine/Plugins/Experimental/PythonScriptPlugin/Source/PythonScriptPlugin/Private/PyGenUtil.h
+++ b/Engine/Plugins/Experimental/PythonScriptPlugin/Source/PythonScriptPlugin/Private/PyGenUtil.h
@@ -10,14 +10,13 @@
 #include "UObject/ObjectMacros.h"
 #include "UObject/WeakObjectPtr.h"
 #include "UObject/WeakObjectPtrTemplates.h"
+#include "UObject/PropertyAccessUtil.h"
 
 #if WITH_PYTHON
 
 struct FOutParmRec;
 struct FPyWrapperBaseMetaData;
 
-class FPyWrapperOwnerContext;
-
 struct FReportPythonGenerationIssue_Private
 {
 	static FORCEINLINE void ConditionalBreak(const ELogVerbosity::Type InVerbosity)
@@ -858,11 +857,11 @@ namespace PyGenUtil
 		FPythonizeTooltipContext()
 			: Prop(nullptr)
 			, Func(nullptr)
-			, ReadOnlyFlags(CPF_BlueprintReadOnly | CPF_EditConst)
+			, ReadOnlyFlags(PropertyAccessUtil::RuntimeReadOnlyFlags)
 		{
 		}
 
-		explicit FPythonizeTooltipContext(const UProperty* InProp, const uint64 InReadOnlyFlags = CPF_BlueprintReadOnly | CPF_EditConst);
+		explicit FPythonizeTooltipContext(const UProperty* InProp, const uint64 InReadOnlyFlags = PropertyAccessUtil::RuntimeReadOnlyFlags);
 
 		explicit FPythonizeTooltipContext(const UFunction* InFunc, const TSet<FName>& InParamsToIgnore = TSet<FName>());
 
@@ -910,10 +909,10 @@ namespace PyGenUtil
 	bool InvokePythonCallableFromUnrealFunctionThunk(FPyObjectPtr InSelf, PyObject* InCallable, const UFunction* InFunc, UObject* Context, FFrame& Stack, RESULT_DECL);
 
 	/** Get the current value of the given property from the given struct */
-	PyObject* GetPropertyValue(const UStruct* InStruct, void* InStructData, const FGeneratedWrappedProperty& InPropDef, const char *InAttributeName, PyObject* InOwnerPyObject, const TCHAR* InErrorCtxt);
+	PyObject* GetPropertyValue(const UStruct* InStruct, const void* InStructData, const FGeneratedWrappedProperty& InPropDef, const char *InAttributeName, PyObject* InOwnerPyObject, const TCHAR* InErrorCtxt);
 
 	/** Set the current value of the given property from the given struct */
-	int SetPropertyValue(const UStruct* InStruct, void* InStructData, PyObject* InValue, const FGeneratedWrappedProperty& InPropDef, const char *InAttributeName, const FPyWrapperOwnerContext& InChangeOwner, const uint64 InReadOnlyFlags, const bool InOwnerIsTemplate, const TCHAR* InErrorCtxt);
+	int SetPropertyValue(const UStruct* InStruct, void* InStructData, PyObject* InValue, const FGeneratedWrappedProperty& InPropDef, const char *InAttributeName, const FPropertyAccessChangeNotify* InChangeNotify, const uint64 InReadOnlyFlags, const bool InOwnerIsTemplate, const TCHAR* InErrorCtxt);
 
 	/** Build a Python doc string for the given function and arguments list */
 	FString BuildFunctionDocString(const UFunction* InFunc, const FString& InFuncPythonName, const TArray<FGeneratedWrappedMethodParameter>& InInputParams, const TArray<FGeneratedWrappedMethodParameter>& InOutputParams, const bool* InStaticOverride = nullptr);
@@ -964,7 +963,7 @@ namespace PyGenUtil
 	FString PythonizePropertyName(const FString& InName, const EPythonizeNameCase InNameCase);
 
 	/** Given a property tooltip, convert it to a doc string */
-	FString PythonizePropertyTooltip(const FString& InTooltip, const UProperty* InProp, const uint64 InReadOnlyFlags = CPF_BlueprintReadOnly | CPF_EditConst);
+	FString PythonizePropertyTooltip(const FString& InTooltip, const UProperty* InProp, const uint64 InReadOnlyFlags = PropertyAccessUtil::RuntimeReadOnlyFlags);
 
 	/** Given a function tooltip, convert it to a doc string */
 	FString PythonizeFunctionTooltip(const FString& InTooltip, const UFunction* InFunc, const TSet<FName>& ParamsToIgnore = TSet<FName>());
@@ -1045,7 +1044,7 @@ namespace PyGenUtil
 	void AppendPropertyPythonType(const UProperty* InProp, FString& OutStr);
 
 	/** Append the given Python property read /write state to the given string */
-	void AppendPropertyPythonReadWriteState(const UProperty* InProp, FString& OutStr, const uint64 InReadOnlyFlags = CPF_BlueprintReadOnly | CPF_EditConst);
+	void AppendPropertyPythonReadWriteState(const UProperty* InProp, FString& OutStr, const uint64 InReadOnlyFlags = PropertyAccessUtil::RuntimeReadOnlyFlags);
 
 	/** Get the tooltip for the given field */
 	FString GetFieldTooltip(const UField* InField);
diff --git a/Engine/Plugins/Experimental/PythonScriptPlugin/Source/PythonScriptPlugin/Private/PyUtil.cpp b/Engine/Plugins/Experimental/PythonScriptPlugin/Source/PythonScriptPlugin/Private/PyUtil.cpp
--- a/Engine/Plugins/Experimental/PythonScriptPlugin/Source/PythonScriptPlugin/Private/PyUtil.cpp
+++ b/Engine/Plugins/Experimental/PythonScriptPlugin/Source/PythonScriptPlugin/Private/PyUtil.cpp
@@ -710,13 +710,20 @@ UObject* GetOwnerObject(PyObject* InPyObj)
 	return nullptr;
 }
 
-PyObject* GetPropertyValue(const UStruct* InStruct, void* InStructData, const UProperty* InProp, const char *InAttributeName, PyObject* InOwnerPyObject, const TCHAR* InErrorCtxt)
+PyObject* GetPropertyValue(const UStruct* InStruct, const void* InStructData, const UProperty* InProp, const char *InAttributeName, PyObject* InOwnerPyObject, const TCHAR* InErrorCtxt)
 {
 	if (InStruct && InProp && ensureAlways(InStructData))
 	{
-		if (!InProp->HasAnyPropertyFlags(CPF_Edit | CPF_BlueprintVisible | CPF_BlueprintAssignable))
+		const EPropertyAccessResultFlags AccessResult = PropertyAccessUtil::CanGetPropertyValue(InProp);
+		if (EnumHasAnyFlags(AccessResult, EPropertyAccessResultFlags::PermissionDenied))
 		{
-			SetPythonError(PyExc_Exception, InErrorCtxt, *FString::Printf(TEXT("Property '%s' for attribute '%s' on '%s' is protected and cannot be read"), *InProp->GetName(), UTF8_TO_TCHAR(InAttributeName), *InStruct->GetName()));
+			if (EnumHasAnyFlags(AccessResult, EPropertyAccessResultFlags::AccessProtected))
+			{
+				SetPythonError(PyExc_Exception, InErrorCtxt, *FString::Printf(TEXT("Property '%s' for attribute '%s' on '%s' is protected and cannot be read"), *InProp->GetName(), UTF8_TO_TCHAR(InAttributeName), *InStruct->GetName()));
+				return nullptr;
+			}
+
+			SetPythonError(PyExc_Exception, InErrorCtxt, *FString::Printf(TEXT("Property '%s' for attribute '%s' on '%s' cannot be read"), *InProp->GetName(), UTF8_TO_TCHAR(InAttributeName), *InStruct->GetName()));
 			return nullptr;
 		}
 
@@ -732,7 +739,7 @@ PyObject* GetPropertyValue(const UStruct* InStruct, void* InStructData, const UP
 	Py_RETURN_NONE;
 }
 
-int SetPropertyValue(const UStruct* InStruct, void* InStructData, PyObject* InValue, const UProperty* InProp, const char *InAttributeName, const FPyWrapperOwnerContext& InChangeOwner, const uint64 InReadOnlyFlags, const bool InOwnerIsTemplate, const TCHAR* InErrorCtxt)
+int SetPropertyValue(const UStruct* InStruct, void* InStructData, PyObject* InValue, const UProperty* InProp, const char *InAttributeName, const FPropertyAccessChangeNotify* InChangeNotify, const uint64 InReadOnlyFlags, const bool InOwnerIsTemplate, const TCHAR* InErrorCtxt)
 {
 	if (!InValue)
 	{
@@ -742,36 +749,38 @@ int SetPropertyValue(const UStruct* InStruct, void* InStructData, PyObject* InVa
 
 	if (InStruct && InProp && ensureAlways(InStructData))
 	{
-		if (!InProp->HasAnyPropertyFlags(CPF_Edit | CPF_BlueprintVisible | CPF_BlueprintAssignable))
+		const EPropertyAccessResultFlags AccessResult = PropertyAccessUtil::CanSetPropertyValue(InProp, InReadOnlyFlags, InOwnerIsTemplate);
+		if (EnumHasAnyFlags(AccessResult, EPropertyAccessResultFlags::PermissionDenied))
 		{
-			SetPythonError(PyExc_Exception, InErrorCtxt, *FString::Printf(TEXT("Property '%s' for attribute '%s' on '%s' is protected and cannot be set"), *InProp->GetName(), UTF8_TO_TCHAR(InAttributeName), *InStruct->GetName()));
-			return -1;
-		}
+			if (EnumHasAnyFlags(AccessResult, EPropertyAccessResultFlags::AccessProtected))
+			{
+				SetPythonError(PyExc_Exception, InErrorCtxt, *FString::Printf(TEXT("Property '%s' for attribute '%s' on '%s' is protected and cannot be set"), *InProp->GetName(), UTF8_TO_TCHAR(InAttributeName), *InStruct->GetName()));
+				return -1;
+			}
 
-		if (InOwnerIsTemplate)
-		{
-			if (InProp->HasAnyPropertyFlags(CPF_DisableEditOnTemplate))
+			if (EnumHasAnyFlags(AccessResult, EPropertyAccessResultFlags::CannotEditTemplate))
 			{
 				SetPythonError(PyExc_Exception, InErrorCtxt, *FString::Printf(TEXT("Property '%s' for attribute '%s' on '%s' cannot be edited on templates"), *InProp->GetName(), UTF8_TO_TCHAR(InAttributeName), *InStruct->GetName()));
 				return -1;
 			}
-		}
-		else
-		{
-			if (InProp->HasAnyPropertyFlags(CPF_DisableEditOnInstance))
+
+			if (EnumHasAnyFlags(AccessResult, EPropertyAccessResultFlags::CannotEditInstance))
 			{
 				SetPythonError(PyExc_Exception, InErrorCtxt, *FString::Printf(TEXT("Property '%s' for attribute '%s' on '%s' cannot be edited on instances"), *InProp->GetName(), UTF8_TO_TCHAR(InAttributeName), *InStruct->GetName()));
 				return -1;
 			}
-		}
 
-		if (InProp->HasAnyPropertyFlags(InReadOnlyFlags))
-		{
-			SetPythonError(PyExc_Exception, InErrorCtxt, *FString::Printf(TEXT("Property '%s' for attribute '%s' on '%s' is read-only and cannot be set"), *InProp->GetName(), UTF8_TO_TCHAR(InAttributeName), *InStruct->GetName()));
+			if (EnumHasAnyFlags(AccessResult, EPropertyAccessResultFlags::ReadOnly))
+			{
+				SetPythonError(PyExc_Exception, InErrorCtxt, *FString::Printf(TEXT("Property '%s' for attribute '%s' on '%s' is read-only and cannot be set"), *InProp->GetName(), UTF8_TO_TCHAR(InAttributeName), *InStruct->GetName()));
+				return -1;
+			}
+
+			SetPythonError(PyExc_Exception, InErrorCtxt, *FString::Printf(TEXT("Property '%s' for attribute '%s' on '%s' cannot be set"), *InProp->GetName(), UTF8_TO_TCHAR(InAttributeName), *InStruct->GetName()));
 			return -1;
 		}
 
-		if (!PyConversion::NativizeProperty_InContainer(InValue, InProp, InStructData, 0, InChangeOwner))
+		if (!PyConversion::NativizeProperty_InContainer(InValue, InProp, InStructData, 0, InChangeNotify))
 		{
 			SetPythonError(PyExc_TypeError, InErrorCtxt, *FString::Printf(TEXT("Failed to convert type '%s' to property '%s' (%s) for attribute '%s' on '%s'"), *GetFriendlyTypename(InValue), *InProp->GetName(), *InProp->GetClass()->GetName(), UTF8_TO_TCHAR(InAttributeName), *InStruct->GetName()));
 			return -1;
diff --git a/Engine/Plugins/Experimental/PythonScriptPlugin/Source/PythonScriptPlugin/Private/PyUtil.h b/Engine/Plugins/Experimental/PythonScriptPlugin/Source/PythonScriptPlugin/Private/PyUtil.h
--- a/Engine/Plugins/Experimental/PythonScriptPlugin/Source/PythonScriptPlugin/Private/PyUtil.h
+++ b/Engine/Plugins/Experimental/PythonScriptPlugin/Source/PythonScriptPlugin/Private/PyUtil.h
@@ -12,7 +12,7 @@ DECLARE_LOG_CATEGORY_EXTERN(LogPython, Log, All);
 
 #if WITH_PYTHON
 
-class FPyWrapperOwnerContext;
+struct FPropertyAccessChangeNotify;
 
 /** Cast a function pointer to PyCFunction (via a void* to avoid a compiler warning) */
 #define PyCFunctionCast(FUNCPTR) (PyCFunction)(void*)(FUNCPTR)
@@ -229,10 +229,10 @@ namespace PyUtil
 	UObject* GetOwnerObject(PyObject* InPyObj);
 
 	/** Get the current value of the given property from the given struct */
-	PyObject* GetPropertyValue(const UStruct* InStruct, void* InStructData, const UProperty* InProp, const char *InAttributeName, PyObject* InOwnerPyObject, const TCHAR* InErrorCtxt);
+	PyObject* GetPropertyValue(const UStruct* InStruct, const void* InStructData, const UProperty* InProp, const char *InAttributeName, PyObject* InOwnerPyObject, const TCHAR* InErrorCtxt);
 
 	/** Set the current value of the given property from the given struct */
-	int SetPropertyValue(const UStruct* InStruct, void* InStructData, PyObject* InValue, const UProperty* InProp, const char *InAttributeName, const FPyWrapperOwnerContext& InChangeOwner, const uint64 InReadOnlyFlags, const bool InOwnerIsTemplate, const TCHAR* InErrorCtxt);
+	int SetPropertyValue(const UStruct* InStruct, void* InStructData, PyObject* InValue, const UProperty* InProp, const char *InAttributeName, const FPropertyAccessChangeNotify* InChangeNotify, const uint64 InReadOnlyFlags, const bool InOwnerIsTemplate, const TCHAR* InErrorCtxt);
 
 	/**
 	 * Check to see if the given object implements a length function.
diff --git a/Engine/Plugins/Experimental/PythonScriptPlugin/Source/PythonScriptPlugin/Private/PyWrapperObject.cpp b/Engine/Plugins/Experimental/PythonScriptPlugin/Source/PythonScriptPlugin/Private/PyWrapperObject.cpp
--- a/Engine/Plugins/Experimental/PythonScriptPlugin/Source/PythonScriptPlugin/Private/PyWrapperObject.cpp
+++ b/Engine/Plugins/Experimental/PythonScriptPlugin/Source/PythonScriptPlugin/Private/PyWrapperObject.cpp
@@ -132,8 +132,8 @@ int FPyWrapperObject::SetPropertyValue(FPyWrapperObject* InSelf, PyObject* InVal
 		return -1;
 	}
 
-	const FPyWrapperOwnerContext ChangeOwner = InNotifyChange ? FPyWrapperOwnerContext((PyObject*)InSelf, InPropDef.Prop) : FPyWrapperOwnerContext();
-	return PyGenUtil::SetPropertyValue(InSelf->ObjectInstance->GetClass(), InSelf->ObjectInstance, InValue, InPropDef, InPythonAttrName, ChangeOwner, InReadOnlyFlags, InSelf->ObjectInstance->IsTemplate() || InSelf->ObjectInstance->IsAsset(), *PyUtil::GetErrorContext(InSelf));
+	const TUniquePtr<FPropertyAccessChangeNotify> ChangeNotify = InNotifyChange ? FPyWrapperOwnerContext((PyObject*)InSelf, InPropDef.Prop).BuildChangeNotify() : nullptr;
+	return PyGenUtil::SetPropertyValue(InSelf->ObjectInstance->GetClass(), InSelf->ObjectInstance, InValue, InPropDef, InPythonAttrName, ChangeNotify.Get(), InReadOnlyFlags, PropertyAccessUtil::IsObjectTemplate(InSelf->ObjectInstance), *PyUtil::GetErrorContext(InSelf));
 }
 
 PyObject* FPyWrapperObject::CallGetterFunction(FPyWrapperObject* InSelf, const PyGenUtil::FGeneratedWrappedFunction& InFuncDef)
@@ -854,7 +854,7 @@ PyTypeObject InitializePyWrapperObjectType()
 				WrappedPropDef.SetProperty(ResolvedProp);
 			}
 
-			const int Result = FPyWrapperObject::SetPropertyValue(InSelf, PyValueObj, WrappedPropDef, TCHAR_TO_UTF8(*Name.ToString()), /*InNotifyChange*/true, CPF_EditConst);
+			const int Result = FPyWrapperObject::SetPropertyValue(InSelf, PyValueObj, WrappedPropDef, TCHAR_TO_UTF8(*Name.ToString()), /*InNotifyChange*/true, PropertyAccessUtil::EditorReadOnlyFlags);
 			if (Result != 0)
 			{
 				return nullptr;
diff --git a/Engine/Plugins/Experimental/PythonScriptPlugin/Source/PythonScriptPlugin/Private/PyWrapperObject.h b/Engine/Plugins/Experimental/PythonScriptPlugin/Source/PythonScriptPlugin/Private/PyWrapperObject.h
--- a/Engine/Plugins/Experimental/PythonScriptPlugin/Source/PythonScriptPlugin/Private/PyWrapperObject.h
+++ b/Engine/Plugins/Experimental/PythonScriptPlugin/Source/PythonScriptPlugin/Private/PyWrapperObject.h
@@ -5,6 +5,7 @@
 #include "PyWrapperBase.h"
 #include "UObject/Object.h"
 #include "UObject/Class.h"
+#include "UObject/PropertyAccessUtil.h"
 #include "PyWrapperObject.generated.h"
 
 #if WITH_PYTHON
@@ -46,7 +47,7 @@ struct FPyWrapperObject : public FPyWrapperBase
 	static PyObject* GetPropertyValue(FPyWrapperObject* InSelf, const PyGenUtil::FGeneratedWrappedProperty& InPropDef, const char* InPythonAttrName);
 
 	/** Set a property value on this instance (called via generated code) */
-	static int SetPropertyValue(FPyWrapperObject* InSelf, PyObject* InValue, const PyGenUtil::FGeneratedWrappedProperty& InPropDef, const char* InPythonAttrName, const bool InNotifyChange = false, const uint64 InReadOnlyFlags = CPF_EditConst | CPF_BlueprintReadOnly);
+	static int SetPropertyValue(FPyWrapperObject* InSelf, PyObject* InValue, const PyGenUtil::FGeneratedWrappedProperty& InPropDef, const char* InPythonAttrName, const bool InNotifyChange = false, const uint64 InReadOnlyFlags = PropertyAccessUtil::RuntimeReadOnlyFlags);
 
 	/** Call a named getter function on this class using the given instance (called via generated code) */
 	static PyObject* CallGetterFunction(FPyWrapperObject* InSelf, const PyGenUtil::FGeneratedWrappedFunction& InFuncDef);
diff --git a/Engine/Plugins/Experimental/PythonScriptPlugin/Source/PythonScriptPlugin/Private/PyWrapperOwnerContext.cpp b/Engine/Plugins/Experimental/PythonScriptPlugin/Source/PythonScriptPlugin/Private/PyWrapperOwnerContext.cpp
--- a/Engine/Plugins/Experimental/PythonScriptPlugin/Source/PythonScriptPlugin/Private/PyWrapperOwnerContext.cpp
+++ b/Engine/Plugins/Experimental/PythonScriptPlugin/Source/PythonScriptPlugin/Private/PyWrapperOwnerContext.cpp
@@ -1,6 +1,8 @@
 // Copyright 1998-2019 Epic Games, Inc. All Rights Reserved.
 
 #include "PyWrapperOwnerContext.h"
+#include "PyWrapperObject.h"
+#include "PyWrapperStruct.h"
 
 #if WITH_PYTHON
 
@@ -50,4 +52,67 @@ void FPyWrapperOwnerContext::AssertValidConversionMethod(const EPyConversionMeth
 	::AssertValidPyConversionOwner(GetOwnerObject(), InMethod);
 }
 
+TUniquePtr<FPropertyAccessChangeNotify> FPyWrapperOwnerContext::BuildChangeNotify() const
+{
+#if WITH_EDITOR
+	TUniquePtr<FPropertyAccessChangeNotify> ChangeNotify = MakeUnique<FPropertyAccessChangeNotify>();
+
+	auto AppendOwnerPropertyToChain = [&ChangeNotify](const FPyWrapperOwnerContext& InOwnerContext) -> bool
+	{
+		const UProperty* LeafProp = nullptr;
+		if (PyObject_IsInstance(InOwnerContext.GetOwnerObject(), (PyObject*)&PyWrapperObjectType) == 1 || PyObject_IsInstance(InOwnerContext.GetOwnerObject(), (PyObject*)&PyWrapperStructType) == 1)
+		{
+			LeafProp = InOwnerContext.GetOwnerProperty();
+		}
+
+		if (LeafProp)
+		{
+			ChangeNotify->ChangedPropertyChain.AddHead(const_cast<UProperty*>(LeafProp));
+			return true;
+		}
+
+		return false;
+	};
+
+	FPyWrapperOwnerContext OwnerContext = *this;
+	while (OwnerContext.HasOwner() && AppendOwnerPropertyToChain(OwnerContext))
+	{
+		PyObject* PyObj = OwnerContext.GetOwnerObject();
+
+		if (PyObj == GetOwnerObject())
+		{
+			ChangeNotify->ChangedPropertyChain.SetActivePropertyNode(ChangeNotify->ChangedPropertyChain.GetHead()->GetValue());
+		}
+
+		if (PyObject_IsInstance(PyObj, (PyObject*)&PyWrapperObjectType) == 1)
+		{
+			// Found an object, this is the end of the chain
+			ChangeNotify->ChangedObject = ((FPyWrapperObject*)PyObj)->ObjectInstance;
+			ChangeNotify->ChangedPropertyChain.SetActiveMemberPropertyNode(ChangeNotify->ChangedPropertyChain.GetHead()->GetValue());
+			break;
+		}
+
+		if (PyObject_IsInstance(PyObj, (PyObject*)&PyWrapperStructType) == 1)
+		{
+			// Found a struct, recurse up the chain
+			OwnerContext = ((FPyWrapperStruct*)PyObj)->OwnerContext;
+			continue;
+		}
+
+		// Unknown object type - just bail
+		break;
+	}
+
+	// If we didn't find an object in the chain then we can't emit notifications
+	if (!ChangeNotify->ChangedObject)
+	{
+		ChangeNotify.Reset();
+	}
+
+	return ChangeNotify;
+#else
+	return nullptr;
+#endif
+}
+
 #endif	// WITH_PYTHON
diff --git a/Engine/Plugins/Experimental/PythonScriptPlugin/Source/PythonScriptPlugin/Private/PyWrapperOwnerContext.h b/Engine/Plugins/Experimental/PythonScriptPlugin/Source/PythonScriptPlugin/Private/PyWrapperOwnerContext.h
--- a/Engine/Plugins/Experimental/PythonScriptPlugin/Source/PythonScriptPlugin/Private/PyWrapperOwnerContext.h
+++ b/Engine/Plugins/Experimental/PythonScriptPlugin/Source/PythonScriptPlugin/Private/PyWrapperOwnerContext.h
@@ -5,6 +5,7 @@
 #include "IncludePython.h"
 #include "PyPtr.h"
 #include "PyConversionMethod.h"
+#include "UObject/PropertyAccessUtil.h"
 
 #if WITH_PYTHON
 
@@ -36,6 +37,9 @@ public:
 	/** Assert that the given conversion method is valid for this owner context */
 	void AssertValidConversionMethod(const EPyConversionMethod InMethod) const;
 
+	/** Build the property change notify that corresponds to this owner context, or null if this owner context shouldn't emit change notifications */
+	TUniquePtr<FPropertyAccessChangeNotify> BuildChangeNotify() const;
+
 private:
 	/** The Python object that owns the instance being wrapped (if any) */
 	FPyObjectPtr OwnerObject;
diff --git a/Engine/Plugins/Experimental/PythonScriptPlugin/Source/PythonScriptPlugin/Private/PyWrapperStruct.cpp b/Engine/Plugins/Experimental/PythonScriptPlugin/Source/PythonScriptPlugin/Private/PyWrapperStruct.cpp
--- a/Engine/Plugins/Experimental/PythonScriptPlugin/Source/PythonScriptPlugin/Private/PyWrapperStruct.cpp
+++ b/Engine/Plugins/Experimental/PythonScriptPlugin/Source/PythonScriptPlugin/Private/PyWrapperStruct.cpp
@@ -277,7 +277,7 @@ FPyWrapperStruct* FPyWrapperStruct::CastPyObject(PyObject* InPyObject, PyTypeObj
 					break;
 				}
 
-				const int Result = PyUtil::SetPropertyValue(NewStruct->ScriptStruct, NewStruct->StructInstance, SequenceItem, InitParam.ParamProp, InitParam.ParamName.GetData(), FPyWrapperOwnerContext(), 0, false, *PyUtil::GetErrorContext(NewStruct.Get()));
+				const int Result = PyUtil::SetPropertyValue(NewStruct->ScriptStruct, NewStruct->StructInstance, SequenceItem, InitParam.ParamProp, InitParam.ParamName.GetData(), nullptr, 0, false, *PyUtil::GetErrorContext(NewStruct.Get()));
 				if (Result != 0)
 				{
 					return nullptr;
@@ -317,7 +317,7 @@ FPyWrapperStruct* FPyWrapperStruct::CastPyObject(PyObject* InPyObject, PyTypeObj
 			PyObject* MappingItem = PyMapping_GetItemString(InPyObject, (char*)InitParam.ParamName.GetData());
 			if (MappingItem)
 			{
-				const int Result = PyUtil::SetPropertyValue(NewStruct->ScriptStruct, NewStruct->StructInstance, MappingItem, InitParam.ParamProp, InitParam.ParamName.GetData(), FPyWrapperOwnerContext(), 0, false, *PyUtil::GetErrorContext(NewStruct.Get()));
+				const int Result = PyUtil::SetPropertyValue(NewStruct->ScriptStruct, NewStruct->StructInstance, MappingItem, InitParam.ParamProp, InitParam.ParamName.GetData(), nullptr, 0, false, *PyUtil::GetErrorContext(NewStruct.Get()));
 				if (Result != 0)
 				{
 					return nullptr;
@@ -444,11 +444,11 @@ int FPyWrapperStruct::SetPropertyValue(FPyWrapperStruct* InSelf, PyObject* InVal
 	bool OwnerIsTemplate = false;
 	if (const UObject* OwnerObject = PyUtil::GetOwnerObject((PyObject*)InSelf))
 	{
-		OwnerIsTemplate = OwnerObject->IsTemplate() || OwnerObject->IsAsset();
+		OwnerIsTemplate = PropertyAccessUtil::IsObjectTemplate(OwnerObject);
 	}
 
-	const FPyWrapperOwnerContext ChangeOwner = InNotifyChange ? FPyWrapperOwnerContext((PyObject*)InSelf, InPropDef.Prop) : FPyWrapperOwnerContext();
-	return PyGenUtil::SetPropertyValue(InSelf->ScriptStruct, InSelf->StructInstance, InValue, InPropDef, InPythonAttrName, ChangeOwner, InReadOnlyFlags, OwnerIsTemplate, *PyUtil::GetErrorContext(InSelf));
+	const TUniquePtr<FPropertyAccessChangeNotify> ChangeNotify = InNotifyChange ? FPyWrapperOwnerContext((PyObject*)InSelf, InPropDef.Prop).BuildChangeNotify() : nullptr;
+	return PyGenUtil::SetPropertyValue(InSelf->ScriptStruct, InSelf->StructInstance, InValue, InPropDef, InPythonAttrName, ChangeNotify.Get(), InReadOnlyFlags, OwnerIsTemplate, *PyUtil::GetErrorContext(InSelf));
 }
 
 int FPyWrapperStruct::CallMakeFunction_Impl(FPyWrapperStruct* InSelf, PyObject* InArgs, PyObject* InKwds, const PyGenUtil::FGeneratedWrappedFunction& InFuncDef)
@@ -631,7 +631,7 @@ PyObject* FPyWrapperStruct::CallOperatorFunction_Impl(FPyWrapperStruct* InSelf,
 		PyGenUtil::ApplyParamDefaults(FuncParams.GetStructMemory(), InOpFunc.AdditionalParams);
 		if (InOpFunc.OtherParam.ParamProp)
 		{
-			const FPyConversionResult RHSResult = PyConversion::NativizeProperty_InContainer(InRHS, InOpFunc.OtherParam.ParamProp, FuncParams.GetStructMemory(), 0, FPyWrapperOwnerContext(), PyConversion::ESetErrorState::No);
+			const FPyConversionResult RHSResult = PyConversion::NativizeProperty_InContainer(InRHS, InOpFunc.OtherParam.ParamProp, FuncParams.GetStructMemory(), 0, nullptr, PyConversion::ESetErrorState::No);
 			SetOptionalPyConversionResult(RHSResult, OutRHSConversionResult);
 
 			if (!RHSResult)
@@ -1093,7 +1093,7 @@ PyTypeObject InitializePyWrapperStructType()
 				WrappedPropDef.SetProperty(ResolvedProp);
 			}
 
-			const int Result = FPyWrapperStruct::SetPropertyValue(InSelf, PyValueObj, WrappedPropDef, TCHAR_TO_UTF8(*Name.ToString()), /*InNotifyChange*/true, CPF_EditConst);
+			const int Result = FPyWrapperStruct::SetPropertyValue(InSelf, PyValueObj, WrappedPropDef, TCHAR_TO_UTF8(*Name.ToString()), /*InNotifyChange*/true, PropertyAccessUtil::EditorReadOnlyFlags);
 			if (Result != 0)
 			{
 				return nullptr;
diff --git a/Engine/Plugins/Experimental/PythonScriptPlugin/Source/PythonScriptPlugin/Private/PyWrapperStruct.h b/Engine/Plugins/Experimental/PythonScriptPlugin/Source/PythonScriptPlugin/Private/PyWrapperStruct.h
--- a/Engine/Plugins/Experimental/PythonScriptPlugin/Source/PythonScriptPlugin/Private/PyWrapperStruct.h
+++ b/Engine/Plugins/Experimental/PythonScriptPlugin/Source/PythonScriptPlugin/Private/PyWrapperStruct.h
@@ -7,6 +7,7 @@
 #include "PyUtil.h"
 #include "PyConversion.h"
 #include "UObject/Class.h"
+#include "UObject/PropertyAccessUtil.h"
 #include "Templates/TypeCompatibleBytes.h"
 #include "PyWrapperStruct.generated.h"
 
@@ -64,7 +65,7 @@ struct FPyWrapperStruct : public FPyWrapperBase
 	static PyObject* GetPropertyValue(FPyWrapperStruct* InSelf, const PyGenUtil::FGeneratedWrappedProperty& InPropDef, const char* InPythonAttrName);
 
 	/** Set a property value on this instance (called via generated code) */
-	static int SetPropertyValue(FPyWrapperStruct* InSelf, PyObject* InValue, const PyGenUtil::FGeneratedWrappedProperty& InPropDef, const char* InPythonAttrName, const bool InNotifyChange = false, const uint64 InReadOnlyFlags = CPF_EditConst | CPF_BlueprintReadOnly);
+	static int SetPropertyValue(FPyWrapperStruct* InSelf, PyObject* InValue, const PyGenUtil::FGeneratedWrappedProperty& InPropDef, const char* InPythonAttrName, const bool InNotifyChange = false, const uint64 InReadOnlyFlags = PropertyAccessUtil::RuntimeReadOnlyFlags);
 
 	/** Call a make function on this instance (MakeStruct internal use only) */
 	static int CallMakeFunction_Impl(FPyWrapperStruct* InSelf, PyObject* InArgs, PyObject* InKwds, const PyGenUtil::FGeneratedWrappedFunction& InFuncDef);
diff --git a/Engine/Plugins/Experimental/PythonScriptPlugin/Source/PythonScriptPlugin/Private/PyWrapperTypeRegistry.cpp b/Engine/Plugins/Experimental/PythonScriptPlugin/Source/PythonScriptPlugin/Private/PyWrapperTypeRegistry.cpp
--- a/Engine/Plugins/Experimental/PythonScriptPlugin/Source/PythonScriptPlugin/Private/PyWrapperTypeRegistry.cpp
+++ b/Engine/Plugins/Experimental/PythonScriptPlugin/Source/PythonScriptPlugin/Private/PyWrapperTypeRegistry.cpp
@@ -2423,7 +2423,7 @@ void FPyWrapperTypeRegistry::GenerateStubCodeForWrappedType(PyTypeObject* PyType
 		// We use strict typing for return values to aid auto-complete (we also only care about the type and not the value, so structs can be default constructed)
 		static const uint32 PythonizeValueFlags = PyGenUtil::EPythonizeValueFlags::UseStrictTyping | PyGenUtil::EPythonizeValueFlags::DefaultConstructStructs;
 		const FString GetReturnValue = PyGenUtil::PythonizeDefaultValue(InGetSet.Prop.Prop, FString(), PythonizeValueFlags);
-		const bool bIsReadOnly = InGetSet.Prop.Prop->HasAnyPropertyFlags(CPF_BlueprintReadOnly | CPF_EditConst);
+		const bool bIsReadOnly = InGetSet.Prop.Prop->HasAnyPropertyFlags(PropertyAccessUtil::RuntimeReadOnlyFlags);
 		ExportGetSet(UTF8_TO_TCHAR(InGetSet.GetSetName.GetData()), UTF8_TO_TCHAR(InGetSet.GetSetDoc.GetData()), *GetReturnValue, bIsReadOnly);
 	};
 
diff --git a/Engine/Source/Editor/BlueprintGraph/Private/K2Node_EditorPropertyAccess.cpp b/Engine/Source/Editor/BlueprintGraph/Private/K2Node_EditorPropertyAccess.cpp
new file mode 100644
--- /dev/null
+++ b/Engine/Source/Editor/BlueprintGraph/Private/K2Node_EditorPropertyAccess.cpp
@@ -0,0 +1,332 @@
+// Copyright 1998-2019 Epic Games, Inc. All Rights Reserved.
+
+#include "K2Node_EditorPropertyAccess.h"
+#include "EdGraphSchema_K2.h"
+#include "K2Node_CallFunction.h"
+#include "K2Node_IfThenElse.h"
+#include "KismetCompiler.h"
+#include "Kismet2/BlueprintEditorUtils.h"
+#include "BlueprintNodeSpawner.h"
+#include "EditorCategoryUtils.h"
+#include "BlueprintActionDatabaseRegistrar.h"
+#include "Kismet/KismetSystemLibrary.h"
+
+#define LOCTEXT_NAMESPACE "K2Node_EditorPropertyAccess"
+
+namespace EditorPropertyAccessHelper
+{
+
+const FName ObjectPinName = "Object";
+const FName PropertyNamePinName = "PropertyName";
+const FName PropertyValuePinName = "PropertyValue";
+
+}
+
+void UK2Node_EditorPropertyAccessBase::AllocateDefaultPins()
+{
+	Super::AllocateDefaultPins();
+
+	// Add execution pins
+	CreatePin(EGPD_Input, UEdGraphSchema_K2::PC_Exec, UEdGraphSchema_K2::PN_Execute);
+	CreatePin(EGPD_Output, UEdGraphSchema_K2::PC_Exec, UEdGraphSchema_K2::PN_Then);
+
+	// Add Object pin
+	UEdGraphPin* ObjectPin = CreatePin(EGPD_Input, UEdGraphSchema_K2::PC_Object, UObject::StaticClass(), EditorPropertyAccessHelper::ObjectPinName);
+	SetPinToolTip(*ObjectPin, LOCTEXT("ObjectPinDescription", "The object you want to access a property value from"));
+
+	// Add Property Name pin
+	UEdGraphPin* PropertyNamePin = CreatePin(EGPD_Input, UEdGraphSchema_K2::PC_Name, EditorPropertyAccessHelper::PropertyNamePinName);
+	SetPinToolTip(*PropertyNamePin, LOCTEXT("PropertyNamePinDescription", "The name of the property to access from the object"));
+
+	// Add Property Value pin
+	AllocatePropertyValuePin();
+
+	// Add Result Pin
+	UEdGraphPin* ResultPin = CreatePin(EGPD_Output, UEdGraphSchema_K2::PC_Boolean, UEdGraphSchema_K2::PN_ReturnValue);
+	ResultPin->PinFriendlyName = LOCTEXT("ResultPinFriendlyName", "Success?");
+	SetPinToolTip(*ResultPin, LOCTEXT("ResultPinDescription", "Whether the property value was found"));
+}
+
+void UK2Node_EditorPropertyAccessBase::SetPinToolTip(UEdGraphPin& MutatablePin, const FText& PinDescription) const
+{
+	MutatablePin.PinToolTip = UEdGraphSchema_K2::TypeToText(MutatablePin.PinType).ToString();
+
+	const UEdGraphSchema_K2* K2Schema = Cast<const UEdGraphSchema_K2>(GetSchema());
+	if (K2Schema)
+	{
+		MutatablePin.PinToolTip += TEXT(" ");
+		MutatablePin.PinToolTip += K2Schema->GetPinDisplayName(&MutatablePin).ToString();
+	}
+
+	MutatablePin.PinToolTip += FString(TEXT("\n")) + PinDescription.ToString();
+}
+
+void UK2Node_EditorPropertyAccessBase::GetMenuActions(FBlueprintActionDatabaseRegistrar& ActionRegistrar) const
+{
+	// actions get registered under specific object-keys; the idea is that 
+	// actions might have to be updated (or deleted) if their object-key is  
+	// mutated (or removed)... here we use the node's class (so if the node 
+	// type disappears, then the action should go with it)
+	UClass* ActionKey = GetClass();
+	// to keep from needlessly instantiating a UBlueprintNodeSpawner, first   
+	// check to make sure that the registrar is looking for actions of this type
+	// (could be regenerating actions for a specific asset, and therefore the 
+	// registrar would only accept actions corresponding to that asset)
+	if (ActionRegistrar.IsOpenForRegistration(ActionKey))
+	{
+		UBlueprintNodeSpawner* NodeSpawner = UBlueprintNodeSpawner::Create(GetClass());
+		check(NodeSpawner != nullptr);
+
+		ActionRegistrar.AddBlueprintAction(ActionKey, NodeSpawner);
+	}
+}
+
+FText UK2Node_EditorPropertyAccessBase::GetMenuCategory() const
+{
+	return FEditorCategoryUtils::GetCommonCategory(FCommonEditorCategory::Utilities);
+}
+
+bool UK2Node_EditorPropertyAccessBase::CanPasteHere(const UEdGraph* TargetGraph) const
+{
+	bool bCanPaste = Super::CanPasteHere(TargetGraph);
+	if (bCanPaste)
+	{
+		bCanPaste &= FBlueprintEditorUtils::IsEditorUtilityBlueprint(FBlueprintEditorUtils::FindBlueprintForGraphChecked(TargetGraph));
+	}
+	return bCanPaste;
+}
+
+bool UK2Node_EditorPropertyAccessBase::IsActionFilteredOut(const FBlueprintActionFilter& Filter)
+{
+	bool bIsFilteredOut = Super::IsActionFilteredOut(Filter);
+	if (!bIsFilteredOut)
+	{
+		for (UEdGraph* TargetGraph : Filter.Context.Graphs)
+		{
+			bIsFilteredOut |= !CanPasteHere(TargetGraph);
+		}
+	}
+	return bIsFilteredOut;
+}
+
+bool UK2Node_EditorPropertyAccessBase::IsConnectionDisallowed(const UEdGraphPin* MyPin, const UEdGraphPin* OtherPin, FString& OutReason) const
+{
+	if (MyPin->PinName == EditorPropertyAccessHelper::PropertyValuePinName && MyPin->PinType.PinCategory == UEdGraphSchema_K2::PC_Wildcard)
+	{
+		if (OtherPin->PinType.PinCategory == UEdGraphSchema_K2::PC_Exec)
+		{
+			OutReason = TEXT("Cannot connect the property value to an execution pin");
+			return true;
+		}
+	}
+
+	return false;
+}
+
+UEdGraphPin* UK2Node_EditorPropertyAccessBase::GetThenPin()const
+{
+	UEdGraphPin* Pin = FindPinChecked(UEdGraphSchema_K2::PN_Then);
+	check(Pin->Direction == EGPD_Output);
+	return Pin;
+}
+
+UEdGraphPin* UK2Node_EditorPropertyAccessBase::GetObjectPin() const
+{
+	UEdGraphPin* Pin = FindPinChecked(EditorPropertyAccessHelper::ObjectPinName);
+	check(Pin->Direction == EGPD_Input);
+	return Pin;
+}
+
+UEdGraphPin* UK2Node_EditorPropertyAccessBase::GetPropertyNamePin() const
+{
+	UEdGraphPin* Pin = FindPinChecked(EditorPropertyAccessHelper::PropertyNamePinName);
+	check(Pin->Direction == EGPD_Input);
+	return Pin;
+}
+
+UEdGraphPin* UK2Node_EditorPropertyAccessBase::GetResultPin() const
+{
+	UEdGraphPin* Pin = FindPinChecked(UEdGraphSchema_K2::PN_ReturnValue);
+	check(Pin->Direction == EGPD_Output);
+	return Pin;
+}
+
+void UK2Node_EditorPropertyAccessBase::RefreshPropertyValuePin()
+{
+	UEdGraphPin* PropertyValuePin = GetPropertyValuePin();
+
+	FEdGraphPinType NewPinType;
+	if (PropertyValuePin->LinkedTo.Num() > 0)
+	{
+		NewPinType = PropertyValuePin->LinkedTo[0]->PinType;
+	}
+	else
+	{
+		NewPinType.PinCategory = UEdGraphSchema_K2::PC_Wildcard;
+	}
+
+	if (PropertyValuePin->PinType != NewPinType)
+	{
+		PropertyValuePin->PinType = NewPinType;
+
+		GetGraph()->NotifyGraphChanged();
+
+		UBlueprint* Blueprint = GetBlueprint();
+		if (!Blueprint->bBeingCompiled)
+		{
+			FBlueprintEditorUtils::MarkBlueprintAsModified(Blueprint);
+			Blueprint->BroadcastChanged();
+		}
+	}
+}
+
+FSlateIcon UK2Node_EditorPropertyAccessBase::GetIconAndTint(FLinearColor& OutColor) const
+{
+	OutColor = GetNodeTitleColor();
+	static FSlateIcon Icon("EditorStyle", "Kismet.AllClasses.FunctionIcon");
+	return Icon;
+}
+
+void UK2Node_EditorPropertyAccessBase::PostReconstructNode()
+{
+	Super::PostReconstructNode();
+
+	RefreshPropertyValuePin();
+}
+
+void UK2Node_EditorPropertyAccessBase::EarlyValidation(class FCompilerResultsLog& MessageLog) const
+{
+	Super::EarlyValidation(MessageLog);
+
+	const UEdGraphPin* ObjectPin = GetObjectPin();
+	const UEdGraphPin* PropertyNamePin = GetPropertyNamePin();
+	const UEdGraphPin* PropertyValuePin = GetPropertyValuePin();
+	if (!ObjectPin || !PropertyNamePin || !PropertyValuePin)
+	{
+		MessageLog.Error(*LOCTEXT("MissingPins", "Missing pins in @@").ToString(), this);
+		return;
+	}
+
+	if (ObjectPin->LinkedTo.Num() == 0 && !ObjectPin->DefaultObject)
+	{
+		MessageLog.Error(*LOCTEXT("UnsetObject", "No object set on @@").ToString(), this);
+	}
+
+	if (PropertyNamePin->LinkedTo.Num() == 0 && FName(*PropertyNamePin->DefaultValue).IsNone())
+	{
+		MessageLog.Error(*LOCTEXT("UnsetPropertyName", "No property name set on @@").ToString(), this);
+	}
+}
+
+void UK2Node_EditorPropertyAccessBase::NotifyPinConnectionListChanged(UEdGraphPin* Pin)
+{
+	Super::NotifyPinConnectionListChanged(Pin);
+
+	if (Pin->PinName == EditorPropertyAccessHelper::PropertyValuePinName)
+	{
+		RefreshPropertyValuePin();
+	}
+}
+
+void UK2Node_EditorPropertyAccessBase::ExpandNode(FKismetCompilerContext& CompilerContext, UEdGraph* SourceGraph)
+{
+	Super::ExpandNode(CompilerContext, SourceGraph);
+
+	const FName FunctionName = GetUnderlyingFunctionName();
+
+	// Add a CallFunction node for the underlying function
+	UK2Node_CallFunction* CallFunctionNode = CompilerContext.SpawnIntermediateNode<UK2Node_CallFunction>(this, SourceGraph);
+	CallFunctionNode->FunctionReference.SetExternalMember(FunctionName, UKismetSystemLibrary::StaticClass());
+	CallFunctionNode->AllocateDefaultPins();
+
+	// Re-wire the execution pins
+	CompilerContext.MovePinLinksToIntermediate(*GetExecPin(), *CallFunctionNode->GetExecPin());
+	CompilerContext.MovePinLinksToIntermediate(*GetThenPin(), *CallFunctionNode->GetThenPin());
+
+	// Re-wire the Object pin to the function input
+	{
+		UEdGraphPin* FunctionObjectPin = CallFunctionNode->FindPinChecked(TEXT("Object"));
+		CompilerContext.MovePinLinksToIntermediate(*GetObjectPin(), *FunctionObjectPin);
+	}
+
+	// Re-wire the PropertyName pin to the function input
+	{
+		UEdGraphPin* FunctionPropertyNamePin = CallFunctionNode->FindPinChecked(TEXT("PropertyName"));
+		CompilerContext.MovePinLinksToIntermediate(*GetPropertyNamePin(), *FunctionPropertyNamePin);
+	}
+
+	// Re-wire the ValuePtr pin to the function output, and copy its type over
+	{
+		UEdGraphPin* PropertyValuePin = GetPropertyValuePin();
+		UEdGraphPin* FunctionValuePtrPin = CallFunctionNode->FindPinChecked(TEXT("ValuePtr"));
+		CompilerContext.MovePinLinksToIntermediate(*PropertyValuePin, *FunctionValuePtrPin);
+		FunctionValuePtrPin->PinType = PropertyValuePin->PinType;
+	}
+
+	// Re-wire the result pin
+	{
+		UEdGraphPin* FunctionReturnPin = CallFunctionNode->FindPinChecked(UEdGraphSchema_K2::PN_ReturnValue);
+		CompilerContext.MovePinLinksToIntermediate(*GetResultPin(), *FunctionReturnPin);
+	}
+
+	// Disconnect this node
+	BreakAllNodeLinks();
+}
+
+void UK2Node_GetEditorProperty::AllocatePropertyValuePin()
+{
+	UEdGraphPin* PropertyValuePin = CreatePin(EGPD_Output, UEdGraphSchema_K2::PC_Wildcard, EditorPropertyAccessHelper::PropertyValuePinName);
+	SetPinToolTip(*PropertyValuePin, LOCTEXT("GetEditorProperty_PropertyValueDescription", "The returned property value, if found"));
+}
+
+UEdGraphPin* UK2Node_GetEditorProperty::GetPropertyValuePin() const
+{
+	UEdGraphPin* Pin = FindPinChecked(EditorPropertyAccessHelper::PropertyValuePinName);
+	check(Pin->Direction == EGPD_Output);
+	return Pin;
+}
+
+FText UK2Node_GetEditorProperty::GetNodeTitle(ENodeTitleType::Type TitleType) const
+{
+	return LOCTEXT("GetEditorProperty_NodeTitle", "Get Editor Property");
+}
+
+FText UK2Node_GetEditorProperty::GetTooltipText() const
+{
+	return LOCTEXT("GetEditorProperty_NodeTooltip", "Attempts to retrieve the value of a named property from the given object");
+}
+
+FName UK2Node_GetEditorProperty::GetUnderlyingFunctionName() const
+{
+	return GET_FUNCTION_NAME_CHECKED(UKismetSystemLibrary, GetEditorProperty);
+}
+
+void UK2Node_SetEditorProperty::AllocatePropertyValuePin()
+{
+	UEdGraphPin* PropertyValuePin = CreatePin(EGPD_Input, UEdGraphSchema_K2::PC_Wildcard, EditorPropertyAccessHelper::PropertyValuePinName);
+	SetPinToolTip(*PropertyValuePin, LOCTEXT("SetEditorProperty_PropertyValueDescription", "The property value to set"));
+}
+
+UEdGraphPin* UK2Node_SetEditorProperty::GetPropertyValuePin() const
+{
+	UEdGraphPin* Pin = FindPinChecked(EditorPropertyAccessHelper::PropertyValuePinName);
+	check(Pin->Direction == EGPD_Input);
+	return Pin;
+}
+
+FText UK2Node_SetEditorProperty::GetNodeTitle(ENodeTitleType::Type TitleType) const
+{
+	return LOCTEXT("SetEditorProperty_NodeTitle", "Set Editor Property");
+}
+
+FText UK2Node_SetEditorProperty::GetTooltipText() const
+{
+	return LOCTEXT("SetEditorProperty_NodeTooltip", "Attempts to set the value of a named property on the given object");
+}
+
+FName UK2Node_SetEditorProperty::GetUnderlyingFunctionName() const
+{
+	return GET_FUNCTION_NAME_CHECKED(UKismetSystemLibrary, SetEditorProperty);
+}
+
+#undef LOCTEXT_NAMESPACE
diff --git a/Engine/Source/Editor/BlueprintGraph/Private/K2Node_EditorPropertyAccess.h b/Engine/Source/Editor/BlueprintGraph/Private/K2Node_EditorPropertyAccess.h
new file mode 100644
--- /dev/null
+++ b/Engine/Source/Editor/BlueprintGraph/Private/K2Node_EditorPropertyAccess.h
@@ -0,0 +1,102 @@
+// Copyright 1998-2019 Epic Games, Inc. All Rights Reserved.
+
+
+#pragma once
+
+#include "CoreMinimal.h"
+#include "Textures/SlateIcon.h"
+#include "K2Node.h"
+#include "K2Node_EditorPropertyAccess.generated.h"
+
+class UEdGraph;
+class UK2Node_CallFunction;
+class FBlueprintActionDatabaseRegistrar;
+
+UCLASS(Abstract)
+class UK2Node_EditorPropertyAccessBase : public UK2Node
+{
+	GENERATED_BODY()
+
+public:
+	//~ Begin UEdGraphNode Interface.
+	virtual void AllocateDefaultPins() override;
+	virtual FSlateIcon GetIconAndTint(FLinearColor& OutColor) const override;
+	virtual void PostReconstructNode() override;
+	//~ End UEdGraphNode Interface.
+
+	//~ Begin UK2Node Interface
+	virtual void GetMenuActions(FBlueprintActionDatabaseRegistrar& ActionRegistrar) const override;
+	virtual FText GetMenuCategory() const override;
+	virtual bool CanPasteHere(const UEdGraph* TargetGraph) const override;
+	virtual bool IsActionFilteredOut(const class FBlueprintActionFilter& Filter) override;
+	virtual bool IsConnectionDisallowed(const UEdGraphPin* MyPin, const UEdGraphPin* OtherPin, FString& OutReason) const override;
+	virtual void EarlyValidation(class FCompilerResultsLog& MessageLog) const override;
+	virtual void NotifyPinConnectionListChanged(UEdGraphPin* Pin) override;
+	virtual void ExpandNode(class FKismetCompilerContext& CompilerContext, UEdGraph* SourceGraph) override;
+	//~ End UK2Node Interface
+
+	/** Get the then output pin */
+	UEdGraphPin* GetThenPin() const;
+	/** Get the object input pin */
+	UEdGraphPin* GetObjectPin() const;
+	/** Get the property name input pin */
+	UEdGraphPin* GetPropertyNamePin() const;
+	/** Get the property value output pin */
+	virtual UEdGraphPin* GetPropertyValuePin() const PURE_VIRTUAL(UK2Node_EditorPropertyAccessBase::GetPropertyValuePin, return nullptr;);
+	/** Get the result output pin */
+	UEdGraphPin* GetResultPin() const;
+
+protected:
+	/** Allocate the property value pin for this node */
+	virtual void AllocatePropertyValuePin() PURE_VIRTUAL(UK2Node_EditorPropertyAccessBase::AllocatePropertyValuePin, );
+
+	/** Get the function name from UKismetSystemLibrary that the CallFunction node should use for the underlying access function */
+	virtual FName GetUnderlyingFunctionName() const PURE_VIRTUAL(UK2Node_EditorPropertyAccessBase::GetUnderlyingFunctionName, return FName(););
+
+	/** Updates the type of the property value pin based on its connection type */
+	void RefreshPropertyValuePin();
+
+	/**
+	 * Takes the specified "MutatablePin" and sets its 'PinToolTip' field (according to the specified description)
+	 * 
+	 * @param   MutatablePin	The pin you want to set tool-tip text on
+	 * @param   PinDescription	A string describing the pin's purpose
+	 */
+	void SetPinToolTip(UEdGraphPin& MutatablePin, const FText& PinDescription) const;
+};
+
+UCLASS()
+class UK2Node_GetEditorProperty : public UK2Node_EditorPropertyAccessBase
+{
+	GENERATED_BODY()
+
+public:
+	//~ Begin UEdGraphNode Interface.
+	virtual FText GetNodeTitle(ENodeTitleType::Type TitleType) const override;
+	virtual FText GetTooltipText() const override;
+	//~ End UEdGraphNode Interface.
+
+	//~ Begin UK2Node_EditorPropertyAccessBase Interface.
+	virtual void AllocatePropertyValuePin() override;
+	virtual UEdGraphPin* GetPropertyValuePin() const override;
+	virtual FName GetUnderlyingFunctionName() const override;
+	//~ End UK2Node_EditorPropertyAccessBase Interface.
+};
+
+UCLASS()
+class UK2Node_SetEditorProperty : public UK2Node_EditorPropertyAccessBase
+{
+	GENERATED_BODY()
+
+public:
+	//~ Begin UEdGraphNode Interface.
+	virtual FText GetNodeTitle(ENodeTitleType::Type TitleType) const override;
+	virtual FText GetTooltipText() const override;
+	//~ End UEdGraphNode Interface.
+
+	//~ Begin UK2Node_EditorPropertyAccessBase Interface.
+	virtual void AllocatePropertyValuePin() override;
+	virtual UEdGraphPin* GetPropertyValuePin() const override;
+	virtual FName GetUnderlyingFunctionName() const override;
+	//~ End UK2Node_EditorPropertyAccessBase Interface.
+};
diff --git a/Engine/Source/Editor/DetailCustomizations/Private/ComponentTransformDetails.cpp b/Engine/Source/Editor/DetailCustomizations/Private/ComponentTransformDetails.cpp
--- a/Engine/Source/Editor/DetailCustomizations/Private/ComponentTransformDetails.cpp
+++ b/Engine/Source/Editor/DetailCustomizations/Private/ComponentTransformDetails.cpp
@@ -697,7 +697,7 @@ void FComponentTransformDetails::OnSetAbsoluteTransform(ETransformField::Type Tr
 
 	if (bBeganTransaction)
 	{
-		FPropertyChangedEvent PropertyChangedEvent(AbsoluteProperty, EPropertyChangeType::ValueSet, (const TArray<const UObject*>*)&ModifiedObjects);
+		FPropertyChangedEvent PropertyChangedEvent(AbsoluteProperty, EPropertyChangeType::ValueSet, MakeArrayView(ModifiedObjects));
 
 		for (UObject* Object : ModifiedObjects)
 		{
@@ -1027,7 +1027,7 @@ void FComponentTransformDetails::OnSetTransform(ETransformField::Type TransformF
 	bool bBeganTransaction = false;
 	TArray<UObject*> ModifiedObjects;
 
-	FPropertyChangedEvent PropertyChangedEvent(ValueProperty, !bCommitted ? EPropertyChangeType::Interactive : EPropertyChangeType::ValueSet, (const TArray<const UObject*>*)&ModifiedObjects);
+	FPropertyChangedEvent PropertyChangedEvent(ValueProperty, !bCommitted ? EPropertyChangeType::Interactive : EPropertyChangeType::ValueSet, MakeArrayView(ModifiedObjects));
 	FEditPropertyChain PropertyChain;
 
 	if (AxisProperty)
diff --git a/Engine/Source/Editor/DetailCustomizations/Private/MovieSceneEventParametersCustomization.cpp b/Engine/Source/Editor/DetailCustomizations/Private/MovieSceneEventParametersCustomization.cpp
--- a/Engine/Source/Editor/DetailCustomizations/Private/MovieSceneEventParametersCustomization.cpp
+++ b/Engine/Source/Editor/DetailCustomizations/Private/MovieSceneEventParametersCustomization.cpp
@@ -132,7 +132,7 @@ void FMovieSceneEventParametersCustomization::OnStructChanged(const FAssetData&
 		static_cast<FMovieSceneEventParameters*>(Value)->Reassign(NewStruct);
 	}
 
-	FPropertyChangedEvent BubbleChangeEvent(PropertyHandle->GetProperty(), EPropertyChangeType::ValueSet, nullptr);
+	FPropertyChangedEvent BubbleChangeEvent(PropertyHandle->GetProperty(), EPropertyChangeType::ValueSet);
 
 	// post notify
 	if (Hook)
@@ -163,7 +163,7 @@ void FMovieSceneEventParametersCustomization::OnEditStructChildContentsChanged()
 		static_cast<FMovieSceneEventParameters*>(Value)->OverwriteWith(EditStructData->GetStructMemory());
 	}
 
-	FPropertyChangedEvent BubbleChangeEvent(PropertyHandle->GetProperty(), EPropertyChangeType::ValueSet, nullptr);
+	FPropertyChangedEvent BubbleChangeEvent(PropertyHandle->GetProperty(), EPropertyChangeType::ValueSet);
 	PropertyUtilities->NotifyFinishedChangingProperties(BubbleChangeEvent);
 }
 
diff --git a/Engine/Source/Editor/PropertyEditor/Private/PropertyHandleImpl.cpp b/Engine/Source/Editor/PropertyEditor/Private/PropertyHandleImpl.cpp
--- a/Engine/Source/Editor/PropertyEditor/Private/PropertyHandleImpl.cpp
+++ b/Engine/Source/Editor/PropertyEditor/Private/PropertyHandleImpl.cpp
@@ -609,7 +609,7 @@ FPropertyAccess::Result FPropertyValueImpl::ImportText( const TArray<FObjectBase
 			FPropertyValueImpl::GenerateArrayIndexMapToObjectNode( ArrayIndicesPerObject[ObjectIndex], InPropertyNode );
 		}
 
-		FPropertyChangedEvent ChangeEvent(NodeProperty, bFinished ? EPropertyChangeType::ValueSet : EPropertyChangeType::Interactive, &TopLevelObjects);
+		FPropertyChangedEvent ChangeEvent(NodeProperty, bFinished ? EPropertyChangeType::ValueSet : EPropertyChangeType::Interactive, MakeArrayView(TopLevelObjects));
 		ChangeEvent.SetArrayIndexPerObject(ArrayIndicesPerObject);
 
 		// If PreEditChange was called, so should PostEditChange.
@@ -1201,7 +1201,7 @@ void FPropertyValueImpl::AddChild()
 				}
 			}
 
-			FPropertyChangedEvent ChangeEvent(NodeProperty, EPropertyChangeType::ArrayAdd, &TopLevelObjects);
+			FPropertyChangedEvent ChangeEvent(NodeProperty, EPropertyChangeType::ArrayAdd, MakeArrayView(TopLevelObjects));
 			ChangeEvent.SetArrayIndexPerObject(ArrayIndicesPerObject);
 			ChangeEvent.SetInstancesChangedResultPerArchetype(PropagationResultPerObject);
 
@@ -1364,7 +1364,7 @@ void FPropertyValueImpl::ClearChildren()
 				}
 			}
 
-			FPropertyChangedEvent ChangeEvent(NodeProperty, EPropertyChangeType::ArrayClear, &TopLevelObjects);
+			FPropertyChangedEvent ChangeEvent(NodeProperty, EPropertyChangeType::ArrayClear, MakeArrayView(TopLevelObjects));
 
 			if ( bNotifiedPreChange )
 			{
@@ -1449,7 +1449,7 @@ void FPropertyValueImpl::InsertChild( TSharedPtr<FPropertyNode> ChildNodeToInser
 			FPropertyValueImpl::GenerateArrayIndexMapToObjectNode(ArrayIndicesPerObject[ObjectIndex], ChildNodePtr );
 		}
 
-		FPropertyChangedEvent ChangeEvent(ParentNode->GetProperty(), EPropertyChangeType::ArrayAdd, &TopLevelObjects);
+		FPropertyChangedEvent ChangeEvent(ParentNode->GetProperty(), EPropertyChangeType::ArrayAdd, MakeArrayView(TopLevelObjects));
 		ChangeEvent.SetArrayIndexPerObject(ArrayIndicesPerObject);
 		ChangeEvent.SetInstancesChangedResultPerArchetype(PropagationResultPerObject);
 
@@ -1596,7 +1596,7 @@ void FPropertyValueImpl::DeleteChild( TSharedPtr<FPropertyNode> ChildNodeToDelet
 			}
 		}
 
-		FPropertyChangedEvent ChangeEvent(ParentNode->GetProperty(), EPropertyChangeType::ArrayRemove, &TopLevelObjects);
+		FPropertyChangedEvent ChangeEvent(ParentNode->GetProperty(), EPropertyChangeType::ArrayRemove, MakeArrayView(TopLevelObjects));
 		ChangeEvent.SetArrayIndexPerObject(ArrayIndicesPerObject);
 		ChangeEvent.SetInstancesChangedResultPerArchetype(PropagationResultPerObject);
 
@@ -1697,7 +1697,7 @@ void FPropertyValueImpl::SwapChildren( TSharedPtr<FPropertyNode> FirstChildNode,
 			}
 		}
 
-		FPropertyChangedEvent ChangeEvent(ParentNode->GetProperty(), EPropertyChangeType::Unspecified, &TopLevelObjects);
+		FPropertyChangedEvent ChangeEvent(ParentNode->GetProperty(), EPropertyChangeType::Unspecified, MakeArrayView(TopLevelObjects));
 		FirstChildNodePtr->NotifyPostChange(ChangeEvent, NotifyHook);
 		SecondChildNodePtr->NotifyPostChange(ChangeEvent, NotifyHook);
 
@@ -1763,7 +1763,7 @@ void FPropertyValueImpl::MoveElementTo(int32 OriginalIndex, int32 NewIndex)
 					FPropertyValueImpl::GenerateArrayIndexMapToObjectNode(ArrayIndicesPerObject[ObjectIndex], ChildNodePtr);
 				}
 
-				FPropertyChangedEvent ChangeEvent(ParentNode->GetProperty(), EPropertyChangeType::ArrayAdd, &TopLevelObjects);
+				FPropertyChangedEvent ChangeEvent(ParentNode->GetProperty(), EPropertyChangeType::ArrayAdd, MakeArrayView(TopLevelObjects));
 				ChangeEvent.SetArrayIndexPerObject(ArrayIndicesPerObject);
 
 				if (PropertyUtilities.IsValid())
@@ -1829,7 +1829,7 @@ void FPropertyValueImpl::MoveElementTo(int32 OriginalIndex, int32 NewIndex)
 					}
 				}
 
-				FPropertyChangedEvent ChangeEvent(NodeProperty, EPropertyChangeType::ArrayAdd, &TopLevelObjects);
+				FPropertyChangedEvent ChangeEvent(NodeProperty, EPropertyChangeType::ArrayAdd, MakeArrayView(TopLevelObjects));
 				ChangeEvent.SetArrayIndexPerObject(ArrayIndicesPerObject);
 
 				if (PropertyUtilities.IsValid())
@@ -1920,7 +1920,7 @@ void FPropertyValueImpl::MoveElementTo(int32 OriginalIndex, int32 NewIndex)
 				}
 			}
 
-			FPropertyChangedEvent ChangeEvent(ParentNode->GetProperty(), EPropertyChangeType::Unspecified, &TopLevelObjects);
+			FPropertyChangedEvent ChangeEvent(ParentNode->GetProperty(), EPropertyChangeType::Unspecified, MakeArrayView(TopLevelObjects));
 
 			if (PropertyUtilities.IsValid())
 			{
@@ -1989,7 +1989,7 @@ void FPropertyValueImpl::MoveElementTo(int32 OriginalIndex, int32 NewIndex)
 				}
 			}
 
-			FPropertyChangedEvent ChangeEvent(ParentNode->GetProperty(), EPropertyChangeType::Unspecified, &TopLevelObjects);
+			FPropertyChangedEvent ChangeEvent(ParentNode->GetProperty(), EPropertyChangeType::Unspecified, MakeArrayView(TopLevelObjects));
 			if (PropertyUtilities.IsValid())
 			{
 				ChildNodePtr->FixPropertiesInEvent(ChangeEvent);
@@ -2099,7 +2099,7 @@ void FPropertyValueImpl::DuplicateChild( TSharedPtr<FPropertyNode> ChildNodeToDu
 			FPropertyValueImpl::GenerateArrayIndexMapToObjectNode(ArrayIndicesPerObject[0], ChildNodePtr);
 		}
 
-		FPropertyChangedEvent ChangeEvent(ParentNode->GetProperty(), EPropertyChangeType::Duplicate, &TopLevelObjects);
+		FPropertyChangedEvent ChangeEvent(ParentNode->GetProperty(), EPropertyChangeType::Duplicate, MakeArrayView(TopLevelObjects));
 		ChangeEvent.SetArrayIndexPerObject(ArrayIndicesPerObject);
 		ChangeEvent.SetInstancesChangedResultPerArchetype(PropagationResultPerObject);
 
@@ -2978,7 +2978,7 @@ void FPropertyHandleBase::NotifyPostChange( EPropertyChangeType::Type ChangeType
 			}
 		}
 
-		FPropertyChangedEvent PropertyChangedEvent( PropertyNode->GetProperty(), ChangeType, &ObjectsBeingChanged );
+		FPropertyChangedEvent PropertyChangedEvent( PropertyNode->GetProperty(), ChangeType, MakeArrayView(ObjectsBeingChanged) );
 		PropertyNode->NotifyPostChange( PropertyChangedEvent, Implementation->GetNotifyHook());
 	}
 }
diff --git a/Engine/Source/Runtime/CoreUObject/Private/UObject/PropertyAccessUtil.cpp b/Engine/Source/Runtime/CoreUObject/Private/UObject/PropertyAccessUtil.cpp
new file mode 100644
--- /dev/null
+++ b/Engine/Source/Runtime/CoreUObject/Private/UObject/PropertyAccessUtil.cpp
@@ -0,0 +1,264 @@
+// Copyright 1998-2019 Epic Games, Inc. All Rights Reserved.
+
+#include "UObject/PropertyAccessUtil.h"
+#include "UObject/Object.h"
+#include "UObject/Class.h"
+
+namespace PropertyAccessUtil
+{
+
+EPropertyAccessResultFlags GetPropertyValue_Object(const UProperty* InProp, const UObject* InObject, void* InDestValue, const int32 InArrayIndex)
+{
+	check(InObject->IsA(InProp->GetOwnerClass()));
+	return GetPropertyValue_InContainer(InProp, InObject, InDestValue, InArrayIndex);
+}
+
+EPropertyAccessResultFlags GetPropertyValue_InContainer(const UProperty* InProp, const void* InContainerData, void* InDestValue, const int32 InArrayIndex)
+{
+	if (InArrayIndex == INDEX_NONE || InProp->ArrayDim == 1)
+	{
+		const void* SrcValue = InProp->ContainerPtrToValuePtr<void>(InContainerData);
+		return GetPropertyValue_DirectComplete(InProp, SrcValue, InDestValue);
+	}
+	else
+	{
+		check(InArrayIndex < InProp->ArrayDim);
+		const void* SrcValue = InProp->ContainerPtrToValuePtr<void>(InContainerData, InArrayIndex);
+		return GetPropertyValue_DirectSingle(InProp, SrcValue, InDestValue);
+	}
+}
+
+EPropertyAccessResultFlags GetPropertyValue_DirectSingle(const UProperty* InProp, const void* InSrcValue, void* InDestValue)
+{
+	EPropertyAccessResultFlags Result = CanGetPropertyValue(InProp);
+	if (Result != EPropertyAccessResultFlags::Success)
+	{
+		return Result;
+	}
+
+	return GetPropertyValue([InProp, InSrcValue, InDestValue]()
+	{
+		InProp->CopySingleValue(InDestValue, InSrcValue);
+		return true;
+	});
+}
+
+EPropertyAccessResultFlags GetPropertyValue_DirectComplete(const UProperty* InProp, const void* InSrcValue, void* InDestValue)
+{
+	EPropertyAccessResultFlags Result = CanGetPropertyValue(InProp);
+	if (Result != EPropertyAccessResultFlags::Success)
+	{
+		return Result;
+	}
+
+	return GetPropertyValue([InProp, InSrcValue, InDestValue]()
+	{
+		InProp->CopyCompleteValue(InDestValue, InSrcValue);
+		return true;
+	});
+}
+
+EPropertyAccessResultFlags GetPropertyValue(const FPropertyAccessGetFunc& InGetFunc)
+{
+	const bool bGetResult = InGetFunc();
+	return bGetResult
+		? EPropertyAccessResultFlags::Success
+		: EPropertyAccessResultFlags::ConversionFailed;
+}
+
+EPropertyAccessResultFlags CanGetPropertyValue(const UProperty* InProp)
+{
+	if (!InProp->HasAnyPropertyFlags(CPF_Edit | CPF_BlueprintVisible | CPF_BlueprintAssignable))
+	{
+		return EPropertyAccessResultFlags::PermissionDenied | EPropertyAccessResultFlags::AccessProtected;
+	}
+
+	return EPropertyAccessResultFlags::Success;
+}
+
+EPropertyAccessResultFlags SetPropertyValue_Object(const UProperty* InProp, UObject* InObject, const void* InSrcValue, const int32 InArrayIndex, const uint64 InReadOnlyFlags)
+{
+	check(InObject->IsA(InProp->GetOwnerClass()));
+	return SetPropertyValue_InContainer(InProp, InObject, InSrcValue, InArrayIndex, InReadOnlyFlags, IsObjectTemplate(InObject), [InProp, InObject]()
+	{
+		return BuildBasicChangeNotify(InProp, InObject);
+	});
+}
+
+EPropertyAccessResultFlags SetPropertyValue_InContainer(const UProperty* InProp, void* InContainerData, const void* InSrcValue, const int32 InArrayIndex, const uint64 InReadOnlyFlags, const bool InOwnerIsTemplate, const FPropertyAccessBuildChangeNotifyFunc& InBuildChangeNotifyFunc)
+{
+	if (InArrayIndex == INDEX_NONE || InProp->ArrayDim == 1)
+	{
+		void* DestValue = InProp->ContainerPtrToValuePtr<void>(InContainerData);
+		return SetPropertyValue_DirectComplete(InProp, InSrcValue, DestValue, InReadOnlyFlags, InOwnerIsTemplate, InBuildChangeNotifyFunc);
+	}
+	else
+	{
+		check(InArrayIndex < InProp->ArrayDim);
+		void* DestValue = InProp->ContainerPtrToValuePtr<void>(InContainerData, InArrayIndex);
+		return SetPropertyValue_DirectSingle(InProp, InSrcValue, DestValue, InReadOnlyFlags, InOwnerIsTemplate, InBuildChangeNotifyFunc);
+	}
+}
+
+EPropertyAccessResultFlags SetPropertyValue_DirectSingle(const UProperty* InProp, const void* InSrcValue, void* InDestValue, const uint64 InReadOnlyFlags, const bool InOwnerIsTemplate, const FPropertyAccessBuildChangeNotifyFunc& InBuildChangeNotifyFunc)
+{
+	EPropertyAccessResultFlags Result = CanSetPropertyValue(InProp, InReadOnlyFlags, InOwnerIsTemplate);
+	if (Result != EPropertyAccessResultFlags::Success)
+	{
+		return Result;
+	}
+
+	return SetPropertyValue([InProp, InSrcValue, InDestValue](const FPropertyAccessChangeNotify* InChangeNotify)
+	{
+		if (!InProp->Identical(InSrcValue, InDestValue))
+		{
+			EmitPreChangeNotify(InChangeNotify);
+			InProp->CopySingleValue(InDestValue, InSrcValue);
+			EmitPostChangeNotify(InChangeNotify);
+		}
+		return true;
+	}, InBuildChangeNotifyFunc);
+}
+
+EPropertyAccessResultFlags SetPropertyValue_DirectComplete(const UProperty* InProp, const void* InSrcValue, void* InDestValue, const uint64 InReadOnlyFlags, const bool InOwnerIsTemplate, const FPropertyAccessBuildChangeNotifyFunc& InBuildChangeNotifyFunc)
+{
+	EPropertyAccessResultFlags Result = CanSetPropertyValue(InProp, InReadOnlyFlags, InOwnerIsTemplate);
+	if (Result != EPropertyAccessResultFlags::Success)
+	{
+		return Result;
+	}
+
+	return SetPropertyValue([InProp, InSrcValue, InDestValue](const FPropertyAccessChangeNotify* InChangeNotify)
+	{
+		auto IdenticalComplete = [InProp, InSrcValue, InDestValue]()
+		{
+			bool bIsIdentical = true;
+			for (int32 Idx = 0; Idx < InProp->ArrayDim && bIsIdentical; ++Idx)
+			{
+				const void* SrcElemValue = static_cast<const uint8*>(InSrcValue) + (InProp->ElementSize * Idx);
+				const void* DestElemValue = static_cast<const uint8*>(InDestValue) + (InProp->ElementSize * Idx);
+				bIsIdentical &= InProp->Identical(SrcElemValue, DestElemValue);
+			}
+			return bIsIdentical;
+		};
+
+		if (!IdenticalComplete())
+		{
+			EmitPreChangeNotify(InChangeNotify);
+			InProp->CopyCompleteValue(InDestValue, InSrcValue);
+			EmitPostChangeNotify(InChangeNotify);
+		}
+		return true;
+	}, InBuildChangeNotifyFunc);
+}
+
+EPropertyAccessResultFlags SetPropertyValue(const FPropertyAccessSetFunc& InSetFunc, const FPropertyAccessBuildChangeNotifyFunc& InBuildChangeNotifyFunc)
+{
+	TUniquePtr<FPropertyAccessChangeNotify> ChangeNotify = InBuildChangeNotifyFunc();
+	const bool bSetResult = InSetFunc(ChangeNotify.Get());
+	return bSetResult
+		? EPropertyAccessResultFlags::Success
+		: EPropertyAccessResultFlags::ConversionFailed;
+}
+
+EPropertyAccessResultFlags CanSetPropertyValue(const UProperty* InProp, const uint64 InReadOnlyFlags, const bool InOwnerIsTemplate)
+{
+	if (!InProp->HasAnyPropertyFlags(CPF_Edit | CPF_BlueprintVisible | CPF_BlueprintAssignable))
+	{
+		return EPropertyAccessResultFlags::PermissionDenied | EPropertyAccessResultFlags::AccessProtected;
+	}
+
+	if (InOwnerIsTemplate)
+	{
+		if (InProp->HasAnyPropertyFlags(CPF_DisableEditOnTemplate))
+		{
+			return EPropertyAccessResultFlags::PermissionDenied | EPropertyAccessResultFlags::CannotEditTemplate;
+		}
+	}
+	else
+	{
+		if (InProp->HasAnyPropertyFlags(CPF_DisableEditOnInstance))
+		{
+			return EPropertyAccessResultFlags::PermissionDenied | EPropertyAccessResultFlags::CannotEditInstance;
+		}
+	}
+
+	if (InProp->HasAnyPropertyFlags(InReadOnlyFlags))
+	{
+		return EPropertyAccessResultFlags::PermissionDenied | EPropertyAccessResultFlags::ReadOnly;
+	}
+
+	return EPropertyAccessResultFlags::Success;
+}
+
+void EmitPreChangeNotify(const FPropertyAccessChangeNotify* InChangeNotify)
+{
+#if WITH_EDITOR
+	if (InChangeNotify)
+	{
+		check(InChangeNotify->ChangedObject);
+
+		// Notify that a change is about to occur
+		InChangeNotify->ChangedObject->PreEditChange(const_cast<FEditPropertyChain&>(InChangeNotify->ChangedPropertyChain));
+	}
+#endif
+}
+
+void EmitPostChangeNotify(const FPropertyAccessChangeNotify* InChangeNotify)
+{
+#if WITH_EDITOR
+	if (InChangeNotify)
+	{
+		check(InChangeNotify->ChangedObject);
+
+		// Notify that the change has occurred
+		FPropertyChangedEvent PropertyEvent(InChangeNotify->ChangedPropertyChain.GetActiveNode()->GetValue(), InChangeNotify->ChangeType, MakeArrayView(&InChangeNotify->ChangedObject, 1));
+		PropertyEvent.SetActiveMemberProperty(InChangeNotify->ChangedPropertyChain.GetActiveMemberNode()->GetValue());
+		FPropertyChangedChainEvent PropertyChainEvent(const_cast<FEditPropertyChain&>(InChangeNotify->ChangedPropertyChain), PropertyEvent);
+		InChangeNotify->ChangedObject->PostEditChangeChainProperty(PropertyChainEvent);
+	}
+#endif
+}
+
+TUniquePtr<FPropertyAccessChangeNotify> BuildBasicChangeNotify(const UProperty* InProp, const UObject* InObject)
+{
+	check(InObject->IsA(InProp->GetOwnerClass()));
+#if WITH_EDITOR
+	TUniquePtr<FPropertyAccessChangeNotify> ChangeNotify = MakeUnique<FPropertyAccessChangeNotify>();
+	ChangeNotify->ChangedObject = const_cast<UObject*>(InObject);
+	ChangeNotify->ChangedPropertyChain.AddHead(const_cast<UProperty*>(InProp));
+	ChangeNotify->ChangedPropertyChain.SetActivePropertyNode(const_cast<UProperty*>(InProp));
+	ChangeNotify->ChangedPropertyChain.SetActiveMemberPropertyNode(const_cast<UProperty*>(InProp));
+	return ChangeNotify;
+#else
+	return nullptr;
+#endif
+}
+
+bool IsObjectTemplate(const UObject* InObject)
+{
+	return InObject->IsTemplate() || InObject->IsAsset();
+}
+
+UProperty* FindPropertyByName(const FName InPropName, const UStruct* InStruct)
+{
+	UProperty* Prop = InStruct->FindPropertyByName(InPropName);
+
+	if (!Prop)
+	{
+		const FName NewPropName = UProperty::FindRedirectedPropertyName(const_cast<UStruct*>(InStruct), InPropName);
+		if (!NewPropName.IsNone())
+		{
+			Prop = InStruct->FindPropertyByName(NewPropName);
+		}
+	}
+
+	if (!Prop)
+	{
+		Prop = InStruct->CustomFindProperty(InPropName);
+	}
+
+	return Prop;
+}
+
+}
diff --git a/Engine/Source/Runtime/CoreUObject/Public/UObject/PropertyAccessUtil.h b/Engine/Source/Runtime/CoreUObject/Public/UObject/PropertyAccessUtil.h
new file mode 100644
--- /dev/null
+++ b/Engine/Source/Runtime/CoreUObject/Public/UObject/PropertyAccessUtil.h
@@ -0,0 +1,266 @@
+// Copyright 1998-2019 Epic Games, Inc. All Rights Reserved.
+
+#pragma once
+
+#include "CoreMinimal.h"
+#include "UObject/UnrealType.h"
+#include "Templates/UniquePtr.h"
+#include "Misc/EnumClassFlags.h"
+
+class UObject;
+class UStruct;
+class UProperty;
+
+/**
+ * Result flags from property access.
+ */
+enum class EPropertyAccessResultFlags : uint8
+{
+	/** The property was accessed successfully */
+	Success = 0,
+
+	/** The property could not be accessed due to a permission error (the permission flags can give more detail of the error) */
+	PermissionDenied = 1<<0,
+	/** The property could not be read or written due to a failure converting from the source or to the destination */
+	ConversionFailed = 1<<1,
+
+	/** Permission flag added when the property cannot be accessed due to being protected (is not marked for editor or Blueprint access) */
+	AccessProtected = 1<<4,
+	/** Permission flag added when attempting to edit a property on a template that cannot be edited on templates */
+	CannotEditTemplate = 1<<5,
+	/** Permission flag added when attempting to edit a property on an instance that cannot be edited on instances */
+	CannotEditInstance = 1<<6,
+	/** Permission flag added when attempting to edit a property that is read-only (based on the given read-only flags) */
+	ReadOnly = 1<<7,
+};
+ENUM_CLASS_FLAGS(EPropertyAccessResultFlags);
+
+/**
+ * Information needed to emit property change notifications when setting a property value.
+ */
+struct FPropertyAccessChangeNotify
+{
+	/** The kind of change that occurred */
+	EPropertyChangeType::Type ChangeType = EPropertyChangeType::Unspecified;
+	/** The object that is being changed */
+	UObject* ChangedObject = nullptr;
+	/** The chain of properties that are being changed */
+	FEditPropertyChain ChangedPropertyChain;
+};
+
+/**
+ * Callback used to get the value of a property.
+ */
+using FPropertyAccessGetFunc = TFunctionRef<bool()>;
+
+/**
+ * Callback used to set the value of a property.
+ */
+using FPropertyAccessSetFunc = TFunctionRef<bool(const FPropertyAccessChangeNotify*)>;
+
+/**
+ * Callback used to build the information needed to emit property change notifications when setting a property value.
+ */
+using FPropertyAccessBuildChangeNotifyFunc = TFunctionRef<TUniquePtr<FPropertyAccessChangeNotify>()>;
+
+namespace PropertyAccessUtil
+{
+	/** Flags that make a property read-only when settings its value at runtime */
+	static const uint64 RuntimeReadOnlyFlags = CPF_EditConst | CPF_BlueprintReadOnly;
+
+	/** Flags that make a property read-only when settings its value in the editor */
+	static const uint64 EditorReadOnlyFlags = CPF_EditConst;
+
+	/**
+	 * High-level function for getting the value of a property from an object.
+	 * @note This function calls CanGetPropertyValue internally.
+	 *
+	 * @param InProp Property to get the value of.
+	 * @param InObject Object containing the property.
+	 * @param InDestValue Instance to fill with the property value (must be a valid and constructed block of memory that is compatible with the property).
+	 * @param InArrayIndex For fixed-size array properties denotes which index of the array to get, or INDEX_NONE to get the entire property.
+	 *
+	 * @return Flags describing whether the get was successful.
+	 */
+	COREUOBJECT_API EPropertyAccessResultFlags GetPropertyValue_Object(const UProperty* InProp, const UObject* InObject, void* InDestValue, const int32 InArrayIndex);
+
+	/**
+	 * High-level function for getting the value of a property from a property container (object or struct).
+	 * @note This function calls CanGetPropertyValue internally.
+	 *
+	 * @param InProp Property to get the value of.
+	 * @param InContainerData The instance data containing the property.
+	 * @param InDestValue Instance to fill with the property value (must be a valid and constructed block of memory that is compatible with the property).
+	 * @param InArrayIndex For fixed-size array properties denotes which index of the array to get, or INDEX_NONE to get the entire property.
+	 *
+	 * @return Flags describing whether the get was successful.
+	 */
+	COREUOBJECT_API EPropertyAccessResultFlags GetPropertyValue_InContainer(const UProperty* InProp, const void* InContainerData, void* InDestValue, const int32 InArrayIndex);
+
+	/**
+	 * High-level function for getting the single-element value of a property from memory.
+	 * @note This function calls CanGetPropertyValue internally.
+	 *
+	 * @param InProp Property to get the value of.
+	 * @param InSrcValue The property value to copy.
+	 * @param InDestValue Instance to fill with the property value (must be a valid and constructed block of memory that is compatible with the property).
+	 *
+	 * @return Flags describing whether the get was successful.
+	 */
+	COREUOBJECT_API EPropertyAccessResultFlags GetPropertyValue_DirectSingle(const UProperty* InProp, const void* InSrcValue, void* InDestValue);
+
+	/**
+	 * High-level function for getting the multi-element value of a property from memory.
+	 * @note This function calls CanGetPropertyValue internally.
+	 *
+	 * @param InProp Property to get the value of.
+	 * @param InSrcValue The property value to copy.
+	 * @param InDestValue Instance to fill with the property value (must be a valid and constructed block of memory that is compatible with the property).
+	 *
+	 * @return Flags describing whether the get was successful.
+	 */
+	COREUOBJECT_API EPropertyAccessResultFlags GetPropertyValue_DirectComplete(const UProperty* InProp, const void* InSrcValue, void* InDestValue);
+
+	/**
+	 * Low-level function for getting the value of a property.
+	 * @note This function does *not* CanGetPropertyValue internally, you must call it yourself to validate the get.
+	 *
+	 * @param InGetFunc Logic for getting the value of the property.
+	 *
+	 * @return Flags describing whether the get was successful.
+	 */
+	COREUOBJECT_API EPropertyAccessResultFlags GetPropertyValue(const FPropertyAccessGetFunc& InGetFunc);
+
+	/**
+	 * Low-level function for checking whether it's valid to get the value of a property.
+	 *
+	 * @param InProp Property to query.
+	 *
+	 * @return Flags describing whether it's valid to get the value of the property.
+	 */
+	COREUOBJECT_API EPropertyAccessResultFlags CanGetPropertyValue(const UProperty* InProp);
+
+	/**
+	 * High-level function for setting the value of a property on an object.
+	 * @note This function calls CanSetPropertyValue internally, and will emit property change notifications for the object.
+	 *
+	 * @param InProp Property to set the value of.
+	 * @param InObject Object containing the property.
+	 * @param InSrcValue The value to set on the property.
+	 * @param InArrayIndex For fixed-size array properties denotes which index of the array to set, or INDEX_NONE to set the entire property.
+	 * @param InReadOnlyFlags Flags controlling which properties are considered read-only.
+	 *
+	 * @return Flags describing whether the set was successful.
+	 */
+	COREUOBJECT_API EPropertyAccessResultFlags SetPropertyValue_Object(const UProperty* InProp, UObject* InObject, const void* InSrcValue, const int32 InArrayIndex, const uint64 InReadOnlyFlags);
+
+	/**
+	 * High-level function for setting the value of a property on a property container (object or struct).
+	 * @note This function calls CanSetPropertyValue internally.
+	 *
+	 * @param InProp Property to set the value of.
+	 * @param InContainerData The instance data containing the property.
+	 * @param InSrcValue The value to set on the property.
+	 * @param InArrayIndex For fixed-size array properties denotes which index of the array to set, or INDEX_NONE to set the entire property.
+	 * @param InReadOnlyFlags Flags controlling which properties are considered read-only.
+	 * @param InOwnerIsTemplate True if the owner object is considered a template (see IsObjectTemplate).
+	 * @param InBuildChangeNotifyFunc Logic for building the information needed to emit property change notifications when setting a property value (can return nullptr if no notifications are needed or possible).
+	 *
+	 * @return Flags describing whether the set was successful.
+	 */
+	COREUOBJECT_API EPropertyAccessResultFlags SetPropertyValue_InContainer(const UProperty* InProp, void* InContainerData, const void* InSrcValue, const int32 InArrayIndex, const uint64 InReadOnlyFlags, const bool InOwnerIsTemplate, const FPropertyAccessBuildChangeNotifyFunc& InBuildChangeNotifyFunc);
+	
+	/**
+	 * High-level function for setting the single-element value of a property in memory.
+	 * @note This function calls CanSetPropertyValue internally.
+	 *
+	 * @param InProp Property to set the value of.
+	 * @param InSrcValue The value to set on the property.
+	 * @param InDestValue Instance to fill with the property value (must be a valid and constructed block of memory that is compatible with the property).
+	 * @param InReadOnlyFlags Flags controlling which properties are considered read-only.
+	 * @param InOwnerIsTemplate True if the owner object is considered a template (see IsObjectTemplate).
+	 * @param InBuildChangeNotifyFunc Logic for building the information needed to emit property change notifications when setting a property value (can return nullptr if no notifications are needed or possible).
+	 *
+	 * @return Flags describing whether the set was successful.
+	 */
+	COREUOBJECT_API EPropertyAccessResultFlags SetPropertyValue_DirectSingle(const UProperty* InProp, const void* InSrcValue, void* InDestValue, const uint64 InReadOnlyFlags, const bool InOwnerIsTemplate, const FPropertyAccessBuildChangeNotifyFunc& InBuildChangeNotifyFunc);
+
+	/**
+	 * High-level function for setting the multi-element value of a property in memory.
+	 * @note This function calls CanSetPropertyValue internally.
+	 *
+	 * @param InProp Property to set the value of.
+	 * @param InSrcValue The value to set on the property.
+	 * @param InDestValue Instance to fill with the property value (must be a valid and constructed block of memory that is compatible with the property).
+	 * @param InReadOnlyFlags Flags controlling which properties are considered read-only.
+	 * @param InOwnerIsTemplate True if the owner object is considered a template (see IsObjectTemplate).
+	 * @param InBuildChangeNotifyFunc Logic for building the information needed to emit property change notifications when setting a property value (can return nullptr if no notifications are needed or possible).
+	 *
+	 * @return Flags describing whether the set was successful.
+	 */
+	COREUOBJECT_API EPropertyAccessResultFlags SetPropertyValue_DirectComplete(const UProperty* InProp, const void* InSrcValue, void* InDestValue, const uint64 InReadOnlyFlags, const bool InOwnerIsTemplate, const FPropertyAccessBuildChangeNotifyFunc& InBuildChangeNotifyFunc);
+
+	/**
+	 * Low-level function for setting the value of a property.
+	 * @note This function does *not* CanSetPropertyValue internally, you must call it yourself to validate the set.
+	 *
+	 * @param InSetFunc Logic for setting the value of the property (should call EmitPreChangeNotify and EmitPostChangeNotify using the given FPropertyAccessChangeNotify instance).
+	 * @param InBuildChangeNotifyFunc Logic for building the information needed to emit property change notifications when setting a property value (can return nullptr if no notifications are needed or possible).
+	 *
+	 * @return Flags describing whether the get was successful.
+	 */
+	COREUOBJECT_API EPropertyAccessResultFlags SetPropertyValue(const FPropertyAccessSetFunc& InSetFunc, const FPropertyAccessBuildChangeNotifyFunc& InBuildChangeNotifyFunc);
+	
+	/**
+	 * Low-level function for checking whether it's valid to set the value of a property.
+	 *
+	 * @param InProp Property to query.
+	 *
+	 * @return Flags describing whether it's valid to set the value of the property.
+	 */
+	COREUOBJECT_API EPropertyAccessResultFlags CanSetPropertyValue(const UProperty* InProp, const uint64 InReadOnlyFlags, const bool InOwnerIsTemplate);
+
+	/**
+	 * Low-level function called before modifying an object to notify that its value is about to change.
+	 *
+	 * @param InChangeNotify Information needed to emit property change notifications, or nullptr if no notifications are needed or possible.
+	 */
+	COREUOBJECT_API void EmitPreChangeNotify(const FPropertyAccessChangeNotify* InChangeNotify);
+
+	/**
+	 * Low-level function called after modifying an object to notify that its value has changed.
+	 *
+	 * @param InChangeNotify Information needed to emit property change notifications, or nullptr if no notifications are needed or possible.
+	 */
+	COREUOBJECT_API void EmitPostChangeNotify(const FPropertyAccessChangeNotify* InChangeNotify);
+
+	/**
+	 * Low-level function to build the basic information needed to emit property change notifications.
+	 * @note This function can only build the notification for a property directly on the object instance, as more complex cases require external management as the reflection data doesn't let you backtrack over different types.
+	 *
+	 * @param InProp Property being modified.
+	 * @param InObject Object being modified.
+	 *
+	 * @return The information needed to emit property change notifications.
+	 */
+	COREUOBJECT_API TUniquePtr<FPropertyAccessChangeNotify> BuildBasicChangeNotify(const UProperty* InProp, const UObject* InObject);
+
+	/**
+	 * Low-level function for checking whether the given object instance is considered a template for property access.
+	 *
+	 * @param InObject Object to query.
+	 *
+	 * @return True if the object instance is considered a template.
+	 */
+	COREUOBJECT_API bool IsObjectTemplate(const UObject* InObject);
+
+	/**
+	 * Low-level function to find a property by its name, following redirectors if it cannot be found.
+	 *
+	 * @param InPropName The name of the property to find.
+	 * @param InStruct The struct that should contain the property.
+	 *
+	 * @return The found property, or null if the property cannot be found.
+	 */
+	COREUOBJECT_API UProperty* FindPropertyByName(const FName InPropName, const UStruct* InStruct);
+}
diff --git a/Engine/Source/Runtime/CoreUObject/Public/UObject/UnrealType.h b/Engine/Source/Runtime/CoreUObject/Public/UObject/UnrealType.h
--- a/Engine/Source/Runtime/CoreUObject/Public/UObject/UnrealType.h
+++ b/Engine/Source/Runtime/CoreUObject/Public/UObject/UnrealType.h
@@ -10,6 +10,7 @@
 
 #include "Concepts/GetTypeHashable.h"
 #include "Containers/List.h"
+#include "Containers/ArrayView.h"
 #include "Serialization/SerializedPropertyScope.h"
 #include "Templates/Casts.h"
 #include "Templates/Greater.h"
@@ -4735,27 +4736,26 @@ namespace EPropertyChangeType
  */
 struct FPropertyChangedEvent
 {
-	//default constructor
-	FPropertyChangedEvent(UProperty* InProperty)
+	FPropertyChangedEvent(UProperty* InProperty, EPropertyChangeType::Type InChangeType = EPropertyChangeType::Unspecified, TArrayView<const UObject* const> InTopLevelObjects = TArrayView<const UObject* const>())
 		: Property(InProperty)
 		, MemberProperty(InProperty)
-		, ChangeType(EPropertyChangeType::Unspecified)
+		, ChangeType(InChangeType)
 		, ObjectIteratorIndex(INDEX_NONE)
-		, ArrayIndicesPerObject(nullptr)
-		, InstancesChangedResultPerArchetype(nullptr)
-		, TopLevelObjects(nullptr)
+		, TopLevelObjects(InTopLevelObjects)
 	{
 	}
 
-	FPropertyChangedEvent(UProperty* InProperty, EPropertyChangeType::Type InChangeType, const TArray<const UObject*>* InTopLevelObjects = nullptr )
+	UE_DEPRECATED(4.25, "The FPropertyChangedEvent constructor taking a TArray* is deprecated. Use the version taking a TArrayView instead.")
+	FPropertyChangedEvent(UProperty* InProperty, EPropertyChangeType::Type InChangeType, const TArray<const UObject*>* InTopLevelObjects)
 		: Property(InProperty)
 		, MemberProperty(InProperty)
 		, ChangeType(InChangeType)
 		, ObjectIteratorIndex(INDEX_NONE)
-		, ArrayIndicesPerObject(nullptr)
-		, InstancesChangedResultPerArchetype(nullptr)
-		, TopLevelObjects(InTopLevelObjects)
 	{
+		if (InTopLevelObjects)
+		{
+			TopLevelObjects = MakeArrayView(*InTopLevelObjects);
+		}
 	}
 
 	void SetActiveMemberProperty( UProperty* InActiveMemberProperty )
@@ -4766,17 +4766,17 @@ struct FPropertyChangedEvent
 	/**
 	 * Saves off map of array indices per object being set.
 	 */
-	void SetArrayIndexPerObject(const TArray< TMap<FString,int32> >& InArrayIndices)
+	void SetArrayIndexPerObject(TArrayView<const TMap<FString, int32>> InArrayIndices)
 	{
-		ArrayIndicesPerObject = &InArrayIndices; 
+		ArrayIndicesPerObject = InArrayIndices; 
 	}
 
 	/**
 	 * Saves off map of instance changed result per archetype being set.
 	 */
-	void SetInstancesChangedResultPerArchetype(const TArray< TMap<UObject*, bool> >& InInstancesChangedResultPerArchetype)
+	void SetInstancesChangedResultPerArchetype(TArrayView<const TMap<UObject*, bool>> InInstancesChangedResultPerArchetype)
 	{
-		InstancesChangedResultPerArchetype = &InInstancesChangedResultPerArchetype;
+		InstancesChangedResultPerArchetype = InInstancesChangedResultPerArchetype;
 	}
 
 	/**
@@ -4787,9 +4787,9 @@ struct FPropertyChangedEvent
 	{
 		//default to unknown index
 		int32 Retval = -1;
-		if (ArrayIndicesPerObject && ArrayIndicesPerObject->IsValidIndex(ObjectIteratorIndex))
+		if (ArrayIndicesPerObject.IsValidIndex(ObjectIteratorIndex))
 		{
-			const int32* ValuePtr = (*ArrayIndicesPerObject)[ObjectIteratorIndex].Find(InName);
+			const int32* ValuePtr = ArrayIndicesPerObject[ObjectIteratorIndex].Find(InName);
 			if (ValuePtr)
 			{
 				Retval = *ValuePtr;
@@ -4806,9 +4806,9 @@ struct FPropertyChangedEvent
 	{
 		bool Retval = true;
 
-		if (InstancesChangedResultPerArchetype && InstancesChangedResultPerArchetype->IsValidIndex(ObjectIteratorIndex))
+		if (InstancesChangedResultPerArchetype.IsValidIndex(ObjectIteratorIndex))
 		{
-			const bool* ValuePtr = (*InstancesChangedResultPerArchetype)[ObjectIteratorIndex].Find(InInstance);
+			const bool* ValuePtr = InstancesChangedResultPerArchetype[ObjectIteratorIndex].Find(InInstance);
 			if (ValuePtr)
 			{
 				Retval = *ValuePtr;
@@ -4821,7 +4821,7 @@ struct FPropertyChangedEvent
 	/**
 	 * @return The number of objects being edited during this change event
 	 */
-	int32 GetNumObjectsBeingEdited() const { return TopLevelObjects ? TopLevelObjects->Num() : 0; }
+	int32 GetNumObjectsBeingEdited() const { return TopLevelObjects.Num(); }
 
 	/**
 	 * Gets an object being edited by this change event.  Multiple objects could be edited at once
@@ -4829,7 +4829,7 @@ struct FPropertyChangedEvent
 	 * @param Index	The index of the object being edited. Assumes index is valid.  Call GetNumObjectsBeingEdited first to check if there are valid objects
 	 * @return The object being edited or nullptr if no object was found
 	 */
-	const UObject* GetObjectBeingEdited(int32 Index) const { return TopLevelObjects ? (*TopLevelObjects)[Index] : nullptr; }
+	const UObject* GetObjectBeingEdited(int32 Index) const { return TopLevelObjects[Index]; }
 
 	/**
 	 * Simple utility to get the name of the property and takes care of the possible null property.
@@ -4857,13 +4857,13 @@ struct FPropertyChangedEvent
 	int32 ObjectIteratorIndex;
 private:
 	//In the property window, multiple objects can be selected at once.  In the case of adding/inserting to an array, each object COULD have different indices for the new entries in the array
-	const TArray< TMap<FString,int32> >* ArrayIndicesPerObject;
+	TArrayView<const TMap<FString, int32>> ArrayIndicesPerObject;
 	
 	//In the property window, multiple objects can be selected at once. In this case we want to know if an instance was updated for this operation (used in array/set/map context)
-	const TArray< TMap<UObject*, bool> >* InstancesChangedResultPerArchetype;
+	TArrayView<const TMap<UObject*, bool>> InstancesChangedResultPerArchetype;
 
 	/** List of top level objects being changed */
-	const TArray<const UObject*>* TopLevelObjects;
+	TArrayView<const UObject* const> TopLevelObjects;
 };
 
 /**
diff --git a/Engine/Source/Runtime/Engine/Classes/Kismet/KismetSystemLibrary.h b/Engine/Source/Runtime/Engine/Classes/Kismet/KismetSystemLibrary.h
--- a/Engine/Source/Runtime/Engine/Classes/Kismet/KismetSystemLibrary.h
+++ b/Engine/Source/Runtime/Engine/Classes/Kismet/KismetSystemLibrary.h
@@ -1733,6 +1733,22 @@ class ENGINE_API UKismetSystemLibrary : public UBlueprintFunctionLibrary
 	UFUNCTION(BlueprintPure, Category = "Utilities")
 	static bool IsUnattended();
 
+	// --- Property Access ---------------------------
+
+#if WITH_EDITOR
+	/** Read the value of a named property on the given object to the given value pointer */
+    UFUNCTION(BlueprintCallable, CustomThunk, Category = "Properties", meta=(CustomStructureParam="ValuePtr", BlueprintInternalUseOnly="true"))
+    static bool GetEditorProperty(UObject* Object, const FName PropertyName, int32& ValuePtr);
+	static bool Generic_GetEditorProperty(const UObject* Object, const UProperty* Property, void* ValuePtr);
+	DECLARE_FUNCTION(execGetEditorProperty);
+
+	/** Write the value of the given value pointer to a named property on the given object */
+    UFUNCTION(BlueprintCallable, CustomThunk, Category = "Properties", meta=(CustomStructureParam="ValuePtr", BlueprintInternalUseOnly="true"))
+    static bool SetEditorProperty(UObject* Object, const FName PropertyName, const int32& ValuePtr);
+	static bool Generic_SetEditorProperty(UObject* Object, const UProperty* Property, const void* ValuePtr);
+	DECLARE_FUNCTION(execSetEditorProperty);
+#endif
+
 	// --- Transactions ------------------------------
 
 	/**
diff --git a/Engine/Source/Runtime/Engine/Private/KismetSystemLibrary.cpp b/Engine/Source/Runtime/Engine/Private/KismetSystemLibrary.cpp
--- a/Engine/Source/Runtime/Engine/Private/KismetSystemLibrary.cpp
+++ b/Engine/Source/Runtime/Engine/Private/KismetSystemLibrary.cpp
@@ -10,6 +10,7 @@
 #include "Misc/Paths.h"
 #include "Misc/RuntimeErrors.h"
 #include "UObject/GCObject.h"
+#include "UObject/PropertyAccessUtil.h"
 #include "EngineGlobals.h"
 #include "Components/ActorComponent.h"
 #include "TimerManager.h"
@@ -40,9 +41,27 @@
 //////////////////////////////////////////////////////////////////////////
 // UKismetSystemLibrary
 
+#define LOCTEXT_NAMESPACE "UKismetSystemLibrary"
+
+const FName PropertyGetFailedWarning = FName("PropertyGetFailedWarning");
+const FName PropertySetFailedWarning = FName("PropertySetFailedWarning");
+
 UKismetSystemLibrary::UKismetSystemLibrary(const FObjectInitializer& ObjectInitializer)
 	: Super(ObjectInitializer)
 {
+	FBlueprintSupport::RegisterBlueprintWarning(
+		FBlueprintWarningDeclaration(
+			PropertyGetFailedWarning,
+			LOCTEXT("PropertyGetFailedWarning", "Property Get Failed")
+		)
+	);
+
+	FBlueprintSupport::RegisterBlueprintWarning(
+		FBlueprintWarningDeclaration(
+			PropertySetFailedWarning,
+			LOCTEXT("PropertySetFailedWarning", "Property Set Failed")
+		)
+	);
 }
 
 void UKismetSystemLibrary::StackTraceImpl(const FFrame& StackFrame)
@@ -2642,6 +2661,190 @@ bool UKismetSystemLibrary::IsUnattended()
 	return FApp::IsUnattended();
 }
 
+#if WITH_EDITOR
+
+bool UKismetSystemLibrary::GetEditorProperty(UObject* Object, const FName PropertyName, int32& OutValuePtr)
+{
+	// We should never hit this! Stubbed to avoid NoExport on the class.
+	check(0);
+	return false;
+}
+
+bool UKismetSystemLibrary::Generic_GetEditorProperty(const UObject* Object, const UProperty* Property, void* ValuePtr)
+{
+	const EPropertyAccessResultFlags AccessResult = PropertyAccessUtil::GetPropertyValue_Object(Property, Object, ValuePtr, INDEX_NONE);
+
+	if (EnumHasAnyFlags(AccessResult, EPropertyAccessResultFlags::PermissionDenied))
+	{
+		if (EnumHasAnyFlags(AccessResult, EPropertyAccessResultFlags::AccessProtected))
+		{
+			FFrame::KismetExecutionMessage(*FString::Printf(TEXT("Property '%s' on '%s' (%s) is protected and cannot be read"), *Property->GetName(), *Object->GetPathName(), *Object->GetClass()->GetName()), ELogVerbosity::Warning, PropertyGetFailedWarning);
+			return false;
+		}
+
+		FFrame::KismetExecutionMessage(*FString::Printf(TEXT("Property '%s' on '%s' (%s) cannot be read"), *Property->GetName(), *Object->GetPathName(), *Object->GetClass()->GetName()), ELogVerbosity::Warning, PropertyGetFailedWarning);
+		return false;
+	}
+
+	if (EnumHasAnyFlags(AccessResult, EPropertyAccessResultFlags::ConversionFailed))
+	{
+		FFrame::KismetExecutionMessage(*FString::Printf(TEXT("Property '%s' on '%s' (%s) tried to read into a property value of the incorrect type"), *Property->GetName(), *Object->GetPathName(), *Object->GetClass()->GetName()), ELogVerbosity::Warning, PropertyGetFailedWarning);
+		return nullptr;
+	}
+
+	return true;
+}
+
+DEFINE_FUNCTION(UKismetSystemLibrary::execGetEditorProperty)
+{
+	P_GET_OBJECT(UObject, Object);
+	P_GET_PROPERTY(UNameProperty, PropertyName);
+
+	Stack.StepCompiledIn<UProperty>(nullptr);
+	const UProperty* ValueProp = Stack.MostRecentProperty;
+	void* ValuePtr = Stack.MostRecentPropertyAddress;
+
+	P_FINISH;
+
+	if (!ValueProp || !ValuePtr)
+	{
+		FBlueprintExceptionInfo ExceptionInfo(
+			EBlueprintExceptionType::AccessViolation,
+			LOCTEXT("GetEditorProperty_MissingOutputProperty", "Failed to resolve the output parameter for GetEditorProperty.")
+		);
+		FBlueprintCoreDelegates::ThrowScriptException(P_THIS, Stack, ExceptionInfo);
+	}
+
+	if (!Object)
+	{
+		FBlueprintExceptionInfo ExceptionInfo(
+			EBlueprintExceptionType::AccessViolation, 
+			LOCTEXT("GetEditorProperty_AccessNone", "Accessed None attempting to call GetEditorProperty.")
+		);
+		FBlueprintCoreDelegates::ThrowScriptException(P_THIS, Stack, ExceptionInfo);
+	}
+
+	bool bResult = false;
+
+	if (Object)
+	{
+		const UProperty* ObjectProp = PropertyAccessUtil::FindPropertyByName(PropertyName, Object->GetClass());
+		if (ObjectProp && ObjectProp->GetClass() == ValueProp->GetClass()) // Compare the classes as these must be an *exact* match as the read is low-level and without property coercion
+		{
+			P_NATIVE_BEGIN;
+			bResult = Generic_GetEditorProperty(Object, ObjectProp, ValuePtr);
+			P_NATIVE_END;
+		}
+		else
+		{
+			FFrame::KismetExecutionMessage(*FString::Printf(TEXT("Property '%s' on '%s' (%s) was missing or didn't match the output value type (%s)"), *PropertyName.ToString(), *Object->GetPathName(), *Object->GetClass()->GetName(), *ValueProp->GetClass()->GetName()), ELogVerbosity::Warning, PropertyGetFailedWarning);
+		}
+	}
+
+	*(bool*)RESULT_PARAM = bResult;
+}
+
+bool UKismetSystemLibrary::SetEditorProperty(UObject* Object, const FName PropertyName, const int32& OutValuePtr)
+{
+	// We should never hit this! Stubbed to avoid NoExport on the class.
+	check(0);
+	return false;
+}
+
+bool UKismetSystemLibrary::Generic_SetEditorProperty(UObject* Object, const UProperty* Property, const void* ValuePtr)
+{
+	const EPropertyAccessResultFlags AccessResult = PropertyAccessUtil::SetPropertyValue_Object(Property, Object, ValuePtr, INDEX_NONE, PropertyAccessUtil::EditorReadOnlyFlags);
+
+	if (EnumHasAnyFlags(AccessResult, EPropertyAccessResultFlags::PermissionDenied))
+	{
+		if (EnumHasAnyFlags(AccessResult, EPropertyAccessResultFlags::AccessProtected))
+		{
+			FFrame::KismetExecutionMessage(*FString::Printf(TEXT("Property '%s' on '%s' (%s) is protected and cannot be set"), *Property->GetName(), *Object->GetPathName(), *Object->GetClass()->GetName()), ELogVerbosity::Warning, PropertySetFailedWarning);
+			return false;
+		}
+
+		if (EnumHasAnyFlags(AccessResult, EPropertyAccessResultFlags::CannotEditTemplate))
+		{
+			FFrame::KismetExecutionMessage(*FString::Printf(TEXT("Property '%s' on '%s' (%s) cannot be edited on templates"), *Property->GetName(), *Object->GetPathName(), *Object->GetClass()->GetName()), ELogVerbosity::Warning, PropertySetFailedWarning);
+			return false;
+		}
+
+		if (EnumHasAnyFlags(AccessResult, EPropertyAccessResultFlags::CannotEditInstance))
+		{
+			FFrame::KismetExecutionMessage(*FString::Printf(TEXT("Property '%s' on '%s' (%s) cannot be edited on instances"), *Property->GetName(), *Object->GetPathName(), *Object->GetClass()->GetName()), ELogVerbosity::Warning, PropertySetFailedWarning);
+			return false;
+		}
+
+		if (EnumHasAnyFlags(AccessResult, EPropertyAccessResultFlags::ReadOnly))
+		{
+			FFrame::KismetExecutionMessage(*FString::Printf(TEXT("Property '%s' on '%s' (%s) is read-only and cannot be set"), *Property->GetName(), *Object->GetPathName(), *Object->GetClass()->GetName()), ELogVerbosity::Warning, PropertySetFailedWarning);
+			return false;
+		}
+
+		FFrame::KismetExecutionMessage(*FString::Printf(TEXT("Property '%s' on '%s' (%s) cannot be read"), *Property->GetName(), *Object->GetPathName(), *Object->GetClass()->GetName()), ELogVerbosity::Warning, PropertySetFailedWarning);
+		return false;
+	}
+
+	if (EnumHasAnyFlags(AccessResult, EPropertyAccessResultFlags::ConversionFailed))
+	{
+		FFrame::KismetExecutionMessage(*FString::Printf(TEXT("Property '%s' on '%s' (%s) tried to write into a property value of the incorrect type"), *Property->GetName(), *Object->GetPathName(), *Object->GetClass()->GetName()), ELogVerbosity::Warning, PropertySetFailedWarning);
+		return nullptr;
+	}
+
+	return true;
+}
+
+DEFINE_FUNCTION(UKismetSystemLibrary::execSetEditorProperty)
+{
+	P_GET_OBJECT(UObject, Object);
+	P_GET_PROPERTY(UNameProperty, PropertyName);
+
+	Stack.StepCompiledIn<UProperty>(nullptr);
+	const UProperty* ValueProp = Stack.MostRecentProperty;
+	const void* ValuePtr = Stack.MostRecentPropertyAddress;
+
+	P_FINISH;
+
+	if (!ValueProp || !ValuePtr)
+	{
+		FBlueprintExceptionInfo ExceptionInfo(
+			EBlueprintExceptionType::AccessViolation,
+			LOCTEXT("SetEditorProperty_MissingOutputProperty", "Failed to resolve the output parameter for SetEditorProperty.")
+		);
+		FBlueprintCoreDelegates::ThrowScriptException(P_THIS, Stack, ExceptionInfo);
+	}
+
+	if (!Object)
+	{
+		FBlueprintExceptionInfo ExceptionInfo(
+			EBlueprintExceptionType::AccessViolation, 
+			LOCTEXT("SetEditorProperty_AccessNone", "Accessed None attempting to call SetEditorProperty.")
+		);
+		FBlueprintCoreDelegates::ThrowScriptException(P_THIS, Stack, ExceptionInfo);
+	}
+
+	bool bResult = false;
+
+	if (Object)
+	{
+		const UProperty* ObjectProp = PropertyAccessUtil::FindPropertyByName(PropertyName, Object->GetClass());
+		if (ObjectProp && ObjectProp->GetClass() == ValueProp->GetClass()) // Compare the classes as these must be an *exact* match as the read is low-level and without property coercion
+		{
+			P_NATIVE_BEGIN;
+			bResult = Generic_SetEditorProperty(Object, ObjectProp, ValuePtr);
+			P_NATIVE_END;
+		}
+		else
+		{
+			FFrame::KismetExecutionMessage(*FString::Printf(TEXT("Property '%s' on '%s' (%s) was missing or didn't match the output value type (%s)"), *PropertyName.ToString(), *Object->GetPathName(), *Object->GetClass()->GetName(), *ValueProp->GetClass()->GetName()), ELogVerbosity::Warning, PropertySetFailedWarning);
+		}
+	}
+
+	*(bool*)RESULT_PARAM = bResult;
+}
+
+#endif	// WITH_EDITOR
+
 int32 UKismetSystemLibrary::BeginTransaction(const FString& Context, FText Description, UObject* PrimaryObject)
 {
 #if WITH_EDITOR
@@ -2866,4 +3069,6 @@ void UKismetSystemLibrary::GetPrimaryAssetsWithBundleState(const TArray<FName>&
 	{
 		Manager->GetPrimaryAssetsWithBundleState(OutPrimaryAssetIdList, ValidTypes, RequiredBundles, ExcludedBundles, bForceCurrentState);
 	}
-}
\ No newline at end of file
+}
+
+#undef LOCTEXT_NAMESPACE
