Fixes issue with cooked platforms generating warnings for partyhub data that is cooked out for them. Treat those files as acceptable to be missing while loading then check() them on platforms that care.

- Allow FObjectFinder and FObjectFinderOptional to take flags to use during loading
- Fixed issue with GetPackageLinker not checking LOAD_NoWarn / LOAD_Quiet when a package doesn't exist.

diff --git a/Engine/Source/Runtime/CoreUObject/Private/UObject/Linker.cpp b/Engine/Source/Runtime/CoreUObject/Private/UObject/Linker.cpp
--- a/Engine/Source/Runtime/CoreUObject/Private/UObject/Linker.cpp
+++ b/Engine/Source/Runtime/CoreUObject/Private/UObject/Linker.cpp
@@ -678,7 +678,10 @@ FLinkerLoad* GetPackageLinker
 		const bool DoesPackageExist = DoesPackageExistForGetPackageLinker(PackageNameToLoad, CompatibleGuid, NewFilename);
 		if( !DoesPackageExist )
 		{
-			if (!FLinkerLoad::IsKnownMissingPackage(InLongPackageName))
+			// Issue a warning if the caller didn't request nowar/quiet, and the package isn't marked as known to be missing.
+			bool IssueWarning = (LoadFlags & (LOAD_NoWarn | LOAD_Quiet)) == 0 && !FLinkerLoad::IsKnownMissingPackage(InLongPackageName);
+
+			if (IssueWarning)
 			{
 				// try to recover from this instead of throwing, it seems recoverable just by doing this
 				LogGetPackageLinkerError(Result, InExistingContext, InLongPackageName, LOCTEXT("FileNotFoundShort", "Can't find file."), InOuter, LoadFlags);
diff --git a/Engine/Source/Runtime/CoreUObject/Public/UObject/ConstructorHelpers.h b/Engine/Source/Runtime/CoreUObject/Public/UObject/ConstructorHelpers.h
--- a/Engine/Source/Runtime/CoreUObject/Public/UObject/ConstructorHelpers.h
+++ b/Engine/Source/Runtime/CoreUObject/Public/UObject/ConstructorHelpers.h
@@ -14,7 +14,7 @@
 namespace ConstructorHelpersInternal
 {
 	template<typename T>
-	inline T* FindOrLoadObject( FString& PathName )
+	inline T* FindOrLoadObject( FString& PathName, uint32 LoadFlags )
 	{
 		// If there is no dot, add a dot and repeat the object name.
 		int32 PackageDelimPos = INDEX_NONE;
@@ -33,7 +33,7 @@ namespace ConstructorHelpersInternal
 
 		UClass* Class = T::StaticClass();
 		Class->GetDefaultObject(); // force the CDO to be created if it hasn't already
-		T* ObjectPtr = LoadObject<T>(NULL, *PathName);
+		T* ObjectPtr = LoadObject<T>(NULL, *PathName, nullptr, LoadFlags);
 		if (ObjectPtr)
 		{
 			ObjectPtr->AddToRoot();
@@ -42,7 +42,7 @@ namespace ConstructorHelpersInternal
 	}
 
 	template<>
-	inline UPackage* FindOrLoadObject<UPackage>( FString& PathName )
+	inline UPackage* FindOrLoadObject<UPackage>( FString& PathName, uint32 LoadFlags)
 	{
 		// If there is a dot, remove it.
 		int32 PackageDelimPos = INDEX_NONE;
@@ -57,7 +57,7 @@ namespace ConstructorHelpersInternal
 		if( !PackagePtr )
 		{
 			// If it is not in memory, try to load it.
-			PackagePtr = LoadPackage( nullptr, *PathName, LOAD_None );
+			PackagePtr = LoadPackage( nullptr, *PathName, LoadFlags);
 		}
 		if (PackagePtr)
 		{
@@ -100,13 +100,13 @@ public:
 	struct FObjectFinder : public FGCObject
 	{
 		T* Object;
-		FObjectFinder(const TCHAR* ObjectToFind)
+		FObjectFinder(const TCHAR* ObjectToFind, uint32 InLoadFlags = LOAD_None)
 		{
 			CheckIfIsInConstructor(ObjectToFind);
 			FString PathName(ObjectToFind);
 			StripObjectClass(PathName,true);
 
-			Object = ConstructorHelpersInternal::FindOrLoadObject<T>(PathName);
+			Object = ConstructorHelpersInternal::FindOrLoadObject<T>(PathName, InLoadFlags);
 			ValidateObject( Object, PathName, ObjectToFind );
 		}
 		bool Succeeded() const
@@ -131,10 +131,12 @@ public:
 	private:
 		T* Object;
 		const TCHAR* ObjectToFind;
+		uint32 LoadFlags;
 	public:
-		FObjectFinderOptional(const TCHAR* InObjectToFind)
+		FObjectFinderOptional(const TCHAR* InObjectToFind, uint32 InLoadFlags = LOAD_None)
 			: Object(nullptr)
 			, ObjectToFind(InObjectToFind)
+			, LoadFlags(InLoadFlags)
 		{
 		}
 		T* Get()
@@ -145,8 +147,14 @@ public:
 				FString PathName(ObjectToFind);
 				StripObjectClass(PathName,true);
 
-				Object = ConstructorHelpersInternal::FindOrLoadObject<T>(PathName);
-				ValidateObject( Object, PathName, ObjectToFind );
+				Object = ConstructorHelpersInternal::FindOrLoadObject<T>(PathName, LoadFlags);
+
+				bool WarnMissing = (LoadFlags & (LOAD_Quiet | LOAD_NoWarn)) == 0;
+
+				if (Object || WarnMissing)
+				{
+					ValidateObject(Object, PathName, ObjectToFind);
+				}
 
 				ObjectToFind = nullptr; // don't try to look again
 			}
