Prevent name collision as a result of transaction undo/redo - if a transaction buffer wants to use a name for a given object we rename objects that already have that name

diff --git a/Engine/Source/Runtime/CoreUObject/Private/UObject/Obj.cpp b/Engine/Source/Runtime/CoreUObject/Private/UObject/Obj.cpp
--- a/Engine/Source/Runtime/CoreUObject/Private/UObject/Obj.cpp
+++ b/Engine/Source/Runtime/CoreUObject/Private/UObject/Obj.cpp
@@ -1319,6 +1319,28 @@ void UObject::Serialize(FStructuredArchive::FRecord Record)
 					bool bDifferentOuter = LoadOuter != GetOuter();
 					if ( bDifferentName == true || bDifferentOuter == true )
 					{
+						// Clear the name for use by this:
+						UObject* Collision = StaticFindObjectFast(UObject::StaticClass(), LoadOuter, LoadName);
+						if(Collision && Collision != this)
+						{
+							FName NewNameForCollision = MakeUniqueObjectName(LoadOuter, Collision->GetClass(), LoadName);
+							checkf( StaticFindObjectFast(UObject::StaticClass(), LoadOuter, NewNameForCollision) == nullptr,
+								TEXT("Failed to MakeUniqueObjectName for object colliding with transaction buffer state: %s %s"),
+								*LoadName.ToString(),
+								*NewNameForCollision.ToString()
+							);
+							Collision->LowLevelRename(NewNameForCollision,LoadOuter);
+#if DO_CHECK
+							UObject* SubsequentCollision = StaticFindObjectFast(UObject::StaticClass(), LoadOuter, LoadName);
+							checkf( SubsequentCollision == nullptr,
+								TEXT("Multiple name collisions detected in the transaction buffer: %x %x with name %s"),
+								Collision,
+								SubsequentCollision,
+								*LoadName.ToString()
+							);
+#endif
+						}
+						
 						LowLevelRename(LoadName,LoadOuter);
 					}
 				}
