Optimize GetArchetypeInstances, particularly when calling it repeatedly with the same array.
#rnx

diff --git a/Engine/Source/Runtime/CoreUObject/Private/UObject/Obj.cpp b/Engine/Source/Runtime/CoreUObject/Private/UObject/Obj.cpp
--- a/Engine/Source/Runtime/CoreUObject/Private/UObject/Obj.cpp
+++ b/Engine/Source/Runtime/CoreUObject/Private/UObject/Obj.cpp
@@ -705,42 +705,34 @@ bool UObject::CanCreateInCurrentContext(UObject* Template)
 
 void UObject::GetArchetypeInstances( TArray<UObject*>& Instances )
 {
-	Instances.Empty();
+	Instances.Reset();
 
 	if ( HasAnyFlags(RF_ArchetypeObject|RF_ClassDefaultObject) )
 	{
-		// we need to evaluate CDOs as well, but nothing pending kill
-		TArray<UObject*> IterObjects;
-		{
-			const bool bIncludeNestedObjects = true;
-			GetObjectsOfClass(GetClass(), IterObjects, bIncludeNestedObjects, RF_NoFlags, EInternalObjectFlags::PendingKill);
-		}
 
 		// if this object is the class default object, any object of the same class (or derived classes) could potentially be affected
 		if ( !HasAnyFlags(RF_ArchetypeObject) )
 		{
-			Instances.Reserve(IterObjects.Num()-1);
-			for (UObject* It : IterObjects)
+			const bool bIncludeNestedObjects = true;
+			ForEachObjectOfClass(GetClass(), [this, &Instances](UObject* Obj)
 			{
-				UObject* Obj = It;
-				if ( Obj != this )
+				if (Obj != this)
 				{
 					Instances.Add(Obj);
 				}
-			}
+			}, bIncludeNestedObjects, RF_NoFlags, EInternalObjectFlags::PendingKill); // we need to evaluate CDOs as well, but nothing pending kill
 		}
 		else
 		{
-			for (UObject* It : IterObjects)
+			const bool bIncludeNestedObjects = true;
+			ForEachObjectOfClass(GetClass(), [this, &Instances](UObject* Obj)
 			{
-				UObject* Obj = It;
-				
-				// if this object is the correct type and its archetype is this object, add it to the list
-				if ( Obj && Obj != this && Obj->IsBasedOnArchetype(this) )
+				if (Obj != this && Obj->IsBasedOnArchetype(this))
 				{
 					Instances.Add(Obj);
 				}
-			}
+			}, bIncludeNestedObjects, RF_NoFlags, EInternalObjectFlags::PendingKill); // we need to evaluate CDOs as well, but nothing pending kill
+
 		}
 	}
 }
