When saving a class that contains sparse class data we now create the sparse class data if it doesn't already exist on the class.

diff --git a/Engine/Source/Runtime/CoreUObject/Private/UObject/Class.cpp b/Engine/Source/Runtime/CoreUObject/Private/UObject/Class.cpp
--- a/Engine/Source/Runtime/CoreUObject/Private/UObject/Class.cpp
+++ b/Engine/Source/Runtime/CoreUObject/Private/UObject/Class.cpp
@@ -3985,6 +3985,11 @@ void UClass::SerializeSparseClassData(UObject* Object, FStructuredArchive::FSlot
 	// tell the archive that it's allowed to load data for transient properties
 	FArchive& UnderlyingArchive = Slot.GetUnderlyingArchive();
 
+	if (UnderlyingArchive.IsSaving())
+	{
+		GetOrCreateSparseClassData();
+	}
+
 	if (((UnderlyingArchive.IsLoading() || UnderlyingArchive.IsSaving()) && !UnderlyingArchive.WantBinaryPropertySerialization()))
 	{
 		// class default objects do not always have a vtable when saved
