Do not port deprecated properties unless we are importing them from default properties or PPF_UseDeprecatedProperties has been explicitly specified

diff --git a/Engine/Source/Runtime/CoreUObject/Private/UObject/Property.cpp b/Engine/Source/Runtime/CoreUObject/Private/UObject/Property.cpp
--- a/Engine/Source/Runtime/CoreUObject/Private/UObject/Property.cpp
+++ b/Engine/Source/Runtime/CoreUObject/Private/UObject/Property.cpp
@@ -792,27 +792,44 @@ bool UProperty::ShouldPort( uint32 PortFlags/*=0*/ ) const
 {
 	// if no size, don't export
 	if (GetSize() <= 0)
+	{
+		return false;
+	}
+
+	if (HasAnyPropertyFlags(CPF_Deprecated) && !(PortFlags & (PPF_ParsingDefaultProperties | PPF_UseDeprecatedProperties)))
+	{
 		return false;
+	}
 
 	// if we're parsing default properties or the user indicated that transient properties should be included
 	if (HasAnyPropertyFlags(CPF_Transient) && !(PortFlags & (PPF_ParsingDefaultProperties | PPF_IncludeTransient)))
+	{
 		return false;
+	}
 
 	// if we're copying, treat DuplicateTransient as transient
 	if ((PortFlags & PPF_Copy) && HasAnyPropertyFlags(CPF_DuplicateTransient | CPF_TextExportTransient) && !(PortFlags & (PPF_ParsingDefaultProperties | PPF_IncludeTransient)))
+	{
 		return false;
+	}
 
 	// if we're not copying for PIE and NonPIETransient is set, don't export
 	if (!(PortFlags & PPF_DuplicateForPIE) && HasAnyPropertyFlags(CPF_NonPIEDuplicateTransient))
+	{
 		return false;
+	}
 
 	// if we're only supposed to export components and this isn't a component property, don't export
 	if ((PortFlags & PPF_SubobjectsOnly) && !ContainsInstancedObjectProperty())
+	{
 		return false;
+	}
 
 	// hide non-Edit properties when we're exporting for the property window
 	if ((PortFlags & PPF_PropertyWindow) && !(PropertyFlags & CPF_Edit))
+	{
 		return false;
+	}
 
 	return true;
 }
