Imports will always process deprecated properties.
#rnx

diff --git a/Engine/Source/Runtime/CoreUObject/Public/UObject/UnrealType.h b/Engine/Source/Runtime/CoreUObject/Public/UObject/UnrealType.h
--- a/Engine/Source/Runtime/CoreUObject/Public/UObject/UnrealType.h
+++ b/Engine/Source/Runtime/CoreUObject/Public/UObject/UnrealType.h
@@ -23,6 +23,7 @@
 #include "UObject/LazyObjectPtr.h"
 #include "UObject/Object.h"
 #include "UObject/ObjectMacros.h"
+#include "Uobject/PropertyPortFlags.h"
 #include "UObject/PropertyTag.h"
 #include "UObject/ScriptInterface.h"
 #include "UObject/SoftObjectPtr.h"
@@ -321,6 +322,7 @@ public:
 		{
 			return NULL;
 		}
+		PortFlags |= EPropertyPortFlags::PPF_UseDeprecatedProperties; // Imports should always process deprecated properties
 		return ImportText_Internal( Buffer, Data, PortFlags, OwnerObject, ErrorText );
 	}
 protected:
