Fix capitalization causing Linux compile errors
#rnx

diff --git a/Engine/Source/Runtime/CoreUObject/Public/UObject/UnrealType.h b/Engine/Source/Runtime/CoreUObject/Public/UObject/UnrealType.h
--- a/Engine/Source/Runtime/CoreUObject/Public/UObject/UnrealType.h
+++ b/Engine/Source/Runtime/CoreUObject/Public/UObject/UnrealType.h
@@ -23,7 +23,7 @@
 #include "UObject/LazyObjectPtr.h"
 #include "UObject/Object.h"
 #include "UObject/ObjectMacros.h"
-#include "Uobject/PropertyPortFlags.h"
+#include "UObject/PropertyPortFlags.h"
 #include "UObject/PropertyTag.h"
 #include "UObject/ScriptInterface.h"
 #include "UObject/SoftObjectPtr.h"
