Fixed invalid check in BulkDataWriter when attempting to serialize a zero byte buffer.

diff --git a/Engine/Source/Runtime/CoreUObject/Public/Serialization/BulkDataWriter.h b/Engine/Source/Runtime/CoreUObject/Public/Serialization/BulkDataWriter.h
--- a/Engine/Source/Runtime/CoreUObject/Public/Serialization/BulkDataWriter.h
+++ b/Engine/Source/Runtime/CoreUObject/Public/Serialization/BulkDataWriter.h
@@ -37,7 +37,7 @@ public:
 	{
 		// Determine if we need to reallocate the buffer to fit the next item
 		const int64 NewPos = WriterPos + Num;
-		check( NewPos > WriterPos );
+		check( NewPos >= WriterPos );
 		if( NewPos > BulkData.GetBulkDataSize() )
 		{
 			// If so, resize to the new size + 3/8
