Prevent crash in DissociateImportsAndForcedExports due to stale XObject pointer

diff --git a/Engine/Source/Runtime/CoreUObject/Private/UObject/LinkerManager.cpp b/Engine/Source/Runtime/CoreUObject/Private/UObject/LinkerManager.cpp
--- a/Engine/Source/Runtime/CoreUObject/Private/UObject/LinkerManager.cpp
+++ b/Engine/Source/Runtime/CoreUObject/Private/UObject/LinkerManager.cpp
@@ -154,7 +154,14 @@ void FLinkerManager::DissociateImportsAndForcedExports()
 			for (int32 ImportIndex = 0; ImportIndex < Linker->ImportMap.Num(); ImportIndex++)
 			{
 				FObjectImport& Import = Linker->ImportMap[ImportIndex];
-				if (Import.XObject && !Import.XObject->IsNative())
+				// The import object could be stale if it has been replaced by patching
+				// logic or compile on load..
+				bool bIsStale = false;
+				if (Import.SourceLinker && Import.SourceIndex != INDEX_NONE)
+				{
+					bIsStale = Import.SourceLinker->ExportMap[Import.SourceIndex].Object != Import.XObject;
+				}
+				if (bIsStale || (Import.XObject && !Import.XObject->IsNative()))
 				{
 					Import.XObject = nullptr;
 				}
