When reloading assets, serialize the sparse class data for assets that may refer to the asset being reloaded.

#rnx

diff --git a/Engine/Source/Runtime/CoreUObject/Private/UObject/Class.cpp b/Engine/Source/Runtime/CoreUObject/Private/UObject/Class.cpp
--- a/Engine/Source/Runtime/CoreUObject/Private/UObject/Class.cpp
+++ b/Engine/Source/Runtime/CoreUObject/Private/UObject/Class.cpp
@@ -3900,6 +3900,14 @@ void UClass::Serialize( FArchive& Ar )
 		}
 	}
 
+	if (!Ar.IsLoading() && !Ar.IsSaving())
+	{
+		if (GetSparseClassDataStruct() != nullptr)
+		{
+			SerializeSparseClassData(FStructuredArchiveFromArchive(Ar).GetSlot());
+		}
+	}
+
 	// mark the archive we that we are no longer serializing defaults
 	Ar.StopSerializingDefaults();
 
@@ -3975,7 +3983,7 @@ void UClass::SerializeDefaultObject(UObject* Object, FStructuredArchive::FSlot S
 	UnderlyingArchive.StopSerializingDefaults();
 }
 
-void UClass::SerializeSparseClassData(UObject* Object, FStructuredArchive::FSlot Slot)
+void UClass::SerializeSparseClassData(FStructuredArchive::FSlot Slot)
 {
 	if (!SparseClassDataStruct)
 	{
diff --git a/Engine/Source/Runtime/CoreUObject/Public/UObject/Class.h b/Engine/Source/Runtime/CoreUObject/Public/UObject/Class.h
--- a/Engine/Source/Runtime/CoreUObject/Public/UObject/Class.h
+++ b/Engine/Source/Runtime/CoreUObject/Public/UObject/Class.h
@@ -2969,10 +2969,9 @@ public:
 	}
 
 	/** serializes the associated sparse class data for the passed in object using the given archive slot. This should only be called if the class has an associated sparse data structure.
-	 * @param Object the object to serialize as default
 	 * @param Slot the structured archive slot to serialize from
 	 */
-	void SerializeSparseClassData(UObject* Object, FStructuredArchive::FSlot Slot);
+	void SerializeSparseClassData(FStructuredArchive::FSlot Slot);
 
 	/** Wraps the PostLoad() call for the class default object.
 	 * @param Object the default object to call PostLoad() on
diff --git a/Engine/Source/Runtime/Engine/Private/BlueprintGeneratedClass.cpp b/Engine/Source/Runtime/Engine/Private/BlueprintGeneratedClass.cpp
--- a/Engine/Source/Runtime/Engine/Private/BlueprintGeneratedClass.cpp
+++ b/Engine/Source/Runtime/Engine/Private/BlueprintGeneratedClass.cpp
@@ -365,7 +365,7 @@ void UBlueprintGeneratedClass::SerializeDefaultObject(UObject* Object, FStructur
 	{
 		if (Object->GetSparseClassDataStruct())
 		{
-			SerializeSparseClassData(Object, FStructuredArchiveFromArchive(UnderlyingArchive).GetSlot());
+			SerializeSparseClassData(FStructuredArchiveFromArchive(UnderlyingArchive).GetSlot());
 		}
 	}
 }
