Don't write out top level properties to config files if deprecated
#rnx

diff --git a/Engine/Source/Runtime/CoreUObject/Private/UObject/Obj.cpp b/Engine/Source/Runtime/CoreUObject/Private/UObject/Obj.cpp
--- a/Engine/Source/Runtime/CoreUObject/Private/UObject/Obj.cpp
+++ b/Engine/Source/Runtime/CoreUObject/Private/UObject/Obj.cpp
@@ -2471,6 +2471,7 @@ void UObject::SaveConfig( uint64 Flags, const TCHAR* InFilename, FConfigCacheIni
 
 			// Properties that are the same as the parent class' defaults should not be saved to ini
 			// Before modifying any key in the section, first check to see if it is different from the parent.
+			const bool bPropDeprecated = Property->HasAnyPropertyFlags(CPF_Deprecated);
 			const bool bIsPropertyInherited = Property->GetOwnerClass() != GetClass();
 			const bool bShouldCheckIfIdenticalBeforeAdding = !GetClass()->HasAnyClassFlags(CLASS_ConfigDoNotCheckDefaults) && !bPerObject && bIsPropertyInherited;
 			UObject* SuperClassDefaultObject = GetClass()->GetSuperClass()->GetDefaultObject();
@@ -2478,7 +2479,7 @@ void UObject::SaveConfig( uint64 Flags, const TCHAR* InFilename, FConfigCacheIni
 			UArrayProperty* Array   = dynamic_cast<UArrayProperty*>( Property );
 			if( Array )
 			{
-				if ( !bShouldCheckIfIdenticalBeforeAdding || !Property->Identical_InContainer(this, SuperClassDefaultObject) )
+				if (!bPropDeprecated && (!bShouldCheckIfIdenticalBeforeAdding || !Property->Identical_InContainer(this, SuperClassDefaultObject)))
 				{
 					FConfigSection* Sec = Config->GetSectionPrivate( *Section, 1, 0, *PropFileName );
 					check(Sec);
@@ -2495,7 +2496,7 @@ void UObject::SaveConfig( uint64 Flags, const TCHAR* InFilename, FConfigCacheIni
 						Sec->Add(*CompleteKey, *Buffer);
 					}
 				}
-				else if( Property->Identical_InContainer(this, SuperClassDefaultObject) )
+				else
 				{
 					// If we are not writing it to config above, we should make sure that this property isn't stagnant in the cache.
 					FConfigSection* Sec = Config->GetSectionPrivate( *Section, 1, 0, *PropFileName );
@@ -2517,13 +2518,13 @@ void UObject::SaveConfig( uint64 Flags, const TCHAR* InFilename, FConfigCacheIni
 						Key = TempKey;
 					}
 
-					if ( !bShouldCheckIfIdenticalBeforeAdding || !Property->Identical_InContainer(this, SuperClassDefaultObject, Index) )
+					if (!bPropDeprecated && (!bShouldCheckIfIdenticalBeforeAdding || !Property->Identical_InContainer(this, SuperClassDefaultObject, Index)))
 					{
 						FString	Value;
 						Property->ExportText_InContainer( Index, Value, this, this, this, PortFlags );
 						Config->SetString( *Section, *Key, *Value, *PropFileName );
 					}
-					else if( Property->Identical_InContainer(this, SuperClassDefaultObject, Index) )
+					else
 					{
 						// If we are not writing it to config above, we should make sure that this property isn't stagnant in the cache.
 						FConfigSection* Sec = Config->GetSectionPrivate( *Section, 1, 0, *PropFileName );
