Added csv stat for FLinkerManager::DeleteLinkers. Also added early return if there are no pending linkers to delete.

#rnx

diff --git a/Engine/Source/Runtime/CoreUObject/Private/UObject/LinkerManager.cpp b/Engine/Source/Runtime/CoreUObject/Private/UObject/LinkerManager.cpp
--- a/Engine/Source/Runtime/CoreUObject/Private/UObject/LinkerManager.cpp
+++ b/Engine/Source/Runtime/CoreUObject/Private/UObject/LinkerManager.cpp
@@ -9,6 +9,7 @@
 #include "UObject/ObjectResource.h"
 #include "UObject/LinkerLoad.h"
 #include "UObject/UObjectThreadContext.h"
+#include "ProfilingDebugging/CsvProfiler.h"
 
 FLinkerManager& FLinkerManager::Get()
 {
@@ -16,6 +17,15 @@ FLinkerManager& FLinkerManager::Get()
 	return *Singleton;
 }
 
+FLinkerManager::FLinkerManager() :
+	bHasPendingCleanup(false)
+{
+}
+
+FLinkerManager::~FLinkerManager()
+{
+}
+
 bool FLinkerManager::Exec(class UWorld* InWorld, const TCHAR* Cmd, FOutputDevice& Ar)
 {
 #if !UE_BUILD_SHIPPING
@@ -202,28 +212,34 @@ void FLinkerManager::DeleteLinkers()
 {
 	check(IsInGameThread());
 
-	QUICK_SCOPE_CYCLE_COUNTER(STAT_FLinkerManager_DeleteLinkers);
-
-	TArray<FLinkerLoad*> CleanupArray;
+	if(bHasPendingCleanup)
 	{
+		bHasPendingCleanup = false;
+	
+		QUICK_SCOPE_CYCLE_COUNTER(STAT_FLinkerManager_DeleteLinkers);
+		CSV_SCOPED_TIMING_STAT_EXCLUSIVE(DeleteLinkers);
+
+		TArray<FLinkerLoad*> CleanupArray;
+		{
 #if THREADSAFE_UOBJECTS
-		FScopeLock PendingCleanupListLock(&PendingCleanupListCritical);
+			FScopeLock PendingCleanupListLock(&PendingCleanupListCritical);
 #endif
-		CleanupArray = PendingCleanupList.Array();
-		PendingCleanupList.Empty();
-	}
+			CleanupArray = PendingCleanupList.Array();
+			PendingCleanupList.Empty();
+		}
 
-	// Note that even though DeleteLinkers can only be called on the main thread,
-	// we store the IsDeletingLinkers in TLS so that we're sure nothing on
-	// another thread can delete linkers except FLinkerManager at the time
-	// we enter this loop.
-	FUObjectThreadContext& ThreadContext = FUObjectThreadContext::Get();
-	ThreadContext.IsDeletingLinkers = true;
-	for (FLinkerLoad* Linker : CleanupArray)
-	{
-		delete Linker;
+		// Note that even though DeleteLinkers can only be called on the main thread,
+		// we store the IsDeletingLinkers in TLS so that we're sure nothing on
+		// another thread can delete linkers except FLinkerManager at the time
+		// we enter this loop.
+		FUObjectThreadContext& ThreadContext = FUObjectThreadContext::Get();
+		ThreadContext.IsDeletingLinkers = true;
+		for (FLinkerLoad* Linker : CleanupArray)
+		{
+			delete Linker;
+		}
+		ThreadContext.IsDeletingLinkers = false;
 	}
-	ThreadContext.IsDeletingLinkers = false;
 }
 
 void FLinkerManager::RemoveLinker(FLinkerLoad* Linker)
@@ -234,5 +250,6 @@ void FLinkerManager::RemoveLinker(FLinkerLoad* Linker)
 	if (Linker && !PendingCleanupList.Contains(Linker))
 	{
 		PendingCleanupList.Add(Linker);
+		bHasPendingCleanup = true;
 	}
 }
diff --git a/Engine/Source/Runtime/CoreUObject/Private/UObject/LinkerManager.h b/Engine/Source/Runtime/CoreUObject/Private/UObject/LinkerManager.h
--- a/Engine/Source/Runtime/CoreUObject/Private/UObject/LinkerManager.h
+++ b/Engine/Source/Runtime/CoreUObject/Private/UObject/LinkerManager.h
@@ -10,10 +10,13 @@
 #include "UObject/ObjectMacros.h"
 #include "Misc/CoreMisc.h"
 #include "Misc/ScopeLock.h"
+#include "HAL/ThreadSafeBool.h"
 
 class FLinkerManager : private FSelfRegisteringExec
 {
 public:
+	FLinkerManager();
+	~FLinkerManager();
 
 	static FLinkerManager& Get();
 
@@ -152,4 +155,5 @@ private:
 #if THREADSAFE_UOBJECTS
 	FCriticalSection PendingCleanupListCritical;
 #endif
+	FThreadSafeBool bHasPendingCleanup;
 };
