Added FReferencerFinder helper class that finds all objects referencing the provided list of objects.

Similar to FArchiveHasReferences  but up to 16x faster.

diff --git a/Engine/Source/Runtime/CoreUObject/Private/UObject/ReferencerFinder.cpp b/Engine/Source/Runtime/CoreUObject/Private/UObject/ReferencerFinder.cpp
new file mode 100644
--- /dev/null
+++ b/Engine/Source/Runtime/CoreUObject/Private/UObject/ReferencerFinder.cpp
@@ -0,0 +1,102 @@
+// Copyright 1998-2019 Epic Games, Inc. All Rights Reserved.
+
+#include "UObject/ReferencerFinder.h"
+
+#include "UObject/UObjectIterator.h"
+#include "UObject/FastReferenceCollector.h"
+#include "UObject/UObjectArray.h"
+#include "Async/ParallelFor.h"
+
+class FAllReferencesProcessor : public FSimpleReferenceProcessorBase
+{
+	const TSet<UObject*>& PotentiallyReferencedObjects;
+	TSet<UObject*>& ReferencingObjects;
+
+public:
+	FAllReferencesProcessor(const TSet<UObject*>& InPotentiallyReferencedObjects, TSet<UObject*>& OutReferencingObjects)
+		: PotentiallyReferencedObjects(InPotentiallyReferencedObjects)
+		, ReferencingObjects(OutReferencingObjects)
+	{
+	}
+	FORCEINLINE_DEBUGGABLE void HandleTokenStreamObjectReference(TArray<UObject*>& ObjectsToSerialize, UObject* ReferencingObject, UObject*& Object, const int32 TokenIndex, bool bAllowReferenceElimination)
+	{
+		if (Object != nullptr && Object != ReferencingObject)
+		{			
+			if (ReferencingObject && PotentiallyReferencedObjects.Contains(Object))
+			{
+				ReferencingObjects.Add(ReferencingObject);
+			}
+		}
+	}
+};
+typedef TDefaultReferenceCollector<FAllReferencesProcessor> FAllReferencesCollector;
+
+TArray<UObject*> FReferencerFinder::GetAllReferencers(const TArray<UObject*>& Referencees, const TSet<UObject*>* ObjectsToIgnore )
+{
+	return GetAllReferencers(TSet<UObject*>(Referencees), ObjectsToIgnore);
+}
+
+void LockUObjectHashTables();
+void UnlockUObjectHashTables();
+
+TArray<UObject*> FReferencerFinder::GetAllReferencers(const TSet<UObject*>& Referencees, const TSet<UObject*>* ObjectsToIgnore )
+{
+	TArray<UObject*> Ret;
+	if(Referencees.Num() > 0)
+	{
+		FCriticalSection ResultCritical;
+
+		// Lock hashtables so that nothing can add UObjects while we're iterating over the GUObjectArray
+		LockUObjectHashTables();
+
+		const int32 MaxNumberOfObjects = GUObjectArray.GetObjectArrayNum();
+		const int32 NumThreads = FMath::Max(1, FTaskGraphInterface::Get().GetNumWorkerThreads());
+		const int32 NumberOfObjectsPerThread = (MaxNumberOfObjects / NumThreads) + 1;
+
+		ParallelFor(NumThreads, [&Referencees, ObjectsToIgnore, &ResultCritical, &Ret, NumberOfObjectsPerThread, NumThreads, MaxNumberOfObjects](int32 ThreadIndex)
+		{
+			TSet<UObject*> ThreadResult;
+			FAllReferencesProcessor Processor(Referencees, ThreadResult);
+			TFastReferenceCollector<false, FAllReferencesProcessor, FAllReferencesCollector, FGCArrayPool, true> ReferenceCollector(Processor, FGCArrayPool::Get());
+			FGCArrayStruct ArrayStruct;
+
+			ArrayStruct.ObjectsToSerialize.Reserve(NumberOfObjectsPerThread);
+
+			const int32 FirstObjectIndex = ThreadIndex * NumberOfObjectsPerThread;
+			const int32 NumObjects = (ThreadIndex < (NumThreads - 1)) ? NumberOfObjectsPerThread : (MaxNumberOfObjects - (NumThreads - 1)*NumberOfObjectsPerThread);
+
+			// First cache all potential referencers
+			for (int32 ObjectIndex = 0; ObjectIndex < NumObjects && (FirstObjectIndex + ObjectIndex) < MaxNumberOfObjects; ++ObjectIndex)
+			{
+				FUObjectItem& ObjectItem = GUObjectArray.GetObjectItemArrayUnsafe()[FirstObjectIndex + ObjectIndex];
+				if (ObjectItem.Object)
+				{
+					UObject* PotentialReferencer = static_cast<UObject*>(ObjectItem.Object);
+
+					if (ObjectsToIgnore && ObjectsToIgnore->Contains(PotentialReferencer))
+					{
+						continue;
+					}
+
+					if (!Referencees.Contains(PotentialReferencer))
+					{
+						ArrayStruct.ObjectsToSerialize.Add(PotentialReferencer);
+					}
+				}
+			}
+
+			// Now check if any of the potential referencers is referencing any of the referencees
+			ReferenceCollector.CollectReferences(ArrayStruct);
+
+			if (ThreadResult.Num())
+			{
+				// We found objects referencing some of the referencess so add them to the final results array
+				FScopeLock ResultLock(&ResultCritical);
+				Ret.Append(ThreadResult.Array());
+			}
+		});
+
+		UnlockUObjectHashTables();
+	}
+	return Ret;
+}
diff --git a/Engine/Source/Runtime/CoreUObject/Public/UObject/ReferencerFinder.h b/Engine/Source/Runtime/CoreUObject/Public/UObject/ReferencerFinder.h
new file mode 100644
--- /dev/null
+++ b/Engine/Source/Runtime/CoreUObject/Public/UObject/ReferencerFinder.h
@@ -0,0 +1,16 @@
+// Copyright 1998-2019 Epic Games, Inc. All Rights Reserved.
+
+#pragma once
+
+#include "CoreMinimal.h"
+#include "Serialization/ArchiveUObject.h"
+
+/**
+ * Helper class for finding all objects referencing any of the objects in Referencees list
+ */
+class COREUOBJECT_API FReferencerFinder
+{
+public:
+	static TArray<UObject*> GetAllReferencers(const TArray<UObject*>& Referencees, const TSet<UObject*>* ObjectsToIgnore);
+	static TArray<UObject*> GetAllReferencers(const TSet<UObject*>& Referencees, const TSet<UObject*>* ObjectsToIgnore);
+};
\ No newline at end of file
