Added more info when crashing due to an export's class not being mapped when saving a package

diff --git a/Engine/Source/Runtime/CoreUObject/Private/UObject/SavePackage.cpp b/Engine/Source/Runtime/CoreUObject/Private/UObject/SavePackage.cpp
--- a/Engine/Source/Runtime/CoreUObject/Private/UObject/SavePackage.cpp
+++ b/Engine/Source/Runtime/CoreUObject/Private/UObject/SavePackage.cpp
@@ -4954,7 +4954,7 @@ FSavePackageResultStruct UPackage::Save(UPackage* InOuter, UObject* Base, EObjec
 						if (ObjClass != UClass::StaticClass())
 						{
 							Export.ClassIndex = Linker->MapObject(ObjClass);
-							check(!Export.ClassIndex.IsNull());
+							checkf(!Export.ClassIndex.IsNull(), TEXT("Export %s class is not mapped when saving %s"), *Export.Object->GetFullName(), *Linker->LinkerRoot->GetName());
 						}
 						else
 						{
