Live Coding: Workaround for hang when modifying a file containing sparse delegates in Live Coding.

Live++ disambiguates symbols in anonymous namespaces when loading data from COFF files, but does not attempt to disambiguate functions read from the PDB. This assumes that symbols in the PDB are public, and cannot contain anonymous namespaces, but does not account for the fact that a template can reference a symbol in an anonymous namespace via a template parameter.

Change the SparseDelegate implementation to create an inline getter function for the delegate name rather than a global variable in an anonymous namespace.

diff --git a/Engine/Source/Runtime/CoreUObject/Public/UObject/SparseDelegate.h b/Engine/Source/Runtime/CoreUObject/Public/UObject/SparseDelegate.h
--- a/Engine/Source/Runtime/CoreUObject/Public/UObject/SparseDelegate.h
+++ b/Engine/Source/Runtime/CoreUObject/Public/UObject/SparseDelegate.h
@@ -171,7 +171,7 @@ protected:
 	bool bIsBound;
 };
 
-template <typename MulticastDelegate, typename OwningClass, const char* DelegateName, typename DelegateOffsetClass>
+template <typename MulticastDelegate, typename OwningClass, typename DelegateInfoClass>
 struct TSparseDynamicDelegate : public FSparseDelegate
 {
 public:
@@ -180,14 +180,14 @@ public:
 protected:
 	FName GetDelegateName() const
 	{
-		static const FName DelegateFName(DelegateName);
+		static const FName DelegateFName(DelegateInfoClass::GetDelegateName());
 		return DelegateFName;
 	}
 
 private:
 	UObject* GetDelegateOwner() const
 	{
-		const size_t OffsetToOwner = DelegateOffsetClass::template GetDelegateOffset<OwningClass>();
+		const size_t OffsetToOwner = DelegateInfoClass::template GetDelegateOffset<OwningClass>();
 		check(OffsetToOwner);
 		UObject* DelegateOwner = reinterpret_cast<UObject*>((uint8*)this - OffsetToOwner);
 		check(DelegateOwner->IsValidLowLevelFast(false)); // Most likely the delegate is trying to be used on the stack, in an object it wasn't defined for, or for a class member with a different name than it was defined for. It is only valid for a sparse delegate to be used for the exact class/property name it is defined with.
@@ -400,13 +400,13 @@ private:
 
 #define FUNC_DECLARE_DYNAMIC_MULTICAST_SPARSE_DELEGATE(SparseDelegateClass, OwningClass, DelegateName, FuncParamList, FuncParamPassThru, ...) \
 	FUNC_DECLARE_DYNAMIC_MULTICAST_DELEGATE(FWeakObjectPtr, SparseDelegateClass##_MCSignature, SparseDelegateClass##_DelegateWrapper, FUNC_CONCAT(FuncParamList), FUNC_CONCAT(FuncParamPassThru), __VA_ARGS__) \
-	namespace { char OwningClass##Sparse##DelegateName##Str[] = #DelegateName; } \
-	struct SparseDelegateClass##OffsetGetter \
+	struct SparseDelegateClass##InfoGetter \
 	{ \
+		static const char* GetDelegateName() { return #DelegateName; } \
 		template<typename T> \
 		static size_t GetDelegateOffset() { return offsetof(T, DelegateName); } \
 	}; \
-	struct SparseDelegateClass : public TSparseDynamicDelegate<SparseDelegateClass##_MCSignature, OwningClass, OwningClass##Sparse##DelegateName##Str, SparseDelegateClass##OffsetGetter> \
+	struct SparseDelegateClass : public TSparseDynamicDelegate<SparseDelegateClass##_MCSignature, OwningClass, SparseDelegateClass##InfoGetter> \
 	{ \
 	};
 
