Made sure async loader is only deleted by the owner linker

diff --git a/Engine/Source/Runtime/CoreUObject/Private/Serialization/AsyncLoading.cpp b/Engine/Source/Runtime/CoreUObject/Private/Serialization/AsyncLoading.cpp
--- a/Engine/Source/Runtime/CoreUObject/Private/Serialization/AsyncLoading.cpp
+++ b/Engine/Source/Runtime/CoreUObject/Private/Serialization/AsyncLoading.cpp
@@ -7123,7 +7123,7 @@ static FAutoConsoleCommand GDumpSerializeCmd(
 
 static FCriticalSection SummaryRacePreventer;
 
-FAsyncArchive::FAsyncArchive(const TCHAR* InFileName, TFunction<void()>&& InSummaryReadyCallback)
+FAsyncArchive::FAsyncArchive(const TCHAR* InFileName, FLinkerLoad* InOwner, TFunction<void()>&& InSummaryReadyCallback)
 	: Handle(nullptr)
 	, SizeRequestPtr(nullptr)
 	, EditorPrecacheRequestPtr(nullptr)
@@ -7147,6 +7147,7 @@ FAsyncArchive::FAsyncArchive(const TCHAR* InFileName, TFunction<void()>&& InSumm
 	, SummaryReadTime(0.0)
 	, ExportReadTime(0.0)
 	, SummaryReadyCallback(Forward<TFunction<void()>>(InSummaryReadyCallback))
+	, OwnerLinker(InOwner)
 {
 	LogItem(TEXT("Open"));
 	Handle = FPlatformFileManager::Get().GetPlatformFile().OpenAsyncRead(InFileName);
@@ -7172,6 +7173,9 @@ FAsyncArchive::FAsyncArchive(const TCHAR* InFileName, TFunction<void()>&& InSumm
 
 FAsyncArchive::~FAsyncArchive()
 {
+	UE_CLOG(OwnerLinker && !(OwnerLinker->GetLoader_Unsafe() == this && OwnerLinker->IsDestroyingLoader()), LogStreaming, Fatal,
+		TEXT("Destroying FAsyncArchive %s that belongs to linker %s outside of the linker's DestroyLoader code!"), *GetArchiveName(), *OwnerLinker->GetArchiveName());
+
 	// Invalidate any precached data and free memory.
 	FlushCache();
 	if (Handle)
diff --git a/Engine/Source/Runtime/CoreUObject/Private/Serialization/AsyncLoadingPrivate.h b/Engine/Source/Runtime/CoreUObject/Private/Serialization/AsyncLoadingPrivate.h
--- a/Engine/Source/Runtime/CoreUObject/Private/Serialization/AsyncLoadingPrivate.h
+++ b/Engine/Source/Runtime/CoreUObject/Private/Serialization/AsyncLoadingPrivate.h
@@ -20,7 +20,7 @@ public:
 		ProcessingExports,
 	};
 
-	FAsyncArchive(const TCHAR* InFileName, TFunction<void()>&& InSummaryReadyCallback);
+	FAsyncArchive(const TCHAR* InFileName, FLinkerLoad* InOwner, TFunction<void()>&& InSummaryReadyCallback);
 	virtual ~FAsyncArchive ();
 
 	/** Archive overrides */
@@ -126,4 +126,5 @@ private:
 	TFunction<void()> SummaryReadyCallback;
 	FAsyncFileCallBack ReadCallbackFunctionForLinkerLoad;
 
+	FLinkerLoad* OwnerLinker;
 };
diff --git a/Engine/Source/Runtime/CoreUObject/Private/UObject/LinkerLoad.cpp b/Engine/Source/Runtime/CoreUObject/Private/UObject/LinkerLoad.cpp
--- a/Engine/Source/Runtime/CoreUObject/Private/UObject/LinkerLoad.cpp
+++ b/Engine/Source/Runtime/CoreUObject/Private/UObject/LinkerLoad.cpp
@@ -827,6 +827,7 @@ FLinkerLoad::FLinkerLoad(UPackage* InParent, const TCHAR* InFilename, uint32 InL
 , bForceSimpleIndexToObject(false)
 , bLockoutLegacyOperations(false)
 , bIsAsyncLoader(false)
+, bIsDestroyingLoader(false)
 , StructuredArchive(nullptr)
 , StructuredArchiveFormatter(nullptr)
 , Loader(nullptr)
@@ -1027,8 +1028,8 @@ FLinkerLoad::ELinkerStatus FLinkerLoad::CreateLoader(
 				bool bCanUseAsyncLoader = FPlatformProperties::RequiresCookedData() || GAllowCookedDataInEditorBuilds;
 				if (bCanUseAsyncLoader)
 				{
-					Loader = new FAsyncArchive(*Filename
-							, GEventDrivenLoaderEnabled ? Forward<TFunction<void()>>(InSummaryReadyCallback) : TFunction<void()>([]() {})
+					Loader = new FAsyncArchive(*Filename, this,
+							GEventDrivenLoaderEnabled ? Forward<TFunction<void()>>(InSummaryReadyCallback) : TFunction<void()>([]() {})
 						);
 				}
 				else
@@ -1044,8 +1045,7 @@ FLinkerLoad::ELinkerStatus FLinkerLoad::CreateLoader(
 
 				if (Loader->IsError())
 				{
-					delete Loader;
-					Loader = nullptr;
+					DestroyLoader();
 					UE_LOG(LogLinker, Warning, TEXT("Error opening file '%s'."), *Filename);
 					return LINKER_Failed;
 				}
@@ -1061,8 +1061,7 @@ FLinkerLoad::ELinkerStatus FLinkerLoad::CreateLoader(
 					uint32	BufferSize = Loader->TotalSize();
 					void*	Buffer = FMemory::Malloc(BufferSize);
 					Loader->Serialize(Buffer, BufferSize);
-					delete Loader;
-					Loader = nullptr;
+					DestroyLoader();
 					if (bHasHashEntry)
 					{
 						// create buffer reader and spawn SHA verify when it gets closed
@@ -1116,8 +1115,7 @@ FLinkerLoad::ELinkerStatus FLinkerLoad::CreateLoader(
 				int64 Size = Loader->TotalSize();
 				if (Size <= 0)
 				{
-					delete Loader;
-					Loader = nullptr;
+					DestroyLoader();
 					UE_LOG(LogLinker, Warning, TEXT("Error opening file '%s'."), *Filename);
 					return LINKER_Failed;
 				}
@@ -4886,6 +4884,19 @@ void FLinkerLoad::LoadAndDetachAllBulkData()
 #endif
 }
 
+void FLinkerLoad::DestroyLoader()
+{
+	check(!bIsDestroyingLoader); // Destroying loader recursively is not safe
+	bIsDestroyingLoader = true; // Some archives check for this to make sure they're not destroyed by random code
+	FPlatformMisc::MemoryBarrier();
+	if (Loader)
+	{
+		delete Loader;
+		Loader = nullptr;
+	}
+	bIsDestroyingLoader = false;
+}
+
 void FLinkerLoad::Detach()
 {
 #if WITH_EDITOR
@@ -4922,11 +4933,7 @@ void FLinkerLoad::Detach()
 	delete StructuredArchiveFormatter;
 	StructuredArchiveFormatter = nullptr;
 
-	if (Loader)
-	{
-		delete Loader;
-		Loader = nullptr;
-	}
+	DestroyLoader();
 
 	// Empty out no longer used arrays.
 	NameMap.Empty();
diff --git a/Engine/Source/Runtime/CoreUObject/Public/UObject/LinkerLoad.h b/Engine/Source/Runtime/CoreUObject/Public/UObject/LinkerLoad.h
--- a/Engine/Source/Runtime/CoreUObject/Public/UObject/LinkerLoad.h
+++ b/Engine/Source/Runtime/CoreUObject/Public/UObject/LinkerLoad.h
@@ -137,6 +137,9 @@ public:
 
 private:
 
+	/** True if the linker is currently deleting loader */
+	bool					bIsDestroyingLoader;
+
 	/** Structured archive interface. Wraps underlying loader to provide contextual metadata to the values being written
 	 *  which ultimately allows text based serialization of the data
 	 */
@@ -173,10 +176,11 @@ public:
 		return Loader != nullptr;
 	}
 
-	void DestroyLoader()
+	void DestroyLoader();
+
+	FORCEINLINE bool IsDestroyingLoader() const
 	{
-		delete Loader;
-		Loader = nullptr;
+		return bIsDestroyingLoader;
 	}
 
 	/** The async package associated with this linker */
