Plugins can now disable tabs in the tab manager

diff --git a/Engine/Source/Runtime/CoreUObject/Public/Misc/NamedBlacklist.h b/Engine/Source/Runtime/CoreUObject/Public/Misc/NamedBlacklist.h
new file mode 100644
--- /dev/null
+++ b/Engine/Source/Runtime/CoreUObject/Public/Misc/NamedBlacklist.h
@@ -0,0 +1,222 @@
+// Copyright 1998-2019 Epic Games, Inc. All Rights Reserved.
+
+#pragma once
+
+#include "CoreTypes.h"
+
+#include "NamedBlacklist.generated.h"
+
+/** 
+ * Filter for blacklisting or whitelisting items with ownership to allow change cleanup
+ */
+class FNamedBlacklist
+{
+public:
+
+	/** Returns true if passes filter restrictions */
+	bool PassesFilter(const FName ItemName) const
+	{
+		if (BlacklistAll.Num() > 0)
+		{
+			return false;
+		}
+		else if (Blacklist.Contains(ItemName))
+		{
+			return false;
+		}
+		else if (Whitelist.Num() > 0 && !Whitelist.Contains(ItemName))
+		{
+			return false;
+		}
+
+		return true;
+	}
+	
+	/** Add item to blacklist, this specific item will be filtered out */
+	void AddBlacklistItem(const FName OwnerName, const FName ItemName)
+	{
+		Blacklist.FindOrAdd(ItemName).AddUnique(OwnerName);
+	}
+
+	/** Add item to whitelist after which all items not in the whitelist will be filtered out */
+	void AddWhitelistItem(const FName OwnerName, const FName ItemName)
+	{
+		Whitelist.FindOrAdd(ItemName).AddUnique(OwnerName);
+	}
+
+	/** Set to filter out all items */
+	void AddBlacklistAll(const FName OwnerName)
+	{
+		BlacklistAll.AddUnique(OwnerName);
+	}
+
+	/** Removes all filtering changes associated with a specific owner name */
+	void UnregisterOwner(const FName OwnerName)
+	{
+		for (auto It = Blacklist.CreateIterator(); It; ++It)
+		{
+			It->Value.Remove(OwnerName);
+			if (It->Value.Num() == 0)
+			{
+				It.RemoveCurrent();
+			}
+		}
+
+		for (auto It = Whitelist.CreateIterator(); It; ++It)
+		{
+			It->Value.Remove(OwnerName);
+			if (It->Value.Num() == 0)
+			{
+				It.RemoveCurrent();
+			}
+		}
+
+		BlacklistAll.Remove(OwnerName);
+	}
+
+	/** True if has filters active */
+	bool HasFiltering() const
+	{
+		return Blacklist.Num() > 0 || Whitelist.Num() > 0 || BlacklistAll.Num() > 0;
+	}
+
+	/** Combine two filters together */
+	void Append(const FNamedBlacklist& Other)
+	{
+		for (const auto& It : Other.Blacklist)
+		{
+			for (const auto& OwnerName : It.Value)
+			{
+				AddBlacklistItem(OwnerName, It.Key);
+			}
+		}
+
+		for (const auto& It : Other.Whitelist)
+		{
+			for (const auto& OwnerName : It.Value)
+			{
+				AddWhitelistItem(OwnerName, It.Key);
+			}
+		}
+
+		for (const auto& OwnerName : Other.BlacklistAll)
+		{
+			AddBlacklistAll(OwnerName);
+		}
+	}
+
+private:
+
+	/** List of owner names that requested a specific item filtered, allows unregistering specific set of changes by a given plugin or system */
+	typedef TArray<FName> FNamedBlacklistOwners;
+
+	/** List if items to filter out */
+	TMap<FName, FNamedBlacklistOwners> Blacklist;
+
+	/** List of items to allow, if not empty all items will be filtered out unless they are in the list */
+	TMap<FName, FNamedBlacklistOwners> Whitelist;
+
+	/** List of owner names that requested all items to be filtered out */
+	FNamedBlacklistOwners BlacklistAll;
+};
+
+/**
+ * Collection of blacklisting filters
+ */
+class FNamedBlacklistCollection
+{
+public:
+
+	/** Returns true if passes filter restrictions */
+	bool PassesFilter(const FName GroupName, const FName ItemName) const
+	{
+		const FNamedBlacklist* Filter = Filters.Find(GroupName);
+		return !Filter || Filter->PassesFilter(ItemName);
+	}
+
+	/** Add item to blacklist */
+	void AddBlacklistItem(const FName OwnerName, const FName GroupName, const FName ItemName)
+	{
+		Filters.FindOrAdd(GroupName).AddBlacklistItem(OwnerName, ItemName);
+	}
+
+	/** Add item to whitelist */
+	void AddWhitelistItem(const FName OwnerName, const FName GroupName, const FName ItemName)
+	{
+		Filters.FindOrAdd(GroupName).AddWhitelistItem(OwnerName, ItemName);
+	}
+
+	/** Sets group to fully blacklisted */
+	void AddBlacklistAll(const FName OwnerName, const FName GroupName)
+	{
+		Filters.FindOrAdd(GroupName).AddBlacklistAll(OwnerName);
+	}
+
+	/** Remove filtering associated with owner name */
+	void UnregisterOwner(const FName OwnerName)
+	{
+		for (auto GroupIt = Filters.CreateIterator(); GroupIt; ++GroupIt)
+		{
+			FNamedBlacklist& Filter = GroupIt->Value;
+			Filter.UnregisterOwner(OwnerName);
+			if (!Filter.HasFiltering())
+			{
+				GroupIt.RemoveCurrent();
+			}
+		}
+	}
+
+private:
+
+	/** List of filters */
+	TMap<FName, FNamedBlacklist> Filters;
+};
+
+/** 
+ * Serializable blacklist (Ownership not included)
+ */
+USTRUCT()
+struct FNamedBlacklistConfig
+{
+	GENERATED_BODY()
+
+	FNamedBlacklistConfig() :
+		bBlacklistAll(false)
+	{
+	}
+
+	/** Optional name for filter */
+	UPROPERTY()
+	FName Name;
+
+	/** Items to filter out */
+	UPROPERTY()
+	TArray<FName> Blacklist;
+							
+	/** Items not on the list will be filtered out when whitelist contains items */
+	UPROPERTY()
+	TArray<FName> Whitelist;
+				
+	/** Filter out all items */
+	UPROPERTY()
+	bool bBlacklistAll;
+
+	/** Append filtering rules in this config friendly struct to the destination filter, specifying an owner name allows unregister later */
+	void AppendTo(const FName OwnerName, FNamedBlacklist& Destination) const
+	{
+		for (const FName ItemName : Blacklist)
+		{
+			Destination.AddBlacklistItem(OwnerName, ItemName);
+		}
+
+		for (const FName ItemName : Whitelist)
+		{
+			Destination.AddWhitelistItem(OwnerName, ItemName);
+		}
+
+		if (bBlacklistAll)
+		{
+			Destination.AddBlacklistAll(OwnerName);
+		}
+	}
+};
diff --git a/Engine/Source/Runtime/Slate/Private/Framework/Docking/TabManager.cpp b/Engine/Source/Runtime/Slate/Private/Framework/Docking/TabManager.cpp
--- a/Engine/Source/Runtime/Slate/Private/Framework/Docking/TabManager.cpp
+++ b/Engine/Source/Runtime/Slate/Private/Framework/Docking/TabManager.cpp
@@ -1021,7 +1021,10 @@ void FTabManager::PopulateTabSpawnerMenu( FMenuBuilder& PopulateMe, TSharedRef<F
 		const TSharedRef<FTabSpawnerEntry>& SpawnerEntry = SpawnerIterator.Value();
 		if ( SpawnerEntry->bAutoGenerateMenuEntry )
 		{
-			AllSpawners->AddUnique(SpawnerEntry);
+			if (SpawnerEntry->TabType == NAME_None || TabBlacklist.PassesFilter(SpawnerEntry->TabType))
+			{
+				AllSpawners->AddUnique(SpawnerEntry);
+			}
 		}
 	}
 
@@ -1031,7 +1034,10 @@ void FTabManager::PopulateTabSpawnerMenu( FMenuBuilder& PopulateMe, TSharedRef<F
 		const TSharedRef<FTabSpawnerEntry>& SpawnerEntry = SpawnerIterator.Value();
 		if ( SpawnerEntry->bAutoGenerateMenuEntry )
 		{
-			AllSpawners->AddUnique(SpawnerEntry);
+			if (SpawnerEntry->TabType == NAME_None || TabBlacklist.PassesFilter(SpawnerEntry->TabType))
+			{
+				AllSpawners->AddUnique(SpawnerEntry);
+			}
 		}
 	}
 
@@ -1618,6 +1624,11 @@ bool FTabManager::HasTabSpawner(FName TabId) const
 	return Spawner != nullptr;
 }
 
+FNamedBlacklist& FTabManager::GetTabBlacklist()
+{
+	return TabBlacklist;
+}
+
 bool FTabManager::IsValidTabForSpawning( const FTab& SomeTab ) const
 {
 	// Nomad tabs being restored from layouts should not be spawned if the nomad tab is already spawned.
diff --git a/Engine/Source/Runtime/Slate/Public/Framework/Docking/TabManager.h b/Engine/Source/Runtime/Slate/Public/Framework/Docking/TabManager.h
--- a/Engine/Source/Runtime/Slate/Public/Framework/Docking/TabManager.h
+++ b/Engine/Source/Runtime/Slate/Public/Framework/Docking/TabManager.h
@@ -5,6 +5,7 @@
 #include "CoreMinimal.h"
 #include "SlateFwd.h"
 #include "Misc/Attribute.h"
+#include "Misc/NamedBlacklist.h"
 #include "Textures/SlateIcon.h"
 #include "Widgets/SWindow.h"
 #include "Framework/Docking/WorkspaceItem.h"
@@ -801,6 +802,9 @@ class SLATE_API FTabManager : public TSharedFromThis<FTabManager>
 		/** Returns the owner tab (if it exists) */
 		TSharedPtr<SDockTab> GetOwnerTab() { return OwnerTabPtr.Pin(); }
 
+		/** Returns filter for additional control over available tabs */
+		FNamedBlacklist& GetTabBlacklist();
+
 	protected:
 		void InvokeTabForMenu( FName TabId );
 
@@ -966,6 +970,9 @@ class SLATE_API FTabManager : public TSharedFromThis<FTabManager>
 
 		/* Prevent or allow Drag operation. */
 		bool bCanDoDragOperation;
+
+		/** Allow systems to dynamically hide tabs */
+		FNamedBlacklist TabBlacklist;
 };
 
 
