When we set the sparse class data structure we need to clean up the old sparse data if it was present.

diff --git a/Engine/Source/Runtime/CoreUObject/Private/UObject/Class.cpp b/Engine/Source/Runtime/CoreUObject/Private/UObject/Class.cpp
--- a/Engine/Source/Runtime/CoreUObject/Private/UObject/Class.cpp
+++ b/Engine/Source/Runtime/CoreUObject/Private/UObject/Class.cpp
@@ -3720,7 +3720,15 @@ void UClass::SetSuperStruct(UStruct* NewSuperStruct)
 	UnhashObject(this);
 	ClearFunctionMapsCaches();
 	Super::SetSuperStruct(NewSuperStruct);
-	SetSparseClassDataStruct(GetSparseClassDataArchetypeStruct());
+
+	if (!GetSparseClassDataStruct())
+	{
+		if (UScriptStruct* SparseClassDataStructArchetype = GetSparseClassDataArchetypeStruct())
+		{
+			SetSparseClassDataStruct(SparseClassDataStructArchetype);
+		}
+	}
+
 	HashObject(this);
 }
 
@@ -4319,7 +4327,17 @@ UScriptStruct* UClass::GetSparseClassDataStruct() const
 
 void UClass::SetSparseClassDataStruct(UScriptStruct* InSparseClassDataStruct)
 { 
-	SparseClassDataStruct = InSparseClassDataStruct; 
+	if (SparseClassDataStruct != InSparseClassDataStruct)
+	{
+		SparseClassDataStruct = InSparseClassDataStruct;
+
+		// the old type and new type may not match when we do a hot reload
+		if (SparseClassData != nullptr)
+		{
+			delete SparseClassData;
+			SparseClassData = nullptr;
+		}
+	}
 }
 
 #if WITH_HOT_RELOAD
