Back out previous change to unbreak the linux build.

diff --git a/Engine/Source/Runtime/CoreUObject/Private/UObject/Class.cpp b/Engine/Source/Runtime/CoreUObject/Private/UObject/Class.cpp
--- a/Engine/Source/Runtime/CoreUObject/Private/UObject/Class.cpp
+++ b/Engine/Source/Runtime/CoreUObject/Private/UObject/Class.cpp
@@ -3720,15 +3720,7 @@ void UClass::SetSuperStruct(UStruct* NewSuperStruct)
 	UnhashObject(this);
 	ClearFunctionMapsCaches();
 	Super::SetSuperStruct(NewSuperStruct);
-
-	if (!GetSparseClassDataStruct())
-	{
-		if (UScriptStruct* SparseClassDataStructArchetype = GetSparseClassDataArchetypeStruct())
-		{
-			SetSparseClassDataStruct(SparseClassDataStructArchetype);
-		}
-	}
-
+	SetSparseClassDataStruct(GetSparseClassDataArchetypeStruct());
 	HashObject(this);
 }
 
@@ -4327,17 +4319,7 @@ UScriptStruct* UClass::GetSparseClassDataStruct() const
 
 void UClass::SetSparseClassDataStruct(UScriptStruct* InSparseClassDataStruct)
 { 
-	if (SparseClassDataStruct != InSparseClassDataStruct)
-	{
-		SparseClassDataStruct = InSparseClassDataStruct;
-
-		// the old type and new type may not match when we do a hot reload
-		if (SparseClassData != nullptr)
-		{
-			delete SparseClassData;
-			SparseClassData = nullptr;
-		}
-	}
+	SparseClassDataStruct = InSparseClassDataStruct; 
 }
 
 #if WITH_HOT_RELOAD
