Added a new function for cleaning up the sparse class data.
Use FMemory::Malloc and FMemory::Free to allocate and deallocate the sparse class data. This removes the need for virtual functions on individual classes doing this.
Removed code gen that was creating virtual function implementations for allocating sparse class data based on the class.

#rnx

diff --git a/Engine/Source/Programs/UnrealHeaderTool/Private/CodeGenerator.cpp b/Engine/Source/Programs/UnrealHeaderTool/Private/CodeGenerator.cpp
--- a/Engine/Source/Programs/UnrealHeaderTool/Private/CodeGenerator.cpp
+++ b/Engine/Source/Programs/UnrealHeaderTool/Private/CodeGenerator.cpp
@@ -4953,12 +4953,6 @@ void FNativeClassHeaderGenerator::ExportNativeFunctions(FOutputDevice& OutGenera
 	FString FullClassName = ((FClass*)Class)->GetNameWithPrefix();
 	for (const FString& SparseClassDataString : SparseClassDataTypes)
 	{
-		RuntimeStringBuilders.AutogeneratedStaticData += TEXT("virtual void* CreateSparseClassData() override\r\n");
-		RuntimeStringBuilders.AutogeneratedStaticData += TEXT("{\r\n");
-		RuntimeStringBuilders.AutogeneratedStaticData.Logf(TEXT("\tvoid* SparseClassDataObj = new F%s;\r\n"), *SparseClassDataString);
-		RuntimeStringBuilders.AutogeneratedStaticData.Logf(TEXT("\treturn SparseClassDataObj;\r\n"), *SparseClassDataString);
-		RuntimeStringBuilders.AutogeneratedStaticData += TEXT("}\r\n");
-
 		RuntimeStringBuilders.AutogeneratedStaticData.Logf(TEXT("F%s* Get%s()\r\n"), *SparseClassDataString, *SparseClassDataString);
 		RuntimeStringBuilders.AutogeneratedStaticData += TEXT("{\r\n");
 		RuntimeStringBuilders.AutogeneratedStaticData.Logf(TEXT("\treturn (F%s*)(GetClass()->GetOrCreateSparseClassData());\r\n"), *SparseClassDataString);
diff --git a/Engine/Source/Runtime/CoreUObject/Private/UObject/Class.cpp b/Engine/Source/Runtime/CoreUObject/Private/UObject/Class.cpp
--- a/Engine/Source/Runtime/CoreUObject/Private/UObject/Class.cpp
+++ b/Engine/Source/Runtime/CoreUObject/Private/UObject/Class.cpp
@@ -3906,7 +3906,15 @@ void UClass::SetSuperStruct(UStruct* NewSuperStruct)
 	UnhashObject(this);
 	ClearFunctionMapsCaches();
 	Super::SetSuperStruct(NewSuperStruct);
-	SetSparseClassDataStruct(GetSparseClassDataArchetypeStruct());
+
+	if (!GetSparseClassDataStruct())
+	{
+		if (UScriptStruct* SparseClassDataStructArchetype = GetSparseClassDataArchetypeStruct())
+		{
+			SetSparseClassDataStruct(SparseClassDataStructArchetype);
+		}
+	}
+
 	HashObject(this);
 }
 
@@ -4475,8 +4483,11 @@ void* UClass::CreateSparseClassData()
 {
 	check(SparseClassData == nullptr);
 
-	SparseClassData = GetDefaultObject()->CreateSparseClassData();
-
+	if (SparseClassDataStruct)
+	{
+		SparseClassData = FMemory::Malloc(SparseClassDataStruct->GetStructureSize(), SparseClassDataStruct->GetMinAlignment());
+		SparseClassDataStruct->GetCppStructOps()->Construct(SparseClassData);
+	}
 	if (SparseClassData)
 	{
 		// initialize per class data from the archetype if we have one
@@ -4495,6 +4506,16 @@ void* UClass::CreateSparseClassData()
 	return SparseClassData;
 }
 
+void UClass::CleanupSparseClassData()
+{
+	if (SparseClassData)
+	{
+		SparseClassDataStruct->GetCppStructOps()->Destruct(SparseClassData);
+		FMemory::Free(SparseClassData);
+		SparseClassData = nullptr;
+	}
+}
+
 UScriptStruct* UClass::GetSparseClassDataStruct() const
 {
 	// this info is specified on the object via code generation so we use it instead of looking at the UClass
@@ -4503,7 +4524,13 @@ UScriptStruct* UClass::GetSparseClassDataStruct() const
 
 void UClass::SetSparseClassDataStruct(UScriptStruct* InSparseClassDataStruct)
 { 
-	SparseClassDataStruct = InSparseClassDataStruct; 
+	if (SparseClassDataStruct != InSparseClassDataStruct)
+	{
+		SparseClassDataStruct = InSparseClassDataStruct;
+
+		// the old type and new type may not match when we do a hot reload so get rid of the old data
+		CleanupSparseClassData();
+	}
 }
 
 #if WITH_HOT_RELOAD
diff --git a/Engine/Source/Runtime/CoreUObject/Private/UObject/Obj.cpp b/Engine/Source/Runtime/CoreUObject/Private/UObject/Obj.cpp
--- a/Engine/Source/Runtime/CoreUObject/Private/UObject/Obj.cpp
+++ b/Engine/Source/Runtime/CoreUObject/Private/UObject/Obj.cpp
@@ -1145,11 +1145,6 @@ void UObject::PostLoadSubobjects( FObjectInstancingGraph* OuterInstanceGraph/*=N
 	}
 }
 
-void* UObject::CreateSparseClassData()
-{
-	return nullptr;
-}
-
 UScriptStruct* UObject::GetSparseClassDataStruct() const
 {
 	UClass* Class = GetClass();
diff --git a/Engine/Source/Runtime/CoreUObject/Public/UObject/Class.h b/Engine/Source/Runtime/CoreUObject/Public/UObject/Class.h
--- a/Engine/Source/Runtime/CoreUObject/Public/UObject/Class.h
+++ b/Engine/Source/Runtime/CoreUObject/Public/UObject/Class.h
@@ -2549,6 +2549,8 @@ public:
 private:
 	void* CreateSparseClassData();
 
+	void CleanupSparseClassData();
+
 #if WITH_EDITOR
 	/** Provides access to attributes of the underlying C++ class. Should never be unset. */
 	TOptional<FCppClassTypeInfo> CppTypeInfo;
diff --git a/Engine/Source/Runtime/CoreUObject/Public/UObject/Object.h b/Engine/Source/Runtime/CoreUObject/Public/UObject/Object.h
--- a/Engine/Source/Runtime/CoreUObject/Public/UObject/Object.h
+++ b/Engine/Source/Runtime/CoreUObject/Public/UObject/Object.h
@@ -541,9 +541,6 @@ public:
 	/** Return a one line description of an object for viewing in the thumbnail view of the generic browser */
 	virtual FString GetDesc() { return TEXT( "" ); }
 
-	/** Allocates the sidecar data structure that stores data that is constant for all instances of this class. */
-	virtual void* CreateSparseClassData();
-
 	/** Return the UStruct corresponding to the sidecar data structure that stores data that is constant for all instances of this class. */
 	virtual UScriptStruct* GetSparseClassDataStruct() const;
 
