Fix BulkData being unable to create streaming requests from uncooked assets in the editor.

Added a check for GEventDrivenLoaderEnabled to the logic that determines if we should convert a uasset/umap filepath reference to a uexp reference or not, matching all other places in BulkData where this is done.

#rnx

diff --git a/Engine/Source/Runtime/CoreUObject/Private/Serialization/BulkData.cpp b/Engine/Source/Runtime/CoreUObject/Private/Serialization/BulkData.cpp
--- a/Engine/Source/Runtime/CoreUObject/Private/Serialization/BulkData.cpp
+++ b/Engine/Source/Runtime/CoreUObject/Private/Serialization/BulkData.cpp
@@ -1606,12 +1606,12 @@ IBulkDataIORequest* FUntypedBulkData::CreateStreamingRequest(int64 OffsetInBulkD
 	FString AdjustedFilename = Filename; 
 	int64 AdjustedBulkDataOffsetInFile = BulkDataOffsetInFile;
 
-	// Fix up the Filename/Offset to work with streaming if we are loading from a .uexp file
-	if (AdjustedFilename.EndsWith(TEXT(".uasset")) || AdjustedFilename.EndsWith(TEXT(".umap")))
+	// Fix up the Filename/Offset to work with streaming if EDL is enabled and the filename is still referencing a uasset or umap
+	if (GEventDrivenLoaderEnabled && (AdjustedFilename.EndsWith(TEXT(".uasset")) || AdjustedFilename.EndsWith(TEXT(".umap"))))
 	{
 		AdjustedBulkDataOffsetInFile -= IFileManager::Get().FileSize(*AdjustedFilename);
-
 		AdjustedFilename = FPaths::GetBaseFilename(AdjustedFilename, false) + TEXT(".uexp");
+
 		UE_LOG(LogSerialization, Error, TEXT("Streaming from the .uexp file '%s' this MUST be in a ubulk instead for best performance."), *AdjustedFilename);
 	}
 
