Fix uninitialized variable warning
#rnx

diff --git a/Engine/Source/Runtime/CoreUObject/Private/UObject/EnumProperty.cpp b/Engine/Source/Runtime/CoreUObject/Private/UObject/EnumProperty.cpp
--- a/Engine/Source/Runtime/CoreUObject/Private/UObject/EnumProperty.cpp
+++ b/Engine/Source/Runtime/CoreUObject/Private/UObject/EnumProperty.cpp
@@ -367,7 +367,7 @@ EConvertFromTypeResult UEnumProperty::ConvertFromType(const FPropertyTag& Tag, F
 
 	if (Tag.Type == NAME_ByteProperty)
 	{
-		uint8 PreviousValue;
+		uint8 PreviousValue = 0;
 		if (Tag.EnumName == NAME_None)
 		{
 			// If we're a nested property the EnumName tag got lost. Handle this case for backward compatibility reasons
diff --git a/Engine/Source/Runtime/CoreUObject/Private/UObject/PropertyByte.cpp b/Engine/Source/Runtime/CoreUObject/Private/UObject/PropertyByte.cpp
--- a/Engine/Source/Runtime/CoreUObject/Private/UObject/PropertyByte.cpp
+++ b/Engine/Source/Runtime/CoreUObject/Private/UObject/PropertyByte.cpp
@@ -177,7 +177,7 @@ EConvertFromTypeResult UByteProperty::ConvertFromType(const FPropertyTag& Tag, F
 	{
 		// a byte property gained or lost an enum
 		// attempt to convert it
-		uint8 PreviousValue;
+		uint8 PreviousValue = 0;
 		if (Tag.EnumName == NAME_None)
 		{
 			// If we're a nested property the EnumName tag got lost. Fail to read in this case
