When doing incremental unshashing, Garbage Collection will no longer log each iteration to reduce log output

diff --git a/Engine/Source/Runtime/CoreUObject/Private/UObject/GarbageCollection.cpp b/Engine/Source/Runtime/CoreUObject/Private/UObject/GarbageCollection.cpp
--- a/Engine/Source/Runtime/CoreUObject/Private/UObject/GarbageCollection.cpp
+++ b/Engine/Source/Runtime/CoreUObject/Private/UObject/GarbageCollection.cpp
@@ -1986,6 +1986,7 @@ bool UnhashUnreachableObjects(bool bUseTimeLimit, float TimeLimit)
 	const int32 TimeLimitEnforcementGranularityForBeginDestroy = 10;
 	int32 Items = 0;
 	int32 TimePollCounter = 0;
+	const bool bFirstIteration = (GUnrechableObjectIndex == 0);
 
 	while (GUnrechableObjectIndex < GUnreachableObjects.Num())
 	{
@@ -2008,16 +2009,31 @@ bool UnhashUnreachableObjects(bool bUseTimeLimit, float TimeLimit)
 		}
 	}
 
-	UE_LOG(LogGarbage, Log, TEXT("%f ms for %sunhashing unreachable objects. Items %d (%d/%d)"),
-		(FPlatformTime::Seconds() - StartTime) * 1000,
-		bUseTimeLimit ? TEXT("incrementally ") : TEXT(""),
-		Items,
-		GUnrechableObjectIndex, GUnreachableObjects.Num());
+	const bool bTimeLimitReached = (GUnrechableObjectIndex < GUnreachableObjects.Num());
+
+	if (!bUseTimeLimit)
+	{
+		UE_LOG(LogGarbage, Log, TEXT("%f ms for %sunhashing unreachable objects (%d objects unhashed)"),
+			(FPlatformTime::Seconds() - StartTime) * 1000,
+			bUseTimeLimit ? TEXT("incrementally ") : TEXT(""),
+			Items,
+			GUnrechableObjectIndex, GUnreachableObjects.Num());
+	}
+	else if (!bTimeLimitReached)
+	{
+		// When doing incremental unhashing log only the first and last iteration (this was the last one)
+		UE_LOG(LogGarbage, Log, TEXT("Finished unhashing unreachable objects (%d objects unhashed)."), GUnreachableObjects.Num());
+	}
+	else if (bFirstIteration)
+	{
+		// When doing incremental unhashing log only the first and last iteration (this was the first one)
+		UE_LOG(LogGarbage, Log, TEXT("Starting unhashing unreachable objects (%d objects to unhash)."), GUnreachableObjects.Num());
+	}
 
 	FCoreUObjectDelegates::PostGarbageCollectConditionalBeginDestroy.Broadcast();
 
 	// Return true if time limit has been reached
-	return GUnrechableObjectIndex < GUnreachableObjects.Num();
+	return bTimeLimitReached;
 }
 
 void CollectGarbage(EObjectFlags KeepFlags, bool bPerformFullPurge)
