#UE4 Fix crash that can happen if a non main thread callls FilenameToLongPackageName while content mount points are being modified

diff --git a/Engine/Source/Runtime/CoreUObject/Private/Misc/PackageName.cpp b/Engine/Source/Runtime/CoreUObject/Private/Misc/PackageName.cpp
--- a/Engine/Source/Runtime/CoreUObject/Private/Misc/PackageName.cpp
+++ b/Engine/Source/Runtime/CoreUObject/Private/Misc/PackageName.cpp
@@ -17,8 +17,10 @@
 #include "UObject/Linker.h"
 #include "Interfaces/IPluginManager.h"
 #include "Internationalization/PackageLocalizationManager.h"
+#include "HAL/CriticalSection.h"
 #include "HAL/ThreadHeartBeat.h"
 #include "Misc/AutomationTest.h"
+#include "Misc/ScopeLock.h"
 
 DEFINE_LOG_CATEGORY_STATIC(LogPackageName, Log, All);
 
@@ -26,6 +28,7 @@ FString FPackageName::AssetPackageExtension = TEXT(".uasset");
 FString FPackageName::MapPackageExtension = TEXT(".umap");
 FString FPackageName::TextAssetPackageExtension = TEXT(".utxt");
 FString FPackageName::TextMapPackageExtension = TEXT(".utxtmap");
+static FCriticalSection ContentMountPointCriticalSection;
 
 /** Event that is triggered when a new content path is mounted */
 FPackageName::FOnContentPathMountedEvent FPackageName::OnContentPathMountedEvent;
@@ -152,7 +155,11 @@ struct FLongPackagePathsSingleton
 	{
 		OutRoots.Add(EngineRootPath);
 		OutRoots.Add(GameRootPath);
-		OutRoots += MountPointRootPaths;
+
+		{
+			FScopeLock ScopeLock(&ContentMountPointCriticalSection);
+			OutRoots += MountPointRootPaths;
+		}
 
 		if (bIncludeReadOnlyRoots)
 		{
@@ -199,9 +206,12 @@ struct FLongPackagePathsSingleton
 		}
 
 		FPathPair Pair(RootPath, RelativeContentPath);
-		ContentRootToPath.Insert(Pair, 0);
-		ContentPathToRoot.Insert(Pair, 0);
-		MountPointRootPaths.Add( RootPath );
+		{
+			FScopeLock ScopeLock(&ContentMountPointCriticalSection);
+			ContentRootToPath.Insert(Pair, 0);
+			ContentPathToRoot.Insert(Pair, 0);
+			MountPointRootPaths.Add(RootPath);
+		}
 
 		// Let subscribers know that a new content path was mounted
 		FPackageName::OnContentPathMounted().Broadcast( RootPath, RelativeContentPath);
@@ -219,15 +229,24 @@ struct FLongPackagePathsSingleton
 			RelativeContentPath += TEXT("/");
 		}
 
-		if ( MountPointRootPaths.Remove(RootPath) > 0 )
+		bool bFirePathDismountedDelegate = false;
 		{
-			FPathPair Pair(RootPath, RelativeContentPath);
-			ContentRootToPath.Remove(Pair);
-			ContentPathToRoot.Remove(Pair);
-			MountPointRootPaths.Remove(RootPath);
+			FScopeLock ScopeLock(&ContentMountPointCriticalSection);
+			if ( MountPointRootPaths.Remove(RootPath) > 0 )
+			{
+				FPathPair Pair(RootPath, RelativeContentPath);
+				ContentRootToPath.Remove(Pair);
+				ContentPathToRoot.Remove(Pair);
+				MountPointRootPaths.Remove(RootPath);
+
+				// Let subscribers know that a new content path was unmounted
+				bFirePathDismountedDelegate = true;
+			}
+		}
 
-			// Let subscribers know that a new content path was unmounted
-			FPackageName::OnContentPathDismounted().Broadcast( RootPath, RelativeContentPath);
+		if (bFirePathDismountedDelegate)
+		{
+			FPackageName::OnContentPathDismounted().Broadcast(RootPath, RelativeContentPath);
 		}
 	}
 
@@ -257,6 +276,8 @@ private:
 		GameScriptPathRebased  = RebasedGameDir / TEXT("Script/");
 		GameSavedPathRebased   = RebasedGameDir / TEXT("Saved/");
 		
+		FScopeLock ScopeLock(&ContentMountPointCriticalSection);
+
 		ContentPathToRoot.Empty(11);
 		ContentPathToRoot.Emplace(EngineRootPath, EngineContentPath);
 		if (FPaths::IsSamePath(GameContentPath, ContentPathShort))
@@ -302,12 +323,15 @@ FString FPackageName::InternalFilenameToLongPackageName(const FString& InFilenam
 
 	// Convert to relative path if it's not already a long package name
 	bool bIsValidLongPackageName = false;
-	for (const auto& Pair : Paths.ContentRootToPath)
 	{
-		if (Filename.StartsWith(Pair.RootPath))
+		FScopeLock ScopeLock(&ContentMountPointCriticalSection);
+		for (const auto& Pair : Paths.ContentRootToPath)
 		{
-			bIsValidLongPackageName = true;
-			break;
+			if (Filename.StartsWith(Pair.RootPath))
+			{
+				bIsValidLongPackageName = true;
+				break;
+			}
 		}
 	}
 
@@ -331,12 +355,15 @@ FString FPackageName::InternalFilenameToLongPackageName(const FString& InFilenam
 	FString Result              = Filename.Mid(0, PackageNameStartsAt + PackageName.Len());
 	Result.ReplaceInline(TEXT("\\"), TEXT("/"), ESearchCase::CaseSensitive);
 
-	for (const auto& Pair : Paths.ContentPathToRoot)
 	{
-		if (Result.StartsWith(Pair.ContentPath))
+		FScopeLock ScopeLock(&ContentMountPointCriticalSection);
+		for (const auto& Pair : Paths.ContentPathToRoot)
 		{
-			Result = Pair.RootPath + Result.Mid(Pair.ContentPath.Len());
-			break;
+			if (Result.StartsWith(Pair.ContentPath))
+			{
+				Result = Pair.RootPath + Result.Mid(Pair.ContentPath.Len());
+				break;
+			}
 		}
 	}
 
@@ -408,6 +435,7 @@ FString FPackageName::FilenameToLongPackageName(const FString& InFilename)
 bool FPackageName::TryConvertLongPackageNameToFilename(const FString& InLongPackageName, FString& OutFilename, const FString& InExtension)
 {
 	const auto& Paths = FLongPackagePathsSingleton::Get();
+	FScopeLock ScopeLock(&ContentMountPointCriticalSection);
 	for (const auto& Pair : Paths.ContentRootToPath)
 	{
 		if (InLongPackageName.StartsWith(Pair.RootPath))
@@ -424,6 +452,7 @@ bool FPackageName::TryConvertLongPackageNameToFilename(const FString& InLongPack
 bool FPackageName::ConvertRootPathToContentPath( const FString& RootPath, FString& OutContentPath)
 {
 	const auto& Paths = FLongPackagePathsSingleton::Get();
+	FScopeLock ScopeLock(&ContentMountPointCriticalSection);
 	for (const auto& Pair : Paths.ContentRootToPath)
 	{
 		if (RootPath.StartsWith(Pair.RootPath))
@@ -825,6 +854,7 @@ bool FPackageName::FixPackageNameCase(FString& LongPackageName, const FString& E
 {
 	// Find the matching long package root
 	const FLongPackagePathsSingleton& Paths = FLongPackagePathsSingleton::Get();
+	FScopeLock ScopeLock(&ContentMountPointCriticalSection);
 	for (const FPathPair& Pair : Paths.ContentRootToPath)
 	{
 		if (LongPackageName.StartsWith(Pair.RootPath))
