#UE4 Skip fatal message about a missing linker table entry for an object if it is transient. This happens when a transient native subobject is inherited by a BP class then a BP subclass of that class is saved

diff --git a/Engine/Source/Runtime/CoreUObject/Private/UObject/SavePackage.cpp b/Engine/Source/Runtime/CoreUObject/Private/UObject/SavePackage.cpp
--- a/Engine/Source/Runtime/CoreUObject/Private/UObject/SavePackage.cpp
+++ b/Engine/Source/Runtime/CoreUObject/Private/UObject/SavePackage.cpp
@@ -5070,7 +5070,7 @@ FSavePackageResultStruct UPackage::Save(UPackage* InOuter, UObject* Base, EObjec
 									AddTo.Add(Index);
 									return;
 								}
-								else
+								else if (!ToTest->HasAnyFlags(RF_Transient))
 								{
 									UE_CLOG(Outermost->HasAnyPackageFlags(PKG_CompiledIn), LogSavePackage, Verbose, TEXT("A compiled in dependency '%s' of '%s' was not actually in the linker tables and so will be ignored (%d)."), *ToTest->GetFullName(), *ForObj->GetFullName(), CallSite);
 									UE_CLOG(!Outermost->HasAnyPackageFlags(PKG_CompiledIn), LogSavePackage, Fatal, TEXT("A dependency '%s' of '%s' was not actually in the linker tables and so will be ignored (%d)."), *ToTest->GetFullName(), *ForObj->GetFullName(), CallSite);
