Add named events to package saving to help with profiling

diff --git a/Engine/Source/Runtime/CoreUObject/Private/UObject/SavePackage.cpp b/Engine/Source/Runtime/CoreUObject/Private/UObject/SavePackage.cpp
--- a/Engine/Source/Runtime/CoreUObject/Private/UObject/SavePackage.cpp
+++ b/Engine/Source/Runtime/CoreUObject/Private/UObject/SavePackage.cpp
@@ -667,6 +667,9 @@ bool IsEditorOnlyObject(const UObject* InObject, bool bCheckRecursive, bool bChe
 static void ConditionallyExcludeObjectForTarget(UObject* Obj, EObjectMark ExcludedObjectMarks, const ITargetPlatform* TargetPlatform, const bool bIsCooking)
 {
 #if WITH_EDITOR
+
+	SCOPED_NAMED_EVENT(ConditionallyExcludeObjectForTarget, FColor::Emerald);
+
 	if (!Obj || Obj->GetOutermost()->GetFName() == GLongCoreUObjectPackageName)
 	{
 		// No object or in CoreUObject, don't exclude
@@ -3391,12 +3394,13 @@ void AddFileToHash(FString const &Filename, FMD5 &Hash)
 }
 
 FSavePackageResultStruct UPackage::Save(UPackage* InOuter, UObject* Base, EObjectFlags TopLevelFlags, const TCHAR* Filename,
-	FOutputDevice* Error, FLinkerNull* ConformNO, bool bForceByteSwapping, bool bWarnOfLongFilename, uint32 SaveFlags, 
-	const class ITargetPlatform* TargetPlatform, const FDateTime&  FinalTimeStamp, bool bSlowTask, FArchiveDiffMap* InOutDiffMap,
+	FOutputDevice* Error, FLinkerNull* ConformNO, bool bForceByteSwapping, bool bWarnOfLongFilename, uint32 SaveFlags,
+	const class ITargetPlatform* TargetPlatform, const FDateTime& FinalTimeStamp, bool bSlowTask, FArchiveDiffMap* InOutDiffMap,
 	FSavePackageContext* SavePackageContext)
 {
 	COOK_STAT(FScopedDurationTimer FuncSaveTimer(SavePackageStats::SavePackageTimeSec));
 	COOK_STAT(SavePackageStats::NumPackagesSaved++);
+	SCOPED_NAMED_EVENT(UPackage_Save, FColor::Emerald);
 
 	FLinkerLoad* Conform = nullptr;
 
@@ -3651,7 +3655,7 @@ FSavePackageResultStruct UPackage::Save(UPackage* InOuter, UObject* Base, EObjec
 
 		// Size of serialized out package in bytes. This is before compression.
 		int32 PackageSize = INDEX_NONE;
-		{
+		{			
 			// TODO: Require a SavePackageContext and move to EditorEngine
 			FPackageNameMapSaver NameMapSaver;
 
@@ -3669,6 +3673,8 @@ FSavePackageResultStruct UPackage::Save(UPackage* InOuter, UObject* Base, EObjec
 			// Tag exports and route presave.
 			FPackageExportTagger PackageExportTagger(Base, TopLevelFlags, InOuter, TargetPlatform);
 			{
+				SCOPED_NAMED_EVENT(UPackage_Save_TagExportsWithPresave, FColor::Red);
+
 				COOK_STAT(FScopedDurationTimer SaveTimer(SavePackageStats::TagPackageExportsPresaveTimeSec));
 				// Do not route presave if saving concurrently. This should have been done before the concurrent save started.
 				// Also if we're trying to diff the package and gathering callstacks, Presave has already been done
@@ -3737,6 +3743,7 @@ FSavePackageResultStruct UPackage::Save(UPackage* InOuter, UObject* Base, EObjec
 
 			
 				{
+					SCOPED_NAMED_EVENT(UPackage_Save_TagExports, FColor::Green);
 					COOK_STAT(FScopedDurationTimer SaveTimer(SavePackageStats::TagPackageExportsTimeSec));
 					// Clear all marks (OBJECTMARK_TagExp and exclusion marks) again as we need to redo tagging below.
 					UnMarkAllObjects();
@@ -3783,76 +3790,80 @@ FSavePackageResultStruct UPackage::Save(UPackage* InOuter, UObject* Base, EObjec
 				FArchiveFormatterType* Formatter = nullptr;
 				FArchive* TextFormatArchive = nullptr;
 				const bool bTextFormat = FString(Filename).EndsWith(FPackageName::GetTextAssetPackageExtension()) || FString(Filename).EndsWith(FPackageName::GetTextMapPackageExtension());
-				
+	
+				{
+					SCOPED_NAMED_EVENT(UPackage_Save_CreateLinkerSave, FColor::Blue);
+
 #if WITH_EDITOR
-				FString DiffCookedPackagesPath;
-				// if we are cooking and we have diff cooked packages on the commandline then do some special stuff
+					FString DiffCookedPackagesPath;
+					// if we are cooking and we have diff cooked packages on the commandline then do some special stuff
 
-				// Finds the asset object within a package
-				auto FindAssetInPackage = [](UPackage* Package) -> UObject*
-				{
-					for (UObject* Object : TObjectRange<UObject>())
+					// Finds the asset object within a package
+					auto FindAssetInPackage = [](UPackage* Package) -> UObject*
 					{
-						if (Object->GetOuter() == Package && Object->IsAsset())
+						for (UObject* Object : TObjectRange<UObject>())
 						{
-							return Object;
+							if (Object->GetOuter() == Package && Object->IsAsset())
+							{
+								return Object;
+							}
 						}
-					}
 
-					return nullptr;
-				};
+						return nullptr;
+					};
 
-				if (TargetPlatform != nullptr && (SaveFlags & SAVE_DiffCallstack))
-				{
-					// The entire package will be serialized to memory and then compared against package on disk.
-					// Each difference will be log with its Serialize call stack trace
-					FArchive* Saver = new FArchiveStackTrace(FindAssetInPackage(InOuter), *InOuter->FileName.ToString(), true, InOutDiffMap);
-					Linker = TUniquePtr<FLinkerSave>(new FLinkerSave(InOuter, Saver, bForceByteSwapping, bSaveUnversioned));
-				}
-				else if (TargetPlatform != nullptr && (SaveFlags & SAVE_DiffOnly))
-				{
-					// The entire package will be serialized to memory and then compared against package on disk
-					FArchive* Saver = new FArchiveStackTrace(FindAssetInPackage(InOuter), *InOuter->FileName.ToString(), false);
-					Linker = TUniquePtr<FLinkerSave>(new FLinkerSave(InOuter, Saver, bForceByteSwapping, bSaveUnversioned));
-				}
-				else if ((!!TargetPlatform) && FParse::Value(FCommandLine::Get(), TEXT("DiffCookedPackages="), DiffCookedPackagesPath))
-				{
-					FString TestArchiveFilename = Filename;
-					// TestArchiveFilename.ReplaceInline(TEXT("Cooked"), TEXT("CookedDiff"));
-					DiffCookedPackagesPath.ReplaceInline(TEXT("\\"), TEXT("/"));
-					FString CookedPath = FPaths::ConvertRelativePathToFull(FPaths::ProjectSavedDir() + TEXT("Cooked/"));
-					CookedPath.ReplaceInline(TEXT("\\"), TEXT("/"));
-					TestArchiveFilename.ReplaceInline(*CookedPath, *DiffCookedPackagesPath);
-					
-					FArchive* TestArchive = IFileManager::Get().CreateFileReader(*TestArchiveFilename); 
-					FArchive* Saver = new FDiffSerializeArchive(*InOuter->FileName.ToString(), TestArchive);
-					Linker = TUniquePtr<FLinkerSave>(new FLinkerSave(InOuter, Saver, bForceByteSwapping));
-				}
-				else 
+					if (TargetPlatform != nullptr && (SaveFlags & SAVE_DiffCallstack))
+					{
+						// The entire package will be serialized to memory and then compared against package on disk.
+						// Each difference will be log with its Serialize call stack trace
+						FArchive* Saver = new FArchiveStackTrace(FindAssetInPackage(InOuter), *InOuter->FileName.ToString(), true, InOutDiffMap);
+						Linker = TUniquePtr<FLinkerSave>(new FLinkerSave(InOuter, Saver, bForceByteSwapping, bSaveUnversioned));
+					}
+					else if (TargetPlatform != nullptr && (SaveFlags & SAVE_DiffOnly))
+					{
+						// The entire package will be serialized to memory and then compared against package on disk
+						FArchive* Saver = new FArchiveStackTrace(FindAssetInPackage(InOuter), *InOuter->FileName.ToString(), false);
+						Linker = TUniquePtr<FLinkerSave>(new FLinkerSave(InOuter, Saver, bForceByteSwapping, bSaveUnversioned));
+					}
+					else if ((!!TargetPlatform) && FParse::Value(FCommandLine::Get(), TEXT("DiffCookedPackages="), DiffCookedPackagesPath))
+					{
+						FString TestArchiveFilename = Filename;
+						// TestArchiveFilename.ReplaceInline(TEXT("Cooked"), TEXT("CookedDiff"));
+						DiffCookedPackagesPath.ReplaceInline(TEXT("\\"), TEXT("/"));
+						FString CookedPath = FPaths::ConvertRelativePathToFull(FPaths::ProjectSavedDir() + TEXT("Cooked/"));
+						CookedPath.ReplaceInline(TEXT("\\"), TEXT("/"));
+						TestArchiveFilename.ReplaceInline(*CookedPath, *DiffCookedPackagesPath);
+
+						FArchive* TestArchive = IFileManager::Get().CreateFileReader(*TestArchiveFilename);
+						FArchive* Saver = new FDiffSerializeArchive(*InOuter->FileName.ToString(), TestArchive);
+						Linker = TUniquePtr<FLinkerSave>(new FLinkerSave(InOuter, Saver, bForceByteSwapping));
+					}
+					else
 #endif
-				if (bSaveAsync)
-				{
-					// Allocate the linker with a memory writer, forcing byte swapping if wanted.
-					Linker = TUniquePtr<FLinkerSave>(new FLinkerSave(InOuter, bForceByteSwapping, bSaveUnversioned));
-				}
-				else
-				{
-					// Allocate the linker, forcing byte swapping if wanted.
-					Linker = TUniquePtr<FLinkerSave>(new FLinkerSave(InOuter, *TempFilename, bForceByteSwapping, bSaveUnversioned));
-				}
+						if (bSaveAsync)
+						{
+							// Allocate the linker with a memory writer, forcing byte swapping if wanted.
+							Linker = TUniquePtr<FLinkerSave>(new FLinkerSave(InOuter, bForceByteSwapping, bSaveUnversioned));
+						}
+						else
+						{
+							// Allocate the linker, forcing byte swapping if wanted.
+							Linker = TUniquePtr<FLinkerSave>(new FLinkerSave(InOuter, *TempFilename, bForceByteSwapping, bSaveUnversioned));
+						}
 
 #if WITH_TEXT_ARCHIVE_SUPPORT
-				if (bTextFormat)
-				{
-					TextFormatArchive = IFileManager::Get().CreateFileWriter(*TextFormatTempFilename);
-					FJsonArchiveOutputFormatter* OutputFormatter = new FJsonArchiveOutputFormatter(*TextFormatArchive);
-					OutputFormatter->SetObjectIndicesMap(&Linker->ObjectIndicesMap);
-					Formatter = OutputFormatter;
-				}
-				else
+					if (bTextFormat)
+					{
+						TextFormatArchive = IFileManager::Get().CreateFileWriter(*TextFormatTempFilename);
+						FJsonArchiveOutputFormatter* OutputFormatter = new FJsonArchiveOutputFormatter(*TextFormatArchive);
+						OutputFormatter->SetObjectIndicesMap(&Linker->ObjectIndicesMap);
+						Formatter = OutputFormatter;
+					}
+					else
 #endif
-				{
-					Formatter = new FBinaryArchiveFormatter(*(FArchive*)Linker.Get());
+					{
+						Formatter = new FBinaryArchiveFormatter(*(FArchive*)Linker.Get());
+					}
 				}
 
 				FStructuredArchive* StructuredArchive = new FStructuredArchive(*Formatter);
@@ -3959,6 +3970,8 @@ FSavePackageResultStruct UPackage::Save(UPackage* InOuter, UObject* Base, EObjec
 				// Import objects & names.
 				TSet<UPackage*> PrestreamPackages;
 				{
+					SCOPED_NAMED_EVENT(UPackage_Save_TagImports, FColor::Emerald);
+					
 					TArray<UObject*> TagExpObjects;
 					GetObjectsWithAnyMarks(TagExpObjects, OBJECTMARK_TagExp);
 					for(int32 Index = 0; Index < TagExpObjects.Num(); Index++)
@@ -4134,6 +4147,8 @@ FSavePackageResultStruct UPackage::Save(UPackage* InOuter, UObject* Base, EObjec
 
 				// Tag the names for all relevant object, classes, and packages.
 				{
+					SCOPED_NAMED_EVENT(UPackage_Save_TagNames, FColor::Red);
+
 					TArray<UObject*> TagExpImpObjects;
 					GetObjectsWithAnyMarks(TagExpImpObjects, EObjectMark(OBJECTMARK_TagExp|OBJECTMARK_TagImp));
 					for(int32 Index = 0; Index < TagExpImpObjects.Num(); Index++)
@@ -4421,6 +4436,7 @@ FSavePackageResultStruct UPackage::Save(UPackage* InOuter, UObject* Base, EObjec
 				// Build NameMap.
 				Linker->Summary.NameOffset = Linker->Tell();
 				{
+					SCOPED_NAMED_EVENT(UPackage_Save_BuildNameMap, FColor::Red);
 #if WITH_EDITOR
 					FArchive::FScopeSetDebugSerializationFlags S(*Linker, DSF_IgnoreDiff, true);
 					FArchiveStackTraceIgnoreScope IgnoreSummaryDiffsScope(DiffSettings.bIgnoreHeaderDiffs);
@@ -4438,6 +4454,8 @@ FSavePackageResultStruct UPackage::Save(UPackage* InOuter, UObject* Base, EObjec
 				Linker->Summary.GatherableTextDataCount = 0;
 				if (!(Linker->Summary.PackageFlags & PKG_FilterEditorOnly))
 				{
+					SCOPED_NAMED_EVENT(UPackage_Save_WriteGatherableTextData, FColor::Yellow);
+
 					// The Editor version is used as part of the check to see if a package is too old to use the gather cache, so we always have to add it if we have gathered loc for this asset
 					// Note that using custom version here only works because we already added it to the export tagger before the package summary was serialized
 					Linker->UsingCustomVersion(FEditorObjectVersion::GUID);
@@ -4474,6 +4492,8 @@ FSavePackageResultStruct UPackage::Save(UPackage* InOuter, UObject* Base, EObjec
 
 				// Build ImportMap.
 				{
+					SCOPED_NAMED_EVENT(UPackage_Save_BuildImportMap, FColor::Orange);
+
 					TArray<UObject*> TagImpObjects;
 
 					const EObjectMark ExcludedObjectMarks = GetExcludedObjectMarksForTargetPlatform(TargetPlatform, Linker->IsCooking());
@@ -4545,8 +4565,11 @@ FSavePackageResultStruct UPackage::Save(UPackage* InOuter, UObject* Base, EObjec
 
 				// sort and conform imports
 				FObjectImportSortHelper ImportSortHelper;
-				ImportSortHelper.SortImports( Linker.Get(), Conform );
-				Linker->Summary.ImportCount = Linker->ImportMap.Num();
+				{
+					SCOPED_NAMED_EVENT(UPackage_Save_SortImports, FColor::Red);
+					ImportSortHelper.SortImports(Linker.Get(), Conform);
+					Linker->Summary.ImportCount = Linker->ImportMap.Num();
+				}
 
 				if ( EndSavingIfCancelled( Linker.Get(), TempFilename ) )
 				{ 
@@ -4557,6 +4580,8 @@ FSavePackageResultStruct UPackage::Save(UPackage* InOuter, UObject* Base, EObjec
 				
 				// Build ExportMap.
 				{
+					SCOPED_NAMED_EVENT(UPackage_Save_BuildExportMap, FColor::Turquoise);
+
 					TArray<UObject*> TagExpObjects;
 					GetObjectsWithAnyMarks(TagExpObjects, OBJECTMARK_TagExp);
 					for(int32 Index = 0; Index < TagExpObjects.Num(); Index++)
@@ -4593,10 +4618,14 @@ FSavePackageResultStruct UPackage::Save(UPackage* InOuter, UObject* Base, EObjec
 
 				// Sort exports alphabetically and conform the export table (if necessary)
 				FObjectExportSortHelper ExportSortHelper;
-				ExportSortHelper.SortExports( Linker.Get(), Conform );
+				{
+					SCOPED_NAMED_EVENT(UPackage_Save_SortExports, FColor::Turquoise);
+					ExportSortHelper.SortExports(Linker.Get(), Conform);
+				}
 				
 				// Sort exports for seek-free loading.
 				{
+					SCOPED_NAMED_EVENT(UPackage_Save_SortExportsForSeekFree, FColor::Turquoise);
 					COOK_STAT(FScopedDurationTimer SaveTimer(SavePackageStats::SortExportsSeekfreeInnerTimeSec));
 					FObjectExportSeekFreeSorter SeekFreeSorter;
 					SeekFreeSorter.SortExports( Linker.Get(), Conform );
@@ -4627,64 +4656,67 @@ FSavePackageResultStruct UPackage::Save(UPackage* InOuter, UObject* Base, EObjec
 				}
 
 				// go back over the (now sorted) exports and fill out the DependsMap
-				for (int32 ExpIndex = 0; ExpIndex < Linker->ExportMap.Num(); ExpIndex++)
 				{
-					UObject* Object = Linker->ExportMap[ExpIndex].Object;
-					// sorting while conforming can create NULL export map entries, so skip those depends map
-					if (Object == nullptr)
+					SCOPED_NAMED_EVENT(UPackage_Save_BuildExportDependsMap, FColor::Orange);
+					for (int32 ExpIndex = 0; ExpIndex < Linker->ExportMap.Num(); ExpIndex++)
 					{
-						UE_LOG(LogSavePackage, Warning, TEXT("Object is missing for an export, unable to save dependency map. Most likely this is caused my conforming against a package that is missing this object. See log for more info"));
-						if (!(SaveFlags & SAVE_NoError))
+						UObject* Object = Linker->ExportMap[ExpIndex].Object;
+						// sorting while conforming can create NULL export map entries, so skip those depends map
+						if (Object == nullptr)
 						{
-							Error->Logf(ELogVerbosity::Warning, TEXT("%s"), *FText::Format( NSLOCTEXT( "Core", "SavePackageObjectIsMissingExport", "Object is missing for an export, unable to save dependency map for asset '{0}'. Most likely this is caused my conforming against a asset that is missing this object. See log for more info" ), FText::FromString( FString( Filename ) ) ).ToString() );
+							UE_LOG(LogSavePackage, Warning, TEXT("Object is missing for an export, unable to save dependency map. Most likely this is caused my conforming against a package that is missing this object. See log for more info"));
+							if (!(SaveFlags & SAVE_NoError))
+							{
+								Error->Logf(ELogVerbosity::Warning, TEXT("%s"), *FText::Format(NSLOCTEXT("Core", "SavePackageObjectIsMissingExport", "Object is missing for an export, unable to save dependency map for asset '{0}'. Most likely this is caused my conforming against a asset that is missing this object. See log for more info"), FText::FromString(FString(Filename))).ToString());
+							}
+							continue;
 						}
-						continue;
-					}
 
-					// add a dependency map entry also
-					TArray<FPackageIndex>& DependIndices = Linker->DependsMap[ExpIndex];
-					// find all the objects needed by this export
-					TArray<UObject*>* SrcDepends = ObjectDependencies.Find(Object);
-					checkf(SrcDepends,TEXT("Couldn't find dependency map for %s"), *Object->GetFullName());
+						// add a dependency map entry also
+						TArray<FPackageIndex>& DependIndices = Linker->DependsMap[ExpIndex];
+						// find all the objects needed by this export
+						TArray<UObject*>* SrcDepends = ObjectDependencies.Find(Object);
+						checkf(SrcDepends, TEXT("Couldn't find dependency map for %s"), *Object->GetFullName());
 
-					// go through each object and...
-					DependIndices.Reserve(SrcDepends->Num());
-					for (int32 DependIndex = 0; DependIndex < SrcDepends->Num(); DependIndex++)
-					{
-						UObject* DependentObject = (*SrcDepends)[DependIndex];
+						// go through each object and...
+						DependIndices.Reserve(SrcDepends->Num());
+						for (int32 DependIndex = 0; DependIndex < SrcDepends->Num(); DependIndex++)
+						{
+							UObject* DependentObject = (*SrcDepends)[DependIndex];
 
-						FPackageIndex DependencyIndex;
+							FPackageIndex DependencyIndex;
+
+							// if the dependency is in the same pacakge, we need to save an index into our ExportMap
+							if (DependentObject->GetOutermost() == Linker->LinkerRoot)
+							{
+								// ... find the associated ExportIndex
+								DependencyIndex = ExportToIndexMap.FindRef(DependentObject);
+							}
+							// otherwise we need to save an index into the ImportMap
+							else
+							{
+								// ... find the associated ImportIndex
+								DependencyIndex = ImportToIndexMap.FindRef(DependentObject);
+							}
 
-						// if the dependency is in the same pacakge, we need to save an index into our ExportMap
-						if (DependentObject->GetOutermost() == Linker->LinkerRoot)
-						{
-							// ... find the associated ExportIndex
-							DependencyIndex = ExportToIndexMap.FindRef(DependentObject);
-						}
-						// otherwise we need to save an index into the ImportMap
-						else
-						{
-							// ... find the associated ImportIndex
-							DependencyIndex = ImportToIndexMap.FindRef(DependentObject);
-						}
-					
 #if WITH_EDITOR
-						// If we still didn't find index, maybe it was a duplicate export which got removed.
-						// Check if we have a redirect to original.
-						if (DependencyIndex.IsNull() && DuplicateRedirects.Contains(DependentObject))
-						{
-							UObject** const RedirectObj = DuplicateRedirects.Find(DependentObject);
-							if (RedirectObj)
+							// If we still didn't find index, maybe it was a duplicate export which got removed.
+							// Check if we have a redirect to original.
+							if (DependencyIndex.IsNull() && DuplicateRedirects.Contains(DependentObject))
 							{
-								DependencyIndex = ExportToIndexMap.FindRef(*RedirectObj);
+								UObject** const RedirectObj = DuplicateRedirects.Find(DependentObject);
+								if (RedirectObj)
+								{
+									DependencyIndex = ExportToIndexMap.FindRef(*RedirectObj);
+								}
 							}
-						}
 #endif
-						// if we didn't find it (FindRef returns 0 on failure, which is good in this case), then we are in trouble, something went wrong somewhere
-						checkf(!DependencyIndex.IsNull(), TEXT("Failed to find dependency index for %s (%s)"), *DependentObject->GetFullName(), *Object->GetFullName());
+							// if we didn't find it (FindRef returns 0 on failure, which is good in this case), then we are in trouble, something went wrong somewhere
+							checkf(!DependencyIndex.IsNull(), TEXT("Failed to find dependency index for %s (%s)"), *DependentObject->GetFullName(), *Object->GetFullName());
 
-						// add the import as an import for this export
-						DependIndices.Add(DependencyIndex);
+							// add the import as an import for this export
+							DependIndices.Add(DependencyIndex);
+						}
 					}
 				}
 
@@ -4815,6 +4847,7 @@ FSavePackageResultStruct UPackage::Save(UPackage* InOuter, UObject* Base, EObjec
 				// Save dummy import map, overwritten later.
 				if (!bTextFormat)
 				{
+					SCOPED_NAMED_EVENT(UPackage_Save_WriteDummyImportMap, FColor::Cyan);
 #if WITH_EDITOR
 					FArchiveStackTraceIgnoreScope IgnoreSummaryDiffsScope(DiffSettings.bIgnoreHeaderDiffs);
 #endif // WITH_EDITOR
@@ -4837,6 +4870,7 @@ FSavePackageResultStruct UPackage::Save(UPackage* InOuter, UObject* Base, EObjec
 				// Save dummy export map, overwritten later.
 				if (!bTextFormat)
 				{
+					SCOPED_NAMED_EVENT(UPackage_Save_WriteDummyExportMap, FColor::Green);
 #if WITH_EDITOR
 					FArchiveStackTraceIgnoreScope IgnoreSummaryDiffsScope(DiffSettings.bIgnoreHeaderDiffs);
 #endif // WITH_EDITOR
@@ -4858,6 +4892,8 @@ FSavePackageResultStruct UPackage::Save(UPackage* InOuter, UObject* Base, EObjec
 
 				if (!bTextFormat)
 				{
+					SCOPED_NAMED_EVENT(UPackage_Save_WriteDependsMap, FColor::Red);
+
 					FStructuredArchive::FStream DependsStream = StructuredArchiveRoot.EnterStream(SA_FIELD_NAME(TEXT("DependsMap")));
 					if (Linker->IsCooking())
 					{
@@ -4895,6 +4931,8 @@ FSavePackageResultStruct UPackage::Save(UPackage* InOuter, UObject* Base, EObjec
 				// Only save string asset and searchable name map if saving for editor
 				if (!(Linker->Summary.PackageFlags & PKG_FilterEditorOnly))
 				{
+					SCOPED_NAMED_EVENT(UPackage_Save_SaveSoftPackagesAndSearchableNames, FColor::Turquoise);
+
 					// Save soft package references
 					Linker->Summary.SoftPackageReferencesOffset = Linker->Tell();
 					Linker->Summary.SoftPackageReferencesCount = Linker->SoftPackageReferenceList.Num();
@@ -4927,96 +4965,107 @@ FSavePackageResultStruct UPackage::Save(UPackage* InOuter, UObject* Base, EObjec
 #endif // WITH_EDITOR
 
 					// Save thumbnails
-					SaveThumbnails(InOuter, Linker.Get(), StructuredArchiveRoot.EnterField(SA_FIELD_NAME(TEXT("Thumbnails"))));
+					{
+						SCOPED_NAMED_EVENT(UPackage_Save_SaveThumbnails, FColor::Turquoise);
+						SaveThumbnails(InOuter, Linker.Get(), StructuredArchiveRoot.EnterField(SA_FIELD_NAME(TEXT("Thumbnails"))));
+					}
 
 					if (!bTextFormat)
 					{	
 						// Save asset registry data so the editor can search for information about assets in this package
+						SCOPED_NAMED_EVENT(UPackage_Save_SaveAssetRegistryData, FColor::Turquoise);
 						SaveAssetRegistryData(InOuter, Linker.Get(), StructuredArchiveRoot.EnterField(SA_FIELD_NAME(TEXT("AssetRegistry"))));
 					}
 
 					// Save level information used by World browser
-					SaveWorldLevelInfo(InOuter, Linker.Get(), StructuredArchiveRoot);
+					{
+						SCOPED_NAMED_EVENT(UPackage_Save_WorldLevelData, FColor::Turquoise);
+						SaveWorldLevelInfo(InOuter, Linker.Get(), StructuredArchiveRoot);
+					}
 				}
 
-
-				for (int32 i = 0; i < Linker->ExportMap.Num(); i++)
+				// Map export indices
 				{
-					FObjectExport& Export = Linker->ExportMap[i];
-					if (Export.Object)
+					SCOPED_NAMED_EVENT(UPackage_Save_MapExportIndices, FColor::Turquoise);
+
+					for (int32 i = 0; i < Linker->ExportMap.Num(); i++)
 					{
-						// Set class index.
-						// If this is *exactly* a UClass, store null instead; for anything else, including UClass-derived classes, map it
-						UClass* ObjClass = Export.Object->GetClass();
-						if (ObjClass != UClass::StaticClass())
-						{
-							Export.ClassIndex = Linker->MapObject(ObjClass);
-							checkf(!Export.ClassIndex.IsNull(), TEXT("Export %s class is not mapped when saving %s"), *Export.Object->GetFullName(), *Linker->LinkerRoot->GetName());
-						}
-						else
+						FObjectExport& Export = Linker->ExportMap[i];
+						if (Export.Object)
 						{
-							Export.ClassIndex = FPackageIndex();
-						}
+							// Set class index.
+							// If this is *exactly* a UClass, store null instead; for anything else, including UClass-derived classes, map it
+							UClass* ObjClass = Export.Object->GetClass();
+							if (ObjClass != UClass::StaticClass())
+							{
+								Export.ClassIndex = Linker->MapObject(ObjClass);
+								checkf(!Export.ClassIndex.IsNull(), TEXT("Export %s class is not mapped when saving %s"), *Export.Object->GetFullName(), *Linker->LinkerRoot->GetName());
+							}
+							else
+							{
+								Export.ClassIndex = FPackageIndex();
+							}
 
-						if (IsEventDrivenLoaderEnabledInCookedBuilds() && TargetPlatform)
-						{
-							UObject* Archetype = Export.Object->GetArchetype();
-							check(Archetype);
-							check(Archetype->IsA(Export.Object->HasAnyFlags(RF_ClassDefaultObject) ? ObjClass->GetSuperClass() : ObjClass));
-							Export.TemplateIndex = Linker->MapObject(Archetype);
-							UE_CLOG(Export.TemplateIndex.IsNull(), LogSavePackage, Fatal, TEXT("%s was an archetype of %s but returned a null index mapping the object."), *Archetype->GetFullName(), *Export.Object->GetFullName());
-							check(!Export.TemplateIndex.IsNull());
-						}
+							if (IsEventDrivenLoaderEnabledInCookedBuilds() && TargetPlatform)
+							{
+								UObject* Archetype = Export.Object->GetArchetype();
+								check(Archetype);
+								check(Archetype->IsA(Export.Object->HasAnyFlags(RF_ClassDefaultObject) ? ObjClass->GetSuperClass() : ObjClass));
+								Export.TemplateIndex = Linker->MapObject(Archetype);
+								UE_CLOG(Export.TemplateIndex.IsNull(), LogSavePackage, Fatal, TEXT("%s was an archetype of %s but returned a null index mapping the object."), *Archetype->GetFullName(), *Export.Object->GetFullName());
+								check(!Export.TemplateIndex.IsNull());
+							}
 
-						// Set the parent index, if this export represents a UStruct-derived object
-						if (UStruct* Struct = dynamic_cast<UStruct*>(Export.Object))
-						{
-							if (Struct->GetSuperStruct() != nullptr)
+							// Set the parent index, if this export represents a UStruct-derived object
+							if (UStruct* Struct = dynamic_cast<UStruct*>(Export.Object))
 							{
-								Export.SuperIndex = Linker->MapObject(Struct->GetSuperStruct());
-								checkf(!Export.SuperIndex.IsNull(),
-									TEXT("Export Struct (%s) of type (%s) inheriting from (%s) of type (%s) has not mapped super struct."),
-									*GetPathNameSafe(Struct),
-									*(Struct->GetClass()->GetName()),
-									*GetPathNameSafe(Struct->GetSuperStruct()),
-									*(Struct->GetSuperStruct()->GetClass()->GetName())
-								);
+								if (Struct->GetSuperStruct() != nullptr)
+								{
+									Export.SuperIndex = Linker->MapObject(Struct->GetSuperStruct());
+									checkf(!Export.SuperIndex.IsNull(),
+										TEXT("Export Struct (%s) of type (%s) inheriting from (%s) of type (%s) has not mapped super struct."),
+										*GetPathNameSafe(Struct),
+										*(Struct->GetClass()->GetName()),
+										*GetPathNameSafe(Struct->GetSuperStruct()),
+										*(Struct->GetSuperStruct()->GetClass()->GetName())
+									);
+								}
+								else
+								{
+									Export.SuperIndex = FPackageIndex();
+								}
 							}
 							else
 							{
 								Export.SuperIndex = FPackageIndex();
 							}
-						}
-						else
-						{
-							Export.SuperIndex = FPackageIndex();
-						}
 
-						// Set FPackageIndex for this export's Outer. If the export's Outer
-						// is the UPackage corresponding to this package's LinkerRoot, the
-						if (Export.Object->GetOuter() != InOuter)
-						{
-							check(Export.Object->GetOuter());
-							checkf(Export.Object->GetOuter()->IsIn(InOuter),
-								TEXT("Export Object (%s) Outer (%s) mismatch."),
-								*(Export.Object->GetPathName()),
-								*(Export.Object->GetOuter()->GetPathName()));
-							Export.OuterIndex = Linker->MapObject(Export.Object->GetOuter());
-							checkf(!Export.OuterIndex.IsImport(),
-								TEXT("Export Object (%s) Outer (%s) is an Import."),
-								*(Export.Object->GetPathName()),
-								*(Export.Object->GetOuter()->GetPathName()));
-
-							if (Linker->IsCooking() && IsEventDrivenLoaderEnabledInCookedBuilds())
+							// Set FPackageIndex for this export's Outer. If the export's Outer
+							// is the UPackage corresponding to this package's LinkerRoot, the
+							if (Export.Object->GetOuter() != InOuter)
 							{
-								// Only packages are allowed to have no outer
-								ensureMsgf(Export.OuterIndex != FPackageIndex() || Export.Object->IsA(UPackage::StaticClass()), TEXT("Export %s has no valid outer when cooking!"), *Export.Object->GetPathName());
+								check(Export.Object->GetOuter());
+								checkf(Export.Object->GetOuter()->IsIn(InOuter),
+									TEXT("Export Object (%s) Outer (%s) mismatch."),
+									*(Export.Object->GetPathName()),
+									*(Export.Object->GetOuter()->GetPathName()));
+								Export.OuterIndex = Linker->MapObject(Export.Object->GetOuter());
+								checkf(!Export.OuterIndex.IsImport(),
+									TEXT("Export Object (%s) Outer (%s) is an Import."),
+									*(Export.Object->GetPathName()),
+									*(Export.Object->GetOuter()->GetPathName()));
+
+								if (Linker->IsCooking() && IsEventDrivenLoaderEnabledInCookedBuilds())
+								{
+									// Only packages are allowed to have no outer
+									ensureMsgf(Export.OuterIndex != FPackageIndex() || Export.Object->IsA(UPackage::StaticClass()), TEXT("Export %s has no valid outer when cooking!"), *Export.Object->GetPathName());
+								}
+							}
+							else
+							{
+								// this export's Outer is the LinkerRoot for this package
+								Export.OuterIndex = FPackageIndex();
 							}
-						}
-						else
-						{
-							// this export's Outer is the LinkerRoot for this package
-							Export.OuterIndex = FPackageIndex();
 						}
 					}
 				}
@@ -5321,6 +5370,7 @@ FSavePackageResultStruct UPackage::Save(UPackage* InOuter, UObject* Base, EObjec
 #endif
 				{
 					COOK_STAT(FScopedDurationTimer SaveTimer(SavePackageStats::SerializeExportsTimeSec));
+					SCOPED_NAMED_EVENT(UPackage_Save_SaveExports, FColor::Orange);
 #if WITH_EDITOR
 					FArchive::FScopeSetDebugSerializationFlags S(*Linker, DSF_IgnoreDiff, true);
 #endif
@@ -5341,6 +5391,8 @@ FSavePackageResultStruct UPackage::Save(UPackage* InOuter, UObject* Base, EObjec
 						FObjectExport& Export = Linker->ExportMap[i];
 						if (Export.Object)
 						{
+							SCOPED_NAMED_EVENT(UPackage_Save_SaveExport, FColor::Blue);
+
 							// Save the object data.
 							Export.SerialOffset = Linker->Tell();
 							Linker->CurrentlySavingExport = FPackageIndex::FromExport(i);
