Optimized IsEventDrivenLoaderEnabledInCookedBuilds() to not query the config system for the current state, but just use the global cvar variable that the config system writes to

diff --git a/Engine/Source/Runtime/CoreUObject/Private/Serialization/AsyncLoading.cpp b/Engine/Source/Runtime/CoreUObject/Private/Serialization/AsyncLoading.cpp
--- a/Engine/Source/Runtime/CoreUObject/Private/Serialization/AsyncLoading.cpp
+++ b/Engine/Source/Runtime/CoreUObject/Private/Serialization/AsyncLoading.cpp
@@ -8145,4 +8145,10 @@ static FAutoConsoleCommand DumpSerializeCmd(
 	);
 #endif
 
+bool IsEventDrivenLoaderEnabledInCookedBuilds()
+{
+	static bool bNoEDL = FParse::Param(FCommandLine::Get(), TEXT("NOEDL"));
+	return (GEventDrivenLoaderEnabledInCookedBuilds != 0) && !bNoEDL;
+}
+
 //#pragma clang optimize on
diff --git a/Engine/Source/Runtime/CoreUObject/Private/Serialization/AsyncPackageLoader.cpp b/Engine/Source/Runtime/CoreUObject/Private/Serialization/AsyncPackageLoader.cpp
--- a/Engine/Source/Runtime/CoreUObject/Private/Serialization/AsyncPackageLoader.cpp
+++ b/Engine/Source/Runtime/CoreUObject/Private/Serialization/AsyncPackageLoader.cpp
@@ -390,43 +390,6 @@ static FEDLBootNotificationManager& GetGEDLBootNotificationManager()
 	return Singleton;
 }
 
-bool IsEventDrivenLoaderEnabledInCookedBuilds()
-{
-	static struct FEventDrivenLoaderEnabledInCookedBuildsInit
-	{
-		bool bEventDrivenLoaderEnabled;
-		FEventDrivenLoaderEnabledInCookedBuildsInit()
-			: bEventDrivenLoaderEnabled(false)
-		{
-			SetEventDrivenLoaderEnabled();
-		}
-
-		void SetEventDrivenLoaderEnabled()
-		{
-			check(GConfig || IsEngineExitRequested());
-			if (GConfig)
-			{
-				GConfig->GetBool(TEXT("/Script/Engine.StreamingSettings"), TEXT("s.EventDrivenLoaderEnabled"), bEventDrivenLoaderEnabled, GEngineIni);
-#if !UE_BUILD_SHIPPING
-				if (FParse::Param(FCommandLine::Get(), TEXT("NOEDL")))
-				{
-					bEventDrivenLoaderEnabled = false;
-				}
-#endif
-			}
-		}
-	} EventDrivenLoaderEnabledInCookedBuilds;
-
-#if WITH_EDITOR	
-	// when building from the UE4 Editor, s.EventDrivenLoaderEnabled can be changed from Project Settings, so we need to test it at every call
-	if (GIsEditor && !IsRunningCommandlet())
-	{
-		EventDrivenLoaderEnabledInCookedBuilds.SetEventDrivenLoaderEnabled();
-	}
-#endif
-	return EventDrivenLoaderEnabledInCookedBuilds.bEventDrivenLoaderEnabled;
-}
-
 bool IsEventDrivenLoaderEnabled()
 {
 	static struct FEventDrivenLoaderEnabledInit
