More trace instrumentation of package/level saving

diff --git a/Engine/Source/Editor/UnrealEd/Private/EditorEngine.cpp b/Engine/Source/Editor/UnrealEd/Private/EditorEngine.cpp
--- a/Engine/Source/Editor/UnrealEd/Private/EditorEngine.cpp
+++ b/Engine/Source/Editor/UnrealEd/Private/EditorEngine.cpp
@@ -4510,6 +4510,8 @@ void UEditorEngine::OnPreSaveWorld(uint32 SaveFlags, UWorld* World)
 		return;
 	}
 
+	TRACE_CPUPROFILER_EVENT_SCOPE(UEditorEngine_OnPreSaveWorld);
+
 	check(World->PersistentLevel);
 
 	// Pre save world event
diff --git a/Engine/Source/Editor/UnrealEd/Private/FileHelpers.cpp b/Engine/Source/Editor/UnrealEd/Private/FileHelpers.cpp
--- a/Engine/Source/Editor/UnrealEd/Private/FileHelpers.cpp
+++ b/Engine/Source/Editor/UnrealEd/Private/FileHelpers.cpp
@@ -3230,6 +3230,8 @@ static bool InternalSavePackages(const TArray<UPackage*>& PackagesToSave, bool b
 
 void FEditorFileUtils::SaveMapDataPackages(UWorld* WorldToSave, bool bCheckDirty)
 {
+	TRACE_CPUPROFILER_EVENT_SCOPE(FEditorFileUtils_SaveMapDataPackages);
+
 	TArray<UPackage*> PackagesToSave;
 	ULevel* Level = WorldToSave->PersistentLevel;
 	UPackage* WorldPackage = WorldToSave->GetOutermost();
@@ -3436,6 +3438,8 @@ bool FEditorFileUtils::SaveCurrentLevel()
 {
 	bool bReturnCode = false;
 
+	TRACE_CPUPROFILER_EVENT_SCOPE(FEditorFileUtils_SaveCurrentLevel);
+
 	ULevel* Level = GWorld->GetCurrentLevel();
 	if ( Level && FEditorFileUtils::PromptToCheckoutLevels( false, Level ) )
 	{
diff --git a/Engine/Source/Runtime/CoreUObject/Private/UObject/SavePackage.cpp b/Engine/Source/Runtime/CoreUObject/Private/UObject/SavePackage.cpp
--- a/Engine/Source/Runtime/CoreUObject/Private/UObject/SavePackage.cpp
+++ b/Engine/Source/Runtime/CoreUObject/Private/UObject/SavePackage.cpp
@@ -668,7 +668,7 @@ static void ConditionallyExcludeObjectForTarget(UObject* Obj, EObjectMark Exclud
 {
 #if WITH_EDITOR
 
-	SCOPED_NAMED_EVENT(ConditionallyExcludeObjectForTarget, FColor::Emerald);
+	TRACE_CPUPROFILER_EVENT_SCOPE(ConditionallyExcludeObjectForTarget);
 
 	if (!Obj || Obj->GetOutermost()->GetFName() == GLongCoreUObjectPackageName)
 	{
@@ -3400,7 +3400,7 @@ FSavePackageResultStruct UPackage::Save(UPackage* InOuter, UObject* Base, EObjec
 {
 	COOK_STAT(FScopedDurationTimer FuncSaveTimer(SavePackageStats::SavePackageTimeSec));
 	COOK_STAT(SavePackageStats::NumPackagesSaved++);
-	SCOPED_NAMED_EVENT(UPackage_Save, FColor::Emerald);
+	TRACE_CPUPROFILER_EVENT_SCOPE(UPackage_Save);
 
 	FLinkerLoad* Conform = nullptr;
 
@@ -3673,7 +3673,7 @@ FSavePackageResultStruct UPackage::Save(UPackage* InOuter, UObject* Base, EObjec
 			// Tag exports and route presave.
 			FPackageExportTagger PackageExportTagger(Base, TopLevelFlags, InOuter, TargetPlatform);
 			{
-				SCOPED_NAMED_EVENT(UPackage_Save_TagExportsWithPresave, FColor::Red);
+				TRACE_CPUPROFILER_EVENT_SCOPE(UPackage_Save_TagExportsWithPresave);
 
 				COOK_STAT(FScopedDurationTimer SaveTimer(SavePackageStats::TagPackageExportsPresaveTimeSec));
 				// Do not route presave if saving concurrently. This should have been done before the concurrent save started.
@@ -3743,7 +3743,7 @@ FSavePackageResultStruct UPackage::Save(UPackage* InOuter, UObject* Base, EObjec
 
 			
 				{
-					SCOPED_NAMED_EVENT(UPackage_Save_TagExports, FColor::Green);
+					TRACE_CPUPROFILER_EVENT_SCOPE(UPackage_Save_TagExports);
 					COOK_STAT(FScopedDurationTimer SaveTimer(SavePackageStats::TagPackageExportsTimeSec));
 					// Clear all marks (OBJECTMARK_TagExp and exclusion marks) again as we need to redo tagging below.
 					UnMarkAllObjects();
@@ -3792,7 +3792,7 @@ FSavePackageResultStruct UPackage::Save(UPackage* InOuter, UObject* Base, EObjec
 				const bool bTextFormat = FString(Filename).EndsWith(FPackageName::GetTextAssetPackageExtension()) || FString(Filename).EndsWith(FPackageName::GetTextMapPackageExtension());
 	
 				{
-					SCOPED_NAMED_EVENT(UPackage_Save_CreateLinkerSave, FColor::Blue);
+					TRACE_CPUPROFILER_EVENT_SCOPE(UPackage_Save_CreateLinkerSave);
 
 #if WITH_EDITOR
 					FString DiffCookedPackagesPath;
@@ -3970,7 +3970,7 @@ FSavePackageResultStruct UPackage::Save(UPackage* InOuter, UObject* Base, EObjec
 				// Import objects & names.
 				TSet<UPackage*> PrestreamPackages;
 				{
-					SCOPED_NAMED_EVENT(UPackage_Save_TagImports, FColor::Emerald);
+					TRACE_CPUPROFILER_EVENT_SCOPE(UPackage_Save_TagImports);
 					
 					TArray<UObject*> TagExpObjects;
 					GetObjectsWithAnyMarks(TagExpObjects, OBJECTMARK_TagExp);
@@ -4147,7 +4147,7 @@ FSavePackageResultStruct UPackage::Save(UPackage* InOuter, UObject* Base, EObjec
 
 				// Tag the names for all relevant object, classes, and packages.
 				{
-					SCOPED_NAMED_EVENT(UPackage_Save_TagNames, FColor::Red);
+					TRACE_CPUPROFILER_EVENT_SCOPE(UPackage_Save_TagNames);
 
 					TArray<UObject*> TagExpImpObjects;
 					GetObjectsWithAnyMarks(TagExpImpObjects, EObjectMark(OBJECTMARK_TagExp|OBJECTMARK_TagImp));
@@ -4436,7 +4436,7 @@ FSavePackageResultStruct UPackage::Save(UPackage* InOuter, UObject* Base, EObjec
 				// Build NameMap.
 				Linker->Summary.NameOffset = Linker->Tell();
 				{
-					SCOPED_NAMED_EVENT(UPackage_Save_BuildNameMap, FColor::Red);
+					TRACE_CPUPROFILER_EVENT_SCOPE(UPackage_Save_BuildNameMap);
 #if WITH_EDITOR
 					FArchive::FScopeSetDebugSerializationFlags S(*Linker, DSF_IgnoreDiff, true);
 					FArchiveStackTraceIgnoreScope IgnoreSummaryDiffsScope(DiffSettings.bIgnoreHeaderDiffs);
@@ -4454,7 +4454,7 @@ FSavePackageResultStruct UPackage::Save(UPackage* InOuter, UObject* Base, EObjec
 				Linker->Summary.GatherableTextDataCount = 0;
 				if (!(Linker->Summary.PackageFlags & PKG_FilterEditorOnly))
 				{
-					SCOPED_NAMED_EVENT(UPackage_Save_WriteGatherableTextData, FColor::Yellow);
+					TRACE_CPUPROFILER_EVENT_SCOPE(UPackage_Save_WriteGatherableTextData);
 
 					// The Editor version is used as part of the check to see if a package is too old to use the gather cache, so we always have to add it if we have gathered loc for this asset
 					// Note that using custom version here only works because we already added it to the export tagger before the package summary was serialized
@@ -4492,7 +4492,7 @@ FSavePackageResultStruct UPackage::Save(UPackage* InOuter, UObject* Base, EObjec
 
 				// Build ImportMap.
 				{
-					SCOPED_NAMED_EVENT(UPackage_Save_BuildImportMap, FColor::Orange);
+					TRACE_CPUPROFILER_EVENT_SCOPE(UPackage_Save_BuildImportMap);
 
 					TArray<UObject*> TagImpObjects;
 
@@ -4566,7 +4566,7 @@ FSavePackageResultStruct UPackage::Save(UPackage* InOuter, UObject* Base, EObjec
 				// sort and conform imports
 				FObjectImportSortHelper ImportSortHelper;
 				{
-					SCOPED_NAMED_EVENT(UPackage_Save_SortImports, FColor::Red);
+					TRACE_CPUPROFILER_EVENT_SCOPE(UPackage_Save_SortImports);
 					ImportSortHelper.SortImports(Linker.Get(), Conform);
 					Linker->Summary.ImportCount = Linker->ImportMap.Num();
 				}
@@ -4580,7 +4580,7 @@ FSavePackageResultStruct UPackage::Save(UPackage* InOuter, UObject* Base, EObjec
 				
 				// Build ExportMap.
 				{
-					SCOPED_NAMED_EVENT(UPackage_Save_BuildExportMap, FColor::Turquoise);
+					TRACE_CPUPROFILER_EVENT_SCOPE(UPackage_Save_BuildExportMap);
 
 					TArray<UObject*> TagExpObjects;
 					GetObjectsWithAnyMarks(TagExpObjects, OBJECTMARK_TagExp);
@@ -4619,13 +4619,13 @@ FSavePackageResultStruct UPackage::Save(UPackage* InOuter, UObject* Base, EObjec
 				// Sort exports alphabetically and conform the export table (if necessary)
 				FObjectExportSortHelper ExportSortHelper;
 				{
-					SCOPED_NAMED_EVENT(UPackage_Save_SortExports, FColor::Turquoise);
+					TRACE_CPUPROFILER_EVENT_SCOPE(UPackage_Save_SortExports);
 					ExportSortHelper.SortExports(Linker.Get(), Conform);
 				}
 				
 				// Sort exports for seek-free loading.
 				{
-					SCOPED_NAMED_EVENT(UPackage_Save_SortExportsForSeekFree, FColor::Turquoise);
+					TRACE_CPUPROFILER_EVENT_SCOPE(UPackage_Save_SortExportsForSeekFree);
 					COOK_STAT(FScopedDurationTimer SaveTimer(SavePackageStats::SortExportsSeekfreeInnerTimeSec));
 					FObjectExportSeekFreeSorter SeekFreeSorter;
 					SeekFreeSorter.SortExports( Linker.Get(), Conform );
@@ -4657,7 +4657,7 @@ FSavePackageResultStruct UPackage::Save(UPackage* InOuter, UObject* Base, EObjec
 
 				// go back over the (now sorted) exports and fill out the DependsMap
 				{
-					SCOPED_NAMED_EVENT(UPackage_Save_BuildExportDependsMap, FColor::Orange);
+					TRACE_CPUPROFILER_EVENT_SCOPE(UPackage_Save_BuildExportDependsMap);
 					for (int32 ExpIndex = 0; ExpIndex < Linker->ExportMap.Num(); ExpIndex++)
 					{
 						UObject* Object = Linker->ExportMap[ExpIndex].Object;
@@ -4847,7 +4847,7 @@ FSavePackageResultStruct UPackage::Save(UPackage* InOuter, UObject* Base, EObjec
 				// Save dummy import map, overwritten later.
 				if (!bTextFormat)
 				{
-					SCOPED_NAMED_EVENT(UPackage_Save_WriteDummyImportMap, FColor::Cyan);
+					TRACE_CPUPROFILER_EVENT_SCOPE(UPackage_Save_WriteDummyImportMap);
 #if WITH_EDITOR
 					FArchiveStackTraceIgnoreScope IgnoreSummaryDiffsScope(DiffSettings.bIgnoreHeaderDiffs);
 #endif // WITH_EDITOR
@@ -4870,7 +4870,7 @@ FSavePackageResultStruct UPackage::Save(UPackage* InOuter, UObject* Base, EObjec
 				// Save dummy export map, overwritten later.
 				if (!bTextFormat)
 				{
-					SCOPED_NAMED_EVENT(UPackage_Save_WriteDummyExportMap, FColor::Green);
+					TRACE_CPUPROFILER_EVENT_SCOPE(UPackage_Save_WriteDummyExportMap);
 #if WITH_EDITOR
 					FArchiveStackTraceIgnoreScope IgnoreSummaryDiffsScope(DiffSettings.bIgnoreHeaderDiffs);
 #endif // WITH_EDITOR
@@ -4892,7 +4892,7 @@ FSavePackageResultStruct UPackage::Save(UPackage* InOuter, UObject* Base, EObjec
 
 				if (!bTextFormat)
 				{
-					SCOPED_NAMED_EVENT(UPackage_Save_WriteDependsMap, FColor::Red);
+					TRACE_CPUPROFILER_EVENT_SCOPE(UPackage_Save_WriteDependsMap);
 
 					FStructuredArchive::FStream DependsStream = StructuredArchiveRoot.EnterStream(SA_FIELD_NAME(TEXT("DependsMap")));
 					if (Linker->IsCooking())
@@ -4931,7 +4931,7 @@ FSavePackageResultStruct UPackage::Save(UPackage* InOuter, UObject* Base, EObjec
 				// Only save string asset and searchable name map if saving for editor
 				if (!(Linker->Summary.PackageFlags & PKG_FilterEditorOnly))
 				{
-					SCOPED_NAMED_EVENT(UPackage_Save_SaveSoftPackagesAndSearchableNames, FColor::Turquoise);
+					TRACE_CPUPROFILER_EVENT_SCOPE(UPackage_Save_SaveSoftPackagesAndSearchableNames);
 
 					// Save soft package references
 					Linker->Summary.SoftPackageReferencesOffset = Linker->Tell();
@@ -4966,27 +4966,27 @@ FSavePackageResultStruct UPackage::Save(UPackage* InOuter, UObject* Base, EObjec
 
 					// Save thumbnails
 					{
-						SCOPED_NAMED_EVENT(UPackage_Save_SaveThumbnails, FColor::Turquoise);
+						TRACE_CPUPROFILER_EVENT_SCOPE(UPackage_Save_SaveThumbnails);
 						SaveThumbnails(InOuter, Linker.Get(), StructuredArchiveRoot.EnterField(SA_FIELD_NAME(TEXT("Thumbnails"))));
 					}
 
 					if (!bTextFormat)
 					{	
 						// Save asset registry data so the editor can search for information about assets in this package
-						SCOPED_NAMED_EVENT(UPackage_Save_SaveAssetRegistryData, FColor::Turquoise);
+						TRACE_CPUPROFILER_EVENT_SCOPE(UPackage_Save_SaveAssetRegistryData);
 						SaveAssetRegistryData(InOuter, Linker.Get(), StructuredArchiveRoot.EnterField(SA_FIELD_NAME(TEXT("AssetRegistry"))));
 					}
 
 					// Save level information used by World browser
 					{
-						SCOPED_NAMED_EVENT(UPackage_Save_WorldLevelData, FColor::Turquoise);
+						TRACE_CPUPROFILER_EVENT_SCOPE(UPackage_Save_WorldLevelData);
 						SaveWorldLevelInfo(InOuter, Linker.Get(), StructuredArchiveRoot);
 					}
 				}
 
 				// Map export indices
 				{
-					SCOPED_NAMED_EVENT(UPackage_Save_MapExportIndices, FColor::Turquoise);
+					TRACE_CPUPROFILER_EVENT_SCOPE(UPackage_Save_MapExportIndices);
 
 					for (int32 i = 0; i < Linker->ExportMap.Num(); i++)
 					{
@@ -5370,7 +5370,7 @@ FSavePackageResultStruct UPackage::Save(UPackage* InOuter, UObject* Base, EObjec
 #endif
 				{
 					COOK_STAT(FScopedDurationTimer SaveTimer(SavePackageStats::SerializeExportsTimeSec));
-					SCOPED_NAMED_EVENT(UPackage_Save_SaveExports, FColor::Orange);
+					TRACE_CPUPROFILER_EVENT_SCOPE(UPackage_Save_SaveExports);
 #if WITH_EDITOR
 					FArchive::FScopeSetDebugSerializationFlags S(*Linker, DSF_IgnoreDiff, true);
 #endif
@@ -5391,7 +5391,7 @@ FSavePackageResultStruct UPackage::Save(UPackage* InOuter, UObject* Base, EObjec
 						FObjectExport& Export = Linker->ExportMap[i];
 						if (Export.Object)
 						{
-							SCOPED_NAMED_EVENT(UPackage_Save_SaveExport, FColor::Blue);
+							TRACE_CPUPROFILER_EVENT_SCOPE(UPackage_Save_SaveExport);
 
 							// Save the object data.
 							Export.SerialOffset = Linker->Tell();
