Fixed bug in ULazyObjectProperty and UWeakObjectProperty copy to from script VM methods

Moved optimized fast reference copying from UObjectPropertyBase to UObjectProperty and removed USoftObjectProperty overloads.

#rnx

diff --git a/Engine/Source/Runtime/CoreUObject/Private/UObject/PropertyBaseObject.cpp b/Engine/Source/Runtime/CoreUObject/Private/UObject/PropertyBaseObject.cpp
--- a/Engine/Source/Runtime/CoreUObject/Private/UObject/PropertyBaseObject.cpp
+++ b/Engine/Source/Runtime/CoreUObject/Private/UObject/PropertyBaseObject.cpp
@@ -542,33 +542,6 @@ bool UObjectPropertyBase::SameType(const UProperty* Other) const
 	return Super::SameType(Other) && (PropertyClass == ((UObjectPropertyBase*)Other)->PropertyClass);
 }
 
-void UObjectPropertyBase::CopySingleValueToScriptVM( void* Dest, void const* Src ) const
-{
-	*(UObject**)Dest = GetObjectPropertyValue(Src);
-}
-
-void UObjectPropertyBase::CopyCompleteValueToScriptVM( void* Dest, void const* Src ) const
-{
-	for (int32 Index = 0; Index < ArrayDim; Index++)
-	{
-		((UObject**)Dest)[Index] = GetObjectPropertyValue(((uint8*)Src) + Index * ElementSize);
-	}
-}
-
-void UObjectPropertyBase::CopySingleValueFromScriptVM( void* Dest, void const* Src ) const
-{
-	SetObjectPropertyValue(Dest, *(UObject**)Src);
-}
-
-void UObjectPropertyBase::CopyCompleteValueFromScriptVM( void* Dest, void const* Src ) const
-{
-	checkSlow(ElementSize == sizeof(UObject*)); // the idea that script pointers are the same size as weak pointers is maybe required, maybe not
-	for (int32 Index = 0; Index < ArrayDim; Index++)
-	{
-		SetObjectPropertyValue(((uint8*)Dest) + Index * ElementSize, ((UObject**)Src)[Index]);
-	}
-}
-
 IMPLEMENT_CORE_INTRINSIC_CLASS(UObjectPropertyBase, UProperty,
 	{
 		Class->EmitObjectReference(STRUCT_OFFSET(UObjectProperty, PropertyClass), TEXT("PropertyClass"));
diff --git a/Engine/Source/Runtime/CoreUObject/Private/UObject/PropertyObject.cpp b/Engine/Source/Runtime/CoreUObject/Private/UObject/PropertyObject.cpp
--- a/Engine/Source/Runtime/CoreUObject/Private/UObject/PropertyObject.cpp
+++ b/Engine/Source/Runtime/CoreUObject/Private/UObject/PropertyObject.cpp
@@ -181,6 +181,34 @@ void UObjectProperty::SetObjectPropertyValue(void* PropertyValueAddress, UObject
 	SetPropertyValue(PropertyValueAddress, Value);
 }
 
+
+void UObjectProperty::CopySingleValueToScriptVM( void* Dest, void const* Src ) const
+{
+	*(UObject**)Dest = GetObjectPropertyValue(Src);
+}
+
+void UObjectProperty::CopyCompleteValueToScriptVM( void* Dest, void const* Src ) const
+{
+	for (int32 Index = 0; Index < ArrayDim; Index++)
+	{
+		((UObject**)Dest)[Index] = GetObjectPropertyValue(((uint8*)Src) + Index * ElementSize);
+	}
+}
+
+void UObjectProperty::CopySingleValueFromScriptVM( void* Dest, void const* Src ) const
+{
+	SetObjectPropertyValue(Dest, *(UObject**)Src);
+}
+
+void UObjectProperty::CopyCompleteValueFromScriptVM( void* Dest, void const* Src ) const
+{
+	checkSlow(ElementSize == sizeof(UObject*)); // the idea that script pointers are the same size as weak pointers is maybe required, maybe not
+	for (int32 Index = 0; Index < ArrayDim; Index++)
+	{
+		SetObjectPropertyValue(((uint8*)Dest) + Index * ElementSize, ((UObject**)Src)[Index]);
+	}
+}
+
 IMPLEMENT_CORE_INTRINSIC_CLASS(UObjectProperty, UObjectPropertyBase,
 	{
 	}
diff --git a/Engine/Source/Runtime/CoreUObject/Private/UObject/PropertySoftObjectPtr.cpp b/Engine/Source/Runtime/CoreUObject/Private/UObject/PropertySoftObjectPtr.cpp
--- a/Engine/Source/Runtime/CoreUObject/Private/UObject/PropertySoftObjectPtr.cpp
+++ b/Engine/Source/Runtime/CoreUObject/Private/UObject/PropertySoftObjectPtr.cpp
@@ -203,26 +203,6 @@ uint32 USoftObjectProperty::GetValueTypeHashInternal(const void* Src) const
 	return GetTypeHash(GetPropertyValue(Src));
 }
 
-void USoftObjectProperty::CopySingleValueToScriptVM(void* Dest, void const* Src) const
-{
-	CopySingleValue(Dest, Src);
-}
-
-void USoftObjectProperty::CopyCompleteValueToScriptVM(void* Dest, void const* Src) const
-{
-	CopyCompleteValue(Dest, Src);
-}
-
-void USoftObjectProperty::CopySingleValueFromScriptVM(void* Dest, void const* Src) const
-{
-	CopySingleValue(Dest, Src);
-}
-
-void USoftObjectProperty::CopyCompleteValueFromScriptVM(void* Dest, void const* Src) const
-{
-	CopyCompleteValue(Dest, Src);
-}
-
 IMPLEMENT_CORE_INTRINSIC_CLASS(USoftObjectProperty, UObjectPropertyBase,
 	{
 	}
diff --git a/Engine/Source/Runtime/CoreUObject/Public/UObject/UnrealType.h b/Engine/Source/Runtime/CoreUObject/Public/UObject/UnrealType.h
--- a/Engine/Source/Runtime/CoreUObject/Public/UObject/UnrealType.h
+++ b/Engine/Source/Runtime/CoreUObject/Public/UObject/UnrealType.h
@@ -1917,49 +1917,6 @@ class COREUOBJECT_API UObjectPropertyBase : public UProperty
 	virtual FName GetID() const override;
 	virtual void InstanceSubobjects( void* Data, void const* DefaultData, UObject* Owner, struct FObjectInstancingGraph* InstanceGraph ) override;
 	virtual bool SameType(const UProperty* Other) const override;
-	/**
-	 * Copy the value for a single element of this property. To the script VM.
-	 * 
-	 * @param	Dest				the address where the value should be copied to.  This should always correspond to the BASE + OFFSET + INDEX * SIZE, where
-	 *									BASE = (for member properties) the address of the UObject which contains this data, (for locals/parameters) the address of the space allocated for the function's locals
-	 *									OFFSET = the Offset of this UProperty
-	 *									INDEX = the index that you want to copy.  for properties which are not arrays, this should always be 0
-	 *									SIZE = the ElementSize of this UProperty
-	 * @param	Src					the address of the value to copy from. should be evaluated the same way as Dest
-	 */
-	virtual void CopySingleValueToScriptVM( void* Dest, void const* Src ) const override;
-
-	/**
-	 * Copy the value for all elements of this property. To the script VM.
-	 * 
-	 * @param	Dest				the address where the value should be copied to.  This should always correspond to the BASE + OFFSET, where
-	 *									BASE = (for member properties) the address of the UObject which contains this data, (for locals/parameters) the address of the space allocated for the function's locals
-	 *									OFFSET = the Offset of this UProperty
-	 * @param	Src					the address of the value to copy from. should be evaluated the same way as Dest
-	 */
-	virtual void CopyCompleteValueToScriptVM( void* Dest, void const* Src ) const override;
-
-	/**
-	 * Copy the value for a single element of this property. From the script VM.
-	 * 
-	 * @param	Dest				the address where the value should be copied to.  This should always correspond to the BASE + OFFSET + INDEX * SIZE, where
-	 *									BASE = (for member properties) the address of the UObject which contains this data, (for locals/parameters) the address of the space allocated for the function's locals
-	 *									OFFSET = the Offset of this UProperty
-	 *									INDEX = the index that you want to copy.  for properties which are not arrays, this should always be 0
-	 *									SIZE = the ElementSize of this UProperty
-	 * @param	Src					the address of the value to copy from. should be evaluated the same way as Dest
-	 */
-	virtual void CopySingleValueFromScriptVM( void* Dest, void const* Src ) const override;
-
-	/**
-	 * Copy the value for all elements of this property. From the script VM.
-	 * 
-	 * @param	Dest				the address where the value should be copied to.  This should always correspond to the BASE + OFFSET, where
-	 *									BASE = (for member properties) the address of the UObject which contains this data, (for locals/parameters) the address of the space allocated for the function's locals
-	 *									OFFSET = the Offset of this UProperty
-	 * @param	Src					the address of the value to copy from. should be evaluated the same way as Dest
-	 */
-	virtual void CopyCompleteValueFromScriptVM( void* Dest, void const* Src ) const override;
 	// End of UProperty interface
 
 	// UObjectPropertyBase interface
@@ -2108,6 +2065,12 @@ class COREUOBJECT_API UObjectProperty : public TUObjectPropertyBase<UObject*>
 private:
 	virtual uint32 GetValueTypeHashInternal(const void* Src) const override;
 public:
+
+	virtual void CopySingleValueToScriptVM(void* Dest, void const* Src) const override;
+	virtual void CopyCompleteValueToScriptVM(void* Dest, void const* Src) const override;
+	virtual void CopySingleValueFromScriptVM(void* Dest, void const* Src) const override;
+	virtual void CopyCompleteValueFromScriptVM(void* Dest, void const* Src) const override;
+
 	// End of UProperty interface
 
 	// UObjectPropertyBase interface
@@ -2239,14 +2202,7 @@ class COREUOBJECT_API USoftObjectProperty : public TUObjectPropertyBase<FSoftObj
 
 private:
 	virtual uint32 GetValueTypeHashInternal(const void* Src) const override;
-public:
-
-	// ScriptVM should store Asset as FSoftObjectPtr not as UObject.
 
-	virtual void CopySingleValueToScriptVM(void* Dest, void const* Src) const override;
-	virtual void CopyCompleteValueToScriptVM(void* Dest, void const* Src) const override;
-	virtual void CopySingleValueFromScriptVM(void* Dest, void const* Src) const override;
-	virtual void CopyCompleteValueFromScriptVM(void* Dest, void const* Src) const override;
 	// End of UObjectProperty interface
 };
 
