Exposed ConditionalUpdateCache via the FPackageLocalizationManager

You may need to call this yourself (from the game thread) if manually mounting content paths.

diff --git a/Engine/Source/Runtime/CoreUObject/Private/Internationalization/PackageLocalizationManager.cpp b/Engine/Source/Runtime/CoreUObject/Private/Internationalization/PackageLocalizationManager.cpp
--- a/Engine/Source/Runtime/CoreUObject/Private/Internationalization/PackageLocalizationManager.cpp
+++ b/Engine/Source/Runtime/CoreUObject/Private/Internationalization/PackageLocalizationManager.cpp
@@ -145,6 +145,14 @@ FName FPackageLocalizationManager::FindLocalizedPackageNameNoCache(const FName I
 	return NAME_None;
 }
 
+void FPackageLocalizationManager::ConditionalUpdateCache()
+{
+	if (ActiveCache.IsValid())
+	{
+		ActiveCache->ConditionalUpdateCache();
+	}
+}
+
 FPackageLocalizationManager& FPackageLocalizationManager::Get()
 {
 	static FPackageLocalizationManager PackageLocalizationManager;
diff --git a/Engine/Source/Runtime/CoreUObject/Public/Internationalization/PackageLocalizationManager.h b/Engine/Source/Runtime/CoreUObject/Public/Internationalization/PackageLocalizationManager.h
--- a/Engine/Source/Runtime/CoreUObject/Public/Internationalization/PackageLocalizationManager.h
+++ b/Engine/Source/Runtime/CoreUObject/Public/Internationalization/PackageLocalizationManager.h
@@ -56,6 +56,11 @@ public:
 	 */
 	FName FindLocalizedPackageNameForCulture(const FName InSourcePackageName, const FString& InCultureName);
 
+	/**
+	 * Update this cache, but only if it is dirty.
+	 */
+	void ConditionalUpdateCache();
+
 	/**
 	 * Singleton accessor.
 	 *
