Don't recycle objects when creating default subobjects in the constructor.
#rnx

diff --git a/Engine/Source/Runtime/CoreUObject/Private/UObject/UObjectGlobals.cpp b/Engine/Source/Runtime/CoreUObject/Private/UObject/UObjectGlobals.cpp
--- a/Engine/Source/Runtime/CoreUObject/Private/UObject/UObjectGlobals.cpp
+++ b/Engine/Source/Runtime/CoreUObject/Private/UObject/UObjectGlobals.cpp
@@ -3075,12 +3075,13 @@ UObject* StaticConstructObject_Internal
 			!InTemplate || 
 			(InName != NAME_None && (bAssumeTemplateIsArchetype || InTemplate == UObject::GetArchetypeFromRequiredInfo(InClass, InOuter, InName, InFlags)))
 			);
+	const bool bCanRecycleSubobjects = bIsNativeFromCDO && (!(InFlags & RF_DefaultSubObject) || !FUObjectThreadContext::Get().IsInConstructor)
 #if WITH_HOT_RELOAD
 	// Do not recycle subobjects when performing hot-reload as they may contain old property values.
-	const bool bCanRecycleSubobjects = bIsNativeFromCDO && !GIsHotReload;
-#else
-	const bool bCanRecycleSubobjects = bIsNativeFromCDO;
+	&& !GIsHotReload
 #endif
+		;
+
 	bool bRecycledSubobject = false;	
 	Result = StaticAllocateObject(InClass, InOuter, InName, InFlags, InternalSetFlags, bCanRecycleSubobjects, &bRecycledSubobject);
 	check(Result != NULL);
@@ -3658,7 +3659,7 @@ UObject* FObjectInitializer::CreateDefaultSubobject(UObject* Outer, FName Subobj
 		if (!OverrideClass->HasAnyClassFlags(CLASS_Abstract) || !bAbstract)
 		{
 			UObject* Template = OverrideClass->GetDefaultObject(); // force the CDO to be created if it hasn't already
-			EObjectFlags SubobjectFlags = Outer->GetMaskedFlags(RF_PropagateToSubObjects);
+			EObjectFlags SubobjectFlags = Outer->GetMaskedFlags(RF_PropagateToSubObjects) | RF_DefaultSubObject;
 			bool bOwnerArchetypeIsNotNative;
 			UClass* OuterArchetypeClass;
 
@@ -3708,7 +3709,6 @@ UObject* FObjectInitializer::CreateDefaultSubobject(UObject* Outer, FName Subobj
 #endif
 				Outer->GetClass()->AddDefaultSubobject(Result, ReturnType);
 			}
-			Result->SetFlags(RF_DefaultSubObject);
 			// Clear PendingKill flag in case we recycled a subobject of a dead object.
 			// @todo: we should not be recycling subobjects unless we're currently loading from a package
 			Result->ClearInternalFlags(EInternalObjectFlags::PendingKill);
