Added iterators for FScriptMapHelper and FScriptSetHelper.

diff --git a/Engine/Source/Runtime/CoreUObject/Public/UObject/UnrealType.h b/Engine/Source/Runtime/CoreUObject/Public/UObject/UnrealType.h
--- a/Engine/Source/Runtime/CoreUObject/Public/UObject/UnrealType.h
+++ b/Engine/Source/Runtime/CoreUObject/Public/UObject/UnrealType.h
@@ -3473,6 +3473,40 @@ public:
 		return ScriptMapHelper;
 	}
 
+	class FIterator
+	{
+	public:
+		explicit FIterator(const FScriptMapHelper& InMap) :
+			Map(InMap),
+			CurrentIndex(-1)
+		{
+			Advance();
+		}
+
+		FIterator& operator++() { Advance(); return *this; }
+		FIterator& operator++(int) { Advance(); return *this; }
+		explicit operator bool() const { return Map.IsValidIndex(CurrentIndex); }
+		int32 operator*() const { return CurrentIndex; }
+
+	private:
+		const FScriptMapHelper& Map;
+		int32 CurrentIndex;
+
+		void Advance()
+		{
+			++CurrentIndex;
+			while (CurrentIndex < Map.GetMaxIndex() && !Map.IsValidIndex(CurrentIndex))
+			{
+				++CurrentIndex;
+			}
+		}
+	};
+
+	FScriptMapHelper::FIterator CreateIterator() const
+	{
+		return FIterator(*this);
+	}
+
 private:
 	FScriptMapHelper()
 		: KeyProp(nullptr)
@@ -4005,6 +4039,40 @@ public:
 		return ScriptSetHelper;
 	}
 
+	class FIterator
+	{
+	public:
+		explicit FIterator(const FScriptSetHelper& InSet) :
+			Set(InSet),
+			CurrentIndex(-1)
+		{
+			Advance();
+		}
+
+		FIterator& operator++() { Advance(); return *this; }
+		FIterator& operator++(int) { Advance(); return *this; }
+		explicit operator bool() const { return Set.IsValidIndex(CurrentIndex); }
+		int32 operator*() const { return CurrentIndex; }
+
+	private:
+		const FScriptSetHelper& Set;
+		int32 CurrentIndex;
+
+		void Advance()
+		{
+			++CurrentIndex;
+			while (CurrentIndex < Set.GetMaxIndex() && !Set.IsValidIndex(CurrentIndex))
+			{
+				++CurrentIndex;
+			}
+		}
+	};
+
+	FScriptSetHelper::FIterator CreateIterator() const
+	{
+		return FIterator(*this);
+	}
+
 private: 
 	FScriptSetHelper()
 		: ElementProp(nullptr)
