Fixed typo

#rnx

diff --git a/Engine/Source/Runtime/CoreUObject/Private/UObject/PropertyAccessUtil.cpp b/Engine/Source/Runtime/CoreUObject/Private/UObject/PropertyAccessUtil.cpp
--- a/Engine/Source/Runtime/CoreUObject/Private/UObject/PropertyAccessUtil.cpp
+++ b/Engine/Source/Runtime/CoreUObject/Private/UObject/PropertyAccessUtil.cpp
@@ -137,7 +137,7 @@ EPropertyAccessResultFlags SetPropertyValue_DirectComplete(const UProperty* InPr
 			{
 				const void* SrcElemValue = static_cast<const uint8*>(InSrcValue) + (InProp->ElementSize * Idx);
 				const void* DestElemValue = static_cast<const uint8*>(InDestValue) + (InProp->ElementSize * Idx);
-				bIsIdentical |= InProp->Identical(SrcElemValue, DestElemValue);
+				bIsIdentical &= InProp->Identical(SrcElemValue, DestElemValue);
 			}
 			return bIsIdentical;
 		};
