Cleanup named blacklist header
#rnx

diff --git a/Engine/Source/Runtime/CoreUObject/Public/Misc/NamedBlacklist.h b/Engine/Source/Runtime/CoreUObject/Public/Misc/NamedBlacklist.h
--- a/Engine/Source/Runtime/CoreUObject/Public/Misc/NamedBlacklist.h
+++ b/Engine/Source/Runtime/CoreUObject/Public/Misc/NamedBlacklist.h
@@ -4,8 +4,6 @@
 
 #include "CoreTypes.h"
 
-#include "NamedBlacklist.generated.h"
-
 /** 
  * Filter for blacklisting or whitelisting items with ownership to allow change cleanup
  */
@@ -171,52 +169,3 @@ private:
 	/** List of filters */
 	TMap<FName, FNamedBlacklist> Filters;
 };
-
-/** 
- * Serializable blacklist (Ownership not included)
- */
-USTRUCT()
-struct FNamedBlacklistConfig
-{
-	GENERATED_BODY()
-
-	FNamedBlacklistConfig() :
-		bBlacklistAll(false)
-	{
-	}
-
-	/** Optional name for filter */
-	UPROPERTY()
-	FName Name;
-
-	/** Items to filter out */
-	UPROPERTY()
-	TArray<FName> Blacklist;
-							
-	/** Items not on the list will be filtered out when whitelist contains items */
-	UPROPERTY()
-	TArray<FName> Whitelist;
-				
-	/** Filter out all items */
-	UPROPERTY()
-	bool bBlacklistAll;
-
-	/** Append filtering rules in this config friendly struct to the destination filter, specifying an owner name allows unregister later */
-	void AppendTo(const FName OwnerName, FNamedBlacklist& Destination) const
-	{
-		for (const FName ItemName : Blacklist)
-		{
-			Destination.AddBlacklistItem(OwnerName, ItemName);
-		}
-
-		for (const FName ItemName : Whitelist)
-		{
-			Destination.AddWhitelistItem(OwnerName, ItemName);
-		}
-
-		if (bBlacklistAll)
-		{
-			Destination.AddBlacklistAll(OwnerName);
-		}
-	}
-};
