Second attempt at optimizing IsEventDrivenLoaderEnabledInCookedBuilds(). The previous attempt used the global cvar variable which is fed it's value via the config system. In this case, the function is called BEFORE the appropriate ApplyCVarSettingsFromIni() which forces the value out of config and into the cvar. It didn't matter in editor builds because the editor doesn't rely on that function for anything, but cooked platforms use it to decide how to serialize. Now, on the first call to that function, ApplyCVarSettingsFromIni() is called for the streaming settings category which forces the value out into the cvar so it can be used correctly.

diff --git a/Engine/Source/Runtime/CoreUObject/Private/Serialization/AsyncLoading.cpp b/Engine/Source/Runtime/CoreUObject/Private/Serialization/AsyncLoading.cpp
--- a/Engine/Source/Runtime/CoreUObject/Private/Serialization/AsyncLoading.cpp
+++ b/Engine/Source/Runtime/CoreUObject/Private/Serialization/AsyncLoading.cpp
@@ -8145,4 +8145,35 @@ static FAutoConsoleCommand DumpSerializeCmd(
 	);
 #endif
 
+bool IsEventDrivenLoaderEnabledInCookedBuilds()
+{
+	static struct FEventDrivenLoaderEnabledInCookedBuildsInit
+	{
+		FEventDrivenLoaderEnabledInCookedBuildsInit()
+		{
+			check(GConfig || IsEngineExitRequested());
+			if (GConfig)
+			{
+				// Ensure that the streaming settings from the config have been applied 
+				ApplyCVarSettingsFromIni(TEXT("/Script/Engine.StreamingSettings"), *GEngineIni, ECVF_SetByProjectSetting);
+			}
+		}
+	} EventDrivenLoaderEnabledInCookedBuilds;
+
+	static const bool bNoEDL = !UE_BUILD_SHIPPING && FParse::Param(FCommandLine::Get(), TEXT("NOEDL"));
+	return !bNoEDL && (GEventDrivenLoaderEnabledInCookedBuilds != 0);
+}
+
+bool IsEventDrivenLoaderEnabled()
+{
+	static struct FEventDrivenLoaderEnabledInit
+	{
+		FEventDrivenLoaderEnabledInit()
+		{
+			GEventDrivenLoaderEnabled = IsEventDrivenLoaderEnabledInCookedBuilds() && FPlatformProperties::RequiresCookedData();
+		}
+	} EventDrivenLoaderEnabledInit;
+	return GEventDrivenLoaderEnabled;
+}
+
 //#pragma clang optimize on
diff --git a/Engine/Source/Runtime/CoreUObject/Private/Serialization/AsyncPackageLoader.cpp b/Engine/Source/Runtime/CoreUObject/Private/Serialization/AsyncPackageLoader.cpp
--- a/Engine/Source/Runtime/CoreUObject/Private/Serialization/AsyncPackageLoader.cpp
+++ b/Engine/Source/Runtime/CoreUObject/Private/Serialization/AsyncPackageLoader.cpp
@@ -391,55 +391,6 @@ static FEDLBootNotificationManager& GetGEDLBootNotificationManager()
 	return Singleton;
 }
 
-bool IsEventDrivenLoaderEnabledInCookedBuilds()
-{
-	static struct FEventDrivenLoaderEnabledInCookedBuildsInit
-	{
-		bool bEventDrivenLoaderEnabled;
-		FEventDrivenLoaderEnabledInCookedBuildsInit()
-			: bEventDrivenLoaderEnabled(false)
-		{
-			SetEventDrivenLoaderEnabled();
-		}
-
-		void SetEventDrivenLoaderEnabled()
-		{
-			check(GConfig || IsEngineExitRequested());
-			if (GConfig)
-			{
-				GConfig->GetBool(TEXT("/Script/Engine.StreamingSettings"), TEXT("s.EventDrivenLoaderEnabled"), bEventDrivenLoaderEnabled, GEngineIni);
-#if !UE_BUILD_SHIPPING
-				if (FParse::Param(FCommandLine::Get(), TEXT("NOEDL")))
-				{
-					bEventDrivenLoaderEnabled = false;
-				}
-#endif
-			}
-		}
-	} EventDrivenLoaderEnabledInCookedBuilds;
-
-#if WITH_EDITOR	
-	// when building from the UE4 Editor, s.EventDrivenLoaderEnabled can be changed from Project Settings, so we need to test it at every call
-	if (GIsEditor && !IsRunningCommandlet())
-	{
-		EventDrivenLoaderEnabledInCookedBuilds.SetEventDrivenLoaderEnabled();
-	}
-#endif
-	return EventDrivenLoaderEnabledInCookedBuilds.bEventDrivenLoaderEnabled;
-}
-
-bool IsEventDrivenLoaderEnabled()
-{
-	static struct FEventDrivenLoaderEnabledInit
-	{
-		FEventDrivenLoaderEnabledInit()
-		{
-			GEventDrivenLoaderEnabled = IsEventDrivenLoaderEnabledInCookedBuilds() && FPlatformProperties::RequiresCookedData();
-		}
-	} EventDrivenLoaderEnabledInit;
-	return GEventDrivenLoaderEnabled;
-}
-
 bool IsFullyLoadedObj(UObject* Obj)
 {
 	if (!Obj)
