Fixing TFastReferenceCollector not setting the serializing object when processing TMap and TSet references which results in issues with classes like FReferencerFinder (not finding all referencers).

diff --git a/Engine/Source/Runtime/CoreUObject/Public/UObject/FastReferenceCollector.h b/Engine/Source/Runtime/CoreUObject/Public/UObject/FastReferenceCollector.h
--- a/Engine/Source/Runtime/CoreUObject/Public/UObject/FastReferenceCollector.h
+++ b/Engine/Source/Runtime/CoreUObject/Public/UObject/FastReferenceCollector.h
@@ -792,7 +792,12 @@ private:
 						void*         Map = StackEntryData + ReferenceInfo.Offset;
 						FMapProperty* MapProperty = (FMapProperty*)TokenStream->ReadPointer(TokenStreamIndex);
 						TokenReturnCount = ReferenceInfo.ReturnCount;
-						MapProperty->SerializeItem(FStructuredArchiveFromArchive(ReferenceCollector.GetVerySlowReferenceCollectorArchive()).GetSlot(), Map, nullptr);
+
+						FReferenceCollectorArchive& Ar = ReferenceCollector.GetVerySlowReferenceCollectorArchive();
+						checkSlow(Ar.GetSerializingObject() == nullptr);
+						Ar.SetSerializingObject(CurrentObject);
+						MapProperty->SerializeItem(FStructuredArchiveFromArchive(Ar).GetSlot(), Map, nullptr);
+						Ar.SetSerializingObject(nullptr);
 					}
 					break;
 					case GCRT_AddTSetReferencedObjects:
@@ -800,7 +805,12 @@ private:
 						void*         Set = StackEntryData + ReferenceInfo.Offset;
 						FSetProperty* SetProperty = (FSetProperty*)TokenStream->ReadPointer(TokenStreamIndex);
 						TokenReturnCount = ReferenceInfo.ReturnCount;
-						SetProperty->SerializeItem(FStructuredArchiveFromArchive(ReferenceCollector.GetVerySlowReferenceCollectorArchive()).GetSlot(), Set, nullptr);
+
+						FReferenceCollectorArchive& Ar = ReferenceCollector.GetVerySlowReferenceCollectorArchive();
+						checkSlow(Ar.GetSerializingObject() == nullptr);
+						Ar.SetSerializingObject(CurrentObject);
+						SetProperty->SerializeItem(FStructuredArchiveFromArchive(Ar).GetSlot(), Set, nullptr);
+						Ar.SetSerializingObject(nullptr);
 					}
 					break;
 					case GCRT_AddFieldPathReferencedObject:
