Fixing bulk data crashes in cooked builds

diff --git a/Engine/Source/Runtime/CoreUObject/Private/Serialization/BulkData2.cpp b/Engine/Source/Runtime/CoreUObject/Private/Serialization/BulkData2.cpp
--- a/Engine/Source/Runtime/CoreUObject/Private/Serialization/BulkData2.cpp
+++ b/Engine/Source/Runtime/CoreUObject/Private/Serialization/BulkData2.cpp
@@ -711,6 +711,7 @@ void* FBulkDataBase::Lock(uint32 LockFlags)
 const void* FBulkDataBase::LockReadOnly() const
 {
 	check(LockStatus == LOCKSTATUS_Unlocked);
+	LockStatus = LOCKSTATUS_ReadOnlyLock;
 	return DataBuffer;
 }
 
