Allow PropertyAccessUtil to correctly get/set bitfield bools from a non-bitfield bool

This involves changing the API to pass in both the source and dest properties so that the data can be read/written correctly. While these functions still don't support type coercion, passing in this info means that they could potentially be made to in the future.

diff --git a/Engine/Source/Runtime/CoreUObject/Private/UObject/PropertyAccessUtil.cpp b/Engine/Source/Runtime/CoreUObject/Private/UObject/PropertyAccessUtil.cpp
--- a/Engine/Source/Runtime/CoreUObject/Private/UObject/PropertyAccessUtil.cpp
+++ b/Engine/Source/Runtime/CoreUObject/Private/UObject/PropertyAccessUtil.cpp
@@ -7,54 +7,167 @@
 namespace PropertyAccessUtil
 {
 
-EPropertyAccessResultFlags GetPropertyValue_Object(const FProperty* InProp, const UObject* InObject, void* InDestValue, const int32 InArrayIndex)
+bool ArePropertiesCompatible(const FProperty* InSrcProp, const FProperty* InDestProp)
 {
-	check(InObject->IsA(InProp->GetOwnerClass()));
-	return GetPropertyValue_InContainer(InProp, InObject, InDestValue, InArrayIndex);
+	// Compare the classes as these must be an *exact* match as the access is low-level and without property coercion
+	if (InSrcProp->GetClass() != InDestProp->GetClass())
+	{
+		return false;
+	}
+
+	// Containers also need to check their inner types
+	if (const FArrayProperty* SrcArrayProp = CastField<FArrayProperty>(InSrcProp))
+	{
+		const FArrayProperty* DestArrayProp = CastFieldChecked<FArrayProperty>(InDestProp);
+		return ArePropertiesCompatible(SrcArrayProp->Inner, DestArrayProp->Inner);
+	}
+	if (const FSetProperty* SrcSetProp = CastField<FSetProperty>(InSrcProp))
+	{
+		const FSetProperty* DestSetProp = CastFieldChecked<FSetProperty>(InDestProp);
+		return ArePropertiesCompatible(SrcSetProp->ElementProp, DestSetProp->ElementProp);
+	}
+	if (const FMapProperty* SrcMapProp = CastField<FMapProperty>(InSrcProp))
+	{
+		const FMapProperty* DestMapProp = CastFieldChecked<FMapProperty>(InDestProp);
+		return ArePropertiesCompatible(SrcMapProp->KeyProp, DestMapProp->KeyProp)
+			&& ArePropertiesCompatible(SrcMapProp->ValueProp, DestMapProp->ValueProp);
+	}
+
+	return true;
+}
+
+bool IsSinglePropertyIdentical(const FProperty* InSrcProp, const void* InSrcValue, const FProperty* InDestProp, const void* InDestValue)
+{
+	if (!ArePropertiesCompatible(InSrcProp, InDestProp))
+	{
+		return false;
+	}
+
+	if (const FBoolProperty* SrcBoolProp = CastField<FBoolProperty>(InSrcProp))
+	{
+		const FBoolProperty* DestBoolProp = CastFieldChecked<FBoolProperty>(InDestProp);
+
+		// Bools can be represented as bitfields, we we have to handle the compare a little differently to only check the bool we want
+		const bool bSrcBoolValue = SrcBoolProp->GetPropertyValue(InSrcValue);
+		const bool bDestBoolValue = DestBoolProp->GetPropertyValue(InDestValue);
+		return bSrcBoolValue == bDestBoolValue;
+	}
+	else
+	{
+		return InSrcProp->Identical(InSrcValue, InDestValue);
+	}
+}
+
+bool IsCompletePropertyIdentical(const FProperty* InSrcProp, const void* InSrcValue, const FProperty* InDestProp, const void* InDestValue)
+{
+	bool bIsIdentical = InSrcProp->ArrayDim == InDestProp->ArrayDim;
+	for (int32 Idx = 0; Idx < InSrcProp->ArrayDim && bIsIdentical; ++Idx)
+	{
+		const void* SrcElemValue = static_cast<const uint8*>(InSrcValue) + (InSrcProp->ElementSize * Idx);
+		const void* DestElemValue = static_cast<const uint8*>(InDestValue) + (InDestProp->ElementSize * Idx);
+		bIsIdentical &= IsSinglePropertyIdentical(InSrcProp, SrcElemValue, InDestProp, DestElemValue);
+	}
+	return bIsIdentical;
+}
+
+bool CopySinglePropertyValue(const FProperty* InSrcProp, const void* InSrcValue, const FProperty* InDestProp, void* InDestValue)
+{
+	if (!ArePropertiesCompatible(InSrcProp, InDestProp))
+	{
+		return false;
+	}
+
+	if (const FBoolProperty* SrcBoolProp = CastField<FBoolProperty>(InSrcProp))
+	{
+		const FBoolProperty* DestBoolProp = CastFieldChecked<FBoolProperty>(InDestProp);
+
+		// Bools can be represented as bitfields, we we have to handle the copy a little differently to only extract the bool we want
+		const bool bBoolValue = SrcBoolProp->GetPropertyValue(InSrcValue);
+		DestBoolProp->SetPropertyValue(InDestValue, bBoolValue);
+	}
+	else
+	{
+		InSrcProp->CopySingleValue(InDestValue, InSrcValue);
+	}
+
+	return true;
 }
 
-EPropertyAccessResultFlags GetPropertyValue_InContainer(const FProperty* InProp, const void* InContainerData, void* InDestValue, const int32 InArrayIndex)
+bool CopyCompletePropertyValue(const FProperty* InSrcProp, const void* InSrcValue, const FProperty* InDestProp, void* InDestValue)
 {
-	if (InArrayIndex == INDEX_NONE || InProp->ArrayDim == 1)
+	if (!ArePropertiesCompatible(InSrcProp, InDestProp) || InSrcProp->ArrayDim != InDestProp->ArrayDim)
 	{
-		const void* SrcValue = InProp->ContainerPtrToValuePtr<void>(InContainerData);
-		return GetPropertyValue_DirectComplete(InProp, SrcValue, InDestValue);
+		return false;
+	}
+
+	if (const FBoolProperty* SrcBoolProp = CastField<FBoolProperty>(InSrcProp))
+	{
+		const FBoolProperty* DestBoolProp = CastFieldChecked<FBoolProperty>(InDestProp);
+
+		for (int32 Idx = 0; Idx < InSrcProp->ArrayDim; ++Idx)
+		{
+			const void* SrcElemValue = static_cast<const uint8*>(InSrcValue) + (InSrcProp->ElementSize * Idx);
+			void* DestElemValue = static_cast<uint8*>(InDestValue) + (InDestProp->ElementSize * Idx);
+
+			// Bools can be represented as bitfields, we we have to handle the copy a little differently to only extract the bool we want
+			const bool bBoolValue = SrcBoolProp->GetPropertyValue(SrcElemValue);
+			DestBoolProp->SetPropertyValue(DestElemValue, bBoolValue);
+		}
 	}
 	else
 	{
-		check(InArrayIndex < InProp->ArrayDim);
-		const void* SrcValue = InProp->ContainerPtrToValuePtr<void>(InContainerData, InArrayIndex);
-		return GetPropertyValue_DirectSingle(InProp, SrcValue, InDestValue);
+		InSrcProp->CopyCompleteValue(InDestValue, InSrcValue);
 	}
+
+	return true;
+}
+
+EPropertyAccessResultFlags GetPropertyValue_Object(const FProperty* InObjectProp, const UObject* InObject, const FProperty* InDestProp, void* InDestValue, const int32 InArrayIndex)
+{
+	check(InObject->IsA(InObjectProp->GetOwnerClass()));
+	return GetPropertyValue_InContainer(InObjectProp, InObject, InDestProp, InDestValue, InArrayIndex);
 }
 
-EPropertyAccessResultFlags GetPropertyValue_DirectSingle(const FProperty* InProp, const void* InSrcValue, void* InDestValue)
+EPropertyAccessResultFlags GetPropertyValue_InContainer(const FProperty* InContainerProp, const void* InContainerData, const FProperty* InDestProp, void* InDestValue, const int32 InArrayIndex)
 {
-	EPropertyAccessResultFlags Result = CanGetPropertyValue(InProp);
+	if (InArrayIndex == INDEX_NONE || InContainerProp->ArrayDim == 1)
+	{
+		const void* SrcValue = InContainerProp->ContainerPtrToValuePtr<void>(InContainerData);
+		return GetPropertyValue_DirectComplete(InContainerProp, SrcValue, InDestProp, InDestValue);
+	}
+	else
+	{
+		check(InArrayIndex < InContainerProp->ArrayDim);
+		const void* SrcValue = InContainerProp->ContainerPtrToValuePtr<void>(InContainerData, InArrayIndex);
+		return GetPropertyValue_DirectSingle(InContainerProp, SrcValue, InDestProp, InDestValue);
+	}
+}
+
+EPropertyAccessResultFlags GetPropertyValue_DirectSingle(const FProperty* InSrcProp, const void* InSrcValue, const FProperty* InDestProp, void* InDestValue)
+{
+	EPropertyAccessResultFlags Result = CanGetPropertyValue(InSrcProp);
 	if (Result != EPropertyAccessResultFlags::Success)
 	{
 		return Result;
 	}
 
-	return GetPropertyValue([InProp, InSrcValue, InDestValue]()
+	return GetPropertyValue([InSrcProp, InSrcValue, InDestProp, InDestValue]()
 	{
-		InProp->CopySingleValue(InDestValue, InSrcValue);
-		return true;
+		return CopySinglePropertyValue(InSrcProp, InSrcValue, InDestProp, InDestValue);
 	});
 }
 
-EPropertyAccessResultFlags GetPropertyValue_DirectComplete(const FProperty* InProp, const void* InSrcValue, void* InDestValue)
+EPropertyAccessResultFlags GetPropertyValue_DirectComplete(const FProperty* InSrcProp, const void* InSrcValue, const FProperty* InDestProp, void* InDestValue)
 {
-	EPropertyAccessResultFlags Result = CanGetPropertyValue(InProp);
+	EPropertyAccessResultFlags Result = CanGetPropertyValue(InSrcProp);
 	if (Result != EPropertyAccessResultFlags::Success)
 	{
 		return Result;
 	}
 
-	return GetPropertyValue([InProp, InSrcValue, InDestValue]()
+	return GetPropertyValue([InSrcProp, InSrcValue, InDestProp, InDestValue]()
 	{
-		InProp->CopyCompleteValue(InDestValue, InSrcValue);
-		return true;
+		return CopyCompletePropertyValue(InSrcProp, InSrcValue, InDestProp, InDestValue);
 	});
 }
 
@@ -76,79 +189,69 @@ EPropertyAccessResultFlags CanGetPropertyValue(const FProperty* InProp)
 	return EPropertyAccessResultFlags::Success;
 }
 
-EPropertyAccessResultFlags SetPropertyValue_Object(const FProperty* InProp, UObject* InObject, const void* InSrcValue, const int32 InArrayIndex, const uint64 InReadOnlyFlags)
+EPropertyAccessResultFlags SetPropertyValue_Object(const FProperty* InObjectProp, UObject* InObject, const FProperty* InSrcProp, const void* InSrcValue, const int32 InArrayIndex, const uint64 InReadOnlyFlags)
 {
-	check(InObject->IsA(InProp->GetOwnerClass()));
-	return SetPropertyValue_InContainer(InProp, InObject, InSrcValue, InArrayIndex, InReadOnlyFlags, IsObjectTemplate(InObject), [InProp, InObject]()
+	check(InObject->IsA(InObjectProp->GetOwnerClass()));
+	return SetPropertyValue_InContainer(InObjectProp, InObject, InSrcProp, InSrcValue, InArrayIndex, InReadOnlyFlags, IsObjectTemplate(InObject), [InObjectProp, InObject]()
 	{
-		return BuildBasicChangeNotify(InProp, InObject);
+		return BuildBasicChangeNotify(InObjectProp, InObject);
 	});
 }
 
-EPropertyAccessResultFlags SetPropertyValue_InContainer(const FProperty* InProp, void* InContainerData, const void* InSrcValue, const int32 InArrayIndex, const uint64 InReadOnlyFlags, const bool InOwnerIsTemplate, const FPropertyAccessBuildChangeNotifyFunc& InBuildChangeNotifyFunc)
+EPropertyAccessResultFlags SetPropertyValue_InContainer(const FProperty* InContainerProp, void* InContainerData, const FProperty* InSrcProp, const void* InSrcValue, const int32 InArrayIndex, const uint64 InReadOnlyFlags, const bool InOwnerIsTemplate, const FPropertyAccessBuildChangeNotifyFunc& InBuildChangeNotifyFunc)
 {
-	if (InArrayIndex == INDEX_NONE || InProp->ArrayDim == 1)
+	if (InArrayIndex == INDEX_NONE || InContainerProp->ArrayDim == 1)
 	{
-		void* DestValue = InProp->ContainerPtrToValuePtr<void>(InContainerData);
-		return SetPropertyValue_DirectComplete(InProp, InSrcValue, DestValue, InReadOnlyFlags, InOwnerIsTemplate, InBuildChangeNotifyFunc);
+		void* DestValue = InContainerProp->ContainerPtrToValuePtr<void>(InContainerData);
+		return SetPropertyValue_DirectComplete(InSrcProp, InSrcValue, InContainerProp, DestValue, InReadOnlyFlags, InOwnerIsTemplate, InBuildChangeNotifyFunc);
 	}
 	else
 	{
-		check(InArrayIndex < InProp->ArrayDim);
-		void* DestValue = InProp->ContainerPtrToValuePtr<void>(InContainerData, InArrayIndex);
-		return SetPropertyValue_DirectSingle(InProp, InSrcValue, DestValue, InReadOnlyFlags, InOwnerIsTemplate, InBuildChangeNotifyFunc);
+		check(InArrayIndex < InContainerProp->ArrayDim);
+		void* DestValue = InContainerProp->ContainerPtrToValuePtr<void>(InContainerData, InArrayIndex);
+		return SetPropertyValue_DirectSingle(InSrcProp, InSrcValue, InContainerProp, DestValue, InReadOnlyFlags, InOwnerIsTemplate, InBuildChangeNotifyFunc);
 	}
 }
 
-EPropertyAccessResultFlags SetPropertyValue_DirectSingle(const FProperty* InProp, const void* InSrcValue, void* InDestValue, const uint64 InReadOnlyFlags, const bool InOwnerIsTemplate, const FPropertyAccessBuildChangeNotifyFunc& InBuildChangeNotifyFunc)
+EPropertyAccessResultFlags SetPropertyValue_DirectSingle(const FProperty* InSrcProp, const void* InSrcValue, const FProperty* InDestProp, void* InDestValue, const uint64 InReadOnlyFlags, const bool InOwnerIsTemplate, const FPropertyAccessBuildChangeNotifyFunc& InBuildChangeNotifyFunc)
 {
-	EPropertyAccessResultFlags Result = CanSetPropertyValue(InProp, InReadOnlyFlags, InOwnerIsTemplate);
+	EPropertyAccessResultFlags Result = CanSetPropertyValue(InDestProp, InReadOnlyFlags, InOwnerIsTemplate);
 	if (Result != EPropertyAccessResultFlags::Success)
 	{
 		return Result;
 	}
 
-	return SetPropertyValue([InProp, InSrcValue, InDestValue](const FPropertyAccessChangeNotify* InChangeNotify)
+	return SetPropertyValue([InSrcProp, InSrcValue, InDestProp, InDestValue](const FPropertyAccessChangeNotify* InChangeNotify)
 	{
-		if (!InProp->Identical(InSrcValue, InDestValue))
+		bool bResult = true;
+		if (!IsSinglePropertyIdentical(InSrcProp, InSrcValue, InDestProp, InDestValue))
 		{
 			EmitPreChangeNotify(InChangeNotify);
-			InProp->CopySingleValue(InDestValue, InSrcValue);
+			bResult = CopySinglePropertyValue(InSrcProp, InSrcValue, InDestProp, InDestValue);
 			EmitPostChangeNotify(InChangeNotify);
 		}
-		return true;
+		return bResult;
 	}, InBuildChangeNotifyFunc);
 }
 
-EPropertyAccessResultFlags SetPropertyValue_DirectComplete(const FProperty* InProp, const void* InSrcValue, void* InDestValue, const uint64 InReadOnlyFlags, const bool InOwnerIsTemplate, const FPropertyAccessBuildChangeNotifyFunc& InBuildChangeNotifyFunc)
+EPropertyAccessResultFlags SetPropertyValue_DirectComplete(const FProperty* InSrcProp, const void* InSrcValue, const FProperty* InDestProp, void* InDestValue, const uint64 InReadOnlyFlags, const bool InOwnerIsTemplate, const FPropertyAccessBuildChangeNotifyFunc& InBuildChangeNotifyFunc)
 {
-	EPropertyAccessResultFlags Result = CanSetPropertyValue(InProp, InReadOnlyFlags, InOwnerIsTemplate);
+	EPropertyAccessResultFlags Result = CanSetPropertyValue(InDestProp, InReadOnlyFlags, InOwnerIsTemplate);
 	if (Result != EPropertyAccessResultFlags::Success)
 	{
 		return Result;
 	}
 
-	return SetPropertyValue([InProp, InSrcValue, InDestValue](const FPropertyAccessChangeNotify* InChangeNotify)
+	return SetPropertyValue([InSrcProp, InSrcValue, InDestProp, InDestValue](const FPropertyAccessChangeNotify* InChangeNotify)
 	{
-		auto IdenticalComplete = [InProp, InSrcValue, InDestValue]()
-		{
-			bool bIsIdentical = true;
-			for (int32 Idx = 0; Idx < InProp->ArrayDim && bIsIdentical; ++Idx)
-			{
-				const void* SrcElemValue = static_cast<const uint8*>(InSrcValue) + (InProp->ElementSize * Idx);
-				const void* DestElemValue = static_cast<const uint8*>(InDestValue) + (InProp->ElementSize * Idx);
-				bIsIdentical &= InProp->Identical(SrcElemValue, DestElemValue);
-			}
-			return bIsIdentical;
-		};
-
-		if (!IdenticalComplete())
+		bool bResult = true;
+		if (!IsCompletePropertyIdentical(InSrcProp, InSrcValue, InDestProp, InDestValue))
 		{
 			EmitPreChangeNotify(InChangeNotify);
-			InProp->CopyCompleteValue(InDestValue, InSrcValue);
+			bResult = CopyCompletePropertyValue(InSrcProp, InSrcValue, InDestProp, InDestValue);
 			EmitPostChangeNotify(InChangeNotify);
 		}
-		return true;
+		return bResult;
 	}, InBuildChangeNotifyFunc);
 }
 
diff --git a/Engine/Source/Runtime/CoreUObject/Public/UObject/PropertyAccessUtil.h b/Engine/Source/Runtime/CoreUObject/Public/UObject/PropertyAccessUtil.h
--- a/Engine/Source/Runtime/CoreUObject/Public/UObject/PropertyAccessUtil.h
+++ b/Engine/Source/Runtime/CoreUObject/Public/UObject/PropertyAccessUtil.h
@@ -75,51 +75,55 @@ namespace PropertyAccessUtil
 	 * High-level function for getting the value of a property from an object.
 	 * @note This function calls CanGetPropertyValue internally.
 	 *
-	 * @param InProp Property to get the value of.
+	 * @param InObjectProp Property to get the value of.
 	 * @param InObject Object containing the property.
+	 * @param InDestProp Property of the value to set (must be compatible with the source property).
 	 * @param InDestValue Instance to fill with the property value (must be a valid and constructed block of memory that is compatible with the property).
 	 * @param InArrayIndex For fixed-size array properties denotes which index of the array to get, or INDEX_NONE to get the entire property.
 	 *
 	 * @return Flags describing whether the get was successful.
 	 */
-	COREUOBJECT_API EPropertyAccessResultFlags GetPropertyValue_Object(const FProperty* InProp, const UObject* InObject, void* InDestValue, const int32 InArrayIndex);
+	COREUOBJECT_API EPropertyAccessResultFlags GetPropertyValue_Object(const FProperty* InObjectProp, const UObject* InObject, const FProperty* InDestProp, void* InDestValue, const int32 InArrayIndex);
 
 	/**
 	 * High-level function for getting the value of a property from a property container (object or struct).
 	 * @note This function calls CanGetPropertyValue internally.
 	 *
-	 * @param InProp Property to get the value of.
+	 * @param InContainerProp Property to get the value of.
 	 * @param InContainerData The instance data containing the property.
+	 * @param InDestProp Property of the value to set (must be compatible with the source property).
 	 * @param InDestValue Instance to fill with the property value (must be a valid and constructed block of memory that is compatible with the property).
 	 * @param InArrayIndex For fixed-size array properties denotes which index of the array to get, or INDEX_NONE to get the entire property.
 	 *
 	 * @return Flags describing whether the get was successful.
 	 */
-	COREUOBJECT_API EPropertyAccessResultFlags GetPropertyValue_InContainer(const FProperty* InProp, const void* InContainerData, void* InDestValue, const int32 InArrayIndex);
+	COREUOBJECT_API EPropertyAccessResultFlags GetPropertyValue_InContainer(const FProperty* InContainerProp, const void* InContainerData, const FProperty* InDestProp, void* InDestValue, const int32 InArrayIndex);
 
 	/**
 	 * High-level function for getting the single-element value of a property from memory.
 	 * @note This function calls CanGetPropertyValue internally.
 	 *
-	 * @param InProp Property to get the value of.
+	 * @param InSrcProp Property to get the value of.
 	 * @param InSrcValue The property value to copy.
+	 * @param InDestProp Property of the value to set (must be compatible with the source property).
 	 * @param InDestValue Instance to fill with the property value (must be a valid and constructed block of memory that is compatible with the property).
 	 *
 	 * @return Flags describing whether the get was successful.
 	 */
-	COREUOBJECT_API EPropertyAccessResultFlags GetPropertyValue_DirectSingle(const FProperty* InProp, const void* InSrcValue, void* InDestValue);
+	COREUOBJECT_API EPropertyAccessResultFlags GetPropertyValue_DirectSingle(const FProperty* InSrcProp, const void* InSrcValue, const FProperty* InDestProp, void* InDestValue);
 
 	/**
 	 * High-level function for getting the multi-element value of a property from memory.
 	 * @note This function calls CanGetPropertyValue internally.
 	 *
-	 * @param InProp Property to get the value of.
+	 * @param InSrcProp Property to get the value of.
 	 * @param InSrcValue The property value to copy.
+	 * @param InDestProp Property of the value to set (must be compatible with the source property).
 	 * @param InDestValue Instance to fill with the property value (must be a valid and constructed block of memory that is compatible with the property).
 	 *
 	 * @return Flags describing whether the get was successful.
 	 */
-	COREUOBJECT_API EPropertyAccessResultFlags GetPropertyValue_DirectComplete(const FProperty* InProp, const void* InSrcValue, void* InDestValue);
+	COREUOBJECT_API EPropertyAccessResultFlags GetPropertyValue_DirectComplete(const FProperty* InSrcProp, const void* InSrcValue, const FProperty* InDestProp, void* InDestValue);
 
 	/**
 	 * Low-level function for getting the value of a property.
@@ -144,22 +148,24 @@ namespace PropertyAccessUtil
 	 * High-level function for setting the value of a property on an object.
 	 * @note This function calls CanSetPropertyValue internally, and will emit property change notifications for the object.
 	 *
-	 * @param InProp Property to set the value of.
+	 * @param InObjectProp Property to set the value of.
 	 * @param InObject Object containing the property.
+	 * @param InSrcProp Property of the value to set (must be compatible with the dest property).
 	 * @param InSrcValue The value to set on the property.
 	 * @param InArrayIndex For fixed-size array properties denotes which index of the array to set, or INDEX_NONE to set the entire property.
 	 * @param InReadOnlyFlags Flags controlling which properties are considered read-only.
 	 *
 	 * @return Flags describing whether the set was successful.
 	 */
-	COREUOBJECT_API EPropertyAccessResultFlags SetPropertyValue_Object(const FProperty* InProp, UObject* InObject, const void* InSrcValue, const int32 InArrayIndex, const uint64 InReadOnlyFlags);
+	COREUOBJECT_API EPropertyAccessResultFlags SetPropertyValue_Object(const FProperty* InObjectProp, UObject* InObject, const FProperty* InSrcProp, const void* InSrcValue, const int32 InArrayIndex, const uint64 InReadOnlyFlags);
 
 	/**
 	 * High-level function for setting the value of a property on a property container (object or struct).
 	 * @note This function calls CanSetPropertyValue internally.
 	 *
-	 * @param InProp Property to set the value of.
+	 * @param InContainerProp Property to set the value of.
 	 * @param InContainerData The instance data containing the property.
+	 * @param InSrcProp Property of the value to set (must be compatible with the dest property).
 	 * @param InSrcValue The value to set on the property.
 	 * @param InArrayIndex For fixed-size array properties denotes which index of the array to set, or INDEX_NONE to set the entire property.
 	 * @param InReadOnlyFlags Flags controlling which properties are considered read-only.
@@ -168,14 +174,15 @@ namespace PropertyAccessUtil
 	 *
 	 * @return Flags describing whether the set was successful.
 	 */
-	COREUOBJECT_API EPropertyAccessResultFlags SetPropertyValue_InContainer(const FProperty* InProp, void* InContainerData, const void* InSrcValue, const int32 InArrayIndex, const uint64 InReadOnlyFlags, const bool InOwnerIsTemplate, const FPropertyAccessBuildChangeNotifyFunc& InBuildChangeNotifyFunc);
+	COREUOBJECT_API EPropertyAccessResultFlags SetPropertyValue_InContainer(const FProperty* InContainerProp, void* InContainerData, const FProperty* InSrcProp, const void* InSrcValue, const int32 InArrayIndex, const uint64 InReadOnlyFlags, const bool InOwnerIsTemplate, const FPropertyAccessBuildChangeNotifyFunc& InBuildChangeNotifyFunc);
 	
 	/**
 	 * High-level function for setting the single-element value of a property in memory.
 	 * @note This function calls CanSetPropertyValue internally.
 	 *
-	 * @param InProp Property to set the value of.
+	 * @param InSrcProp Property to set the value of.
 	 * @param InSrcValue The value to set on the property.
+	 * @param InDestProp Property to get the value from (must be compatible with the source property).
 	 * @param InDestValue Instance to fill with the property value (must be a valid and constructed block of memory that is compatible with the property).
 	 * @param InReadOnlyFlags Flags controlling which properties are considered read-only.
 	 * @param InOwnerIsTemplate True if the owner object is considered a template (see IsObjectTemplate).
@@ -183,14 +190,15 @@ namespace PropertyAccessUtil
 	 *
 	 * @return Flags describing whether the set was successful.
 	 */
-	COREUOBJECT_API EPropertyAccessResultFlags SetPropertyValue_DirectSingle(const FProperty* InProp, const void* InSrcValue, void* InDestValue, const uint64 InReadOnlyFlags, const bool InOwnerIsTemplate, const FPropertyAccessBuildChangeNotifyFunc& InBuildChangeNotifyFunc);
+	COREUOBJECT_API EPropertyAccessResultFlags SetPropertyValue_DirectSingle(const FProperty* InSrcProp, const void* InSrcValue, const FProperty* InDestProp, void* InDestValue, const uint64 InReadOnlyFlags, const bool InOwnerIsTemplate, const FPropertyAccessBuildChangeNotifyFunc& InBuildChangeNotifyFunc);
 
 	/**
 	 * High-level function for setting the multi-element value of a property in memory.
 	 * @note This function calls CanSetPropertyValue internally.
 	 *
-	 * @param InProp Property to set the value of.
+	 * @param InSrcProp Property to set the value of.
 	 * @param InSrcValue The value to set on the property.
+	 * @param InDestProp Property to get the value from (must be compatible with the source property).
 	 * @param InDestValue Instance to fill with the property value (must be a valid and constructed block of memory that is compatible with the property).
 	 * @param InReadOnlyFlags Flags controlling which properties are considered read-only.
 	 * @param InOwnerIsTemplate True if the owner object is considered a template (see IsObjectTemplate).
@@ -198,7 +206,7 @@ namespace PropertyAccessUtil
 	 *
 	 * @return Flags describing whether the set was successful.
 	 */
-	COREUOBJECT_API EPropertyAccessResultFlags SetPropertyValue_DirectComplete(const FProperty* InProp, const void* InSrcValue, void* InDestValue, const uint64 InReadOnlyFlags, const bool InOwnerIsTemplate, const FPropertyAccessBuildChangeNotifyFunc& InBuildChangeNotifyFunc);
+	COREUOBJECT_API EPropertyAccessResultFlags SetPropertyValue_DirectComplete(const FProperty* InSrcProp, const void* InSrcValue, const FProperty* InDestProp, void* InDestValue, const uint64 InReadOnlyFlags, const bool InOwnerIsTemplate, const FPropertyAccessBuildChangeNotifyFunc& InBuildChangeNotifyFunc);
 
 	/**
 	 * Low-level function for setting the value of a property.
diff --git a/Engine/Source/Runtime/Engine/Classes/Kismet/KismetSystemLibrary.h b/Engine/Source/Runtime/Engine/Classes/Kismet/KismetSystemLibrary.h
--- a/Engine/Source/Runtime/Engine/Classes/Kismet/KismetSystemLibrary.h
+++ b/Engine/Source/Runtime/Engine/Classes/Kismet/KismetSystemLibrary.h
@@ -1743,13 +1743,13 @@ class ENGINE_API UKismetSystemLibrary : public UBlueprintFunctionLibrary
 	/** Read the value of a named property on the given object to the given value pointer */
     UFUNCTION(BlueprintCallable, CustomThunk, Category = "Properties", meta=(CustomStructureParam="ValuePtr", BlueprintInternalUseOnly="true"))
     static bool GetEditorProperty(UObject* Object, const FName PropertyName, int32& ValuePtr);
-	static bool Generic_GetEditorProperty(const UObject* Object, const FProperty* Property, void* ValuePtr);
+	static bool Generic_GetEditorProperty(const UObject* Object, const FProperty* ObjectProp, void* ValuePtr, const FProperty* ValueProp);
 	DECLARE_FUNCTION(execGetEditorProperty);
 
 	/** Write the value of the given value pointer to a named property on the given object */
     UFUNCTION(BlueprintCallable, CustomThunk, Category = "Properties", meta=(CustomStructureParam="ValuePtr", BlueprintInternalUseOnly="true"))
     static bool SetEditorProperty(UObject* Object, const FName PropertyName, const int32& ValuePtr);
-	static bool Generic_SetEditorProperty(UObject* Object, const FProperty* Property, const void* ValuePtr);
+	static bool Generic_SetEditorProperty(UObject* Object, const FProperty* ObjectProp, const void* ValuePtr, const FProperty* ValueProp);
 	DECLARE_FUNCTION(execSetEditorProperty);
 #endif
 
diff --git a/Engine/Source/Runtime/Engine/Private/KismetSystemLibrary.cpp b/Engine/Source/Runtime/Engine/Private/KismetSystemLibrary.cpp
--- a/Engine/Source/Runtime/Engine/Private/KismetSystemLibrary.cpp
+++ b/Engine/Source/Runtime/Engine/Private/KismetSystemLibrary.cpp
@@ -2681,25 +2681,25 @@ bool UKismetSystemLibrary::GetEditorProperty(UObject* Object, const FName Proper
 	return false;
 }
 
-bool UKismetSystemLibrary::Generic_GetEditorProperty(const UObject* Object, const FProperty* Property, void* ValuePtr)
+bool UKismetSystemLibrary::Generic_GetEditorProperty(const UObject* Object, const FProperty* ObjectProp, void* ValuePtr, const FProperty* ValueProp)
 {
-	const EPropertyAccessResultFlags AccessResult = PropertyAccessUtil::GetPropertyValue_Object(Property, Object, ValuePtr, INDEX_NONE);
+	const EPropertyAccessResultFlags AccessResult = PropertyAccessUtil::GetPropertyValue_Object(ObjectProp, Object, ValueProp, ValuePtr, INDEX_NONE);
 
 	if (EnumHasAnyFlags(AccessResult, EPropertyAccessResultFlags::PermissionDenied))
 	{
 		if (EnumHasAnyFlags(AccessResult, EPropertyAccessResultFlags::AccessProtected))
 		{
-			FFrame::KismetExecutionMessage(*FString::Printf(TEXT("Property '%s' on '%s' (%s) is protected and cannot be read"), *Property->GetName(), *Object->GetPathName(), *Object->GetClass()->GetName()), ELogVerbosity::Warning, PropertyGetFailedWarning);
+			FFrame::KismetExecutionMessage(*FString::Printf(TEXT("Property '%s' on '%s' (%s) is protected and cannot be read"), *ObjectProp->GetName(), *Object->GetPathName(), *Object->GetClass()->GetName()), ELogVerbosity::Warning, PropertyGetFailedWarning);
 			return false;
 		}
 
-		FFrame::KismetExecutionMessage(*FString::Printf(TEXT("Property '%s' on '%s' (%s) cannot be read"), *Property->GetName(), *Object->GetPathName(), *Object->GetClass()->GetName()), ELogVerbosity::Warning, PropertyGetFailedWarning);
+		FFrame::KismetExecutionMessage(*FString::Printf(TEXT("Property '%s' on '%s' (%s) cannot be read"), *ObjectProp->GetName(), *Object->GetPathName(), *Object->GetClass()->GetName()), ELogVerbosity::Warning, PropertyGetFailedWarning);
 		return false;
 	}
 
 	if (EnumHasAnyFlags(AccessResult, EPropertyAccessResultFlags::ConversionFailed))
 	{
-		FFrame::KismetExecutionMessage(*FString::Printf(TEXT("Property '%s' on '%s' (%s) tried to read into a property value of the incorrect type"), *Property->GetName(), *Object->GetPathName(), *Object->GetClass()->GetName()), ELogVerbosity::Warning, PropertyGetFailedWarning);
+		FFrame::KismetExecutionMessage(*FString::Printf(TEXT("Property '%s' (%s) on '%s' (%s) tried to read into a property value of the incorrect type (%s)"), *ObjectProp->GetName(), *ObjectProp->GetClass()->GetName(), *Object->GetPathName(), *Object->GetClass()->GetName(), *ValueProp->GetClass()->GetName()), ELogVerbosity::Warning, PropertyGetFailedWarning);
 		return false;
 	}
 
@@ -2740,15 +2740,15 @@ DEFINE_FUNCTION(UKismetSystemLibrary::execGetEditorProperty)
 	if (Object)
 	{
 		const FProperty* ObjectProp = PropertyAccessUtil::FindPropertyByName(PropertyName, Object->GetClass());
-		if (ObjectProp && ObjectProp->GetClass() == ValueProp->GetClass()) // Compare the classes as these must be an *exact* match as the read is low-level and without property coercion
+		if (ObjectProp)
 		{
 			P_NATIVE_BEGIN;
-			bResult = Generic_GetEditorProperty(Object, ObjectProp, ValuePtr);
+			bResult = Generic_GetEditorProperty(Object, ObjectProp, ValuePtr, ValueProp);
 			P_NATIVE_END;
 		}
 		else
 		{
-			FFrame::KismetExecutionMessage(*FString::Printf(TEXT("Property '%s' on '%s' (%s) was missing or didn't match the output value type (%s)"), *PropertyName.ToString(), *Object->GetPathName(), *Object->GetClass()->GetName(), *ValueProp->GetClass()->GetName()), ELogVerbosity::Warning, PropertyGetFailedWarning);
+			FFrame::KismetExecutionMessage(*FString::Printf(TEXT("Property '%s' on '%s' (%s) was missing"), *PropertyName.ToString(), *Object->GetPathName(), *Object->GetClass()->GetName()), ELogVerbosity::Warning, PropertyGetFailedWarning);
 		}
 	}
 
@@ -2762,43 +2762,43 @@ bool UKismetSystemLibrary::SetEditorProperty(UObject* Object, const FName Proper
 	return false;
 }
 
-bool UKismetSystemLibrary::Generic_SetEditorProperty(UObject* Object, const FProperty* Property, const void* ValuePtr)
+bool UKismetSystemLibrary::Generic_SetEditorProperty(UObject* Object, const FProperty* ObjectProp, const void* ValuePtr, const FProperty* ValueProp)
 {
-	const EPropertyAccessResultFlags AccessResult = PropertyAccessUtil::SetPropertyValue_Object(Property, Object, ValuePtr, INDEX_NONE, PropertyAccessUtil::EditorReadOnlyFlags);
+	const EPropertyAccessResultFlags AccessResult = PropertyAccessUtil::SetPropertyValue_Object(ObjectProp, Object, ValueProp, ValuePtr, INDEX_NONE, PropertyAccessUtil::EditorReadOnlyFlags);
 
 	if (EnumHasAnyFlags(AccessResult, EPropertyAccessResultFlags::PermissionDenied))
 	{
 		if (EnumHasAnyFlags(AccessResult, EPropertyAccessResultFlags::AccessProtected))
 		{
-			FFrame::KismetExecutionMessage(*FString::Printf(TEXT("Property '%s' on '%s' (%s) is protected and cannot be set"), *Property->GetName(), *Object->GetPathName(), *Object->GetClass()->GetName()), ELogVerbosity::Warning, PropertySetFailedWarning);
+			FFrame::KismetExecutionMessage(*FString::Printf(TEXT("Property '%s' on '%s' (%s) is protected and cannot be set"), *ObjectProp->GetName(), *Object->GetPathName(), *Object->GetClass()->GetName()), ELogVerbosity::Warning, PropertySetFailedWarning);
 			return false;
 		}
 
 		if (EnumHasAnyFlags(AccessResult, EPropertyAccessResultFlags::CannotEditTemplate))
 		{
-			FFrame::KismetExecutionMessage(*FString::Printf(TEXT("Property '%s' on '%s' (%s) cannot be edited on templates"), *Property->GetName(), *Object->GetPathName(), *Object->GetClass()->GetName()), ELogVerbosity::Warning, PropertySetFailedWarning);
+			FFrame::KismetExecutionMessage(*FString::Printf(TEXT("Property '%s' on '%s' (%s) cannot be edited on templates"), *ObjectProp->GetName(), *Object->GetPathName(), *Object->GetClass()->GetName()), ELogVerbosity::Warning, PropertySetFailedWarning);
 			return false;
 		}
 
 		if (EnumHasAnyFlags(AccessResult, EPropertyAccessResultFlags::CannotEditInstance))
 		{
-			FFrame::KismetExecutionMessage(*FString::Printf(TEXT("Property '%s' on '%s' (%s) cannot be edited on instances"), *Property->GetName(), *Object->GetPathName(), *Object->GetClass()->GetName()), ELogVerbosity::Warning, PropertySetFailedWarning);
+			FFrame::KismetExecutionMessage(*FString::Printf(TEXT("Property '%s' on '%s' (%s) cannot be edited on instances"), *ObjectProp->GetName(), *Object->GetPathName(), *Object->GetClass()->GetName()), ELogVerbosity::Warning, PropertySetFailedWarning);
 			return false;
 		}
 
 		if (EnumHasAnyFlags(AccessResult, EPropertyAccessResultFlags::ReadOnly))
 		{
-			FFrame::KismetExecutionMessage(*FString::Printf(TEXT("Property '%s' on '%s' (%s) is read-only and cannot be set"), *Property->GetName(), *Object->GetPathName(), *Object->GetClass()->GetName()), ELogVerbosity::Warning, PropertySetFailedWarning);
+			FFrame::KismetExecutionMessage(*FString::Printf(TEXT("Property '%s' on '%s' (%s) is read-only and cannot be set"), *ObjectProp->GetName(), *Object->GetPathName(), *Object->GetClass()->GetName()), ELogVerbosity::Warning, PropertySetFailedWarning);
 			return false;
 		}
 
-		FFrame::KismetExecutionMessage(*FString::Printf(TEXT("Property '%s' on '%s' (%s) cannot be read"), *Property->GetName(), *Object->GetPathName(), *Object->GetClass()->GetName()), ELogVerbosity::Warning, PropertySetFailedWarning);
+		FFrame::KismetExecutionMessage(*FString::Printf(TEXT("Property '%s' on '%s' (%s) cannot be set"), *ObjectProp->GetName(), *Object->GetPathName(), *Object->GetClass()->GetName()), ELogVerbosity::Warning, PropertySetFailedWarning);
 		return false;
 	}
 
 	if (EnumHasAnyFlags(AccessResult, EPropertyAccessResultFlags::ConversionFailed))
 	{
-		FFrame::KismetExecutionMessage(*FString::Printf(TEXT("Property '%s' on '%s' (%s) tried to write into a property value of the incorrect type"), *Property->GetName(), *Object->GetPathName(), *Object->GetClass()->GetName()), ELogVerbosity::Warning, PropertySetFailedWarning);
+		FFrame::KismetExecutionMessage(*FString::Printf(TEXT("Property '%s' (%s) on '%s' (%s) tried to write into a property value of the incorrect type (%s)"), *ObjectProp->GetName(), *ObjectProp->GetClass()->GetName(), *Object->GetPathName(), *Object->GetClass()->GetName(), *ValueProp->GetClass()->GetName()), ELogVerbosity::Warning, PropertySetFailedWarning);
 		return false;
 	}
 
@@ -2839,15 +2839,15 @@ DEFINE_FUNCTION(UKismetSystemLibrary::execSetEditorProperty)
 	if (Object)
 	{
 		const FProperty* ObjectProp = PropertyAccessUtil::FindPropertyByName(PropertyName, Object->GetClass());
-		if (ObjectProp && ObjectProp->GetClass() == ValueProp->GetClass()) // Compare the classes as these must be an *exact* match as the read is low-level and without property coercion
+		if (ObjectProp)
 		{
 			P_NATIVE_BEGIN;
-			bResult = Generic_SetEditorProperty(Object, ObjectProp, ValuePtr);
+			bResult = Generic_SetEditorProperty(Object, ObjectProp, ValuePtr, ValueProp);
 			P_NATIVE_END;
 		}
 		else
 		{
-			FFrame::KismetExecutionMessage(*FString::Printf(TEXT("Property '%s' on '%s' (%s) was missing or didn't match the output value type (%s)"), *PropertyName.ToString(), *Object->GetPathName(), *Object->GetClass()->GetName(), *ValueProp->GetClass()->GetName()), ELogVerbosity::Warning, PropertySetFailedWarning);
+			FFrame::KismetExecutionMessage(*FString::Printf(TEXT("Property '%s' on '%s' (%s) was missing"), *PropertyName.ToString(), *Object->GetPathName(), *Object->GetClass()->GetName()), ELogVerbosity::Warning, PropertySetFailedWarning);
 		}
 	}
 
