Added warning when loading map properties with lazy object pointer keys

Actual fix targeting UE4.24.2

diff --git a/Engine/Source/Runtime/CoreUObject/Private/UObject/PropertyMap.cpp b/Engine/Source/Runtime/CoreUObject/Private/UObject/PropertyMap.cpp
--- a/Engine/Source/Runtime/CoreUObject/Private/UObject/PropertyMap.cpp
+++ b/Engine/Source/Runtime/CoreUObject/Private/UObject/PropertyMap.cpp
@@ -364,6 +364,14 @@ void FMapProperty::SerializeItem(FStructuredArchive::FSlot Slot, void* Value, co
 			KeyProp->InitializeValue(TempKeyStorage);
 		}
 
+		if (KeyProp->IsA<ULazyObjectProperty>())
+		{
+			UE_LOG(LogProperty, Warning, 
+				TEXT("Loading map properties with lazy object pointer keys has a known bug in UE4.24.1 (UE-85796). "
+					 "Resaving this package may cause data loss. '%s' in package '%s'. "),
+				*GetFullName(), *UnderlyingArchive.GetArchiveName());
+		}
+
 		// Read remaining items into container
 		for (; NumEntries; --NumEntries)
 		{
