Fix incoming code in light of FProperty change
#rnx

diff --git a/Engine/Source/Runtime/CoreUObject/Private/UObject/PropertyMap.cpp b/Engine/Source/Runtime/CoreUObject/Private/UObject/PropertyMap.cpp
--- a/Engine/Source/Runtime/CoreUObject/Private/UObject/PropertyMap.cpp
+++ b/Engine/Source/Runtime/CoreUObject/Private/UObject/PropertyMap.cpp
@@ -364,7 +364,7 @@ void FMapProperty::SerializeItem(FStructuredArchive::FSlot Slot, void* Value, co
 			KeyProp->InitializeValue(TempKeyStorage);
 		}
 
-		if (KeyProp->IsA<ULazyObjectProperty>())
+		if (KeyProp->IsA<FLazyObjectProperty>())
 		{
 			UE_LOG(LogProperty, Warning, 
 				TEXT("Loading map properties with lazy object pointer keys has a known bug in UE4.24.1 (UE-85796). "
