Removed some redundant field iteration and class enumeration when scraping loc data from assets

#rnx

diff --git a/Engine/Source/Runtime/CoreUObject/Private/Serialization/PropertyLocalizationDataGathering.cpp b/Engine/Source/Runtime/CoreUObject/Private/Serialization/PropertyLocalizationDataGathering.cpp
--- a/Engine/Source/Runtime/CoreUObject/Private/Serialization/PropertyLocalizationDataGathering.cpp
+++ b/Engine/Source/Runtime/CoreUObject/Private/Serialization/PropertyLocalizationDataGathering.cpp
@@ -23,25 +23,16 @@ FPropertyLocalizationDataGatherer::FPropertyLocalizationDataGatherer(TArray<FGat
 	, AllObjectsInPackage()
 {
 	// Build up the list of objects that are within our package - we won't follow object references to things outside of our package
+	ForEachObjectWithOuter(Package, [this](UObject* Object)
 	{
-		TArray<UObject*> AllObjectsInPackageArray;
-		GetObjectsWithOuter(Package, AllObjectsInPackageArray, true, RF_Transient, EInternalObjectFlags::PendingKill);
-
-		AllObjectsInPackage.Reserve(AllObjectsInPackageArray.Num());
-		for (UObject* Object : AllObjectsInPackageArray)
-		{
-			AllObjectsInPackage.Add(Object);
-		}
-	}
-
-	TArray<UObject*> RootObjectsInPackage;
-	GetObjectsWithOuter(Package, RootObjectsInPackage, false, RF_Transient, EInternalObjectFlags::PendingKill);
+		AllObjectsInPackage.Add(Object);
+	}, true, RF_Transient, EInternalObjectFlags::PendingKill);
 
 	// Iterate over each root object in the package
-	for (const UObject* Object : RootObjectsInPackage)
+	ForEachObjectWithOuter(Package, [this](UObject* Object)
 	{
 		GatherLocalizationDataFromObjectWithCallbacks(Object, EPropertyLocalizationGathererTextFlags::None);
-	}
+	}, false, RF_Transient, EInternalObjectFlags::PendingKill);
 
 	// Iterate any bytecode containing objects
 	for (const FObjectAndGatherFlags& BytecodeToGather : BytecodePendingGather)
@@ -88,30 +79,52 @@ const FPropertyLocalizationDataGatherer::FGatherableFieldsForType& FPropertyLoca
 {
 	TUniquePtr<FGatherableFieldsForType> GatherableFieldsForType = MakeUnique<FGatherableFieldsForType>();
 
-	for (TFieldIterator<FProperty> FieldIt(InType, EFieldIteratorFlags::IncludeSuper, EFieldIteratorFlags::ExcludeDeprecated, EFieldIteratorFlags::IncludeInterfaces); FieldIt; ++FieldIt)
+	// Include the parent fields (this will recursively cache any parent types)
+	if (const UStruct* SuperType = InType->GetSuperStruct())
+	{
+		const FGatherableFieldsForType& GatherableFieldsForSuperType = GetGatherableFieldsForType(SuperType);
+		*GatherableFieldsForType = GatherableFieldsForSuperType;
+	}
+
+	// See if we have a custom handler for this type
+	if (const UClass* Class = Cast<UClass>(InType))
+	{
+		const FLocalizationDataGatheringCallback* CustomCallback = GetTypeSpecificLocalizationDataGatheringCallbacks().Find(Class);
+		if (CustomCallback)
+		{
+			GatherableFieldsForType->CustomCallback = CustomCallback;
+		}
+	}
+
+	// Look for potential properties
+	for (TFieldIterator<const FProperty> FieldIt(InType, EFieldIteratorFlags::ExcludeSuper, EFieldIteratorFlags::ExcludeDeprecated, EFieldIteratorFlags::IncludeInterfaces); FieldIt; ++FieldIt)
 	{
-		// Look for potential properties
+		auto ProcessInnerProperty = [this, &GatherableFieldsForType](const FProperty* InProp, const FProperty* InTypeProp)
 		{
-			auto ProcessInnerProperty = [this, &GatherableFieldsForType](const FProperty* InProp, const FProperty* InTypeProp)
+			if (CanGatherFromInnerProperty(InProp))
 			{
-				if (CanGatherFromInnerProperty(InProp))
-				{
-					GatherableFieldsForType->Properties.AddUnique(InTypeProp);
-				}
-			};
+				checkSlow(!GatherableFieldsForType->Properties.Contains(InTypeProp));
+				GatherableFieldsForType->Properties.Add(InTypeProp);
+				return true;
+			}
+			return false;
+		};
 
-			const FProperty* PropertyField = *FieldIt;
-			if (PropertyField)
+		const FProperty* PropertyField = *FieldIt;
+		if (PropertyField)
+		{
+			if (!ProcessInnerProperty(PropertyField, PropertyField))
 			{
-				ProcessInnerProperty(PropertyField, PropertyField);
 				if (const FArrayProperty* ArrayProp = CastField<const FArrayProperty>(PropertyField))
 				{
 					ProcessInnerProperty(ArrayProp->Inner, PropertyField);
 				}
 				if (const FMapProperty* MapProp = CastField<const FMapProperty>(PropertyField))
 				{
-					ProcessInnerProperty(MapProp->KeyProp, PropertyField);
-					ProcessInnerProperty(MapProp->ValueProp, PropertyField);
+					if (!ProcessInnerProperty(MapProp->KeyProp, PropertyField))
+					{
+						ProcessInnerProperty(MapProp->ValueProp, PropertyField);
+					}
 				}
 				if (const FSetProperty* SetProp = CastField<const FSetProperty>(PropertyField))
 				{
@@ -120,15 +133,14 @@ const FPropertyLocalizationDataGatherer::FGatherableFieldsForType& FPropertyLoca
 			}
 		}
 	}
+
+	// Look for potential functions
 	for (TFieldIterator<const UField> FieldIt(InType, EFieldIteratorFlags::IncludeSuper, EFieldIteratorFlags::ExcludeDeprecated, EFieldIteratorFlags::IncludeInterfaces); FieldIt; ++FieldIt)
 	{
-		// Look for potential functions
+		const UFunction* FunctionField = Cast<UFunction>(*FieldIt);
+		if (FunctionField && FunctionField->Script.Num() > 0 && IsObjectValidForGather(FunctionField))
 		{
-			const UFunction* FunctionField = Cast<UFunction>(*FieldIt);
-			if (FunctionField && IsObjectValidForGather(FunctionField))
-			{
-				GatherableFieldsForType->Functions.Add(FunctionField);
-			}
+			GatherableFieldsForType->Functions.Add(FunctionField);
 		}
 	}
 
@@ -154,25 +166,15 @@ bool FPropertyLocalizationDataGatherer::CanGatherFromInnerProperty(const FProper
 
 void FPropertyLocalizationDataGatherer::GatherLocalizationDataFromObjectWithCallbacks(const UObject* Object, const EPropertyLocalizationGathererTextFlags GatherTextFlags)
 {
-	// See if we have a custom handler for this type
-	FLocalizationDataGatheringCallback* CustomCallback = nullptr;
-	for (const UClass* Class = Object->GetClass(); Class != nullptr; Class = Class->GetSuperClass())
-	{
-		CustomCallback = GetTypeSpecificLocalizationDataGatheringCallbacks().Find(Class);
-		if (CustomCallback)
-		{
-			break;
-		}
-	}
-
-	if (CustomCallback)
+	const FGatherableFieldsForType& GatherableFieldsForType = GetGatherableFieldsForType(Object->GetClass());
+	if (GatherableFieldsForType.CustomCallback)
 	{
 		checkf(IsObjectValidForGather(Object), TEXT("Cannot gather for objects outside of the current package! Package: '%s'. Object: '%s'."), *Package->GetFullName(), *Object->GetFullName());
 
 		if (ShouldProcessObject(Object, GatherTextFlags))
 		{
 			MarkObjectProcessed(Object, GatherTextFlags);
-			(*CustomCallback)(Object, *this, GatherTextFlags);
+			(*GatherableFieldsForType.CustomCallback)(Object, *this, GatherTextFlags);
 		}
 	}
 	else if (ShouldProcessObject(Object, GatherTextFlags))
@@ -198,22 +200,22 @@ void FPropertyLocalizationDataGatherer::GatherLocalizationDataFromObject(const U
 			ResultFlags |= EPropertyLocalizationGathererResultFlags::HasScript;
 		}
 
-		if (Object->IsA<UStruct>())
+		if (const UStruct* Struct = Cast<UStruct>(Object))
 		{
-			BytecodePendingGather.Add(FObjectAndGatherFlags(Object, GatherTextFlags));
+			if (Struct->Script.Num() > 0)
+			{
+				BytecodePendingGather.Add(FObjectAndGatherFlags(Struct, GatherTextFlags));
+			}
 		}
 	}
 
 	// Gather from anything that has us as their outer, as not all objects are reachable via a property pointer.
 	if (!(GatherTextFlags & EPropertyLocalizationGathererTextFlags::SkipSubObjects))
 	{
-		TArray<UObject*> ChildObjects;
-		GetObjectsWithOuter(Object, ChildObjects, false, RF_Transient, EInternalObjectFlags::PendingKill);
-
-		for (const UObject* ChildObject : ChildObjects)
+		ForEachObjectWithOuter(Object, [this, GatherTextFlags](UObject* ChildObject)
 		{
 			GatherLocalizationDataFromObjectWithCallbacks(ChildObject, GatherTextFlags);
-		}
+		}, false, RF_Transient, EInternalObjectFlags::PendingKill);
 	}
 }
 
@@ -230,7 +232,7 @@ void FPropertyLocalizationDataGatherer::GatherLocalizationDataFromObjectFields(c
 
 		if (ArchetypeObject)
 		{
-			const FProperty* ArchetypePropertyField = FindField<FProperty>(ArchetypeObject->GetClass(), *PropertyField->GetName());
+			const FProperty* ArchetypePropertyField = ArchetypeObject->GetClass()->IsChildOf(Object->GetClass()) ? PropertyField : FindField<FProperty>(ArchetypeObject->GetClass(), *PropertyField->GetName());
 			if (ArchetypePropertyField && ArchetypePropertyField->IsA(PropertyField->GetClass()))
 			{
 				DefaultValueAddress = ArchetypePropertyField->ContainerPtrToValuePtr<void>(ArchetypeObject);
@@ -705,9 +707,8 @@ void FPropertyLocalizationDataGatherer::GatherScriptBytecode(const FString& Path
 	if (ScriptData.Num() > 0)
 	{
 		ResultFlags |= EPropertyLocalizationGathererResultFlags::HasScript;
+		FGatherTextFromScriptBytecode(*PathToScript, ScriptData, *this, bIsEditorOnly);
 	}
-
-	FGatherTextFromScriptBytecode(*PathToScript, ScriptData, *this, bIsEditorOnly);
 }
 
 bool FPropertyLocalizationDataGatherer::IsDefaultTextInstance(const FText& Text) const
diff --git a/Engine/Source/Runtime/CoreUObject/Public/Serialization/PropertyLocalizationDataGathering.h b/Engine/Source/Runtime/CoreUObject/Public/Serialization/PropertyLocalizationDataGathering.h
--- a/Engine/Source/Runtime/CoreUObject/Public/Serialization/PropertyLocalizationDataGathering.h
+++ b/Engine/Source/Runtime/CoreUObject/Public/Serialization/PropertyLocalizationDataGathering.h
@@ -74,10 +74,14 @@ ENUM_CLASS_FLAGS(EPropertyLocalizationGathererResultFlags);
 class COREUOBJECT_API FPropertyLocalizationDataGatherer
 {
 public:
+	typedef TFunction<void(const UObject* const, FPropertyLocalizationDataGatherer&, const EPropertyLocalizationGathererTextFlags)> FLocalizationDataGatheringCallback;
+	typedef TMap<const UClass*, FLocalizationDataGatheringCallback> FLocalizationDataGatheringCallbackMap;
+
 	struct FGatherableFieldsForType
 	{
 		TArray<const FProperty*> Properties;
 		TArray<const UFunction*> Functions;
+		const FLocalizationDataGatheringCallback* CustomCallback = nullptr;
 
 		bool HasFields() const
 		{
@@ -85,9 +89,6 @@ public:
 		}
 	};
 
-	typedef TFunction<void(const UObject* const, FPropertyLocalizationDataGatherer&, const EPropertyLocalizationGathererTextFlags)> FLocalizationDataGatheringCallback;
-	typedef TMap<const UClass*, FLocalizationDataGatheringCallback> FLocalizationDataGatheringCallbackMap;
-
 	FPropertyLocalizationDataGatherer(TArray<FGatherableTextData>& InOutGatherableTextDataArray, const UPackage* const InPackage, EPropertyLocalizationGathererResultFlags& OutResultFlags);
 
 	// Non-copyable
