Manual merge of CL: 10759497
Move ResetLoaders to after Presave, to handle Presaves that cause a circular reference that loads the package being saved again.  Since it is possible to do so, we move it to immediately before it is needed, when the file on disk is written out (or copied from temp).

diff --git a/Engine/Source/Runtime/CoreUObject/Private/UObject/Linker.cpp b/Engine/Source/Runtime/CoreUObject/Private/UObject/Linker.cpp
--- a/Engine/Source/Runtime/CoreUObject/Private/UObject/Linker.cpp
+++ b/Engine/Source/Runtime/CoreUObject/Private/UObject/Linker.cpp
@@ -818,25 +818,22 @@ FLinkerLoad* LoadPackageLinker(UPackage* InOuter, const TCHAR* InLongPackageName
 void ResetLoadersForSave(UObject* InOuter, const TCHAR *Filename)
 {
 	UPackage* Package = dynamic_cast<UPackage*>(InOuter);
-	// If we have a loader for the package, unload it to prevent conflicts if we are resaving to the same filename
-	FLinkerLoad* Loader = FLinkerLoad::FindExistingLinkerForPackage(Package);
-	// This is the loader corresponding to the package we're saving.
-	if( Loader )
-	{
-		if (!Package->HasAnyPackageFlags(PKG_FilterEditorOnly))
-		{
-			// Before we save the package, make sure that we load up any thumbnails that aren't already
-			// in memory so that they won't be wiped out during this save
-			Loader->SerializeThumbnails();
-		}
+	EnsureLoadingComplete(Package);
 
+	if(FLinkerLoad* Loader = FLinkerLoad::FindExistingLinkerForPackage(Package))
+	{
 		// Compare absolute filenames to see whether we're trying to save over an existing file.
 		if( FPaths::ConvertRelativePathToFull(Filename) == FPaths::ConvertRelativePathToFull( Loader->Filename ) )
 		{
 			// Detach all exports from the linker and dissociate the linker.
-			ResetLoaders( InOuter );
+			ResetLoaders( Package );
 		}
 	}
 }
 
+void EnsureLoadingComplete(UPackage* Package)
+{
+	FLinkerManager::Get().EnsureLoadingComplete(Package);
+}
+
 #undef LOCTEXT_NAMESPACE
diff --git a/Engine/Source/Runtime/CoreUObject/Private/UObject/LinkerManager.cpp b/Engine/Source/Runtime/CoreUObject/Private/UObject/LinkerManager.cpp
--- a/Engine/Source/Runtime/CoreUObject/Private/UObject/LinkerManager.cpp
+++ b/Engine/Source/Runtime/CoreUObject/Private/UObject/LinkerManager.cpp
@@ -152,6 +152,24 @@ void FLinkerManager::ResetLoaders(UObject* InPkg)
 	}
 }
 
+void FLinkerManager::EnsureLoadingComplete(UPackage* Package)
+{
+	if (!Package)
+	{
+		return;
+	}
+	FLinkerLoad* Linker = FLinkerLoad::FindExistingLinkerForPackage(Package);
+	if (!Linker)
+	{
+		return;
+	}
+	if (!Package->HasAnyPackageFlags(PKG_FilterEditorOnly))
+	{
+		Linker->SerializeThumbnails();
+	}
+	Linker->LoadAndDetachAllBulkData();
+}
+
 void FLinkerManager::DissociateImportsAndForcedExports()
 {
 	{
diff --git a/Engine/Source/Runtime/CoreUObject/Private/UObject/LinkerManager.h b/Engine/Source/Runtime/CoreUObject/Private/UObject/LinkerManager.h
--- a/Engine/Source/Runtime/CoreUObject/Private/UObject/LinkerManager.h
+++ b/Engine/Source/Runtime/CoreUObject/Private/UObject/LinkerManager.h
@@ -113,6 +113,9 @@ public:
 	/** Empty the loaders */
 	void ResetLoaders(UObject* InPkg);
 
+	/** Complete all loading (thumbnails/bulkdata) for the given Package */
+	void EnsureLoadingComplete(UPackage* Package);
+
 	/**
 	* Dissociates all linker import and forced export object references. This currently needs to
 	* happen as the referred objects might be destroyed at any time.
diff --git a/Engine/Source/Runtime/CoreUObject/Private/UObject/SavePackage.cpp b/Engine/Source/Runtime/CoreUObject/Private/UObject/SavePackage.cpp
--- a/Engine/Source/Runtime/CoreUObject/Private/UObject/SavePackage.cpp
+++ b/Engine/Source/Runtime/CoreUObject/Private/UObject/SavePackage.cpp
@@ -85,7 +85,8 @@ namespace SavePackageStats
 	static double SavePackageTimeSec = 0.0;
 	static double TagPackageExportsPresaveTimeSec = 0.0;
 	static double TagPackageExportsTimeSec = 0.0;
-	static double ResetLoadersForSaveTimeSec = 0.0;
+	static double FullyLoadLoadersTimeSec = 0.0;
+	static double ResetLoadersTimeSec = 0.0;
 	static double TagPackageExportsGetObjectsWithOuter = 0.0;
 	static double TagPackageExportsGetObjectsWithMarks = 0.0;
 	static double SerializeImportsTimeSec = 0.0;
@@ -106,7 +107,8 @@ namespace SavePackageStats
 		ADD_COOK_STAT(SavePackageTimeSec);
 		ADD_COOK_STAT(TagPackageExportsPresaveTimeSec);
 		ADD_COOK_STAT(TagPackageExportsTimeSec);
-		ADD_COOK_STAT(ResetLoadersForSaveTimeSec);
+		ADD_COOK_STAT(FullyLoadLoadersTimeSec);
+		ADD_COOK_STAT(ResetLoadersTimeSec);
 		ADD_COOK_STAT(TagPackageExportsGetObjectsWithOuter);
 		ADD_COOK_STAT(TagPackageExportsGetObjectsWithMarks);
 		ADD_COOK_STAT(SerializeImportsTimeSec);
@@ -3627,8 +3629,9 @@ FSavePackageResultStruct UPackage::Save(UPackage* InOuter, UObject* Base, EObjec
 		bool Success = true;
 		bool bRequestStub = false;
 		{
-			COOK_STAT(FScopedDurationTimer SaveTimer(SavePackageStats::ResetLoadersForSaveTimeSec));
-			ResetLoadersForSave(InOuter, Filename);
+			// FullyLoad the package's Loader, so that anything we need to serialize (bulkdata, thumbnails) is available
+			COOK_STAT(FScopedDurationTimer SaveTimer(SavePackageStats::FullyLoadLoadersTimeSec));
+			EnsureLoadingComplete(InOuter);
 		}
 		SlowTask.EnterProgressFrame();
 
@@ -5700,6 +5703,13 @@ FSavePackageResultStruct UPackage::Save(UPackage* InOuter, UObject* Base, EObjec
 
 				if( Success == true )
 				{
+					{
+						// ResetLoaders on the Package so that we drop the handle to the file on disk and can write to it
+						COOK_STAT(FScopedDurationTimer SaveTimer(SavePackageStats::ResetLoadersTimeSec));
+						ResetLoaders(InOuter);
+					}
+
+					// Compress the temporarily file to destination.
 					if (bSaveAsync)
 					{						
 						FString NewPathToSave = NewPath;
diff --git a/Engine/Source/Runtime/CoreUObject/Public/UObject/Linker.h b/Engine/Source/Runtime/CoreUObject/Public/UObject/Linker.h
--- a/Engine/Source/Runtime/CoreUObject/Public/UObject/Linker.h
+++ b/Engine/Source/Runtime/CoreUObject/Public/UObject/Linker.h
@@ -605,3 +605,10 @@ COREUOBJECT_API FString GetPrestreamPackageLinkerName(const TCHAR* InLongPackage
  * @param	Filename		The filename we are saving too
  */
 COREUOBJECT_API void ResetLoadersForSave(UObject* InOuter, const TCHAR *Filename);
+
+/*
+ * Ensure all data that can be loaded from the linker (thumbnails, bulk data) is loaded, in preparation for saving out the given package
+ *
+ * @param Package	The the package for which the linker should be fully loaded
+ */
+COREUOBJECT_API void EnsureLoadingComplete(UPackage* Package);
\ No newline at end of file
