Fixed CIS warning

#rnx

diff --git a/Engine/Source/Runtime/CoreUObject/Private/Serialization/PropertyLocalizationDataGathering.cpp b/Engine/Source/Runtime/CoreUObject/Private/Serialization/PropertyLocalizationDataGathering.cpp
--- a/Engine/Source/Runtime/CoreUObject/Private/Serialization/PropertyLocalizationDataGathering.cpp
+++ b/Engine/Source/Runtime/CoreUObject/Private/Serialization/PropertyLocalizationDataGathering.cpp
@@ -228,17 +228,7 @@ void FPropertyLocalizationDataGatherer::GatherLocalizationDataFromObjectFields(c
 	for (const FProperty* PropertyField : GatherableFieldsForType.Properties)
 	{
 		const void* ValueAddress = PropertyField->ContainerPtrToValuePtr<void>(Object);
-		const void* DefaultValueAddress = nullptr;
-
-		if (ArchetypeObject)
-		{
-			const FProperty* ArchetypePropertyField = ArchetypeObject->GetClass()->IsChildOf(Object->GetClass()) ? PropertyField : FindField<FProperty>(ArchetypeObject->GetClass(), *PropertyField->GetName());
-			if (ArchetypePropertyField && ArchetypePropertyField->IsA(PropertyField->GetClass()))
-			{
-				DefaultValueAddress = ArchetypePropertyField->ContainerPtrToValuePtr<void>(ArchetypeObject);
-			}
-		}
-
+		const void* DefaultValueAddress = (ArchetypeObject && ArchetypeObject->IsA(PropertyField->GetOwnerClass())) ? PropertyField->ContainerPtrToValuePtr<void>(ArchetypeObject) : nullptr;
 		GatherLocalizationDataFromChildTextProperties(PathToParent, PropertyField, ValueAddress, DefaultValueAddress, GatherTextFlags | (PropertyField->HasAnyPropertyFlags(CPF_EditorOnly) ? EPropertyLocalizationGathererTextFlags::ForceEditorOnly : EPropertyLocalizationGathererTextFlags::None));
 	}
 
