Skipping CollectGarabe calls made during initial load. It's not safe since token streams may not be assembled yet.

diff --git a/Engine/Source/Runtime/CoreUObject/Private/UObject/GarbageCollection.cpp b/Engine/Source/Runtime/CoreUObject/Private/UObject/GarbageCollection.cpp
--- a/Engine/Source/Runtime/CoreUObject/Private/UObject/GarbageCollection.cpp
+++ b/Engine/Source/Runtime/CoreUObject/Private/UObject/GarbageCollection.cpp
@@ -1837,6 +1837,12 @@ void GatherUnreachableObjects(bool bForceSingleThreaded)
  */
 void CollectGarbageInternal(EObjectFlags KeepFlags, bool bPerformFullPurge)
 {
+	if (GIsInitialLoad)
+	{
+		// During initial load classes may not yet have their GC token streams assembled
+		UE_LOG(LogGarbage, Log, TEXT("Skipping CollectGarbage() call during initial load. It's not safe."));
+		return;
+	}
 	SCOPE_TIME_GUARD(TEXT("Collect Garbage"));
 	SCOPED_NAMED_EVENT(CollectGarbageInternal, FColor::Red);
 	CSV_EVENT_GLOBAL(TEXT("GC"));
