FField::GetLinker() will no longer assert when there's no owner UObject to allow support for 'loose' FProperties in Python

diff --git a/Engine/Source/Runtime/CoreUObject/Private/UObject/Field.cpp b/Engine/Source/Runtime/CoreUObject/Private/UObject/Field.cpp
--- a/Engine/Source/Runtime/CoreUObject/Private/UObject/Field.cpp
+++ b/Engine/Source/Runtime/CoreUObject/Private/UObject/Field.cpp
@@ -551,8 +551,7 @@ bool FField::IsIn(const FField* InOwner) const
 FLinkerLoad* FField::GetLinker() const
 {
 	UObject* OwnerObject = GetOwnerUObject();
-	check(OwnerObject);
-	return OwnerObject->GetLinker();
+	return OwnerObject ? OwnerObject->GetLinker() : nullptr;
 }
 
 void FField::AddCppProperty(FProperty* Property)
