#online Disabled implicit conversions from FUniqueNetId to FUniqueNetIdWrapper for a couple of reasons:
- the conversions were hiding AsShared calls, which can be a minefield since not all net IDs are owned by shared pointers (probably ought to be)
- removal of the conversions showed quite a lot of mistaken dereferencing, only to be converted back to a wrapper again

diff --git a/Engine/Source/Runtime/CoreUObject/Public/UObject/CoreOnline.h b/Engine/Source/Runtime/CoreUObject/Public/UObject/CoreOnline.h
--- a/Engine/Source/Runtime/CoreUObject/Public/UObject/CoreOnline.h
+++ b/Engine/Source/Runtime/CoreUObject/Public/UObject/CoreOnline.h
@@ -102,21 +102,11 @@ class FUniqueNetId : public TSharedFromThis<FUniqueNetId>
 {
 protected:
 
-	/** Hidden on purpose */
-	FUniqueNetId()
-	{
-	}
+	/** Only constructible by derived type */
+	FUniqueNetId() = default;
 
-	/** Hidden on purpose */
-	FUniqueNetId(const FUniqueNetId& Src)
-	{
-	}
-
-	/** Hidden on purpose */
-	FUniqueNetId& operator=(const FUniqueNetId& Src)
-	{
-		return *this;
-	}
+	FUniqueNetId(const FUniqueNetId& Src) = default;
+	FUniqueNetId& operator=(const FUniqueNetId& Src) = default;
 
 	virtual bool Compare(const FUniqueNetId& Other) const
 	{
@@ -126,17 +116,19 @@ protected:
 
 public:
 
+	virtual ~FUniqueNetId() = default;
+
 	/**
 	 *	Comparison operator
 	 */
-	bool operator==(const FUniqueNetId& Other) const
+	friend bool operator==(const FUniqueNetId& Lhs, const FUniqueNetId& Rhs)
 	{
-		return Other.Compare(*this);
+		return Lhs.Compare(Rhs);
 	}
 
-	bool operator!=(const FUniqueNetId& Other) const
+	friend bool operator!=(const FUniqueNetId& Lhs, const FUniqueNetId& Rhs)
 	{
-		return !(FUniqueNetId::operator==(Other));
+		return !Lhs.Compare(Rhs);
 	}
 
 	/**
@@ -145,7 +137,7 @@ public:
 	 *
 	 * @return FName representing the Type
 	 */
-	virtual FName GetType() const { return NAME_None; /* This should be pure virtual, however, older versions of the OSS plugins cannot handle that */ };
+	virtual FName GetType() const { return NAME_None; /* This should be pure virtual, however, older versions of the OSS plugins cannot handle that */ }
 
 	/** 
 	 * Get the raw byte representation of this opaque data
@@ -200,23 +192,16 @@ public:
 	{
 		return CityHash32(reinterpret_cast<const char*>(Value.GetBytes()), Value.GetSize());
 	}
-
-	virtual ~FUniqueNetId() {}
 };
 
 struct FUniqueNetIdWrapper
 {
 public:
 
-	FUniqueNetIdWrapper() 
-		: UniqueNetId()
-	{
-	}
+	FUniqueNetIdWrapper() = default;
+	virtual ~FUniqueNetIdWrapper() = default;
 
-	FUniqueNetIdWrapper(const FUniqueNetId& InUniqueNetId)
-		: UniqueNetId(InUniqueNetId.AsShared())
-	{
-	}
+	// copy operators generated by compiler
 
 	FUniqueNetIdWrapper(const TSharedRef<const FUniqueNetId>& InUniqueNetId)
 		: UniqueNetId(InUniqueNetId)
@@ -228,50 +213,6 @@ public:
 	{
 	}
 
-	FUniqueNetIdWrapper(const FUniqueNetIdWrapper& InUniqueNetId)
-		: UniqueNetId(InUniqueNetId.UniqueNetId)
-	{
-	}
-
-	virtual ~FUniqueNetIdWrapper()
-	{
-	}
-
-	/** Assignment operator */
-	FUniqueNetIdWrapper& operator=(const FUniqueNetIdWrapper& Other)
-	{
-		if (&Other != this)
-		{
-			UniqueNetId = Other.UniqueNetId;
-		}
-		return *this;
-	}
-
-	/** Comparison operator */
-	bool operator==(FUniqueNetIdWrapper const& Other) const
-	{
-		// Both invalid structs or both valid and deep comparison equality
-		bool bBothValid = IsValid() && Other.IsValid();
-		bool bBothInvalid = !IsValid() && !Other.IsValid();
-		return (bBothInvalid || (bBothValid && (*UniqueNetId == *Other.UniqueNetId)));
-	}
-
-	bool operator==(TSharedPtr<const FUniqueNetId> const& Other) const
-	{
-		return Other == (*this);
-	}
-
-	bool operator==(TSharedRef<const FUniqueNetId> const& Other) const
-	{
-		return Other == (*this);
-	}
-
-	/** Comparison operator */
-	bool operator!=(FUniqueNetIdWrapper const& Other) const
-	{
-		return !(*this == Other);
-	}
-
 	FName GetType() const
 	{
 		return IsValid() ? UniqueNetId->GetType() : NAME_None;
@@ -349,30 +290,57 @@ public:
 		return InvalidId;
 	}
 
-	bool friend operator==(TSharedPtr<const FUniqueNetId> const& LHS, FUniqueNetIdWrapper const& RHS)
+	friend bool operator==(const FUniqueNetIdWrapper& Lhs, const FUniqueNetIdWrapper& Rhs)
+	{
+		bool bLhsValid = Lhs.IsValid();
+		return bLhsValid == Rhs.IsValid() && (!bLhsValid || *Lhs == *Rhs);
+	}
+
+	friend bool operator!=(const FUniqueNetIdWrapper& Lhs, const FUniqueNetIdWrapper& Rhs)
+	{
+		return !(Lhs == Rhs);
+	}
+
+	friend bool operator==(const FUniqueNetIdWrapper& Lhs, const FUniqueNetId& Rhs)
+	{
+		bool bLhsValid = Lhs.IsValid();
+		return bLhsValid == Rhs.IsValid() && (!bLhsValid || *Lhs == Rhs);
+	}
+
+	friend bool operator!=(const FUniqueNetIdWrapper& Lhs, const FUniqueNetId& Rhs)
+	{
+		return !(Lhs == Rhs);
+	}
+
+	friend bool operator==(const FUniqueNetId& Lhs, const FUniqueNetIdWrapper& Rhs)
+	{
+		return Rhs == Lhs;
+	}
+
+	friend bool operator!=(const FUniqueNetId& Lhs, const FUniqueNetIdWrapper& Rhs)
+	{
+		return Rhs != Lhs;
+	}
+
+	// comparison with nullptr (alternative to IsValid)
+	friend bool operator==(const FUniqueNetIdWrapper& NetIdWrapper, TYPE_OF_NULLPTR)
 	{
-		// Both invalid structs or both valid and deep comparison equality
-		bool bBothValid = RHS.IsValid() && LHS.IsValid() && LHS->IsValid();
-		bool bBothInvalid = !RHS.IsValid() && (!LHS.IsValid() || !LHS->IsValid());
-		return (bBothInvalid || (bBothValid && (*LHS == *RHS)));
+		return !NetIdWrapper.IsValid();
 	}
 
-	bool friend operator==(TSharedRef<const FUniqueNetId> const& LHS, FUniqueNetIdWrapper const& RHS)
+	friend bool operator!=(const FUniqueNetIdWrapper& NetIdWrapper, TYPE_OF_NULLPTR)
 	{
-		// Both invalid structs or both valid and deep comparison equality
-		bool bBothValid = RHS.IsValid() && LHS->IsValid();
-		bool bBothInvalid = !RHS.IsValid() && !LHS->IsValid();
-		return (bBothInvalid || (bBothValid && (*LHS == *RHS)));
+		return NetIdWrapper.IsValid();
 	}
 
-	bool friend operator!=(TSharedPtr<const FUniqueNetId> const& LHS, FUniqueNetIdWrapper const& RHS)
+	friend bool operator==(TYPE_OF_NULLPTR, const FUniqueNetIdWrapper& NetIdWrapper)
 	{
-		return !(LHS == RHS);
+		return !NetIdWrapper.IsValid();
 	}
 
-	bool friend operator!=(TSharedRef<const FUniqueNetId> const& LHS, FUniqueNetIdWrapper const& RHS)
+	friend bool operator!=(TYPE_OF_NULLPTR, const FUniqueNetIdWrapper& NetIdWrapper)
 	{
-		return !(LHS == RHS);
+		return NetIdWrapper.IsValid();
 	}
 
 protected:
diff --git a/Engine/Source/Runtime/Engine/Classes/GameFramework/OnlineReplStructs.h b/Engine/Source/Runtime/Engine/Classes/GameFramework/OnlineReplStructs.h
--- a/Engine/Source/Runtime/Engine/Classes/GameFramework/OnlineReplStructs.h
+++ b/Engine/Source/Runtime/Engine/Classes/GameFramework/OnlineReplStructs.h
@@ -39,11 +39,6 @@ struct FUniqueNetIdRepl : public FUniqueNetIdWrapper
 	{
 	}
 
-	FUniqueNetIdRepl(const FUniqueNetId& InUniqueNetId)
-		: FUniqueNetIdWrapper(InUniqueNetId)
-	{
-	}
-
 	FUniqueNetIdRepl(const TSharedRef<const FUniqueNetId>& InUniqueNetId)
 		: FUniqueNetIdWrapper(InUniqueNetId)
 	{
