#online Temporarily restored implicit conversion from FUniqueNetId to FUniqueNetIdWrapper while I fix some platform-specific code

diff --git a/Engine/Source/Runtime/CoreUObject/Public/UObject/CoreOnline.h b/Engine/Source/Runtime/CoreUObject/Public/UObject/CoreOnline.h
--- a/Engine/Source/Runtime/CoreUObject/Public/UObject/CoreOnline.h
+++ b/Engine/Source/Runtime/CoreUObject/Public/UObject/CoreOnline.h
@@ -213,6 +213,12 @@ public:
 	{
 	}
 
+	// temporarily restored implicit conversion from FUniqueNetId
+	FUniqueNetIdWrapper(const FUniqueNetId& InUniqueNetId)
+		: UniqueNetId(InUniqueNetId.AsShared())
+	{
+	}
+
 	FName GetType() const
 	{
 		return IsValid() ? UniqueNetId->GetType() : NAME_None;
