Fix a crash when running the 'memreport' console command.

We only need to check if the FArchive is valid for the new BulkData serialize path if we are actually going to serialize it.

diff --git a/Engine/Source/Runtime/CoreUObject/Private/Serialization/BulkData2.cpp b/Engine/Source/Runtime/CoreUObject/Private/Serialization/BulkData2.cpp
--- a/Engine/Source/Runtime/CoreUObject/Private/Serialization/BulkData2.cpp
+++ b/Engine/Source/Runtime/CoreUObject/Private/Serialization/BulkData2.cpp
@@ -523,13 +523,12 @@ void FBulkDataBase::Serialize(FArchive& Ar, UObject* Owner, int32 /*Index*/, boo
 	SCOPED_LOADTIMER(BulkData_Serialize);
 
 #if WITH_EDITOR == 0 && WITH_EDITORONLY_DATA == 0
-	check(Ar.IsLoading());				// Only support loading from cooked data!
-	check(!GIsEditor);					// The editor path is not supported
-	check(GEventDrivenLoaderEnabled);	// We are assuming the EDL path is enabled ( TODO: Might need to remove this check)
-	check(LockStatus == LOCKSTATUS_Unlocked);
-
 	if (Ar.IsPersistent() && !Ar.IsObjectReferenceCollector() && !Ar.ShouldSkipBulkData())
 	{
+		check(Ar.IsLoading());				// Only support loading from cooked data!
+		check(!GIsEditor);					// The editor path is not supported
+		check(LockStatus == LOCKSTATUS_Unlocked);
+
 		Ar << BulkDataFlags;
 
 		int64 ElementCount = 0;
