Garbage Collector can now perform UObject destruction on a worker thread so that it doesn't need to wait for it on the game thread.

+ Optimized the destruction phase so that it no longer iterates over the entire UObject array but reuses the unreachable objects array instead
+ Added gc.MultithreadedDestructionEnabled setting to be able to disable async destruction
+ Exit Purge code cleanup to actually use the same code as GC for routing BeginDestroy to objects

diff --git a/Engine/Plugins/Experimental/AppleImageUtils/Source/AppleImageUtils/Public/AppleImageUtilsBlueprintProxy.h b/Engine/Plugins/Experimental/AppleImageUtils/Source/AppleImageUtils/Public/AppleImageUtilsBlueprintProxy.h
--- a/Engine/Plugins/Experimental/AppleImageUtils/Source/AppleImageUtils/Public/AppleImageUtilsBlueprintProxy.h
+++ b/Engine/Plugins/Experimental/AppleImageUtils/Source/AppleImageUtils/Public/AppleImageUtilsBlueprintProxy.h
@@ -100,6 +100,10 @@ public:
 	virtual TStatId GetStatId() const override { RETURN_QUICK_DECLARE_CYCLE_STAT(UAppleImageUtilsBaseAsyncTaskBlueprintProxy, STATGROUP_Tickables); }
 	//~ End FTickableObject Interface
 
+	//~ Begin UObject Interface
+	virtual bool IsDestructionThreadSafe() const override { return false; }
+	//~ End UObject Interface
+
 	/** The async task to check during Tick() */
 	TSharedPtr<FAppleImageUtilsConversionTaskBase, ESPMode::ThreadSafe> ConversionTask;
 
diff --git a/Engine/Plugins/Experimental/AppleVision/Source/AppleVision/Public/AppleVisionBlueprintProxy.h b/Engine/Plugins/Experimental/AppleVision/Source/AppleVision/Public/AppleVisionBlueprintProxy.h
--- a/Engine/Plugins/Experimental/AppleVision/Source/AppleVision/Public/AppleVisionBlueprintProxy.h
+++ b/Engine/Plugins/Experimental/AppleVision/Source/AppleVision/Public/AppleVisionBlueprintProxy.h
@@ -44,6 +44,10 @@ public:
 	virtual TStatId GetStatId() const override { RETURN_QUICK_DECLARE_CYCLE_STAT(UAppleVisionDetectFacesAsyncTaskBlueprintProxy, STATGROUP_Tickables); }
 	//~ End FTickableObject Interface
 
+	//~ Begin UObject Interface
+	virtual bool IsDestructionThreadSafe() const override { return false; }
+	//~ End UObject Interface
+
 	/** The async task to check during Tick() */
 	TSharedPtr<FAppleVisionDetectFacesAsyncTaskBase, ESPMode::ThreadSafe> AsyncTask;
 
diff --git a/Engine/Plugins/Experimental/StereoPanorama/Source/StereoPanorama/Private/SceneCapturer.h b/Engine/Plugins/Experimental/StereoPanorama/Source/StereoPanorama/Private/SceneCapturer.h
--- a/Engine/Plugins/Experimental/StereoPanorama/Source/StereoPanorama/Private/SceneCapturer.h
+++ b/Engine/Plugins/Experimental/StereoPanorama/Source/StereoPanorama/Private/SceneCapturer.h
@@ -72,6 +72,10 @@ public:
 		RETURN_QUICK_DECLARE_CYCLE_STAT( USceneCapturer, STATGROUP_Tickables );
 	}
 
+	//~ Begin UObject Interface
+	virtual bool IsDestructionThreadSafe() const override { return false; }
+	//~ End UObject Interface
+
 public:
 
 	void InitCaptureComponent( USceneCaptureComponent2D* CaptureComponent, float HFov, float VFov, EStereoscopicPass InStereoPass );
diff --git a/Engine/Plugins/Lumin/MagicLeap/Source/MagicLeapIdentity/Public/MagicLeapIdentity.h b/Engine/Plugins/Lumin/MagicLeap/Source/MagicLeapIdentity/Public/MagicLeapIdentity.h
--- a/Engine/Plugins/Lumin/MagicLeap/Source/MagicLeapIdentity/Public/MagicLeapIdentity.h
+++ b/Engine/Plugins/Lumin/MagicLeap/Source/MagicLeapIdentity/Public/MagicLeapIdentity.h
@@ -88,6 +88,10 @@ public:
 	/** FTickableGameObject interface */
 	virtual UWorld* GetTickableGameObjectWorld() const override;
 
+	//~ Begin UObject Interface
+	virtual bool IsDestructionThreadSafe() const override { return false; }
+	//~ End UObject Interface
+
 private:
 	class FIdentityImpl *Impl;
 };
diff --git a/Engine/Plugins/NetcodeUnitTest/NetcodeUnitTest/Source/NetcodeUnitTest/Classes/UnitTestManager.h b/Engine/Plugins/NetcodeUnitTest/NetcodeUnitTest/Source/NetcodeUnitTest/Classes/UnitTestManager.h
--- a/Engine/Plugins/NetcodeUnitTest/NetcodeUnitTest/Source/NetcodeUnitTest/Classes/UnitTestManager.h
+++ b/Engine/Plugins/NetcodeUnitTest/NetcodeUnitTest/Source/NetcodeUnitTest/Classes/UnitTestManager.h
@@ -297,7 +297,11 @@ public:
 	 */
 	bool Exec(UWorld* InWorld, const TCHAR* Cmd, FOutputDevice& Ar);
 
-
+	/**
+	 * UObject methods
+	 */
+	virtual bool IsDestructionThreadSafe() const override { return false; }
+	
 	/**
 	 * FOutputDevice methods
 	 */
diff --git a/Engine/Plugins/Runtime/Oculus/OculusVR/Source/OculusHMD/Private/OculusSceneCaptureCubemap.h b/Engine/Plugins/Runtime/Oculus/OculusVR/Source/OculusHMD/Private/OculusSceneCaptureCubemap.h
--- a/Engine/Plugins/Runtime/Oculus/OculusVR/Source/OculusHMD/Private/OculusSceneCaptureCubemap.h
+++ b/Engine/Plugins/Runtime/Oculus/OculusVR/Source/OculusHMD/Private/OculusSceneCaptureCubemap.h
@@ -56,6 +56,9 @@ public:
 	static void CaptureCubemapCommandHandler(const TArray<FString>& Args, UWorld* World, FOutputDevice& Ar);
 #endif // UE_BUILD_SHIPPING
 
+	// UObject implementation
+	virtual bool IsDestructionThreadSafe() const override { return false; }
+
 private:
 
 	enum EStage
diff --git a/Engine/Source/Developer/AITestSuite/Classes/MockAI.h b/Engine/Source/Developer/AITestSuite/Classes/MockAI.h
--- a/Engine/Source/Developer/AITestSuite/Classes/MockAI.h
+++ b/Engine/Source/Developer/AITestSuite/Classes/MockAI.h
@@ -60,4 +60,7 @@ class UMockAI : public UObject
 	void SetEnableTicking(bool bShouldTick);
 
 	virtual void TickMe(float DeltaTime);
+
+	// UObject implementation
+	virtual bool IsDestructionThreadSafe() const override { return false; }
 };
diff --git a/Engine/Source/Developer/AssetTools/Private/AssetTools.h b/Engine/Source/Developer/AssetTools/Private/AssetTools.h
--- a/Engine/Source/Developer/AssetTools/Private/AssetTools.h
+++ b/Engine/Source/Developer/AssetTools/Private/AssetTools.h
@@ -55,6 +55,9 @@ class UAssetToolsImpl : public UObject, public IAssetTools
 public:
 	UAssetToolsImpl(const FObjectInitializer& ObjectInitializer);
 
+	// UObject implementation
+	virtual bool IsDestructionThreadSafe() const override { return false; }
+
 	// IAssetTools implementation
 	virtual void RegisterAssetTypeActions(const TSharedRef<IAssetTypeActions>& NewActions) override;
 	virtual void UnregisterAssetTypeActions(const TSharedRef<IAssetTypeActions>& ActionsToRemove) override;
diff --git a/Engine/Source/Editor/SkeletonEditor/Public/BoneProxy.h b/Engine/Source/Editor/SkeletonEditor/Public/BoneProxy.h
--- a/Engine/Source/Editor/SkeletonEditor/Public/BoneProxy.h
+++ b/Engine/Source/Editor/SkeletonEditor/Public/BoneProxy.h
@@ -24,6 +24,7 @@ public:
 	/** UObject interface */
 	virtual void PreEditChange(FEditPropertyChain& PropertyAboutToChange) override;
 	virtual void PostEditChangeProperty(FPropertyChangedEvent& PropertyChangedEvent) override;
+	virtual bool IsDestructionThreadSafe() const override { return false; }
 
 	/** FTickableEditorObject interface */
 	virtual void Tick(float DeltaTime) override;
diff --git a/Engine/Source/Editor/UnrealEd/Classes/CookOnTheSide/CookOnTheFlyServer.h b/Engine/Source/Editor/UnrealEd/Classes/CookOnTheSide/CookOnTheFlyServer.h
--- a/Engine/Source/Editor/UnrealEd/Classes/CookOnTheSide/CookOnTheFlyServer.h
+++ b/Engine/Source/Editor/UnrealEd/Classes/CookOnTheSide/CookOnTheFlyServer.h
@@ -281,6 +281,11 @@ public:
 	*/
 	virtual bool Exec(class UWorld* InWorld, const TCHAR* Cmd, FOutputDevice& Ar) override;
 
+	/**
+	  * UObject interface
+	  */
+	virtual bool IsDestructionThreadSafe() const override { return false; }
+
 	/**
 	 * Dumps cooking stats to the log
 	 *  run from the exec command "Cook stats"
diff --git a/Engine/Source/Runtime/AIModule/Classes/AISubsystem.h b/Engine/Source/Runtime/AIModule/Classes/AISubsystem.h
--- a/Engine/Source/Runtime/AIModule/Classes/AISubsystem.h
+++ b/Engine/Source/Runtime/AIModule/Classes/AISubsystem.h
@@ -30,7 +30,11 @@ public:
 	virtual void Tick(float DeltaTime) override {}
 	virtual ETickableTickType GetTickableTickType() const override;
 	virtual TStatId GetStatId() const override;
-	// FTickableGameObject begin
+	// FTickableGameObject end
+
+	// UObject begin
+	virtual bool IsDestructionThreadSafe() const override { return false; }
+	// UObject end
 
 	UWorld* GetWorldFast() const { return AISystem ? AISystem->GetOuterWorld() : GetOuter()->GetWorld(); }
 };
diff --git a/Engine/Source/Runtime/AIModule/Classes/EnvironmentQuery/EnvQueryManager.h b/Engine/Source/Runtime/AIModule/Classes/EnvironmentQuery/EnvQueryManager.h
--- a/Engine/Source/Runtime/AIModule/Classes/EnvironmentQuery/EnvQueryManager.h
+++ b/Engine/Source/Runtime/AIModule/Classes/EnvironmentQuery/EnvQueryManager.h
@@ -160,8 +160,11 @@ class AIMODULE_API UEnvQueryManager : public UAISubsystem, public FSelfRegisteri
 	// makes sure we don't have any UEnvQueryManager instances serialized in. 
 	// Any loaded instance will get marked as PendingKill
 	virtual void PostLoad() override;
-
 	virtual void PostInitProperties() override;
+	virtual bool IsDestructionThreadSafe() const
+	{
+		return false;
+	}
 
 	// FTickableGameObject begin
 	virtual void Tick(float DeltaTime) override;
diff --git a/Engine/Source/Runtime/AIModule/Classes/Perception/AIPerceptionSystem.h b/Engine/Source/Runtime/AIModule/Classes/Perception/AIPerceptionSystem.h
--- a/Engine/Source/Runtime/AIModule/Classes/Perception/AIPerceptionSystem.h
+++ b/Engine/Source/Runtime/AIModule/Classes/Perception/AIPerceptionSystem.h
@@ -39,6 +39,12 @@ public:
 	virtual TStatId GetStatId() const override;
 	// FTickableGameObject end
 
+	// UObject begin
+	virtual bool IsDestructionThreadSafe() const
+	{
+		return false;
+	}
+	// UObject end
 protected:	
 	AIPerception::FListenerMap ListenerContainer;
 
diff --git a/Engine/Source/Runtime/AugmentedReality/Public/ARBlueprintProxy.h b/Engine/Source/Runtime/AugmentedReality/Public/ARBlueprintProxy.h
--- a/Engine/Source/Runtime/AugmentedReality/Public/ARBlueprintProxy.h
+++ b/Engine/Source/Runtime/AugmentedReality/Public/ARBlueprintProxy.h
@@ -28,6 +28,10 @@ public:
 	virtual TStatId GetStatId() const override { RETURN_QUICK_DECLARE_CYCLE_STAT(UARBaseAsyncTaskBlueprintProxy, STATGROUP_Tickables); }
 	//~ End FTickableObject Interface
 	
+	//~ Begin UObject Interface
+	virtual bool IsDestructionThreadSafe() const override { return false; }
+	//~ End UObject Interface
+
 	virtual void ReportSuccess() { check(0); }
 	virtual void ReportFailure() { check(0); }
 
diff --git a/Engine/Source/Runtime/CoreUObject/Private/UObject/GarbageCollection.cpp b/Engine/Source/Runtime/CoreUObject/Private/UObject/GarbageCollection.cpp
--- a/Engine/Source/Runtime/CoreUObject/Private/UObject/GarbageCollection.cpp
+++ b/Engine/Source/Runtime/CoreUObject/Private/UObject/GarbageCollection.cpp
@@ -25,6 +25,8 @@
 #include "UObject/GarbageCollectionVerification.h"
 #include "Async/ParallelFor.h"
 #include "ProfilingDebugging/CsvProfiler.h"
+#include "HAL/Runnable.h"
+#include "HAL/RunnableThread.h"
 
 /*-----------------------------------------------------------------------------
    Garbage collection.
@@ -37,14 +39,10 @@ DEFINE_LOG_CATEGORY(LogGarbage);
 
 /** Object count during last mark phase																				*/
 FThreadSafeCounter		GObjectCountDuringLastMarkPhase;
-/** Count of objects purged since last mark phase																	*/
-int32		GPurgedObjectCountSinceLastMarkPhase	= 0;
 /** Whether incremental object purge is in progress										*/
 bool GObjIncrementalPurgeIsInProgress = false;
 /** Whether GC is currently routing BeginDestroy to objects										*/
 bool GObjUnhashUnreachableIsInProgress = false;
-/** Time the GC started, needs to be reset on return from being in the background on some OSs */
-double GCStartTime = 0.;
 /** Whether FinishDestroy has already been routed to all unreachable objects. */
 static bool GObjFinishDestroyHasBeenRoutedToAllObjects	= false;
 /** 
@@ -57,13 +55,9 @@ static int32 GGCObjectsPendingDestructionCount = 0;
 /** Whether we need to purge objects or not.											*/
 static bool GObjPurgeIsRequired = false;
 /** Current object index for incremental purge.											*/
-static FRawObjectIterator GObjCurrentPurgeObjectIndex;
+static int32 GObjCurrentPurgeObjectIndex = 0;
 /** Current object index for incremental purge.											*/
 static bool GObjCurrentPurgeObjectIndexNeedsReset = true;
-static bool GObjCurrentPurgeObjectIndexResetPastPermanent = false;
-
-/** Whether we are currently purging an object in the GC purge pass. */
-static bool GIsPurgingObject = false;
 
 /** Contains a list of objects that stayed marked as unreachable after the last reachability analysis */
 static TArray<FUObjectItem*> GUnreachableObjects;
@@ -137,12 +131,6 @@ bool IsGarbageCollectionWaiting()
 	return FGCCSyncObject::Get().IsGCWaiting();
 }
 
-/** Called on shutdown to free GC memory */
-void CleanupGCArrayPools()
-{
-	FGCArrayPool::Get().Cleanup();
-}
-
 // Minimum number of objects to spawn a GC sub-task for
 static int32 GMinDesiredObjectsPerSubTask = 128;
 static FAutoConsoleVariableRef CVarMinDesiredObjectsPerSubTask(
@@ -160,6 +148,14 @@ static FAutoConsoleVariableRef CIncrementalBeginDestroyEnabled(
 	ECVF_Default
 );
 
+static int32 GMultithreadedDestructionEnabled = 1;
+static FAutoConsoleVariableRef CMultithreadedDestructionEnabled(
+	TEXT("gc.MultithreadedDestructionEnabled"),
+	GMultithreadedDestructionEnabled,
+	TEXT("If true, the engine will free objects' memory from a worker thread"),
+	ECVF_Default
+);
+
 #if PERF_DETAILED_PER_CLASS_GC_STATS
 /** Map from a UClass' FName to the number of objects that were purged during the last purge phase of this class.	*/
 static TMap<const FName,uint32> GClassToPurgeCountMap;
@@ -240,6 +236,261 @@ static void LogClassCountInfo( const TCHAR* LogText, TMap<const FName,uint32>& C
 };
 #endif
 
+/**
+ * Helper class for destroying UObjects on a worker thread
+ */
+class FAsyncPurge : public FRunnable
+{
+	/** Thread to run the worker FRunnable on. Destroys objects. */
+	volatile FRunnableThread* Thread;
+	/** Id of the worker thread */
+	uint32 AsyncPurgeThreadId;
+	/** Stops this thread */
+	FThreadSafeCounter StopTaskCounter;
+	/** Event that triggers the UObject destruction */
+	FEvent* BeginPurgeEvent;
+	/** Event that signales the UObject destruction is finished */
+	FEvent* FinishedPurgeEvent;
+	/** Current index into the global unreachable objects array (GUnreachableObjects) of the object being destroyed */
+	int32 ObjCurrentPurgeObjectIndex;
+	/** A list of objects that were not thread safe to destroy on the worker thread */
+	TLockFreePointerListUnordered<UObject, PLATFORM_CACHE_LINE_SIZE> GameThreadObjects;
+	/** Stats for the number of objects destroyed */
+	int32 ObjectsDestroyedSinceLastMarkPhase;
+
+	/** [PURGE/GAME THREAD] Destroys objects that are unreachable */
+	bool TickDestroyObjects(bool bUseTimeLimit, float TimeLimit, double StartTime)
+	{
+		const int32 TimeLimitEnforcementGranularityForDeletion = 100;
+		int32 ProcessedObjectsCount = 0;
+		bool bFinishedDestroyingObjects = true;
+
+		while (ObjCurrentPurgeObjectIndex < GUnreachableObjects.Num())
+		{
+			FUObjectItem* ObjectItem = GUnreachableObjects[ObjCurrentPurgeObjectIndex];
+			check(ObjectItem->IsUnreachable());
+
+			UObject* Object = (UObject*)ObjectItem->Object;
+			check(Object->HasAllFlags(RF_FinishDestroyed | RF_BeginDestroyed));
+			if (!Thread || Object->IsDestructionThreadSafe())
+			{
+				Object->~UObject();
+				GUObjectAllocator.FreeUObject(Object);
+			}
+			else
+			{
+				// This object will be destroyed incrementally on the game thread
+				GameThreadObjects.Push(Object);
+			}
+			++ProcessedObjectsCount;
+			++ObjectsDestroyedSinceLastMarkPhase;
+			++ObjCurrentPurgeObjectIndex;
+
+			// Time slicing when running on the game thread
+			if (bUseTimeLimit && (ProcessedObjectsCount == TimeLimitEnforcementGranularityForDeletion))
+			{
+				ProcessedObjectsCount = 0;
+				if ((FPlatformTime::Seconds() - StartTime) > TimeLimit)
+				{
+					bFinishedDestroyingObjects = false;
+					break;
+				}				
+			}
+		}
+		return bFinishedDestroyingObjects;
+	}
+
+	/** [GAME THREAD] Destroys objects that are unreachable and couldn't be destroyed on the worker thread */
+	bool TickDestroyGameThreadObjects(bool bUseTimeLimit, float TimeLimit, double StartTime)
+	{
+		const int32 TimeLimitEnforcementGranularityForDeletion = 100;
+		int32 ProcessedObjectsCount = 0;
+		bool bFinishedDestroyingObjects = true;
+
+		while (UObject* Object = GameThreadObjects.Pop())
+		{
+			Object->~UObject();
+			GUObjectAllocator.FreeUObject(Object);
+
+			if (bUseTimeLimit && (ProcessedObjectsCount == TimeLimitEnforcementGranularityForDeletion))
+			{
+				ProcessedObjectsCount = 0;
+				if ((FPlatformTime::Seconds() - StartTime) > TimeLimit)
+				{
+					bFinishedDestroyingObjects = false;
+					break;
+				}				
+			}
+		}
+
+		return bFinishedDestroyingObjects;
+	}
+
+	/** Waits for the worker thread to finish destroying objects */
+	void WaitForAsyncDestructionToFinish()
+	{
+		FinishedPurgeEvent->Wait();
+	}
+
+public:
+
+	/** 
+	 * Constructor
+	 * @param bMultithreaded if true, the destruction of objects will happen on a worker thread
+	 */
+	FAsyncPurge(bool bMultithreaded)
+		: Thread(nullptr)
+		, AsyncPurgeThreadId(0)
+		, BeginPurgeEvent(nullptr)
+		, FinishedPurgeEvent(nullptr)
+		, ObjectsDestroyedSinceLastMarkPhase(0)
+	{
+		BeginPurgeEvent = FPlatformProcess::GetSynchEventFromPool(true);
+		FinishedPurgeEvent = FPlatformProcess::GetSynchEventFromPool(true);
+		FinishedPurgeEvent->Trigger();
+		if (bMultithreaded)
+		{
+			check(FPlatformProcess::SupportsMultithreading());
+			FPlatformAtomics::InterlockedExchangePtr((void**)&Thread, FRunnableThread::Create(this, TEXT("FAsyncPurge"), 0, TPri_BelowNormal));			
+		}
+		else
+		{
+			AsyncPurgeThreadId = GGameThreadId;
+		}
+	}
+
+	virtual ~FAsyncPurge()
+	{
+		check(IsFinished());
+		delete Thread;
+		Thread = nullptr;
+		FPlatformProcess::ReturnSynchEventToPool(BeginPurgeEvent);
+		FPlatformProcess::ReturnSynchEventToPool(FinishedPurgeEvent);
+		BeginPurgeEvent = nullptr;
+		FinishedPurgeEvent = nullptr;
+	}
+
+	/** [MAIN THREAD] Adds objects to the purge queue */
+	void BeginPurge()
+	{
+		check(IsFinished()); // In single-threaded mode we need to be finished or the condition below will hang
+		if (FinishedPurgeEvent->Wait())
+		{
+			FinishedPurgeEvent->Reset();
+
+			ObjCurrentPurgeObjectIndex = 0;
+			ObjectsDestroyedSinceLastMarkPhase = 0;
+
+			BeginPurgeEvent->Trigger();
+		}
+	}
+
+	/** [GAME THREAD] Ticks the purge process on the game thread */
+	void TickPurge(bool bUseTimeLimit, float TimeLimit, double StartTime)
+	{
+		bool bCanStartDestroyingGameThreadObjects = true;
+		if (!Thread)
+		{
+			// If we're running single-threaded we need to tick the main loop here too
+			bCanStartDestroyingGameThreadObjects = TickDestroyObjects(bUseTimeLimit, TimeLimit, StartTime);
+		}
+		else if (!bUseTimeLimit)
+		{
+			// Wait for the worker thread to finish processing all objects
+			WaitForAsyncDestructionToFinish();
+		}
+		if (bCanStartDestroyingGameThreadObjects)
+		{
+			// Deal with objects that couldn't be destroyed on the worker thread. This will do nothing when running single-threaded
+			bool bFinishedDestroyingObjectsOnGameThread = TickDestroyGameThreadObjects(bUseTimeLimit, TimeLimit, StartTime);
+			if (!Thread && bFinishedDestroyingObjectsOnGameThread)
+			{
+				FinishedPurgeEvent->Trigger();
+			}
+		}
+	}
+
+	/** Returns true if the destruction process is finished */
+	bool IsFinished() const
+	{
+		return FinishedPurgeEvent->Wait(0, true) && GameThreadObjects.IsEmpty();
+	}
+
+	/** Returns the number of objects already destroyed */
+	int32 GetObjectsDestroyedSinceLastMarkPhase() const
+	{
+		return ObjectsDestroyedSinceLastMarkPhase;
+	}
+
+	/** Resets the number of objects already destroyed */
+	void ResetObjectsDestroyedSinceLastMarkPhase()
+	{
+		ObjectsDestroyedSinceLastMarkPhase = 0;
+	}
+
+	/** 
+	  * Returns true if this function is called from the async destruction thread. 
+	  * It will also return true if we're running single-threaded and this function is called on the game thread
+	  */
+	bool IsInAsyncPurgeThread() const
+	{
+		return AsyncPurgeThreadId == FPlatformTLS::GetCurrentThreadId();
+	}
+
+	/* Returns true if it can run multi-threaded destruction */
+	bool IsMultithreaded() const
+	{
+		return !!Thread;
+	}
+
+	//~ Begin FRunnable Interface.
+	virtual bool Init()
+	{
+		return true;
+	}
+
+	virtual uint32 Run()
+	{
+		AsyncPurgeThreadId = FPlatformTLS::GetCurrentThreadId();
+		
+		while (StopTaskCounter.GetValue() == 0)
+		{
+			if (BeginPurgeEvent->Wait(15, true))
+			{
+				BeginPurgeEvent->Reset();
+				TickDestroyObjects(/* bUseTimeLimit = */ false, /* TimeLimit = */ 0.0f, /* StartTime = */ 0.0);
+				FinishedPurgeEvent->Trigger();
+			}
+		}
+		FinishedPurgeEvent->Trigger();
+		return 0;
+	}
+
+	virtual void Stop()
+	{
+		StopTaskCounter.Increment();
+	}
+	//~ End FRunnable Interface
+};
+static FAsyncPurge* GAsyncPurge = nullptr;
+
+/**
+  * Returns true if this function is called from the async destruction thread.
+  * It will also return true if we're running single-threaded and this function is called on the game thread
+  */
+bool IsInGarbageCollectorThread()
+{
+	return GAsyncPurge ? GAsyncPurge->IsInAsyncPurgeThread() : IsInGameThread();
+}
+
+/** Called on shutdown to free GC memory */
+void ShutdownGarbageCollection()
+{
+	FGCArrayPool::Get().Cleanup();
+	delete GAsyncPurge;
+	GAsyncPurge = nullptr;
+}
+
 /**
 * Handles UObject references found by TFastReferenceCollector
 */
@@ -946,6 +1197,27 @@ public:
 	}
 };
 
+// Allow parallel GC to be overridden to single threaded via console command.
+static int32 GAllowParallelGC = 1;
+
+static FAutoConsoleVariableRef CVarAllowParallelGC(
+	TEXT("gc.AllowParallelGC"),
+	GAllowParallelGC,
+	TEXT("sed to control parallel GC."),
+	ECVF_Default
+);
+
+/** Returns true if Garbage Collection should be forced to run on a single thread */
+static bool ShouldForceSingleThreadedGC()
+{
+	const bool bForceSingleThreadedGC = !FApp::ShouldUseThreadingForPerformance() || !FPlatformProcess::SupportsMultithreading() ||
+#if PLATFORM_SUPPORTS_MULTITHREADED_GC
+	(FPlatformMisc::NumberOfCores() < 2 || GAllowParallelGC == 0 || PERF_DETAILED_PER_CLASS_GC_STATS);
+#else	//PLATFORM_SUPPORTS_MULTITHREADED_GC
+		true;
+#endif	//PLATFORM_SUPPORTS_MULTITHREADED_GC
+	return bForceSingleThreadedGC;
+}
 
 static void AcquireGCLock()
 {
@@ -985,6 +1257,8 @@ struct FConditionalGCLock
 	}
 };
 
+static bool IncrementalDestroyGarbage(bool bUseTimeLimit, float TimeLimit, double StartTime);
+
 /**
  * Incrementally purge garbage by deleting all unreferenced objects after routing Destroy.
  *
@@ -994,7 +1268,7 @@ struct FConditionalGCLock
  * @param	bUseTimeLimit	whether the time limit parameter should be used
  * @param	TimeLimit		soft time limit for this function call
  */
-void IncrementalPurgeGarbage( bool bUseTimeLimit, float TimeLimit )
+void IncrementalPurgeGarbage(bool bUseTimeLimit, float TimeLimit)
 {
 	SCOPED_NAMED_EVENT(IncrementalPurgeGarbage, FColor::Red);
 	DECLARE_SCOPE_CYCLE_COUNTER(TEXT("IncrementalPurgeGarbage"), STAT_IncrementalPurgeGarbage, STATGROUP_GC);
@@ -1005,10 +1279,9 @@ void IncrementalPurgeGarbage( bool bUseTimeLimit, float TimeLimit )
 		GObjPurgeIsRequired = true;
 		GUObjectArray.DisableDisregardForGC();
 		GObjCurrentPurgeObjectIndexNeedsReset = true;
-		GObjCurrentPurgeObjectIndexResetPastPermanent = false;
 	}
 	// Early out if there is nothing to do.
-	if( !GObjPurgeIsRequired )
+	if (!GObjPurgeIsRequired)
 	{
 		return;
 	}
@@ -1037,14 +1310,9 @@ void IncrementalPurgeGarbage( bool bUseTimeLimit, float TimeLimit )
 	} ResetPurgeProgress(bCompleted);
 
 	// Keep track of start time to enforce time limit unless bForceFullPurge is true;
-	GCStartTime							= FPlatformTime::Seconds();
-	bool		bTimeLimitReached							= false;
-	// Depending on platform FPlatformTime::Seconds might take a noticeable amount of time if called thousands of times so we avoid
-	// enforcing the time limit too often, especially as neither Destroy nor actual deletion should take significant
-	// amounts of time.
-	const int32	TimeLimitEnforcementGranularityForDestroy	= 10;
-	const int32	TimeLimitEnforcementGranularityForDeletion	= 100;
+	double StartTime = FPlatformTime::Seconds();
 
+	bool bTimeLimitReached = false;
 	if (GUnrechableObjectIndex < GUnreachableObjects.Num())
 	{
 		{
@@ -1057,6 +1325,38 @@ void IncrementalPurgeGarbage( bool bUseTimeLimit, float TimeLimit )
 		}
 	}
 
+	if (!bTimeLimitReached)
+	{
+		bCompleted = IncrementalDestroyGarbage(bUseTimeLimit, TimeLimit, StartTime);
+	}
+}
+
+bool IncrementalDestroyGarbage(bool bUseTimeLimit, float TimeLimit, double StartTime)
+{
+	const bool bMultithreadedPurge = !ShouldForceSingleThreadedGC() && GMultithreadedDestructionEnabled;
+	if (!GAsyncPurge)
+	{
+		GAsyncPurge = new FAsyncPurge(bMultithreadedPurge);
+	}
+	else if (GAsyncPurge->IsMultithreaded() != bMultithreadedPurge)
+	{
+		check(GAsyncPurge->IsFinished());
+		delete GAsyncPurge;
+		GAsyncPurge = new FAsyncPurge(bMultithreadedPurge);
+	}
+
+	bool bCompleted = false;
+	bool bTimeLimitReached = false;
+
+	// Keep track of time it took to destroy objects for stats
+	double IncrementalDestroyGarbageStartTime = FPlatformTime::Seconds();
+
+	// Depending on platform FPlatformTime::Seconds might take a noticeable amount of time if called thousands of times so we avoid
+	// enforcing the time limit too often, especially as neither Destroy nor actual deletion should take significant
+	// amounts of time.
+	const int32	TimeLimitEnforcementGranularityForDestroy = 10;
+	const int32	TimeLimitEnforcementGranularityForDeletion = 100;
+
 	// Set 'I'm garbage collecting' flag - might be checked inside UObject::Destroy etc.
 	FGCScopeLock GCLock;
 
@@ -1072,21 +1372,18 @@ void IncrementalPurgeGarbage( bool bUseTimeLimit, float TimeLimit )
 
 		if (GObjCurrentPurgeObjectIndexNeedsReset)
 		{
-			// iterators don't have an op=, so we destroy it and reconstruct it with a placement new
-			// GObjCurrentPurgeObjectIndex	= FRawObjectIterator(GObjCurrentPurgeObjectIndexResetPastPermanent);
-			GObjCurrentPurgeObjectIndex.~FRawObjectIterator();
-			new (&GObjCurrentPurgeObjectIndex) FRawObjectIterator(GObjCurrentPurgeObjectIndexResetPastPermanent);
+			GObjCurrentPurgeObjectIndex = 0;
 			GObjCurrentPurgeObjectIndexNeedsReset = false;
 		}
 
-		while( GObjCurrentPurgeObjectIndex )
+		while (GObjCurrentPurgeObjectIndex < GUnreachableObjects.Num())
 		{
-			FUObjectItem* ObjectItem = *GObjCurrentPurgeObjectIndex;
+			FUObjectItem* ObjectItem = GUnreachableObjects[GObjCurrentPurgeObjectIndex];
 			checkSlow(ObjectItem);
 
 			//@todo UE4 - A prefetch was removed here. Re-add it. It wasn't right anyway, since it was ten items ahead and the consoles on have 8 prefetch slots
 
-			if (ObjectItem->IsUnreachable())
+			check(ObjectItem->IsUnreachable());
 			{
 				UObject* Object = static_cast<UObject*>(ObjectItem->Object);
 				// Object should always have had BeginDestroy called on it and never already be destroyed
@@ -1123,7 +1420,7 @@ void IncrementalPurgeGarbage( bool bUseTimeLimit, float TimeLimit )
 
 			// Only check time limit every so often to avoid calling FPlatformTime::Seconds too often.
 			const bool bPollTimeLimit = ((TimeLimitTimePollCounter++) % TimeLimitEnforcementGranularityForDestroy == 0);
-			if( bUseTimeLimit && bPollTimeLimit && ((FPlatformTime::Seconds() - GCStartTime) > TimeLimit) )
+			if( bUseTimeLimit && bPollTimeLimit && ((FPlatformTime::Seconds() - StartTime) > TimeLimit) )
 			{
 				bTimeLimitReached = true;
 				break;
@@ -1131,7 +1428,7 @@ void IncrementalPurgeGarbage( bool bUseTimeLimit, float TimeLimit )
 		}
 
 		// Have we finished the first round of attempting to call FinishDestroy on unreachable objects?
-		if( !GObjCurrentPurgeObjectIndex )
+		if (GObjCurrentPurgeObjectIndex >= GUnreachableObjects.Num())
 		{
 			// We've finished iterating over all unreachable objects, but we need still need to handle
 			// objects that were deferred.
@@ -1180,7 +1477,7 @@ void IncrementalPurgeGarbage( bool bUseTimeLimit, float TimeLimit )
 
 					// Only check time limit every so often to avoid calling FPlatformTime::Seconds too often.
 					const bool bPollTimeLimit = ((TimeLimitTimePollCounter++) % TimeLimitEnforcementGranularityForDestroy == 0);
-					if( bUseTimeLimit && bPollTimeLimit && ((FPlatformTime::Seconds() - GCStartTime) > TimeLimit) )
+					if( bUseTimeLimit && bPollTimeLimit && ((FPlatformTime::Seconds() - StartTime) > TimeLimit) )
 					{
 						bTimeLimitReached = true;
 						break;
@@ -1206,7 +1503,7 @@ void IncrementalPurgeGarbage( bool bUseTimeLimit, float TimeLimit )
 #endif
 						// Check if we spent too much time on waiting for FinishDestroy without making any progress
 						if (LastLoopObjectsPendingDestructionCount == GGCObjectsPendingDestructionCount && bPollTimeLimit &&
-							((FPlatformTime::Seconds() - GCStartTime) > MaxTimeForFinishDestroy))
+							((FPlatformTime::Seconds() - StartTime) > MaxTimeForFinishDestroy))
 						{
 							UE_LOG(LogGarbage, Warning, TEXT("Spent more than %.2fs on routing FinishDestroy to objects (objects in queue: %d)"), MaxTimeForFinishDestroy, GGCObjectsPendingDestructionCount);
 							UObject* LastObjectNotReadyForFinishDestroy = nullptr;
@@ -1263,75 +1560,54 @@ void IncrementalPurgeGarbage( bool bUseTimeLimit, float TimeLimit )
 				// Destroy has been routed to all objects so it's safe to delete objects now.
 				GObjFinishDestroyHasBeenRoutedToAllObjects = true;
 				GObjCurrentPurgeObjectIndexNeedsReset = true;
-				GObjCurrentPurgeObjectIndexResetPastPermanent = !GExitPurge;
 			}
 		}
 	}		
-
-	if( GObjFinishDestroyHasBeenRoutedToAllObjects && !bTimeLimitReached )
+	
+	if (GObjFinishDestroyHasBeenRoutedToAllObjects && !bTimeLimitReached)
 	{
 		// Perform actual object deletion.
-		// @warning: Can't use FObjectIterator here because classes may be destroyed before objects.
 		int32 ProcessCount = 0;
 		if (GObjCurrentPurgeObjectIndexNeedsReset)
 		{
-			// iterators don't have an op=, so we destroy it and reconstruct it with a placement new
-			// GObjCurrentPurgeObjectIndex	= FRawObjectIterator(GObjCurrentPurgeObjectIndexResetPastPermanent);
-			GObjCurrentPurgeObjectIndex.~FRawObjectIterator();
-			new (&GObjCurrentPurgeObjectIndex) FRawObjectIterator(GObjCurrentPurgeObjectIndexResetPastPermanent);
+			GAsyncPurge->BeginPurge();
+			// Reset the reset flag but don't reset the actual index yet for stat purposes
 			GObjCurrentPurgeObjectIndexNeedsReset = false;
 		}
-		while( GObjCurrentPurgeObjectIndex )
-		{
-			//@todo UE4 - A prefetch was removed here. Re-add it. It wasn't right anyway, since it was ten items ahead and the consoles on have 8 prefetch slots
-
-			FUObjectItem* ObjectItem = *GObjCurrentPurgeObjectIndex;
-			checkSlow(ObjectItem);
-			if (ObjectItem->IsUnreachable())
-			{
-				UObject* Object = (UObject*)ObjectItem->Object;
-				check(Object->HasAllFlags(RF_FinishDestroyed|RF_BeginDestroyed));
-				GIsPurgingObject				= true; 
-				Object->~UObject();
-				GUObjectAllocator.FreeUObject(Object);
-				GIsPurgingObject				= false;
-				// Keep track of purged stats.
-				GPurgedObjectCountSinceLastMarkPhase++;
-			}
-
-			// Advance to the next object.
-			++GObjCurrentPurgeObjectIndex;
 
-			ProcessCount++;
+		GAsyncPurge->TickPurge(bUseTimeLimit, TimeLimit, StartTime);
 
-			// Only check time limit every so often to avoid calling FPlatformTime::Seconds too often.
-			if( bUseTimeLimit && (ProcessCount == TimeLimitEnforcementGranularityForDeletion))
-			{
-				if ((FPlatformTime::Seconds() - GCStartTime) > TimeLimit)
-				{
-					bTimeLimitReached = true;
-					break;
-				}
-				ProcessCount = 0;
-			}
-		}
-
-		if( !GObjCurrentPurgeObjectIndex )
+		if (GAsyncPurge->IsFinished())
 		{
 			bCompleted = true;
 			// Incremental purge is finished, time to reset variables.
 			GObjFinishDestroyHasBeenRoutedToAllObjects		= false;
 			GObjPurgeIsRequired								= false;
 			GObjCurrentPurgeObjectIndexNeedsReset			= true;
-			GObjCurrentPurgeObjectIndexResetPastPermanent	= true;
 
 			// Log status information.
-			UE_LOG(LogGarbage, Log, TEXT("GC purged %i objects (%i -> %i)"), GPurgedObjectCountSinceLastMarkPhase, GObjectCountDuringLastMarkPhase.GetValue(), GObjectCountDuringLastMarkPhase.GetValue() - GPurgedObjectCountSinceLastMarkPhase );
+			const int32 PurgedObjectCountSinceLastMarkPhase = GAsyncPurge->GetObjectsDestroyedSinceLastMarkPhase();
+			UE_LOG(LogGarbage, Log, TEXT("GC purged %i objects (%i -> %i) in %.3fms"), PurgedObjectCountSinceLastMarkPhase, 
+				GObjectCountDuringLastMarkPhase.GetValue(), 
+				GObjectCountDuringLastMarkPhase.GetValue() - PurgedObjectCountSinceLastMarkPhase,
+				(FPlatformTime::Seconds() - IncrementalDestroyGarbageStartTime) * 1000);
 #if PERF_DETAILED_PER_CLASS_GC_STATS
-			LogClassCountInfo( TEXT("objects of"), GClassToPurgeCountMap, 10, GPurgedObjectCountSinceLastMarkPhase );
+			LogClassCountInfo( TEXT("objects of"), GClassToPurgeCountMap, 10, PurgedObjectCountSinceLastMarkPhase);
 #endif
+			GAsyncPurge->ResetObjectsDestroyedSinceLastMarkPhase();
 		}
 	}
+
+	if (bUseTimeLimit && !bCompleted)
+	{
+		UE_LOG(LogGarbage, Log, TEXT("%.3f ms for incrementally purging unreachable objects (FinishDestroyed: %d, Destroyed: %d / %d)"),
+			(FPlatformTime::Seconds() - IncrementalDestroyGarbageStartTime) * 1000,
+			GObjCurrentPurgeObjectIndex,
+			GAsyncPurge->GetObjectsDestroyedSinceLastMarkPhase(),
+			GUnreachableObjects.Num());
+	}
+
+	return bCompleted;
 }
 
 /**
@@ -1344,15 +1620,6 @@ bool IsIncrementalPurgePending()
 	return GObjIncrementalPurgeIsInProgress || GObjPurgeIsRequired;
 }
 
-// Allow parallel GC to be overridden to single threaded via console command.
-static int32 GAllowParallelGC = 1;
-
-static FAutoConsoleVariableRef CVarAllowParallelGC(
-	TEXT("gc.AllowParallelGC"),
-	GAllowParallelGC,
-	TEXT("sed to control parallel GC."),
-	ECVF_Default
-	);
 
 // This counts how many times GC was skipped
 static int32 GNumAttemptsSinceLastGC = 0;
@@ -1384,7 +1651,7 @@ void GatherUnreachableObjects(bool bForceSingleThreaded)
 	GUnreachableObjects.Reset();
 	GUnrechableObjectIndex = 0;
 
-	int32 MaxNumberOfObjects = GUObjectArray.GetObjectArrayNum() - GUObjectArray.GetFirstGCIndex();
+	int32 MaxNumberOfObjects = GUObjectArray.GetObjectArrayNum() - (GExitPurge ? 0 : GUObjectArray.GetFirstGCIndex());
 	int32 NumThreads = FMath::Max(1, FTaskGraphInterface::Get().GetNumWorkerThreads());
 	int32 NumberOfObjectsPerThread = (MaxNumberOfObjects / NumThreads) + 1;
 
@@ -1395,7 +1662,7 @@ void GatherUnreachableObjects(bool bForceSingleThreaded)
 	// are part of the array so we don't suffer from cache misses as much as we would if we were to check ObjectFlags.
 	ParallelFor(NumThreads, [&ClusterItemsToDestroy, NumberOfObjectsPerThread, NumThreads, MaxNumberOfObjects](int32 ThreadIndex)
 	{
-		int32 FirstObjectIndex = ThreadIndex * NumberOfObjectsPerThread + GUObjectArray.GetFirstGCIndex();
+		int32 FirstObjectIndex = ThreadIndex * NumberOfObjectsPerThread + (GExitPurge ? 0 : GUObjectArray.GetFirstGCIndex());
 		int32 NumObjects = (ThreadIndex < (NumThreads - 1)) ? NumberOfObjectsPerThread : (MaxNumberOfObjects - (NumThreads - 1) * NumberOfObjectsPerThread);
 		int32 LastObjectIndex = FMath::Min(GUObjectArray.GetObjectArrayNum() - 1, FirstObjectIndex + NumObjects - 1);
 		TArray<FUObjectItem*> ThisThreadUnreachableObjects;
@@ -1535,12 +1802,7 @@ void CollectGarbageInternal(EObjectFlags KeepFlags, bool bPerformFullPurge)
 		// Fall back to single threaded GC if processor count is 1 or parallel GC is disabled
 		// or detailed per class gc stats are enabled (not thread safe)
 		// Temporarily forcing single-threaded GC in the editor until Modify() can be safely removed from HandleObjectReference.
-		const bool bForceSingleThreadedGC = !FApp::ShouldUseThreadingForPerformance() || !FPlatformProcess::SupportsMultithreading() ||
-#if PLATFORM_SUPPORTS_MULTITHREADED_GC
-		(FPlatformMisc::NumberOfCores() < 2 || GAllowParallelGC == 0 || PERF_DETAILED_PER_CLASS_GC_STATS);
-#else	//PLATFORM_SUPPORTS_MULTITHREADED_GC
-			true;
-#endif	//PLATFORM_SUPPORTS_MULTITHREADED_GC
+		const bool bForceSingleThreadedGC = ShouldForceSingleThreadedGC();
 
 		// Perform reachability analysis.
 		{
@@ -1575,9 +1837,6 @@ void CollectGarbageInternal(EObjectFlags KeepFlags, bool bPerformFullPurge)
 
 		// Set flag to indicate that we are relying on a purge to be performed.
 		GObjPurgeIsRequired = true;
-		// Reset purged count.
-		GPurgedObjectCountSinceLastMarkPhase = 0;
-		GObjCurrentPurgeObjectIndexResetPastPermanent = true;
 
 		// Perform a full purge by not using a time limit for the incremental purge. The Editor always does a full purge.
 		if (bPerformFullPurge || GIsEditor)
diff --git a/Engine/Source/Runtime/CoreUObject/Private/UObject/Obj.cpp b/Engine/Source/Runtime/CoreUObject/Private/UObject/Obj.cpp
--- a/Engine/Source/Runtime/CoreUObject/Private/UObject/Obj.cpp
+++ b/Engine/Source/Runtime/CoreUObject/Private/UObject/Obj.cpp
@@ -4341,7 +4341,7 @@ void StaticUObjectInit()
 }
 
 // Internal cleanup functions
-void CleanupGCArrayPools();
+void ShutdownGarbageCollection();
 void CleanupLinkerAnnotations();
 void CleanupCachedArchetypes();
 
@@ -4365,7 +4365,9 @@ void StaticExit()
 		GObjTransientPkg = NULL;
 	}
 
-	IncrementalPurgeGarbage( false );
+	GatherUnreachableObjects(false);
+	IncrementalPurgeGarbage(false);
+	GUObjectClusters.DissolveClusters(true);
 
 	// Keep track of how many objects there are for GC stats as we simulate a mark pass.
 	extern FThreadSafeCounter GObjectCountDuringLastMarkPhase;
@@ -4400,19 +4402,8 @@ void StaticExit()
 	// set on all objects that are about to be deleted. One example is FLinkerLoad detaching textures - the SetLinker call needs to 
 	// not kick off texture streaming.
 	//
-	for ( FRawObjectIterator It; It; ++It )
-	{
-		FUObjectItem* ObjItem = *It;
-		checkSlow(ObjItem);
-		if (ObjItem->IsUnreachable())
-		{
-			// Begin the object's asynchronous destruction.
-			UObject* Obj = static_cast<UObject*>(ObjItem->Object);
-			Obj->ConditionalBeginDestroy();
-		}
-	}
-
-	IncrementalPurgeGarbage( false );
+	GatherUnreachableObjects(false);
+	IncrementalPurgeGarbage(false);
 
 	{
 		//Repeat GC for every object, including structures and properties.
@@ -4422,25 +4413,15 @@ void StaticExit()
 			It->SetUnreachable();
 		}
 
-		for (FRawObjectIterator It; It; ++It)
-		{
-			FUObjectItem* ObjItem = *It;
-			checkSlow(ObjItem);
-			if (ObjItem->IsUnreachable())
-			{
-				// Begin the object's asynchronous destruction.
-				UObject* Obj = static_cast<UObject*>(ObjItem->Object);
-				Obj->ConditionalBeginDestroy();
-			}
-		}
-
+		GatherUnreachableObjects(false);
 		IncrementalPurgeGarbage(false);
 	}
 
+	ShutdownGarbageCollection();
 	UObjectBaseShutdown();
+
 	// Empty arrays to prevent falsely-reported memory leaks.
-	FDeferredMessageLog::Cleanup();
-	CleanupGCArrayPools();
+	FDeferredMessageLog::Cleanup();	
 	CleanupLinkerAnnotations();
 	CleanupCachedArchetypes();
 
diff --git a/Engine/Source/Runtime/CoreUObject/Private/UObject/UObjectArray.cpp b/Engine/Source/Runtime/CoreUObject/Private/UObject/UObjectArray.cpp
--- a/Engine/Source/Runtime/CoreUObject/Private/UObject/UObjectArray.cpp
+++ b/Engine/Source/Runtime/CoreUObject/Private/UObject/UObjectArray.cpp
@@ -197,7 +197,7 @@ void FUObjectArray::AllocateUObjectIndex(UObjectBase* Object, bool bMergingThrea
 void FUObjectArray::FreeUObjectIndex(UObjectBase* Object)
 {
 	// This should only be happening on the game thread (GC runs only on game thread when it's freeing objects)
-	check(IsInGameThread());
+	check(IsInGameThread() || IsInGarbageCollectorThread());
 
 	int32 Index = Object->InternalIndex;
 	// At this point no two objects exist with the same index so no need to lock here
@@ -214,9 +214,9 @@ void FUObjectArray::FreeUObjectIndex(UObjectBase* Object)
 	}
 	// You cannot safely recycle indicies in the non-GC range
 	// No point in filling this list when doing exit purge. Nothing should be allocated afterwards anyway.
+	IndexToObject(Index)->ResetSerialNumberAndFlags();
 	if (Index > ObjLastNonGCIndex && !GExitPurge)  
 	{
-		IndexToObject(Index)->ResetSerialNumberAndFlags();
 		ObjAvailableList.Push((int32*)(uintptr_t)Index);
 #if UE_GC_TRACK_OBJ_AVAILABLE
 		ObjAvailableCount.Increment();
diff --git a/Engine/Source/Runtime/CoreUObject/Private/UObject/UObjectBase.cpp b/Engine/Source/Runtime/CoreUObject/Private/UObject/UObjectBase.cpp
--- a/Engine/Source/Runtime/CoreUObject/Private/UObject/UObjectBase.cpp
+++ b/Engine/Source/Runtime/CoreUObject/Private/UObject/UObjectBase.cpp
@@ -125,7 +125,7 @@ UObjectBase::~UObjectBase()
 	{
 		// Validate it.
 		check(IsValidLowLevel());
-		LowLevelRename(NAME_None);
+		check(GetFName() == NAME_None);
 		GUObjectArray.FreeUObjectIndex(this);
 	}
 
diff --git a/Engine/Source/Runtime/CoreUObject/Private/UObject/UObjectClusters.cpp b/Engine/Source/Runtime/CoreUObject/Private/UObject/UObjectClusters.cpp
--- a/Engine/Source/Runtime/CoreUObject/Private/UObject/UObjectClusters.cpp
+++ b/Engine/Source/Runtime/CoreUObject/Private/UObject/UObjectClusters.cpp
@@ -203,11 +203,11 @@ void FUObjectClusterContainer::DissolveClusterAndMarkObjectsAsUnreachable(FUObje
 	}
 }
 
-void FUObjectClusterContainer::DissolveClusters()
+void FUObjectClusterContainer::DissolveClusters(bool bForceDissolveAllClusters /* = false */)
 {
 	for (FUObjectCluster& Cluster : Clusters)
 	{
-		if (Cluster.RootIndex >= 0 && Cluster.bNeedsDissolving)
+		if (Cluster.RootIndex >= 0 && (Cluster.bNeedsDissolving || bForceDissolveAllClusters))
 		{
 			DissolveCluster(Cluster);
 		}
diff --git a/Engine/Source/Runtime/CoreUObject/Public/UObject/Object.h b/Engine/Source/Runtime/CoreUObject/Public/UObject/Object.h
--- a/Engine/Source/Runtime/CoreUObject/Public/UObject/Object.h
+++ b/Engine/Source/Runtime/CoreUObject/Public/UObject/Object.h
@@ -439,6 +439,16 @@ public:
 		return false;
 	}
 
+	/**
+	* Called during garbage collection to determine if an object can have its destructor called on a worker thread.
+	*
+	* @return	true if this object's destructor is thread safe
+	*/
+	virtual bool IsDestructionThreadSafe() const
+	{
+		return true;
+	}
+
 	/**
 	* Called during cooking. Must return all objects that will be Preload()ed when this is serialized at load time. Only used by the EDL.
 	*
diff --git a/Engine/Source/Runtime/CoreUObject/Public/UObject/UObjectArray.h b/Engine/Source/Runtime/CoreUObject/Public/UObject/UObjectArray.h
--- a/Engine/Source/Runtime/CoreUObject/Public/UObject/UObjectArray.h
+++ b/Engine/Source/Runtime/CoreUObject/Public/UObject/UObjectArray.h
@@ -1095,8 +1095,11 @@ public:
 	 */
 	void DissolveCluster(UObjectBaseUtility* ClusterRootOrObjectFromCluster);
 
-	/** Dissolve all clusters marked for dissolving */
-	void DissolveClusters();
+	/** 
+	 * Dissolve all clusters marked for dissolving 
+	 * @param bForceDissolveAllClusters if true, dissolves all clusters even if they're not marked for dissolving
+	 */
+	void DissolveClusters(bool bForceDissolveAllClusters = false);
 
 	/** Dissolve the specified cluster and all clusters that reference it */
 	void DissolveClusterAndMarkObjectsAsUnreachable(FUObjectItem* RootObjectItem);
diff --git a/Engine/Source/Runtime/CoreUObject/Public/UObject/UObjectGlobals.h b/Engine/Source/Runtime/CoreUObject/Public/UObject/UObjectGlobals.h
--- a/Engine/Source/Runtime/CoreUObject/Public/UObject/UObjectGlobals.h
+++ b/Engine/Source/Runtime/CoreUObject/Public/UObject/UObjectGlobals.h
@@ -426,6 +426,11 @@ COREUOBJECT_API float GetAsyncLoadPercentage( const FName& PackageName );
 */
 COREUOBJECT_API bool IsGarbageCollecting();
 
+/**
+* Whether we are running on the Garbage Collector Thread
+*/
+COREUOBJECT_API bool IsInGarbageCollectorThread();
+
 /** 
  * Deletes all unreferenced objects, keeping objects that have any of the passed in KeepFlags set. Will wait for other threads to unlock GC.
  *
@@ -466,6 +471,13 @@ COREUOBJECT_API bool IsIncrementalUnhashPending();
  */
 COREUOBJECT_API bool IsIncrementalPurgePending();
 
+/**
+ * Gathers unreachable objects for IncrementalPurgeGarbage.
+ *
+ * @param bForceSingleThreaded true to force the process to just one thread
+ */
+COREUOBJECT_API void GatherUnreachableObjects(bool bForceSingleThreaded);
+
 /**
  * Incrementally purge garbage by deleting all unreferenced objects after routing Destroy.
  *
diff --git a/Engine/Source/Runtime/Engine/Classes/Components/SkyLightComponent.h b/Engine/Source/Runtime/Engine/Classes/Components/SkyLightComponent.h
--- a/Engine/Source/Runtime/Engine/Classes/Components/SkyLightComponent.h
+++ b/Engine/Source/Runtime/Engine/Classes/Components/SkyLightComponent.h
@@ -181,6 +181,7 @@ class ENGINE_API USkyLightComponent : public ULightComponentBase
 #endif // WITH_EDITOR
 	virtual void BeginDestroy() override;
 	virtual bool IsReadyForFinishDestroy() override;
+	virtual bool IsDestructionThreadSafe() const override { return false; }
 	//~ End UObject Interface
 
 	virtual TStructOnScope<FActorComponentInstanceData> GetComponentInstanceData() const override;
diff --git a/Engine/Source/Runtime/Engine/Classes/Engine/CoreSettings.h b/Engine/Source/Runtime/Engine/Classes/Engine/CoreSettings.h
--- a/Engine/Source/Runtime/Engine/Classes/Engine/CoreSettings.h
+++ b/Engine/Source/Runtime/Engine/Classes/Engine/CoreSettings.h
@@ -175,6 +175,11 @@ protected:
 		ToolTip = "If true, the engine will destroy objects incrementally using time limit each frame."))
 	uint32 IncrementalBeginDestroyEnabled : 1;
 
+	UPROPERTY(EditAnywhere, config, Category = Optimization, meta = (
+		ConsoleVariable = "gc.MultithreadedDestructionEnabled", DisplayName = "Multithreaded Destruction Enabled",
+		ToolTip = "If true, the engine will free objects' memory on a worker thread."))
+	uint32 MultithreadedDestructionEnabled : 1;
+
 	UPROPERTY(EditAnywhere, config, Category = Optimization, meta = (
 		ConsoleVariable = "gc.CreateGCClusters", DisplayName = "Create Garbage Collector UObject Clusters",
 		ToolTip = "If true, the engine will attempt to create clusters of objects for better garbage collection performance."))
diff --git a/Engine/Source/Runtime/Engine/Classes/Engine/Engine.h b/Engine/Source/Runtime/Engine/Classes/Engine/Engine.h
--- a/Engine/Source/Runtime/Engine/Classes/Engine/Engine.h
+++ b/Engine/Source/Runtime/Engine/Classes/Engine/Engine.h
@@ -1933,6 +1933,7 @@ public:
 	virtual void FinishDestroy() override;
 	virtual void Serialize(FArchive& Ar) override;
 	static void AddReferencedObjects(UObject* InThis, FReferenceCollector& Collector);
+	virtual bool IsDestructionThreadSafe() const override { return false; }
 	//~ End UObject Interface.
 
 	/** Initialize the game engine. */
diff --git a/Engine/Source/Runtime/Engine/Classes/Engine/LevelBounds.h b/Engine/Source/Runtime/Engine/Classes/Engine/LevelBounds.h
--- a/Engine/Source/Runtime/Engine/Classes/Engine/LevelBounds.h
+++ b/Engine/Source/Runtime/Engine/Classes/Engine/LevelBounds.h
@@ -37,6 +37,9 @@ class ALevelBounds
 
 	//~ Begin UObject Interface
 	virtual void PostLoad() override;
+#if WITH_EDITOR // because of FEditorTickableLevelBounds
+	virtual bool IsDestructionThreadSafe() const override { return false; }
+#endif
 	//~ End UObject Interface
 	
 	//~ Begin AActor Interface.
diff --git a/Engine/Source/Runtime/Engine/Classes/Engine/MapBuildDataRegistry.h b/Engine/Source/Runtime/Engine/Classes/Engine/MapBuildDataRegistry.h
--- a/Engine/Source/Runtime/Engine/Classes/Engine/MapBuildDataRegistry.h
+++ b/Engine/Source/Runtime/Engine/Classes/Engine/MapBuildDataRegistry.h
@@ -289,6 +289,10 @@ public:
 	ENGINE_API virtual void BeginDestroy() override;
 	ENGINE_API virtual bool IsReadyForFinishDestroy() override;
 	ENGINE_API virtual void FinishDestroy() override;
+	ENGINE_API virtual bool IsDestructionThreadSafe() const
+	{
+		return false;
+	}
 	//~ End UObject Interface
 
 	/** 
diff --git a/Engine/Source/Runtime/Engine/Private/CoreSettings.cpp b/Engine/Source/Runtime/Engine/Private/CoreSettings.cpp
--- a/Engine/Source/Runtime/Engine/Private/CoreSettings.cpp
+++ b/Engine/Source/Runtime/Engine/Private/CoreSettings.cpp
@@ -169,6 +169,7 @@ UGarbageCollectionSettings::UGarbageCollectionSettings()
 	FlushStreamingOnGC = false;
 	AllowParallelGC = true;
 	IncrementalBeginDestroyEnabled = true;
+	MultithreadedDestructionEnabled = true;
 	NumRetriesBeforeForcingGC = 0;
 	MaxObjectsNotConsideredByGC = 0;
 	SizeOfPermanentObjectPool = 0;
diff --git a/Engine/Source/Runtime/Engine/Private/StreamableManager.cpp b/Engine/Source/Runtime/Engine/Private/StreamableManager.cpp
--- a/Engine/Source/Runtime/Engine/Private/StreamableManager.cpp
+++ b/Engine/Source/Runtime/Engine/Private/StreamableManager.cpp
@@ -507,7 +507,7 @@ void FStreamableHandle::StartStalledHandle()
 FStreamableHandle::~FStreamableHandle()
 {
 	TRACE_LOADTIME_DESTROY_STREAMABLE_HANDLE(this);
-	check(IsInGameThread());
+	check(IsInGameThread() || IsInGarbageCollectorThread());
 
 	if (IsActive())
 	{
