TBA: Don't serialize the name table to/from text archives

diff --git a/Engine/Source/Runtime/CoreUObject/Private/UObject/LinkerLoad.cpp b/Engine/Source/Runtime/CoreUObject/Private/UObject/LinkerLoad.cpp
--- a/Engine/Source/Runtime/CoreUObject/Private/UObject/LinkerLoad.cpp
+++ b/Engine/Source/Runtime/CoreUObject/Private/UObject/LinkerLoad.cpp
@@ -1397,6 +1397,12 @@ FLinkerLoad::ELinkerStatus FLinkerLoad::SerializeNameMap()
 {
 	DECLARE_SCOPE_CYCLE_COUNTER( TEXT( "FLinkerLoad::SerializeNameMap" ), STAT_LinkerLoad_SerializeNameMap, STATGROUP_LinkerLoad );
 
+	// Text archives don't have name tables
+	if (IsTextFormat())
+	{
+		return LINKER_Loaded;
+	}
+
 	// The name map is the first item serialized. We wait till all the header information is read
 	// before any serialization. @todo async, @todo seamless: this could be spread out across name,
 	// import and export maps if the package file summary contained more detailed information on
@@ -1405,10 +1411,7 @@ FLinkerLoad::ELinkerStatus FLinkerLoad::SerializeNameMap()
 
 	if( NameMapIndex == 0 && Summary.NameCount > 0 )
 	{
-		if (!IsTextFormat())
-		{
-			Seek(Summary.NameOffset);
-		}
+		Seek(Summary.NameOffset);
 
 		// Make sure there is something to precache first.
 		if( Summary.TotalHeaderSize > 0 )
@@ -1431,9 +1434,6 @@ FLinkerLoad::ELinkerStatus FLinkerLoad::SerializeNameMap()
 		}
 	}
 
-	FStructuredArchive::FStream NameStream = StructuredArchiveRootRecord->EnterStream(FIELD_NAME_TEXT("Names"));
-	const bool bIsTextFormat = StructuredArchiveRootRecord->GetUnderlyingArchive().IsTextFormat();
-
 	NameMap.Reserve(Summary.NameCount);
 
 	while (bFinishedPrecaching && NameMapIndex < Summary.NameCount && !IsTimeLimitExceeded(TEXT("serializing name map"), 100))
@@ -1442,15 +1442,8 @@ FLinkerLoad::ELinkerStatus FLinkerLoad::SerializeNameMap()
 
 		FNameEntrySerialized NameEntry(ENAME_LinkerConstructor);
 
-		if (bIsTextFormat)
-		{
-			NameStream.EnterElement() << NameEntry;
-		}
-		else
-		{
-			// Read the name from the underlying Archive when the format is binary to avoid the overhead from ArchiveProxy
-			NameStream.GetUnderlyingArchive() << NameEntry;
-		}
+		// Read the name from the underlying Archive when the format is binary to avoid the overhead from ArchiveProxy
+		*this << NameEntry;
 
 		// Add it to the name table with no splitting and no hash calculations
 		NameMap.Emplace(FName(NameEntry).GetDisplayIndex());
diff --git a/Engine/Source/Runtime/CoreUObject/Private/UObject/SavePackage.cpp b/Engine/Source/Runtime/CoreUObject/Private/UObject/SavePackage.cpp
--- a/Engine/Source/Runtime/CoreUObject/Private/UObject/SavePackage.cpp
+++ b/Engine/Source/Runtime/CoreUObject/Private/UObject/SavePackage.cpp
@@ -4571,6 +4571,7 @@ FSavePackageResultStruct UPackage::Save(UPackage* InOuter, UObject* Base, EObjec
 				SlowTask.EnterProgressFrame();
 
 				// Save names.
+				if (!bTextFormat)
 				{
 					FStructuredArchive::FStream NameStream = StructuredArchiveRoot.EnterStream(FIELD_NAME_TEXT("Names"));
 #if WITH_EDITOR
