TBA: Serialize tagged properties as a stream of records containing Tag/Value. Slight backwards step, but ensures that existing loading code can deal with arrays properly.

diff --git a/Engine/Source/Runtime/CoreUObject/Private/UObject/Class.cpp b/Engine/Source/Runtime/CoreUObject/Private/UObject/Class.cpp
--- a/Engine/Source/Runtime/CoreUObject/Private/UObject/Class.cpp
+++ b/Engine/Source/Runtime/CoreUObject/Private/UObject/Class.cpp
@@ -922,21 +922,7 @@ void UStruct::SerializeTaggedProperties(FStructuredArchive::FSlot Slot, uint8* D
 	if( UnderlyingArchive.IsLoading() )
 	{
 		// Load tagged properties.
-		TArray<FString> FieldNames;
-
-		TOptional<FStructuredArchive::FRecord> PropertiesRecord;
-		TOptional<FStructuredArchive::FStream> PropertiesStream;
-
-		if (UnderlyingArchive.IsTextFormat())
-		{
-			PropertiesRecord.Emplace(Slot.EnterRecord_TextOnly(FieldNames));
-		}
-		else
-		{
-			PropertiesStream.Emplace(Slot.EnterStream());
-		}
-
-		int32 CurrentFieldNameIdx = UnderlyingArchive.IsTextFormat() ? 0 : -1;
+		FStructuredArchive::FStream PropertiesStream = Slot.EnterStream();
 
 		// This code assumes that properties are loaded in the same order they are saved in. This removes a n^2 search 
 		// and makes it an O(n) when properties are saved in the same order as they are loaded (default case). In the 
@@ -946,9 +932,9 @@ void UStruct::SerializeTaggedProperties(FStructuredArchive::FSlot Slot, uint8* D
 		int32		RemainingArrayDim	= Property ? Property->ArrayDim : 0;
 
 		// Load all stored properties, potentially skipping unknown ones.
-		while (CurrentFieldNameIdx < FieldNames.Num())
+		while (true)
 		{
-			FStructuredArchive::FRecord PropertyRecord = UnderlyingArchive.IsTextFormat() ? PropertiesRecord->EnterRecord(FIELD_NAME(*FieldNames[CurrentFieldNameIdx++])) : PropertiesStream->EnterElement().EnterRecord();
+			FStructuredArchive::FRecord PropertyRecord = PropertiesStream.EnterElement().EnterRecord();
 
 			FPropertyTag Tag;
 			PropertyRecord << NAMED_FIELD(Tag);
@@ -1143,7 +1129,7 @@ void UStruct::SerializeTaggedProperties(FStructuredArchive::FSlot Slot, uint8* D
 	}
 	else
 	{
-		FStructuredArchive::FRecord PropertiesRecord = Slot.EnterRecord();
+		FStructuredArchive::FStream PropertiesStream = Slot.EnterStream();
 
 		check(UnderlyingArchive.IsSaving() || UnderlyingArchive.IsCountingMemory());
 
@@ -1191,7 +1177,7 @@ void UStruct::SerializeTaggedProperties(FStructuredArchive::FSlot Slot, uint8* D
 							Tag.SetPropertyGuid(PropertyGuid);
 						}
 
-						FStructuredArchive::FRecord PropertyRecord = PropertiesRecord.EnterRecord(FIELD_NAME(*Tag.Name.ToString()));
+						FStructuredArchive::FRecord PropertyRecord = PropertiesStream.EnterElement().EnterRecord();
 
 						PropertyRecord << NAMED_FIELD(Tag);
 
@@ -1239,11 +1225,8 @@ void UStruct::SerializeTaggedProperties(FStructuredArchive::FSlot Slot, uint8* D
 			}
 		}
 
-		if (!UnderlyingArchive.IsTextFormat())
-		{
-			static FName Temp(NAME_None);
-			UnderlyingArchive << Temp;
-		}
+		static FName Temp(NAME_None);
+		PropertiesStream.EnterElement().EnterRecord().EnterField(FIELD_NAME_TEXT("Tag")).EnterRecord() << NAMED_ITEM("Name", Temp);
 	}
 }
 void UStruct::FinishDestroy()
