PR #5258: Performance Optimization while Async Loading (Contributed by [REDACTED_PERSON])

diff --git a/Engine/Source/Runtime/CoreUObject/Private/Serialization/AsyncLoading.cpp b/Engine/Source/Runtime/CoreUObject/Private/Serialization/AsyncLoading.cpp
--- a/Engine/Source/Runtime/CoreUObject/Private/Serialization/AsyncLoading.cpp
+++ b/Engine/Source/Runtime/CoreUObject/Private/Serialization/AsyncLoading.cpp
@@ -4790,10 +4790,16 @@ EAsyncPackageState::Type FAsyncLoadingThread::ProcessLoadedPackages(bool bUseTim
 		FScopeLock LoadedPackagesLock(&LoadedPackagesCritical);
 		FScopeLock LoadedPackagesToProcessLock(&LoadedPackagesToProcessCritical);
 #endif
-		LoadedPackagesToProcess.Append(LoadedPackages);
-		LoadedPackagesToProcessNameLookup.Append(LoadedPackagesNameLookup);
-		LoadedPackages.Reset();
-		LoadedPackagesNameLookup.Reset();
+		if (LoadedPackages.Num() != 0)
+		{
+			LoadedPackagesToProcess.Append(LoadedPackages);
+			LoadedPackages.Reset();
+		}
+		if (LoadedPackagesNameLookup.Num() != 0)
+		{
+			LoadedPackagesToProcessNameLookup.Append(LoadedPackagesNameLookup);
+			LoadedPackagesNameLookup.Reset();
+		}
 	}
 #if USE_EVENT_DRIVEN_ASYNC_LOAD_AT_BOOT_TIME
 	if (IsMultithreaded() && GEventDrivenLoaderEnabled &&
