Fixed StaticFindObjectFastInternalThreadSafe() bug caused by recent optimization

diff --git a/Engine/Source/Runtime/CoreUObject/Private/UObject/UObjectHash.cpp b/Engine/Source/Runtime/CoreUObject/Private/UObject/UObjectHash.cpp
--- a/Engine/Source/Runtime/CoreUObject/Private/UObject/UObjectHash.cpp
+++ b/Engine/Source/Runtime/CoreUObject/Private/UObject/UObjectHash.cpp
@@ -7,6 +7,7 @@
 #include "UObject/UObjectHash.h"
 #include "UObject/Class.h"
 #include "UObject/Package.h"
+#include "Misc/AsciiSet.h"
 #include "Misc/PackageName.h"
 #include "HAL/IConsoleManager.h"
 
@@ -536,14 +537,21 @@ UObject* StaticFindObjectFastInternalThreadSafe(FUObjectHashTables& ThreadHash,
 	{
 		// Find an object with the specified name and (optional) class, in any package; if bAnyPackage is false, only matches top-level packages
 		FName ActualObjectName = ObjectName;
-		TCHAR ObjectNameString[NAME_SIZE];
-		ObjectName.GetPlainNameString(ObjectNameString);
-
-		const TCHAR* LastDot = FCString::Strrchr(ObjectNameString, '.');
-		const TCHAR* FirstColon = FCString::Strchr(ObjectNameString, ':');
-		if (const TCHAR* LastDelimiter = (UPTRINT)LastDot > (UPTRINT)FirstColon ? LastDot : FirstColon)
+		FName VerifyOuterName;
+		TCHAR PlainObjectName[NAME_SIZE];
+		int32 PlainObjectNameLen = ObjectName.GetPlainNameString(PlainObjectName);
+		
+		// Reverse scan for first '.' or ':'
+		constexpr FAsciiSet DotColon(".:");
+		for (const TCHAR* PlainIt = PlainObjectName + PlainObjectNameLen - 1; PlainIt > PlainObjectName; --PlainIt)
 		{
-			ActualObjectName = FName(LastDelimiter + 1, ObjectName.GetNumber());
+			if (DotColon.Test(*PlainIt))
+			{
+				ActualObjectName = FName(PlainIt + 1, ObjectName.GetNumber());
+				int32 OuterLen = static_cast<int32>(PlainIt - PlainObjectName);
+				VerifyOuterName = FName(OuterLen, PlainObjectName);
+				break;
+			}
 		}
 
 		const int32 Hash = GetObjectHash(ActualObjectName);
@@ -573,7 +581,7 @@ UObject* StaticFindObjectFastInternalThreadSafe(FUObjectHashTables& ThreadHash,
 					&& !Object->HasAnyInternalFlags(ExclusiveInternalFlags)
 
 					/** Ensure that the partial path provided matches the object found */
-					&& (Object->GetPathName().EndsWith(ObjectNameString)))
+					&& (VerifyOuterName.IsNone() || Object->GetOuter() && Object->GetOuter()->GetFName() == VerifyOuterName))
 				{
 					checkf(!Object->IsUnreachable(), TEXT("%s"), *Object->GetFullName());
 					if (Result)
