Quench clang's thirst for superfluous parentheses

diff --git a/Engine/Source/Runtime/CoreUObject/Private/UObject/UObjectHash.cpp b/Engine/Source/Runtime/CoreUObject/Private/UObject/UObjectHash.cpp
--- a/Engine/Source/Runtime/CoreUObject/Private/UObject/UObjectHash.cpp
+++ b/Engine/Source/Runtime/CoreUObject/Private/UObject/UObjectHash.cpp
@@ -581,7 +581,7 @@ UObject* StaticFindObjectFastInternalThreadSafe(FUObjectHashTables& ThreadHash,
 					&& !Object->HasAnyInternalFlags(ExclusiveInternalFlags)
 
 					/** Ensure that the partial path provided matches the object found */
-					&& (VerifyOuterName.IsNone() || Object->GetOuter() && Object->GetOuter()->GetFName() == VerifyOuterName))
+					&& (VerifyOuterName.IsNone() || (Object->GetOuter() && Object->GetOuter()->GetFName() == VerifyOuterName)))
 				{
 					checkf(!Object->IsUnreachable(), TEXT("%s"), *Object->GetFullName());
 					if (Result)
