Fixing GC crash on platforms that do not support multithreading

diff --git a/Engine/Source/Runtime/CoreUObject/Private/UObject/GarbageCollection.cpp b/Engine/Source/Runtime/CoreUObject/Private/UObject/GarbageCollection.cpp
--- a/Engine/Source/Runtime/CoreUObject/Private/UObject/GarbageCollection.cpp
+++ b/Engine/Source/Runtime/CoreUObject/Private/UObject/GarbageCollection.cpp
@@ -253,6 +253,8 @@ class FAsyncPurge : public FRunnable
 	FEvent* FinishedPurgeEvent;
 	/** Current index into the global unreachable objects array (GUnreachableObjects) of the object being destroyed */
 	int32 ObjCurrentPurgeObjectIndex;
+	/** Number of unreachable objects the last time single-threaded tick was called */
+	int32 LastUnreachableObjectsCount;
 	/** A list of objects that were not thread safe to destroy on the worker thread */
 	TLockFreePointerListUnordered<UObject, PLATFORM_CACHE_LINE_SIZE> GameThreadObjects;
 	/** Stats for the number of objects destroyed */
@@ -343,6 +345,8 @@ public:
 		, AsyncPurgeThreadId(0)
 		, BeginPurgeEvent(nullptr)
 		, FinishedPurgeEvent(nullptr)
+		, ObjCurrentPurgeObjectIndex(0)
+		, LastUnreachableObjectsCount(0)
 		, ObjectsDestroyedSinceLastMarkPhase(0)
 	{
 		BeginPurgeEvent = FPlatformProcess::GetSynchEventFromPool(true);
@@ -392,6 +396,7 @@ public:
 		if (!Thread)
 		{
 			// If we're running single-threaded we need to tick the main loop here too
+			LastUnreachableObjectsCount = GUnreachableObjects.Num();
 			bCanStartDestroyingGameThreadObjects = TickDestroyObjects(bUseTimeLimit, TimeLimit, StartTime);
 		}
 		else if (!bUseTimeLimit)
@@ -413,7 +418,14 @@ public:
 	/** Returns true if the destruction process is finished */
 	bool IsFinished() const
 	{
-		return FinishedPurgeEvent->Wait(0, true) && GameThreadObjects.IsEmpty();
+		if (Thread)
+		{
+			return FinishedPurgeEvent->Wait(0, true) && GameThreadObjects.IsEmpty();
+		}
+		else
+		{
+			return (ObjCurrentPurgeObjectIndex >= LastUnreachableObjectsCount && GameThreadObjects.IsEmpty());
+		}
 	}
 
 	/** Returns the number of objects already destroyed */
