Fixing a rare crash caused by UObjectDelete listeners not being locked when freeing UObjects on the async destruction thread

diff --git a/Engine/Source/Runtime/CoreUObject/Private/UObject/UObjectArray.cpp b/Engine/Source/Runtime/CoreUObject/Private/UObject/UObjectArray.cpp
--- a/Engine/Source/Runtime/CoreUObject/Private/UObject/UObjectArray.cpp
+++ b/Engine/Source/Runtime/CoreUObject/Private/UObject/UObjectArray.cpp
@@ -206,11 +206,15 @@ void FUObjectArray::FreeUObjectIndex(UObjectBase* Object)
 		UE_LOG(LogUObjectArray, Fatal, TEXT("Unexpected concurency while adding new object"));
 	}
 
-	// @todo: threading: delete listeners should be locked while we're doing this
-	// Iterate in reverse order so that when one of the listeners removes itself from the array inside of NotifyUObjectDeleted we don't skip the next listener.
-	for (int32 ListenerIndex = UObjectDeleteListeners.Num() - 1; ListenerIndex >= 0; --ListenerIndex)
 	{
-		UObjectDeleteListeners[ListenerIndex]->NotifyUObjectDeleted(Object, Index);
+#if THREADSAFE_UOBJECTS
+		FScopeLock UObjectDeleteListenersLock(&UObjectDeleteListenersCritical);
+#endif
+		// Iterate in reverse order so that when one of the listeners removes itself from the array inside of NotifyUObjectDeleted we don't skip the next listener.
+		for (int32 ListenerIndex = UObjectDeleteListeners.Num() - 1; ListenerIndex >= 0; --ListenerIndex)
+		{
+			UObjectDeleteListeners[ListenerIndex]->NotifyUObjectDeleted(Object, Index);
+		}
 	}
 	// You cannot safely recycle indicies in the non-GC range
 	// No point in filling this list when doing exit purge. Nothing should be allocated afterwards anyway.
