Allowing TStrongObjectPtr's FInternalReferenceCollector to be destroyed on the GC's async destruction thread as the underlying FGCObject already works fine with it and multithreading with GC in general

diff --git a/Engine/Source/Runtime/CoreUObject/Public/UObject/StrongObjectPtr.h b/Engine/Source/Runtime/CoreUObject/Public/UObject/StrongObjectPtr.h
--- a/Engine/Source/Runtime/CoreUObject/Public/UObject/StrongObjectPtr.h
+++ b/Engine/Source/Runtime/CoreUObject/Public/UObject/StrongObjectPtr.h
@@ -20,7 +20,7 @@ namespace UE4StrongObjectPtr_Private
 
 		virtual ~FInternalReferenceCollector()
 		{
-			check(IsInGameThread());
+			check(IsInGameThread() || IsInGarbageCollectorThread());
 		}
 
 		bool IsValid() const
