Changed text asset extension to utxt/utxtmap

diff --git a/Engine/Source/Runtime/CoreUObject/Private/Misc/PackageName.cpp b/Engine/Source/Runtime/CoreUObject/Private/Misc/PackageName.cpp
--- a/Engine/Source/Runtime/CoreUObject/Private/Misc/PackageName.cpp
+++ b/Engine/Source/Runtime/CoreUObject/Private/Misc/PackageName.cpp
@@ -24,8 +24,8 @@ DEFINE_LOG_CATEGORY_STATIC(LogPackageName, Log, All);
 
 FString FPackageName::AssetPackageExtension = TEXT(".uasset");
 FString FPackageName::MapPackageExtension = TEXT(".umap");
-FString FPackageName::TextAssetPackageExtension = TEXT(".utextasset");
-FString FPackageName::TextMapPackageExtension = TEXT(".utextmap");
+FString FPackageName::TextAssetPackageExtension = TEXT(".utxt");
+FString FPackageName::TextMapPackageExtension = TEXT(".utxtmap");
 
 /** Event that is triggered when a new content path is mounted */
 FPackageName::FOnContentPathMountedEvent FPackageName::OnContentPathMountedEvent;
