FAsyncArchive interface

diff --git a/Engine/Source/Runtime/CoreUObject/Private/Serialization/AsyncLoading.cpp b/Engine/Source/Runtime/CoreUObject/Private/Serialization/AsyncLoading.cpp
--- a/Engine/Source/Runtime/CoreUObject/Private/Serialization/AsyncLoading.cpp
+++ b/Engine/Source/Runtime/CoreUObject/Private/Serialization/AsyncLoading.cpp
@@ -70,7 +70,7 @@ DECLARE_CYCLE_STAT(TEXT("CreateImports AsyncPackage"),STAT_FAsyncPackage_CreateI
 DECLARE_CYCLE_STAT(TEXT("CreateMetaData AsyncPackage"),STAT_FAsyncPackage_CreateMetaData,STATGROUP_AsyncLoad);
 DECLARE_CYCLE_STAT(TEXT("CreateExports AsyncPackage"),STAT_FAsyncPackage_CreateExports,STATGROUP_AsyncLoad);
 DECLARE_CYCLE_STAT(TEXT("FreeReferencedImports AsyncPackage"), STAT_FAsyncPackage_FreeReferencedImports, STATGROUP_AsyncLoad);
-DECLARE_CYCLE_STAT(TEXT("Precache ArchiveAsync"), STAT_FArchiveAsync_Precache, STATGROUP_AsyncLoad);
+DECLARE_CYCLE_STAT(TEXT("Precache AsyncArchive"), STAT_FAsyncArchive_Precache, STATGROUP_AsyncLoad);
 DECLARE_CYCLE_STAT(TEXT("PreLoadObjects AsyncPackage"),STAT_FAsyncPackage_PreLoadObjects,STATGROUP_AsyncLoad);
 DECLARE_CYCLE_STAT(TEXT("ExternalReadDependencies AsyncPackage"),STAT_FAsyncPackage_ExternalReadDependencies,STATGROUP_AsyncLoad);
 DECLARE_CYCLE_STAT(TEXT("PostLoadObjects AsyncPackage"),STAT_FAsyncPackage_PostLoadObjects,STATGROUP_AsyncLoad);
@@ -298,7 +298,7 @@ FORCEINLINE bool FAsyncPackage::IsTimeLimitExceeded()
 }
 
 DEFINE_LOG_CATEGORY_STATIC(LogAsyncArchive, Display, All);
-DECLARE_MEMORY_STAT(TEXT("FArchiveAsync2 Buffers"), STAT_FArchiveAsync2Mem, STATGROUP_Memory);
+DECLARE_MEMORY_STAT(TEXT("FAsyncArchive Buffers"), STAT_FAsyncArchiveMem, STATGROUP_Memory);
 
 //#define ASYNC_WATCH_FILE "SM_Boots_Wings.uasset"
 #define TRACK_SERIALIZE (0)
@@ -307,7 +307,7 @@ DECLARE_MEMORY_STAT(TEXT("FArchiveAsync2 Buffers"), STAT_FArchiveAsync2Mem, STAT
 
 
 
-FORCEINLINE void FArchiveAsync2::LogItem(const TCHAR* Item, int64 Offset, int64 Size, double StartTime)
+FORCEINLINE void FAsyncArchive::LogItem(const TCHAR* Item, int64 Offset, int64 Size, double StartTime)
 {
 	if (UE_LOG_ACTIVE(LogAsyncArchive, Verbose)
 #if defined(ASYNC_WATCH_FILE)
@@ -2697,10 +2697,10 @@ EAsyncPackageState::Type FAsyncPackage::SetupImports_Event()
 	}
 
 #if 0
-	if (bAnyImportArcsAdded && Linker->GetFArchiveAsync2Loader())
+	if (bAnyImportArcsAdded && Linker->GetAsyncLoader())
 	{
 		// we are waiting for imports, so drop our precache requests
-		Linker->GetFArchiveAsync2Loader()->FlushCache();
+		Linker->GetAsyncLoader()->FlushCache();
 	}
 #endif
 	return ImportIndex == Linker->ImportMap.Num() ? EAsyncPackageState::Complete : EAsyncPackageState::TimeOut;
@@ -2888,7 +2888,7 @@ EAsyncPackageState::Type FAsyncPackage::SetupExports_Event()
 	FGCScopeGuard GCGuard;
 
 	FCheckedWeakAsyncPackagePtr WeakThis(this);
-	Linker->GetFArchiveAsync2Loader()->LogItem(TEXT("SetupExports_Event"));
+	Linker->GetAsyncLoader()->LogItem(TEXT("SetupExports_Event"));
 
 	LastTypeOfWorkPerformed = TEXT("SetupExports_Event");
 	LastObjectWorkWasPerformedOn = LinkerRoot;
@@ -3004,9 +3004,9 @@ void FAsyncPackage::LinkImport(int32 LocalImportIndex)
 	if (!Import.XObject && !Import.bImportFailed)
 	{
 		FScopedAddObjectreference OnExitAddReference(*this, Import.XObject);
-		if (Linker->GetFArchiveAsync2Loader())
+		if (Linker->GetAsyncLoader())
 		{
-			Linker->GetFArchiveAsync2Loader()->LogItem(TEXT("LinkImport"));
+			Linker->GetAsyncLoader()->LogItem(TEXT("LinkImport"));
 		}
 		if (Import.SourceLinker)
 		{
@@ -3267,9 +3267,9 @@ void FAsyncPackage::EventDrivenCreateExport(int32 LocalExportIndex)
 		{
 			SCOPED_ACCUM_LOADTIME(Construction, StaticGetNativeClassName(CastEventDrivenIndexToObject<UClass>(Export.ClassIndex, false)));
 
-			if (Linker->GetFArchiveAsync2Loader())
+			if (Linker->GetAsyncLoader())
 			{
-				Linker->GetFArchiveAsync2Loader()->LogItem(TEXT("EventDrivenCreateExport"), Export.SerialOffset, Export.SerialSize);
+				Linker->GetAsyncLoader()->LogItem(TEXT("EventDrivenCreateExport"), Export.SerialOffset, Export.SerialSize);
 			}
 			LastTypeOfWorkPerformed = TEXT("EventDrivenCreateExport");
 			LastObjectWorkWasPerformedOn = nullptr;
@@ -3590,7 +3590,7 @@ void FAsyncPackage::EventDrivenSerializeExport(int32 LocalExportIndex)
 	else if (Object && Object->HasAnyFlags(RF_NeedLoad))
 	{
 
-		Linker->GetFArchiveAsync2Loader()->LogItem(TEXT("EventDrivenSerializeExport"), Export.SerialOffset, Export.SerialSize);
+		Linker->GetAsyncLoader()->LogItem(TEXT("EventDrivenSerializeExport"), Export.SerialOffset, Export.SerialSize);
 
 		LastTypeOfWorkPerformed = TEXT("EventDrivenSerializeExport");
 		LastObjectWorkWasPerformedOn = Object;
@@ -3624,11 +3624,11 @@ void FAsyncPackage::EventDrivenSerializeExport(int32 LocalExportIndex)
 		}
 		check(Export.SerialOffset >= CurrentBlockOffset && Export.SerialOffset + Export.SerialSize <= CurrentBlockOffset + CurrentBlockBytes);
 
-		FArchiveAsync2* FAA2 = Linker->GetFArchiveAsync2Loader();
-		check(FAA2);
+		FAsyncArchive* AsyncLoader = Linker->GetAsyncLoader();
+		check(AsyncLoader);
 
-		const int64 SavedPos = FAA2->Tell();
-		FAA2->Seek(Export.SerialOffset);
+		const int64 SavedPos = AsyncLoader->Tell();
+		AsyncLoader->Seek(Export.SerialOffset);
 
 		Object->ClearFlags(RF_NeedLoad);
 
@@ -3664,19 +3664,19 @@ void FAsyncPackage::EventDrivenSerializeExport(int32 LocalExportIndex)
 		LoadContext->SerializedObject = PrevSerializedObject;
 		Linker->bForceSimpleIndexToObject = false;
 
-		if (FAA2->Tell() - Export.SerialOffset != Export.SerialSize)
+		if (AsyncLoader->Tell() - Export.SerialOffset != Export.SerialSize)
 		{
 			if (Object->GetClass()->HasAnyClassFlags(CLASS_Deprecated))
 			{
-				UE_LOG(LogStreaming, Warning, TEXT("%s"), *FString::Printf(TEXT("%s: Serial size mismatch: Got %d, Expected %d"), *Object->GetFullName(), (int32)(FAA2->Tell() - Export.SerialOffset), Export.SerialSize));
+				UE_LOG(LogStreaming, Warning, TEXT("%s"), *FString::Printf(TEXT("%s: Serial size mismatch: Got %d, Expected %d"), *Object->GetFullName(), (int32)(AsyncLoader->Tell() - Export.SerialOffset), Export.SerialSize));
 			}
 			else
 			{
-				UE_LOG(LogStreaming, Fatal, TEXT("%s"), *FString::Printf(TEXT("%s: Serial size mismatch: Got %d, Expected %d"), *Object->GetFullName(), (int32)(FAA2->Tell() - Export.SerialOffset), Export.SerialSize));
+				UE_LOG(LogStreaming, Fatal, TEXT("%s"), *FString::Printf(TEXT("%s: Serial size mismatch: Got %d, Expected %d"), *Object->GetFullName(), (int32)(AsyncLoader->Tell() - Export.SerialOffset), Export.SerialSize));
 			}
 		}
 
-		FAA2->Seek(SavedPos);
+		AsyncLoader->Seek(SavedPos);
 #if DO_CHECK
 		if (Object->HasAnyFlags(RF_ClassDefaultObject) && Object->GetClass()->HasAnyClassFlags(CLASS_CompiledFromBlueprint))
 		{
@@ -3810,10 +3810,10 @@ void FAsyncPackage::StartPrecacheRequest()
 		}
 	}
 	check(NewReq.ExportsToRead.Num());
-	FArchiveAsync2* FAA2 = Linker->GetFArchiveAsync2Loader();
-	check(FAA2);
+	FAsyncArchive* AsyncLoader = Linker->GetAsyncLoader();
+	check(AsyncLoader);
 
-	IAsyncReadRequest* Precache = FAA2->MakeEventDrivenPrecacheRequest(NewReq.Offset, NewReq.BytesToRead,GPrecacheCallbackHandler.GetCompletionCallback());
+	IAsyncReadRequest* Precache = AsyncLoader->MakeEventDrivenPrecacheRequest(NewReq.Offset, NewReq.BytesToRead,GPrecacheCallbackHandler.GetCompletionCallback());
 
 	NewReq.FirstExportCovered = LocalExportIndex;
 	NewReq.LastExportCovered = LastExportIndex;
@@ -3850,11 +3850,11 @@ void FAsyncPackage::MakeNextPrecacheRequestCurrent()
 
 	GPrecacheCallbackHandler.FinishRequest(Req.BytesToRead);
 
-	FArchiveAsync2* FAA2 = Linker->GetFArchiveAsync2Loader();
-	check(FAA2);
+	FAsyncArchive* AsyncLoader = Linker->GetAsyncLoader();
+	check(AsyncLoader);
 	Read->WaitCompletion();
 
-	bool bReady = FAA2->PrecacheForEvent(CurrentBlockOffset, CurrentBlockBytes);
+	bool bReady = AsyncLoader->PrecacheForEvent(CurrentBlockOffset, CurrentBlockBytes);
 	UE_CLOG(!bReady, LogStreaming, Warning, TEXT("Precache request should have been hot %s."), *Linker->Filename);
 	for (int32 Index = Req.FirstExportCovered; Index <= Req.LastExportCovered; Index++)
 	{
@@ -3876,9 +3876,9 @@ void FAsyncPackage::FlushPrecacheBuffer()
 	CurrentBlockBytes = -1;
 	if (!Linker->bDynamicClassLinker)
 	{
-	FArchiveAsync2* FAA2 = Linker->GetFArchiveAsync2Loader();
-	check(FAA2);
-	FAA2->FlushPrecacheBlock();
+	FAsyncArchive* AsyncLoader = Linker->GetAsyncLoader();
+	check(AsyncLoader);
+	AsyncLoader->FlushPrecacheBlock();
 }
 }
 
@@ -4025,7 +4025,7 @@ EAsyncPackageState::Type FAsyncPackage::ProcessImportsAndExports_Event()
 
 void FAsyncPackage::Event_ExportsDone()
 {
-	Linker->GetFArchiveAsync2Loader()->LogItem(TEXT("Event_ExportsDone"));
+	Linker->GetAsyncLoader()->LogItem(TEXT("Event_ExportsDone"));
 	check(AsyncPackageLoadingState == EAsyncPackageLoadingState::ProcessNewImportsAndExports);
 	bAllExportsSerialized = true;
 	RemoveNode(EEventLoadNode::Package_ExportsSerialized);
@@ -4055,7 +4055,7 @@ void FAsyncPackage::Event_ExportsDone()
 
 void FAsyncPackage::Event_ProcessPostloadWait()
 {
-	Linker->GetFArchiveAsync2Loader()->LogItem(TEXT("Event_ProcessPostloadWait"));
+	Linker->GetAsyncLoader()->LogItem(TEXT("Event_ProcessPostloadWait"));
 	check(AsyncPackageLoadingState == EAsyncPackageLoadingState::WaitingForPostLoad);
 	check(bAllExportsSerialized && !OtherPackagesWaitingForMeBeforePostload.Num());
 	bProcessPostloadWaitInFlight = false;
@@ -4172,7 +4172,7 @@ void FAsyncPackage::Event_StartPostload()
 {
 	LLM_SCOPE(ELLMTag::AsyncLoading);
 
-	Linker->GetFArchiveAsync2Loader()->LogItem(TEXT("Event_StartPostload"));
+	Linker->GetAsyncLoader()->LogItem(TEXT("Event_StartPostload"));
 	check(AsyncPackageLoadingState == EAsyncPackageLoadingState::ReadyForPostLoad);
 	check(!PackagesIMayBeWaitingForBeforePostload.Num());
 	check(!PackagesIAmWaitingForBeforePostload.Num());
@@ -6144,7 +6144,7 @@ EAsyncPackageState::Type FAsyncPackage::CreateLinker()
 					if (Linker->bDynamicClassLinker)
 					{
 						//native blueprint 
-						check(!Linker->GetFArchiveAsync2Loader());
+						check(!Linker->GetAsyncLoader());
 						GPrecacheCallbackHandler.SummaryComplete(WeakPtr);
 					}
 				}
@@ -6522,10 +6522,10 @@ EAsyncPackageState::Type FAsyncPackage::CreateExports()
 		// Precache data and see whether it's already finished.
 		bool bReady;
 
-		FArchiveAsync2* FAA2 = Linker->GetFArchiveAsync2Loader();
-		if (FAA2)
+		FAsyncArchive* AsyncLoader = Linker->GetAsyncLoader();
+		if (AsyncLoader)
 		{
-			bReady = FAA2->Precache(Export.SerialOffset, Export.SerialSize, bUseTimeLimit, bUseFullTimeLimit, TickStartTime, TimeLimit);
+			bReady = AsyncLoader->PrecacheWithTimeLimit(Export.SerialOffset, Export.SerialSize, bUseTimeLimit, bUseFullTimeLimit, TickStartTime, TimeLimit);
 		}
 		else
 		{
@@ -7546,10 +7546,10 @@ void NotifyRegistrationComplete()
 	GetGEDLBootNotificationManager().NotifyRegistrationComplete();
 }
 
-#define USE_DETAILED_FARCHIVEASYNC2_MEMORY_TRACKING 0
+#define USE_DETAILED_FASYNCARCHIVE_MEMORY_TRACKING 0
 
-#if USE_DETAILED_FARCHIVEASYNC2_MEMORY_TRACKING
-class FArchiveAsync2MemTracker
+#if USE_DETAILED_FASYNCARCHIVE_MEMORY_TRACKING
+class FAsyncArchiveMemTracker
 {
 	TMap<FString, int64> AllocatedMem;
 	FCriticalSection AllocatedMemCritical;
@@ -7557,14 +7557,14 @@ class FArchiveAsync2MemTracker
 public:
 
 	void Allocate(const FString& Filename, int64 Mem)
-{
+	{
 		FScopeLock AllocatedMemLock(&AllocatedMemCritical);
 		int64& AllocatedMemAmount = AllocatedMem.FindOrAdd(Filename);
 		AllocatedMemAmount += Mem;
-}
+	}
 
 	void Deallocate(const FString& Filename, int64 Mem)
-{
+	{
 		FScopeLock AllocatedMemLock(&AllocatedMemCritical);
 		int64& AllocatedMemAmount = AllocatedMem.FindOrAdd(Filename);
 		AllocatedMemAmount -= Mem;
@@ -7572,36 +7572,36 @@ public:
 		if (AllocatedMemAmount == 0)
 		{
 			AllocatedMem.Remove(Filename);
+		}
 	}
-}
 
 	void Dump()
-{
+	{
 		FScopeLock AllocatedMemLock(&AllocatedMemCritical);
 
-		UE_LOG(LogStreaming, Display, TEXT("Dumping FArchiveAsync2 allocated memory (%d)"), AllocatedMem.Num());
+		UE_LOG(LogStreaming, Display, TEXT("Dumping FAsyncArchie allocated memory (%d)"), AllocatedMem.Num());
 		for (TPair<FString, int64>& ArchiveMem : AllocatedMem)
-{
+		{
 			UE_LOG(LogStreaming, Display, TEXT("  %s %lldb"), *ArchiveMem.Key, ArchiveMem.Value);
-}
-}
-} GArchiveAsync2MemTracker;
+		}
+	}
+} GAsyncArchiveMemTracker;
 
-void DumpArchiveAsync2Mem(const TArray<FString>& Args)
+void DumpAsyncArchiveMem(const TArray<FString>& Args)
 {
-	GArchiveAsync2MemTracker.Dump();
+	GAsyncArchiveMemTracker.Dump();
 }
 
 static FAutoConsoleCommand GDumpSerializeCmd(
-	TEXT("DumpFArchiveAsync2Mem"),
-	TEXT("Debug command to dump the memory allocated by existing FArhiveAsync2."),
-	FConsoleCommandWithArgsDelegate::CreateStatic(&DumpArchiveAsync2Mem)
+	TEXT("DumpAsyncArchiveMem"),
+	TEXT("Debug command to dump the memory allocated by existing FAsyncArchive."),
+	FConsoleCommandWithArgsDelegate::CreateStatic(&DumpAsyncArchiveMem)
 );
-#endif // USE_DETAILED_FARCHIVEASYNC2_MEMORY_TRACKING
+#endif // USE_DETAILED_FASYNCARCHIVE_MEMORY_TRACKING
 
 static FCriticalSection SummaryRacePreventer;
 
-FArchiveAsync2::FArchiveAsync2(const TCHAR* InFileName, TFunction<void()>&& InSummaryReadyCallback)
+FAsyncArchive::FAsyncArchive(const TCHAR* InFileName, TFunction<void()>&& InSummaryReadyCallback)
 	: Handle(nullptr)
 	, SizeRequestPtr(nullptr)
 	, EditorPrecacheRequestPtr(nullptr)
@@ -7619,7 +7619,6 @@ FArchiveAsync2::FArchiveAsync2(const TCHAR* InFileName, TFunction<void()>&& InSu
 	, HeaderSize(0)
 	, HeaderSizeWhenReadingExportsFromSplitFile(0)
 	, LoadPhase(ELoadPhase::WaitingForSize)
-	, bCookedForEDLInEditor(false)
 	, FileName(InFileName)
 	, OpenTime(FPlatformTime::Seconds())
 	, SummaryReadTime(0.0)
@@ -7648,7 +7647,7 @@ FArchiveAsync2::FArchiveAsync2(const TCHAR* InFileName, TFunction<void()>&& InSu
 
 }
 
-FArchiveAsync2::~FArchiveAsync2()
+FAsyncArchive::~FAsyncArchive()
 {
 	// Invalidate any precached data and free memory.
 	FlushCache();
@@ -7657,10 +7656,10 @@ FArchiveAsync2::~FArchiveAsync2()
 		delete Handle;
 		Handle = nullptr;
 	}
-	LogItem(TEXT("~FArchiveAsync2"), 0, 0);
+	LogItem(TEXT("~FAsyncArchive"), 0, 0);
 }
 
-void FArchiveAsync2::ReadCallback(bool bWasCancelled, IAsyncReadRequest* Request)
+void FAsyncArchive::ReadCallback(bool bWasCancelled, IAsyncReadRequest* Request)
 {
 	if (bWasCancelled || ArIsError)
 	{
@@ -7751,17 +7750,17 @@ void FArchiveAsync2::ReadCallback(bool bWasCancelled, IAsyncReadRequest* Request
 	}
 }
 
-void FArchiveAsync2::FlushPrecacheBlock()
+void FAsyncArchive::FlushPrecacheBlock()
 {
 #if DEVIRTUALIZE_FLinkerLoad_Serialize
 	DiscardInlineBufferAndUpdateCurrentPos();
 #endif
 	if (PrecacheBuffer)
 	{
-		DEC_MEMORY_STAT_BY(STAT_FArchiveAsync2Mem, PrecacheEndPos - PrecacheStartPos);
+		DEC_MEMORY_STAT_BY(STAT_FAsyncArchiveMem, PrecacheEndPos - PrecacheStartPos);
 		FMemory::Free(PrecacheBuffer);
-#if USE_DETAILED_FARCHIVEASYNC2_MEMORY_TRACKING
-		GArchiveAsync2MemTracker.Deallocate(FileName, PrecacheEndPos - PrecacheStartPos);
+#if USE_DETAILED_FASYNCARCHIVE_MEMORY_TRACKING
+		GAAsyncArchieMemTracker.Deallocate(FileName, PrecacheEndPos - PrecacheStartPos);
 #endif
 		check(!GEventDrivenLoaderEnabled || LoadPhase > ELoadPhase::WaitingForHeader);
 	}
@@ -7770,7 +7769,7 @@ void FArchiveAsync2::FlushPrecacheBlock()
 	PrecacheEndPos = 0;
 }
 
-void FArchiveAsync2::FlushCache()
+void FAsyncArchive::FlushCache()
 {
 	bool nNonRedundantFlush = PrecacheEndPos || PrecacheBuffer || ReadRequestPtr;
 	LogItem(TEXT("Flush"));
@@ -7820,7 +7819,7 @@ void FArchiveAsync2::FlushCache()
 
 }
 
-bool FArchiveAsync2::Close()
+bool FAsyncArchive::Close()
 {
 	// Invalidate any precached data and free memory.
 	FlushCache();
@@ -7828,13 +7827,13 @@ bool FArchiveAsync2::Close()
 	return !ArIsError;
 }
 
-bool FArchiveAsync2::SetCompressionMap(TArray<FCompressedChunk>* InCompressedChunks, ECompressionFlags InCompressionFlags)
+bool FAsyncArchive::SetCompressionMap(TArray<FCompressedChunk>* InCompressedChunks, ECompressionFlags InCompressionFlags)
 {
 	check(0); // no support for compression
 	return false;
 }
 
-int64 FArchiveAsync2::TotalSize()
+int64 FAsyncArchive::TotalSize()
 {
 	if (SizeRequestPtr)
 	{
@@ -7851,7 +7850,7 @@ int64 FArchiveAsync2::TotalSize()
 }
 
 #if DEVIRTUALIZE_FLinkerLoad_Serialize
-FORCEINLINE void FArchiveAsync2::SetPosAndUpdatePrecacheBuffer(int64 Pos)
+FORCEINLINE void FAsyncArchive::SetPosAndUpdatePrecacheBuffer(int64 Pos)
 {
 	check(Pos >= 0 && Pos <= TotalSizeOrMaxInt64IfNotReady());
 	if (Pos < PrecacheStartPos || Pos >= PrecacheEndPos)
@@ -7871,7 +7870,7 @@ FORCEINLINE void FArchiveAsync2::SetPosAndUpdatePrecacheBuffer(int64 Pos)
 }
 #endif
 
-void FArchiveAsync2::Seek(int64 InPos)
+void FAsyncArchive::Seek(int64 InPos)
 {
 	if ((GEventDrivenLoaderEnabled || bCookedForEDLInEditor) && LoadPhase < ELoadPhase::ProcessingExports)
 	{
@@ -7881,7 +7880,7 @@ void FArchiveAsync2::Seek(int64 InPos)
 			FirstExportStarting();
 		}
 	}
-	checkf(InPos >= 0 && InPos <= TotalSizeOrMaxInt64IfNotReady(), TEXT("Bad position in FArchiveAsync2::Seek. Filename:%s InPos:%lu, Size:%lu"), *FileName, InPos, TotalSizeOrMaxInt64IfNotReady());
+	checkf(InPos >= 0 && InPos <= TotalSizeOrMaxInt64IfNotReady(), TEXT("Bad position in FAsyncArchive::Seek. Filename:%s InPos:%lu, Size:%lu"), *FileName, InPos, TotalSizeOrMaxInt64IfNotReady());
 #if DEVIRTUALIZE_FLinkerLoad_Serialize
 	SetPosAndUpdatePrecacheBuffer(InPos);
 #else
@@ -7889,7 +7888,7 @@ void FArchiveAsync2::Seek(int64 InPos)
 #endif
 }
 
-bool FArchiveAsync2::WaitRead(float TimeLimit)
+bool FAsyncArchive::WaitRead(float TimeLimit)
 {
 	if (ReadRequestPtr)
 	{
@@ -7909,7 +7908,7 @@ bool FArchiveAsync2::WaitRead(float TimeLimit)
 	return true;
 }
 
-void FArchiveAsync2::CompleteRead()
+void FAsyncArchive::CompleteRead()
 {
 	double StartTime = FPlatformTime::Seconds();
 	check(LoadPhase != ELoadPhase::WaitingForSize && LoadPhase != ELoadPhase::WaitingForSummary);
@@ -7931,10 +7930,10 @@ void FArchiveAsync2::CompleteRead()
 			PrecacheStartPos = ReadRequestOffset;
 			PrecacheEndPos = ReadRequestOffset + ReadRequestSize;
 			check(ReadRequestSize > 0 && PrecacheStartPos >= 0);
-			INC_MEMORY_STAT_BY(STAT_FArchiveAsync2Mem, PrecacheEndPos - PrecacheStartPos);
+			INC_MEMORY_STAT_BY(STAT_FAsyncArchiveMem, PrecacheEndPos - PrecacheStartPos);
 			DEC_MEMORY_STAT_BY(STAT_AsyncFileMemory, ReadRequestSize);
-#if USE_DETAILED_FARCHIVEASYNC2_MEMORY_TRACKING
-			GArchiveAsync2MemTracker.Allocate(FileName, PrecacheEndPos - PrecacheStartPos);
+#if USE_DETAILED_FASYNCARCHIVE_MEMORY_TRACKING
+			GAsyncArchiveMemTracker.Allocate(FileName, PrecacheEndPos - PrecacheStartPos);
 #endif
 			// keeps the last cache block of the header around until we process the first export
 			if (LoadPhase != ELoadPhase::ProcessingExports)
@@ -7952,7 +7951,7 @@ void FArchiveAsync2::CompleteRead()
 	ReadRequestSize = 0;
 }
 
-void FArchiveAsync2::CompleteCancel()
+void FAsyncArchive::CompleteCancel()
 {
 	if (CanceledReadRequestPtr)
 	{
@@ -7966,7 +7965,7 @@ void FArchiveAsync2::CompleteCancel()
 }
 
 
-void FArchiveAsync2::CancelRead()
+void FAsyncArchive::CancelRead()
 {
 	if (ReadRequestPtr)
 	{
@@ -7979,7 +7978,7 @@ void FArchiveAsync2::CancelRead()
 	ReadRequestSize = 0;
 }
 
-bool FArchiveAsync2::WaitForIntialPhases(float InTimeLimit)
+bool FAsyncArchive::WaitForIntialPhases(float InTimeLimit)
 {
 	if (SizeRequestPtr 
 		|| GEventDrivenLoaderEnabled || SummaryRequestPtr || SummaryPrecacheRequestPtr
@@ -8052,7 +8051,7 @@ bool FArchiveAsync2::WaitForIntialPhases(float InTimeLimit)
 	return true;
 }
 
-bool FArchiveAsync2::PrecacheInternal(int64 RequestOffset, int64 RequestSize, bool bApplyMinReadSize)
+bool FAsyncArchive::PrecacheInternal(int64 RequestOffset, int64 RequestSize, bool bApplyMinReadSize)
 {
 	// CAUTION! This is possibly called the first time from a random IO thread.
 
@@ -8079,7 +8078,7 @@ bool FArchiveAsync2::PrecacheInternal(int64 RequestOffset, int64 RequestSize, bo
 		if (ReadRequestPtr)
 		{
 			// this one does not have what we need
-			UE_LOG(LogStreaming, Warning, TEXT("FArchiveAsync2::PrecacheInternal Canceled read for %s  Offset = %lld   Size = %lld"), *FileName, RequestOffset, ReadRequestSize);
+			UE_LOG(LogStreaming, Warning, TEXT("FAsyncArchive::PrecacheInternal Canceled read for %s  Offset = %lld   Size = %lld"), *FileName, RequestOffset, ReadRequestSize);
 			CancelRead();
 		}
 	}
@@ -8139,7 +8138,7 @@ bool FArchiveAsync2::PrecacheInternal(int64 RequestOffset, int64 RequestSize, bo
 	return false;
 }
 
-void FArchiveAsync2::FirstExportStarting()
+void FAsyncArchive::FirstExportStarting()
 {
 	ExportReadTime = FPlatformTime::Seconds();
 	LogItem(TEXT("Exports"));
@@ -8169,7 +8168,7 @@ void FArchiveAsync2::FirstExportStarting()
 	}
 }
 
-IAsyncReadRequest* FArchiveAsync2::MakeEventDrivenPrecacheRequest(int64 Offset, int64 BytesToRead, FAsyncFileCallBack* CompleteCallback)
+IAsyncReadRequest* FAsyncArchive::MakeEventDrivenPrecacheRequest(int64 Offset, int64 BytesToRead, FAsyncFileCallBack* CompleteCallback)
 {
 	check(GEventDrivenLoaderEnabled);
 	if (LoadPhase == ELoadPhase::WaitingForFirstExport)
@@ -8214,7 +8213,7 @@ IAsyncReadRequest* FArchiveAsync2::MakeEventDrivenPrecacheRequest(int64 Offset,
 }
 
 
-bool FArchiveAsync2::Precache(int64 RequestOffset, int64 RequestSize, bool bUseTimeLimit, bool bUseFullTimeLimit, double TickStartTime, float TimeLimit)
+bool FAsyncArchive::PrecacheWithTimeLimit(int64 RequestOffset, int64 RequestSize, bool bUseTimeLimit, bool bUseFullTimeLimit, double TickStartTime, float TimeLimit)
 {
 #if defined(ASYNC_WATCH_FILE)
 	if (FileName.Contains(TEXT(ASYNC_WATCH_FILE)) && RequestOffset == 878129)
@@ -8247,7 +8246,7 @@ bool FArchiveAsync2::Precache(int64 RequestOffset, int64 RequestSize, bool bUseT
 	return bResult;
 }
 
-bool FArchiveAsync2::Precache(int64 RequestOffset, int64 RequestSize)
+bool FAsyncArchive::Precache(int64 RequestOffset, int64 RequestSize)
 {
 	if (LoadPhase == ELoadPhase::WaitingForSize || LoadPhase == ELoadPhase::WaitingForSummary)
 	{
@@ -8261,14 +8260,14 @@ bool FArchiveAsync2::Precache(int64 RequestOffset, int64 RequestSize)
 	return PrecacheInternal(RequestOffset, RequestSize);
 }
 
-bool FArchiveAsync2::PrecacheForEvent(int64 RequestOffset, int64 RequestSize)
+bool FAsyncArchive::PrecacheForEvent(int64 RequestOffset, int64 RequestSize)
 {
 	check(int32(LoadPhase) > int32(ELoadPhase::WaitingForHeader));
 	return PrecacheInternal(RequestOffset, RequestSize, false);
 }
 
 
-void FArchiveAsync2::StartReadingHeader()
+void FAsyncArchive::StartReadingHeader()
 {
 	//LogItem(TEXT("Start Header"));
 	WaitForIntialPhases();
@@ -8283,7 +8282,7 @@ void FArchiveAsync2::StartReadingHeader()
 	}
 }
 
-void FArchiveAsync2::EndReadingHeader()
+void FAsyncArchive::EndReadingHeader()
 {
 	LogItem(TEXT("End Header"));
 	check(LoadPhase == ELoadPhase::WaitingForHeader);
@@ -8291,7 +8290,7 @@ void FArchiveAsync2::EndReadingHeader()
 	FlushPrecacheBlock();
 }
 
-bool FArchiveAsync2::ReadyToStartReadingHeader(bool bUseTimeLimit, bool bUseFullTimeLimit, double TickStartTime, float TimeLimit)
+bool FAsyncArchive::ReadyToStartReadingHeader(bool bUseTimeLimit, bool bUseFullTimeLimit, double TickStartTime, float TimeLimit)
 {
 	if (SummaryReadTime == 0.0)
 	{
@@ -8325,7 +8324,7 @@ bool FArchiveAsync2::ReadyToStartReadingHeader(bool bUseTimeLimit, bool bUseFull
 void CallSerializeHook();
 #endif
 
-void FArchiveAsync2::Serialize(void* Data, int64 Count)
+void FAsyncArchive::Serialize(void* Data, int64 Count)
 {
 	if (!Count || ArIsError)
 	{
@@ -8413,7 +8412,7 @@ void FArchiveAsync2::Serialize(void* Data, int64 Count)
 	}
 	if (BeforeBlockSize)
 	{
-		UE_CLOG(GEventDrivenLoaderEnabled, LogAsyncArchive, Warning, TEXT("FArchiveAsync2::Serialize Backwards streaming in %s  CurrentPos = %lld   BeforeBlockOffset = %lld"), *FileName, CurrentPos, BeforeBlockOffset);
+		UE_CLOG(GEventDrivenLoaderEnabled, LogAsyncArchive, Warning, TEXT("FAsyncArchive::Serialize Backwards streaming in %s  CurrentPos = %lld   BeforeBlockOffset = %lld"), *FileName, CurrentPos, BeforeBlockOffset);
 		LogItem(TEXT("Sync Before Block"), BeforeBlockOffset, BeforeBlockSize);
 		if (!PrecacheInternal(BeforeBlockOffset, BeforeBlockSize))
 		{
@@ -8481,7 +8480,7 @@ void FArchiveAsync2::Serialize(void* Data, int64 Count)
 }
 
 #if DEVIRTUALIZE_FLinkerLoad_Serialize
-void FArchiveAsync2::DiscardInlineBufferAndUpdateCurrentPos()
+void FAsyncArchive::DiscardInlineBufferAndUpdateCurrentPos()
 {
 	CurrentPos += (ActiveFPLB->StartFastPathLoadBuffer - ActiveFPLB->OriginalFastPathLoadBuffer);
 	ActiveFPLB->Reset();
@@ -8494,17 +8493,17 @@ void FArchiveAsync2::DiscardInlineBufferAndUpdateCurrentPos()
 static TAutoConsoleVariable<int32> CVarLogAsyncArchiveSerializeChurn(
 	TEXT("LogAsyncArchiveSerializeChurn.Enable"),
 	0,
-	TEXT("If > 0, then sample game thread FArchiveAsync2::Serialize calls, periodically print a report of the worst offenders."));
+	TEXT("If > 0, then sample game thread FAsyncArchive::Serialize calls, periodically print a report of the worst offenders."));
 
 static TAutoConsoleVariable<int32> CVarLogAsyncArchiveSerializeChurn_Threshhold(
 	TEXT("LogAsyncArchiveSerializeChurn.Threshhold"),
 	1000,
-	TEXT("Minimum average number of FArchiveAsync2::Serialize calls to include in the report."));
+	TEXT("Minimum average number of FAsyncArchive::Serialize calls to include in the report."));
 
 static TAutoConsoleVariable<int32> CVarLogAsyncArchiveSerializeChurn_SampleFrequency(
 	TEXT("LogAsyncArchiveSerializeChurn.SampleFrequency"),
 	1000,
-	TEXT("Number of FArchiveAsync2::Serialize calls per sample. This is used to prevent sampling from slowing the game down too much."));
+	TEXT("Number of FAsyncArchive::Serialize calls per sample. This is used to prevent sampling from slowing the game down too much."));
 
 static TAutoConsoleVariable<int32> CVarLogAsyncArchiveSerializeChurn_StackIgnore(
 	TEXT("LogAsyncArchiveSerializeChurn.StackIgnore"),
diff --git a/Engine/Source/Runtime/CoreUObject/Private/Serialization/AsyncLoadingPrivate.h b/Engine/Source/Runtime/CoreUObject/Private/Serialization/AsyncLoadingPrivate.h
--- a/Engine/Source/Runtime/CoreUObject/Private/Serialization/AsyncLoadingPrivate.h
+++ b/Engine/Source/Runtime/CoreUObject/Private/Serialization/AsyncLoadingPrivate.h
@@ -8,7 +8,7 @@
 
 DECLARE_LOG_CATEGORY_EXTERN(LogLoadingDev, Fatal, All);
 
-class COREUOBJECT_API FArchiveAsync2 final: public FArchive
+class COREUOBJECT_API FAsyncArchive final: public FArchive
 {
 public:
 	enum class ELoadPhase
@@ -20,10 +20,10 @@ public:
 		ProcessingExports,
 	};
 
-	FArchiveAsync2(const TCHAR* InFileName
-		, TFunction<void()>&& InSummaryReadyCallback
-		);
-	virtual ~FArchiveAsync2();
+	FAsyncArchive(const TCHAR* InFileName, TFunction<void()>&& InSummaryReadyCallback);
+	virtual ~FAsyncArchive ();
+
+	/** Archive overrides */
 	virtual bool Close() override;
 	virtual bool SetCompressionMap(TArray<FCompressedChunk>* CompressedChunks, ECompressionFlags CompressionFlags) override;
 	virtual bool Precache(int64 PrecacheOffset, int64 PrecacheSize) override;
@@ -44,16 +44,14 @@ public:
 		return FileName;
 	}
 
-	bool Precache(int64 PrecacheOffset, int64 PrecacheSize, bool bUseTimeLimit, bool bUseFullTimeLimit, double TickStartTime, float TimeLimit);
+	/** AsyncArchive interface */
+	bool PrecacheWithTimeLimit(int64 PrecacheOffset, int64 PrecacheSize, bool bUseTimeLimit, bool bUseFullTimeLimit, double TickStartTime, float TimeLimit);
 	bool PrecacheForEvent(int64 PrecacheOffset, int64 PrecacheSize);
 	void FlushPrecacheBlock();
 	bool ReadyToStartReadingHeader(bool bUseTimeLimit, bool bUseFullTimeLimit, double TickStartTime, float TimeLimit);
 	void StartReadingHeader();
 	void EndReadingHeader();
-	void FirstExportStarting();
-
 	IAsyncReadRequest* MakeEventDrivenPrecacheRequest(int64 Offset, int64 BytesToRead, FAsyncFileCallBack* CompleteCallback);
-
 	void LogItem(const TCHAR* Item, int64 Offset = 0, int64 Size = 0, double StartTime = 0.0);
 
 	bool IsCookedForEDLInEditor() const
@@ -71,7 +69,7 @@ private:
 	void SetPosAndUpdatePrecacheBuffer(int64 Pos);
 
 #endif
-
+	void FirstExportStarting();
 	bool WaitRead(float TimeLimit = 0.0f);
 	void CompleteRead();
 	void CancelRead();
diff --git a/Engine/Source/Runtime/CoreUObject/Private/Serialization/BulkData.cpp b/Engine/Source/Runtime/CoreUObject/Private/Serialization/BulkData.cpp
--- a/Engine/Source/Runtime/CoreUObject/Private/Serialization/BulkData.cpp
+++ b/Engine/Source/Runtime/CoreUObject/Private/Serialization/BulkData.cpp
@@ -1538,7 +1538,7 @@ void FUntypedBulkData::LoadDataIntoMemory( void* Dest )
 	checkf( AttachedAr, TEXT( "Attempted to load bulk data without an attached archive. Most likely the bulk data was loaded twice on console, which is not supported" ) );
 
 	FArchive* BulkDataArchive = nullptr;
-	if (Linker && Linker->GetFArchiveAsync2Loader() && Linker->GetFArchiveAsync2Loader()->IsCookedForEDLInEditor() &&
+	if (Linker && Linker->GetAsyncLoader() && Linker->GetAsyncLoader()->IsCookedForEDLInEditor() &&
 		(BulkDataFlags & BULKDATA_PayloadInSeperateFile))
 	{
 		// The attached archive is a package cooked for EDL loaded in the editor so the actual bulk data sits in a separate ubulk file.
diff --git a/Engine/Source/Runtime/CoreUObject/Private/UObject/LinkerLoad.cpp b/Engine/Source/Runtime/CoreUObject/Private/UObject/LinkerLoad.cpp
--- a/Engine/Source/Runtime/CoreUObject/Private/UObject/LinkerLoad.cpp
+++ b/Engine/Source/Runtime/CoreUObject/Private/UObject/LinkerLoad.cpp
@@ -817,7 +817,7 @@ FLinkerLoad::FLinkerLoad(UPackage* InParent, const TCHAR* InFilename, uint32 InL
 , TemplateForGetArchetypeFromLoader(nullptr)
 , bForceSimpleIndexToObject(false)
 , bLockoutLegacyOperations(false)
-, bLoaderIsFArchiveAsync2(false)
+, bIsAsyncLoader(false)
 , StructuredArchive(nullptr)
 , StructuredArchiveFormatter(nullptr)
 , Loader(nullptr)
@@ -1018,12 +1018,12 @@ FLinkerLoad::ELinkerStatus FLinkerLoad::CreateLoader(
 			else
 #endif
 			{
-				// If want to be able to load cooked data in the editor we need to use FArchiveAsync2 which supports EDL cooked packages,
+				// If want to be able to load cooked data in the editor we need to use FAsyncArchive which supports EDL cooked packages,
 				// otherwise the generic file reader is faster in the editor so use that
-				bool bCanUseFArchiveAsync2 = FPlatformProperties::RequiresCookedData() || GAllowCookedDataInEditorBuilds;
-				if (bCanUseFArchiveAsync2)
+				bool bCanUseAsyncLoader = FPlatformProperties::RequiresCookedData() || GAllowCookedDataInEditorBuilds;
+				if (bCanUseAsyncLoader)
 				{
-					Loader = new FArchiveAsync2(*Filename
+					Loader = new FAsyncArchive(*Filename
 							, GEventDrivenLoaderEnabled ? Forward<TFunction<void()>>(InSummaryReadyCallback) : TFunction<void()>([]() {})
 						);
 				}
@@ -1074,7 +1074,7 @@ FLinkerLoad::ELinkerStatus FLinkerLoad::CreateLoader(
 				}
 				else
 				{
-					bLoaderIsFArchiveAsync2 = bCanUseFArchiveAsync2;
+					bIsAsyncLoader = bCanUseAsyncLoader;
 				}
 			}
 		} 
@@ -1105,9 +1105,9 @@ FLinkerLoad::ELinkerStatus FLinkerLoad::CreateLoader(
 		bool bExecuteNextStep = true;
 		if( bHasSerializedPackageFileSummary == false )
 		{
-			if (bLoaderIsFArchiveAsync2)
+			if (bIsAsyncLoader)
 			{
-				bExecuteNextStep = GetFArchiveAsync2Loader()->ReadyToStartReadingHeader(bUseTimeLimit, bUseFullTimeLimit, TickStartTime, TimeLimit);
+				bExecuteNextStep = GetAsyncLoader()->ReadyToStartReadingHeader(bUseTimeLimit, bUseFullTimeLimit, TickStartTime, TimeLimit);
 			}
 			else
 			{
@@ -1150,9 +1150,9 @@ FLinkerLoad::ELinkerStatus FLinkerLoad::SerializePackageFileSummary()
 			UE_LOG(LogLinker, Warning, TEXT("The file '%s' contains unrecognizable data, check that it is of the expected type."), *Filename);
 			return LINKER_Failed;
 		}
-		if (bLoaderIsFArchiveAsync2)
+		if (bIsAsyncLoader)
 		{
-			GetFArchiveAsync2Loader()->StartReadingHeader();
+			GetAsyncLoader()->StartReadingHeader();
 		}
 
 #if WITH_EDITOR
@@ -1417,9 +1417,9 @@ FLinkerLoad::ELinkerStatus FLinkerLoad::SerializeNameMap()
 			bool bFinishedPrecaching = true;
 
 			// Precache name, import and export map.
-			if (bLoaderIsFArchiveAsync2)
+			if (bIsAsyncLoader)
 			{
-				bFinishedPrecaching = GetFArchiveAsync2Loader()->ReadyToStartReadingHeader(bUseTimeLimit, bUseFullTimeLimit, TickStartTime, TimeLimit);
+				bFinishedPrecaching = GetAsyncLoader()->ReadyToStartReadingHeader(bUseTimeLimit, bUseFullTimeLimit, TickStartTime, TimeLimit);
 				check(!GEventDrivenLoaderEnabled || bFinishedPrecaching || !EVENT_DRIVEN_ASYNC_LOAD_ACTIVE_AT_RUNTIME);
 			}
 			else
@@ -2086,9 +2086,9 @@ FLinkerLoad::ELinkerStatus FLinkerLoad::FinalizeCreation()
 			}
 		}
 
-		if (bLoaderIsFArchiveAsync2)
+		if (bIsAsyncLoader)
 		{
-			GetFArchiveAsync2Loader()->EndReadingHeader();
+			GetAsyncLoader()->EndReadingHeader();
 		}
 
 		if ( !(LoadFlags & LOAD_NoVerify) )
@@ -3515,14 +3515,14 @@ void FLinkerLoad::Preload( UObject* Object )
 				// is stored
 				Seek(Export.SerialOffset);
 
-				FArchiveAsync2* FAA2 = GetFArchiveAsync2Loader();
+				FAsyncArchive* AsyncLoader = GetAsyncLoader();
 
 				{
 					SCOPE_CYCLE_COUNTER(STAT_LinkerPrecache);
 					// tell the file reader to read the raw data from disk
-					if (FAA2)
+					if (AsyncLoader)
 					{
-						bool bReady = FAA2->Precache(Export.SerialOffset, Export.SerialSize, bUseTimeLimit, bUseFullTimeLimit, TickStartTime, TimeLimit);
+						bool bReady = AsyncLoader->PrecacheWithTimeLimit(Export.SerialOffset, Export.SerialSize, bUseTimeLimit, bUseFullTimeLimit, TickStartTime, TimeLimit);
 						UE_CLOG(!(bReady || !bUseTimeLimit || !FPlatformProperties::RequiresCookedData()), LogLinker, Warning, TEXT("Hitch on async loading of %s; this export was not properly precached."), *Object->GetFullName());
 					}
 					else
diff --git a/Engine/Source/Runtime/CoreUObject/Public/UObject/LinkerLoad.h b/Engine/Source/Runtime/CoreUObject/Public/UObject/LinkerLoad.h
--- a/Engine/Source/Runtime/CoreUObject/Public/UObject/LinkerLoad.h
+++ b/Engine/Source/Runtime/CoreUObject/Public/UObject/LinkerLoad.h
@@ -70,7 +70,7 @@ struct FScopedCreateImportCounter
 /**
  * Handles loading Unreal package files, including reading UObject data from disk.
  */
-class FArchiveAsync2;
+class FAsyncArchive;
 
 class FLinkerLoad 
 #if !WITH_EDITOR
@@ -127,11 +127,11 @@ public:
 	bool					bForceSimpleIndexToObject;
 	bool					bLockoutLegacyOperations;
 
-	/** True if Loader is FArchiveAsync2  */
-	bool					bLoaderIsFArchiveAsync2;
-	FORCEINLINE FArchiveAsync2* GetFArchiveAsync2Loader()
+	/** True if Loader is FAsyncArchive  */
+	bool					bIsAsyncLoader;
+	FORCEINLINE FAsyncArchive* GetAsyncLoader()
 	{
-		return bLoaderIsFArchiveAsync2 ? (FArchiveAsync2*)Loader : nullptr;
+		return bIsAsyncLoader ? (FAsyncArchive*)Loader : nullptr;
 	}
 
 private:
