Fixed removal of bCookedForEDLInEditor default value

diff --git a/Engine/Source/Runtime/CoreUObject/Private/Serialization/AsyncLoading.cpp b/Engine/Source/Runtime/CoreUObject/Private/Serialization/AsyncLoading.cpp
--- a/Engine/Source/Runtime/CoreUObject/Private/Serialization/AsyncLoading.cpp
+++ b/Engine/Source/Runtime/CoreUObject/Private/Serialization/AsyncLoading.cpp
@@ -7620,6 +7620,7 @@ FAsyncArchive::FAsyncArchive(const TCHAR* InFileName, TFunction<void()>&& InSumm
 	, HeaderSize(0)
 	, HeaderSizeWhenReadingExportsFromSplitFile(0)
 	, LoadPhase(ELoadPhase::WaitingForSize)
+	, bCookedForEDLInEditor(false)
 	, FileName(InFileName)
 	, OpenTime(FPlatformTime::Seconds())
 	, SummaryReadTime(0.0)
@@ -7761,7 +7762,7 @@ void FAsyncArchive::FlushPrecacheBlock()
 		DEC_MEMORY_STAT_BY(STAT_FAsyncArchiveMem, PrecacheEndPos - PrecacheStartPos);
 		FMemory::Free(PrecacheBuffer);
 #if USE_DETAILED_FASYNCARCHIVE_MEMORY_TRACKING
-		GAAsyncArchieMemTracker.Deallocate(FileName, PrecacheEndPos - PrecacheStartPos);
+		GAsyncArchieMemTracker.Deallocate(FileName, PrecacheEndPos - PrecacheStartPos);
 #endif
 		check(!GEventDrivenLoaderEnabled || LoadPhase > ELoadPhase::WaitingForHeader);
 	}
