Fix asset registry disk size not accurately reflecting disk size numbers.
Let files from CookAdditionalFiles contribute to package total file size, and package hash when enabled.

diff --git a/Engine/Source/Runtime/CoreUObject/Private/UObject/SavePackage.cpp b/Engine/Source/Runtime/CoreUObject/Private/UObject/SavePackage.cpp
--- a/Engine/Source/Runtime/CoreUObject/Private/UObject/SavePackage.cpp
+++ b/Engine/Source/Runtime/CoreUObject/Private/UObject/SavePackage.cpp
@@ -5461,6 +5461,9 @@ FSavePackageResultStruct UPackage::Save(UPackage* InOuter, UObject* Base, EObjec
 					Linker->StartScriptSHAGeneration();
 				}
 
+#if WITH_EDITOR
+				TArray<FLargeMemoryWriter, TInlineAllocator<4>> AdditionalFilesFromExports;
+#endif
 				{
 					COOK_STAT(FScopedDurationTimer SaveTimer(SavePackageStats::SerializeExportsTimeSec));
 #if WITH_EDITOR
@@ -5529,9 +5532,14 @@ FSavePackageResultStruct UPackage::Save(UPackage* InOuter, UObject* Base, EObjec
 								}
 
 #if WITH_EDITOR
-								if (Linker->IsCooking())
+								if (bIsCooking)
 								{
-									Export.Object->CookAdditionalFiles(Filename, Linker->CookingTarget());
+									Export.Object->CookAdditionalFiles(Filename, TargetPlatform,
+										[&AdditionalFilesFromExports](const TCHAR* Filename, void* Data, int64 Size)
+									{
+										FLargeMemoryWriter& Writer = AdditionalFilesFromExports.Emplace_GetRef(0, true, Filename);
+										Writer.Serialize(Data, Size);
+									});
 								}
 #endif
 							}
@@ -5776,6 +5784,30 @@ FSavePackageResultStruct UPackage::Save(UPackage* InOuter, UObject* Base, EObjec
 				}
 
 				Linker->BulkDataToAppend.Empty();
+
+#if WITH_EDITOR
+				if (bIsCooking && AdditionalFilesFromExports.Num() > 0)
+				{
+					const bool bWriteFileToDisk = !bDiffing;
+					for (FLargeMemoryWriter& Writer : AdditionalFilesFromExports)
+					{
+						const int64 Size = Writer.TotalSize();
+						TotalPackageSizeUncompressed += Size;
+						if (bComputeHash)
+						{
+							CookedPackageHash.Update(Writer.GetData(), Size);
+						}
+						if (bWriteFileToDisk)
+						{
+							FLargeMemoryPtr DataPtr(Writer.GetData());
+							Writer.ReleaseOwnership();
+							AsyncWriteFile(MoveTemp(DataPtr), Size, *Writer.GetArchiveName(), FDateTime::MinValue(), false);
+						}
+					}
+					AdditionalFilesFromExports.Empty();
+				}
+#endif
+
 			
 				// write the package post tag
 				if (!bTextFormat)
diff --git a/Engine/Source/Runtime/CoreUObject/Public/UObject/Object.h b/Engine/Source/Runtime/CoreUObject/Public/UObject/Object.h
--- a/Engine/Source/Runtime/CoreUObject/Public/UObject/Object.h
+++ b/Engine/Source/Runtime/CoreUObject/Public/UObject/Object.h
@@ -929,7 +929,43 @@ public:
 	 * @param	PackageFilename full path to the package that this object is being saved to on disk
 	 * @param	TargetPlatform	target platform to cook additional files for
 	 */
-	virtual void CookAdditionalFiles( const TCHAR* PackageFilename, const ITargetPlatform* TargetPlatform ) { }
+	UE_DEPRECATED(4.23, "Use the new CookAdditionalFilesOverride that provides a function to write the files")
+	virtual void CookAdditionalFiles(const TCHAR* PackageFilename, const ITargetPlatform* TargetPlatform) { }
+
+	/**
+	 * Called during cook to allow objects to generate additional cooked files alongside their cooked package.
+	 * @note Implement CookAdditionalFilesOverride to define sub class behavior.
+	 *
+	 * @param	PackageFilename full path to the package that this object is being saved to on disk
+	 * @param	TargetPlatform	target platform to cook additional files for
+	 * @param	WriteAdditionalFile function for writing the additional files
+	 */
+	void CookAdditionalFiles(const TCHAR* PackageFilename, const ITargetPlatform* TargetPlatform,
+		TFunctionRef<void(const TCHAR* Filename, void* Data, int64 Size)> WriteAdditionalFile)
+	{
+		CookAdditionalFilesOverride(PackageFilename, TargetPlatform, WriteAdditionalFile);
+	}
+
+private:
+	/**
+	 * Called during cook to allow objects to generate additional cooked files alongside their cooked package.
+	 * Files written using the provided function will be handled as part of the saved cooked package
+	 * and contribute to package total file size, and package hash when enabled.
+	 * @note These should typically match the name of the package, but with a different extension.
+	 *
+	 * @param	PackageFilename full path to the package that this object is being saved to on disk
+	 * @param	TargetPlatform	target platform to cook additional files for
+	 * @param	WriteAdditionalFile function for writing the additional files
+	 */
+	virtual void CookAdditionalFilesOverride(const TCHAR* PackageFilename, const ITargetPlatform* TargetPlatform,
+		TFunctionRef<void(const TCHAR* Filename, void* Data, int64 Size)> WriteAdditionalFile)
+	{
+		PRAGMA_DISABLE_DEPRECATION_WARNINGS;
+		CookAdditionalFiles(PackageFilename, TargetPlatform);
+		PRAGMA_ENABLE_DEPRECATION_WARNINGS;
+	}
+
+public:
 #endif
 	/**
 	 * Determine if this object has SomeObject in its archetype chain.
diff --git a/Engine/Source/Runtime/Engine/Classes/Engine/FontFace.h b/Engine/Source/Runtime/Engine/Classes/Engine/FontFace.h
--- a/Engine/Source/Runtime/Engine/Classes/Engine/FontFace.h
+++ b/Engine/Source/Runtime/Engine/Classes/Engine/FontFace.h
@@ -31,7 +31,10 @@ public:
 	virtual void PostEditChangeProperty(FPropertyChangedEvent& PropertyChangedEvent) override;
 	virtual void PostEditUndo() override;
 	virtual void GetAssetRegistryTags(TArray<FAssetRegistryTag>& OutTags) const override;
-	virtual void CookAdditionalFiles(const TCHAR* PackageFilename, const ITargetPlatform* TargetPlatform) override;
+private:
+	virtual void CookAdditionalFilesOverride(const TCHAR* PackageFilename, const ITargetPlatform* TargetPlatform,
+		TFunctionRef<void(const TCHAR* Filename, void* Data, int64 Size)> WriteAdditionalFile);
+public:
 #endif // WITH_EDITOR
 	//~ End UObject interface
 
diff --git a/Engine/Source/Runtime/Engine/Private/FontFace.cpp b/Engine/Source/Runtime/Engine/Private/FontFace.cpp
--- a/Engine/Source/Runtime/Engine/Private/FontFace.cpp
+++ b/Engine/Source/Runtime/Engine/Private/FontFace.cpp
@@ -122,22 +122,26 @@ void UFontFace::GetAssetRegistryTags(TArray<FAssetRegistryTag>& OutTags) const
 	OutTags.Add(FAssetRegistryTag(SourceFileTagName(), ImportInfo.ToJson(), FAssetRegistryTag::TT_Hidden));
 }
 
-void UFontFace::CookAdditionalFiles(const TCHAR* PackageFilename, const ITargetPlatform* TargetPlatform)
+void UFontFace::CookAdditionalFilesOverride(const TCHAR* PackageFilename, const ITargetPlatform* TargetPlatform,
+	TFunctionRef<void(const TCHAR* Filename, void* Data, int64 Size)> WriteAdditionalFile)
 {
-	Super::CookAdditionalFiles(PackageFilename, TargetPlatform);
-
 	if (LoadingPolicy != EFontLoadingPolicy::Inline)
 	{
-		// Iterative COTF can't handle the .ufont files generated when this UFontFace is within a UFont asset (rather than its own asset) so emit a warning about that
+		// Iterative COTF can't handle the .ufont files generated when this UFontFace is within a UFont asset (rather than its own asset)
 		if (UFont* OuterFont = GetTypedOuter<UFont>())
 		{
-			UE_LOG(LogFontFace, Warning, TEXT("The font asset '%s' contains nested font faces which can cause issues for iterative cook-on-the-fly. Please edit the font asset and split the font faces into their own assets."), *OuterFont->GetPathName());
+			UE_LOG(LogFontFace, Warning, 
+				TEXT("The font asset '%s' contains nested font faces which can cause issues for iterative cook-on-the-fly."
+					"Please edit the font asset and split the font faces into their own assets."),
+				*OuterFont->GetPathName());
 		}
 
 		// We replace the package name with the cooked font face name
 		// Note: This must match the replacement logic in UFontFace::GetCookedFilename
 		const FString CookedFontFilename = FPaths::GetPath(PackageFilename) / GetName() + TEXT(".ufont");
-		FFileHelper::SaveArrayToFile(FontFaceData->GetData(), *CookedFontFilename);
+		const TArray<uint8>& Data = FontFaceData->GetData();
+		const int32 NumBytes = Data.Num() * Data.GetTypeSize();
+		WriteAdditionalFile(*CookedFontFilename, (void*)Data.GetData(), NumBytes);
 	}
 }
 #endif // WITH_EDITOR
