Restoring the global GCStartTime to fix a hack that was made for one of the mobile platforms and fix CIS

diff --git a/Engine/Source/Runtime/CoreUObject/Private/UObject/GarbageCollection.cpp b/Engine/Source/Runtime/CoreUObject/Private/UObject/GarbageCollection.cpp
--- a/Engine/Source/Runtime/CoreUObject/Private/UObject/GarbageCollection.cpp
+++ b/Engine/Source/Runtime/CoreUObject/Private/UObject/GarbageCollection.cpp
@@ -43,6 +43,8 @@ FThreadSafeCounter		GObjectCountDuringLastMarkPhase;
 bool GObjIncrementalPurgeIsInProgress = false;
 /** Whether GC is currently routing BeginDestroy to objects										*/
 bool GObjUnhashUnreachableIsInProgress = false;
+/** Time the GC started, needs to be reset on return from being in the background on some OSs */
+double GCStartTime = 0.;
 /** Whether FinishDestroy has already been routed to all unreachable objects. */
 static bool GObjFinishDestroyHasBeenRoutedToAllObjects	= false;
 /** 
@@ -1269,7 +1271,7 @@ struct FConditionalGCLock
 	}
 };
 
-static bool IncrementalDestroyGarbage(bool bUseTimeLimit, float TimeLimit, double StartTime);
+static bool IncrementalDestroyGarbage(bool bUseTimeLimit, float TimeLimit);
 
 /**
  * Incrementally purge garbage by deleting all unreferenced objects after routing Destroy.
@@ -1322,7 +1324,7 @@ void IncrementalPurgeGarbage(bool bUseTimeLimit, float TimeLimit)
 	} ResetPurgeProgress(bCompleted);
 
 	// Keep track of start time to enforce time limit unless bForceFullPurge is true;
-	double StartTime = FPlatformTime::Seconds();
+	GCStartTime = FPlatformTime::Seconds();
 
 	bool bTimeLimitReached = false;
 	if (GUnrechableObjectIndex < GUnreachableObjects.Num())
@@ -1339,11 +1341,11 @@ void IncrementalPurgeGarbage(bool bUseTimeLimit, float TimeLimit)
 
 	if (!bTimeLimitReached)
 	{
-		bCompleted = IncrementalDestroyGarbage(bUseTimeLimit, TimeLimit, StartTime);
+		bCompleted = IncrementalDestroyGarbage(bUseTimeLimit, TimeLimit);
 	}
 }
 
-bool IncrementalDestroyGarbage(bool bUseTimeLimit, float TimeLimit, double StartTime)
+bool IncrementalDestroyGarbage(bool bUseTimeLimit, float TimeLimit)
 {
 	const bool bMultithreadedPurge = !ShouldForceSingleThreadedGC() && GMultithreadedDestructionEnabled;
 	if (!GAsyncPurge)
@@ -1432,7 +1434,7 @@ bool IncrementalDestroyGarbage(bool bUseTimeLimit, float TimeLimit, double Start
 
 			// Only check time limit every so often to avoid calling FPlatformTime::Seconds too often.
 			const bool bPollTimeLimit = ((TimeLimitTimePollCounter++) % TimeLimitEnforcementGranularityForDestroy == 0);
-			if( bUseTimeLimit && bPollTimeLimit && ((FPlatformTime::Seconds() - StartTime) > TimeLimit) )
+			if( bUseTimeLimit && bPollTimeLimit && ((FPlatformTime::Seconds() - GCStartTime) > TimeLimit) )
 			{
 				bTimeLimitReached = true;
 				break;
@@ -1489,7 +1491,7 @@ bool IncrementalDestroyGarbage(bool bUseTimeLimit, float TimeLimit, double Start
 
 					// Only check time limit every so often to avoid calling FPlatformTime::Seconds too often.
 					const bool bPollTimeLimit = ((TimeLimitTimePollCounter++) % TimeLimitEnforcementGranularityForDestroy == 0);
-					if( bUseTimeLimit && bPollTimeLimit && ((FPlatformTime::Seconds() - StartTime) > TimeLimit) )
+					if( bUseTimeLimit && bPollTimeLimit && ((FPlatformTime::Seconds() - GCStartTime) > TimeLimit) )
 					{
 						bTimeLimitReached = true;
 						break;
@@ -1515,7 +1517,7 @@ bool IncrementalDestroyGarbage(bool bUseTimeLimit, float TimeLimit, double Start
 #endif
 						// Check if we spent too much time on waiting for FinishDestroy without making any progress
 						if (LastLoopObjectsPendingDestructionCount == GGCObjectsPendingDestructionCount && bPollTimeLimit &&
-							((FPlatformTime::Seconds() - StartTime) > MaxTimeForFinishDestroy))
+							((FPlatformTime::Seconds() - GCStartTime) > MaxTimeForFinishDestroy))
 						{
 							UE_LOG(LogGarbage, Warning, TEXT("Spent more than %.2fs on routing FinishDestroy to objects (objects in queue: %d)"), MaxTimeForFinishDestroy, GGCObjectsPendingDestructionCount);
 							UObject* LastObjectNotReadyForFinishDestroy = nullptr;
@@ -1587,7 +1589,7 @@ bool IncrementalDestroyGarbage(bool bUseTimeLimit, float TimeLimit, double Start
 			GObjCurrentPurgeObjectIndexNeedsReset = false;
 		}
 
-		GAsyncPurge->TickPurge(bUseTimeLimit, TimeLimit, StartTime);
+		GAsyncPurge->TickPurge(bUseTimeLimit, TimeLimit, GCStartTime);
 
 		if (GAsyncPurge->IsFinished())
 		{
