LoadTimeTrace: Fixed missing symbol compile errors when compiling with -disableunity -nopch and LOADTIMEPROFILERTRACE_ENABLED 1.

diff --git a/Engine/Source/Runtime/CoreUObject/Private/Serialization/LoadTimeTrace.cpp b/Engine/Source/Runtime/CoreUObject/Private/Serialization/LoadTimeTrace.cpp
--- a/Engine/Source/Runtime/CoreUObject/Private/Serialization/LoadTimeTrace.cpp
+++ b/Engine/Source/Runtime/CoreUObject/Private/Serialization/LoadTimeTrace.cpp
@@ -1,15 +1,17 @@
 // Copyright 1998-2019 Epic Games, Inc. All Rights Reserved.
 
 #include "Serialization/LoadTimeTrace.h"
+
+#if LOADTIMEPROFILERTRACE_ENABLED
+
 #include "LoadTimeTracePrivate.h"
 #include "Trace/Trace.h"
 #include "Misc/CString.h"
 #include "HAL/PlatformTime.h"
 #include "HAL/PlatformTLS.h"
+#include "UObject/Object.h"
 #include "UObject/UObjectGlobals.h"
 
-#if LOADTIMEPROFILERTRACE_ENABLED
-
 UE_TRACE_EVENT_BEGIN(LoadTime, StartAsyncLoading)
 	UE_TRACE_EVENT_FIELD(uint64, Cycle)
 	UE_TRACE_EVENT_FIELD(uint32, ThreadId)
diff --git a/Engine/Source/Runtime/CoreUObject/Private/Serialization/LoadTimeTracePrivate.h b/Engine/Source/Runtime/CoreUObject/Private/Serialization/LoadTimeTracePrivate.h
--- a/Engine/Source/Runtime/CoreUObject/Private/Serialization/LoadTimeTracePrivate.h
+++ b/Engine/Source/Runtime/CoreUObject/Private/Serialization/LoadTimeTracePrivate.h
@@ -6,6 +6,8 @@
 
 #if LOADTIMEPROFILERTRACE_ENABLED
 
+#include "UObject/UObjectGlobals.h"
+
 class FArchive;
 struct FAsyncPackage;
 class FLinkerLoad;
