Fix for unitialized member access in FExportReferenceSorter (flagged by Valgrind)

diff --git a/Engine/Source/Runtime/CoreUObject/Private/UObject/SavePackage.cpp b/Engine/Source/Runtime/CoreUObject/Private/UObject/SavePackage.cpp
--- a/Engine/Source/Runtime/CoreUObject/Private/UObject/SavePackage.cpp
+++ b/Engine/Source/Runtime/CoreUObject/Private/UObject/SavePackage.cpp
@@ -2292,7 +2292,6 @@ public:
 	 * Constructor
 	 */
 	FExportReferenceSorter()
-		: FArchiveUObject(), CurrentInsertIndex(INDEX_NONE), CoreReferencesOffset(INDEX_NONE), bIgnoreFieldReferences(false), CurrentClass(nullptr)
 	{
 		ArIsObjectReferenceCollector = true;
 		this->SetIsPersistent(true);
@@ -2629,12 +2628,12 @@ private:
 	/**
 	 * The index into the ReferencedObjects array to insert new objects
 	 */
-	int32 CurrentInsertIndex;
+	int32 CurrentInsertIndex = INDEX_NONE;
 
 	/**
 	 * The index into the ReferencedObjects array for the first object not referenced by one of the core classes
 	 */
-	int32 CoreReferencesOffset;
+	int32 CoreReferencesOffset = INDEX_NONE;
 
 	/**
 	 * The classes which are pre-added to the array of ReferencedObjects.  Used for resolving a number of circular dependecy issues between
@@ -2688,7 +2687,7 @@ private:
 	/**
 	 * Controls whether to process UField objects encountered during serialization of an object.
 	 */
-	bool bIgnoreFieldReferences;
+	bool bIgnoreFieldReferences = false;
 
 	/**
 	 * The UClass currently being processed.  This is used to prevent serialization of a UStruct's Children member causing other fields of the same class to be processed too early due
@@ -2697,10 +2696,10 @@ private:
 	 * the "second" function would be created first, which would end up force-loading the struct.  This would cause an unacceptible seek because the struct appears later in the export list, thus
 	 * hasn't been created yet.
 	 */
-	UClass* CurrentClass;
+	UClass* CurrentClass = nullptr;
 
 	/** Package to constrain checks to */
-	UPackage* PackageToSort;
+	UPackage* PackageToSort = nullptr;
 };
 
 /**
