Fixed a bug when creating GC clusters during Async Loading times out and the remaining clusters don't get created

diff --git a/Engine/Source/Runtime/CoreUObject/Private/Serialization/AsyncLoading.cpp b/Engine/Source/Runtime/CoreUObject/Private/Serialization/AsyncLoading.cpp
--- a/Engine/Source/Runtime/CoreUObject/Private/Serialization/AsyncLoading.cpp
+++ b/Engine/Source/Runtime/CoreUObject/Private/Serialization/AsyncLoading.cpp
@@ -6916,8 +6916,17 @@ EAsyncPackageState::Type FAsyncPackage::PostLoadDeferredObjects(double InTickSta
 			}
 		}
 		::IsTimeLimitExceeded(InTickStartTime, bInUseTimeLimit, InOutTimeLimit, LastTypeOfWorkPerformed, LastObjectWorkWasPerformedOn);
-		Result = (DeferredFinalizeIndex == DeferredFinalizeObjects.Num()) ? EAsyncPackageState::Complete : EAsyncPackageState::TimeOut;
-
+		if (DeferredFinalizeIndex == DeferredFinalizeObjects.Num())
+		{
+			DeferredFinalizeIndex = 0;
+			DeferredFinalizeObjects.Reset();
+			Result = EAsyncPackageState::Complete;
+		}
+		else
+		{
+			EAsyncPackageState::TimeOut;
+		}
+	
 		// Mark package as having been fully loaded and update load time.
 		if (Result == EAsyncPackageState::Complete && LinkerRoot && !bLoadHasFailed)
 		{
@@ -6955,9 +6964,9 @@ EAsyncPackageState::Type FAsyncPackage::CreateClusters(double InTickStartTime, b
 	}
 
 	EAsyncPackageState::Type Result;
-	if (DeferredFinalizeIndex == DeferredFinalizeObjects.Num())
+	if (DeferredClusterIndex == DeferredClusterObjects.Num())
 	{
-		DeferredFinalizeIndex = 0;
+		DeferredClusterIndex = 0;
 		DeferredClusterObjects.Reset();
 		Result = EAsyncPackageState::Complete;
 	}
