Fixed typo where Result wasn't getting assigned an EAsyncPackageState value correctly

diff --git a/Engine/Source/Runtime/CoreUObject/Private/Serialization/AsyncLoading.cpp b/Engine/Source/Runtime/CoreUObject/Private/Serialization/AsyncLoading.cpp
--- a/Engine/Source/Runtime/CoreUObject/Private/Serialization/AsyncLoading.cpp
+++ b/Engine/Source/Runtime/CoreUObject/Private/Serialization/AsyncLoading.cpp
@@ -6924,7 +6924,7 @@ EAsyncPackageState::Type FAsyncPackage::PostLoadDeferredObjects(double InTickSta
 		}
 		else
 		{
-			EAsyncPackageState::TimeOut;
+			Result = EAsyncPackageState::TimeOut;
 		}
 	
 		// Mark package as having been fully loaded and update load time.
