Optimize TProperty_WithEqualityAndSerializer (UMapProperty::SerializeItem) by removing temporary string allocs by avoiding copy while comparing properties.

diff --git a/Engine/Source/Runtime/CoreUObject/Public/UObject/UnrealType.h b/Engine/Source/Runtime/CoreUObject/Public/UObject/UnrealType.h
--- a/Engine/Source/Runtime/CoreUObject/Public/UObject/UnrealType.h
+++ b/Engine/Source/Runtime/CoreUObject/Public/UObject/UnrealType.h
@@ -1115,7 +1115,8 @@ public:
 	// UProperty interface.
 	virtual bool Identical( const void* A, const void* B, uint32 PortFlags=0 ) const override
 	{
-		return TTypeFundamentals::GetPropertyValue(A) == TTypeFundamentals::GetOptionalPropertyValue(B);
+		// RHS is the same as TTypeFundamentals::GetOptionalPropertyValue(B) but avoids an unnecessary copy of B
+		return TTypeFundamentals::GetPropertyValue(A) == (B ? TTypeFundamentals::GetPropertyValue(B) : TTypeFundamentals::GetDefaultPropertyValue());
 	}
 	virtual void SerializeItem(FStructuredArchive::FSlot Slot, void* Value, void const* Defaults) const override
 	{
