Optimize UProperty::FindRedirectedPropertyName by removing temporary string allocation and FName initialization when there is only one outer

diff --git a/Engine/Source/Runtime/CoreUObject/Private/UObject/Property.cpp b/Engine/Source/Runtime/CoreUObject/Private/UObject/Property.cpp
--- a/Engine/Source/Runtime/CoreUObject/Private/UObject/Property.cpp
+++ b/Engine/Source/Runtime/CoreUObject/Private/UObject/Property.cpp
@@ -1358,9 +1358,11 @@ FName UProperty::FindRedirectedPropertyName(UStruct* ObjectStruct, FName OldName
 
 	// ObjectStruct may be a nested struct, so extract path
 	UPackage* StructPackage = ObjectStruct->GetOutermost();
-	FString OuterPath = ObjectStruct->GetPathName(StructPackage);
+	FName PackageName = StructPackage->GetFName();
+	// Avoid GetPathName string allocation and FName initialization when there is only one outer
+	FName OuterName = (StructPackage == ObjectStruct->GetOuter()) ? ObjectStruct->GetFName() : FName(*ObjectStruct->GetPathName(StructPackage));
 
-	FCoreRedirectObjectName OldRedirectName(OldName, FName(*OuterPath), StructPackage->GetFName());
+	FCoreRedirectObjectName OldRedirectName(OldName, OuterName, PackageName);
 	FCoreRedirectObjectName NewRedirectName = FCoreRedirects::GetRedirectedName(ECoreRedirectFlags::Type_Property, OldRedirectName);
 
 	if (NewRedirectName != OldRedirectName)
@@ -1420,4 +1422,4 @@ UProperty* UStruct::FindPropertyByName(FName InName) const
 	}
 
 	return NULL;
-}
\ No newline at end of file
+}
