Add FSoftObjectPathThreadContexts for properties being serialized during object config loading
Mark editor startup map as an editor only property

diff --git a/Engine/Source/Runtime/CoreUObject/Private/UObject/Obj.cpp b/Engine/Source/Runtime/CoreUObject/Private/UObject/Obj.cpp
--- a/Engine/Source/Runtime/CoreUObject/Private/UObject/Obj.cpp
+++ b/Engine/Source/Runtime/CoreUObject/Private/UObject/Obj.cpp
@@ -2190,6 +2190,21 @@ void UObject::LoadConfig( UClass* ConfigClass/*=NULL*/, const TCHAR* InFilename/
 
 	for ( UProperty* Property = ConfigClass->PropertyLink; Property; Property = Property->PropertyLinkNext )
 	{
+#if WITH_EDITOR
+		FSoftObjectPathSerializationScope SerializationScope(NAME_None, Property->GetFName(), Property->IsEditorOnlyProperty() ? ESoftObjectPathCollectType::EditorOnlyCollect : ESoftObjectPathCollectType::AlwaysCollect, ESoftObjectPathSerializeType::AlwaysSerialize);
+
+		if (Property->IsEditorOnlyProperty() && Property->IsA(USoftObjectProperty::StaticClass()))
+		{
+			static TSet<FName> Encountered;
+			FName FullName = *(GetClass()->GetName() + TEXT("::") + Property->GetName());
+			if (!Encountered.Contains(FullName))
+			{
+				UE_LOG(LogConfig, Display, TEXT("Editor only config property: %s"), *FullName.ToString());
+				Encountered.Add(FullName);
+			}
+		}
+#endif
+
 		if ( !Property->HasAnyPropertyFlags(CPF_Config) )
 		{
 			continue;
diff --git a/Engine/Source/Runtime/EngineSettings/Classes/GameMapsSettings.h b/Engine/Source/Runtime/EngineSettings/Classes/GameMapsSettings.h
--- a/Engine/Source/Runtime/EngineSettings/Classes/GameMapsSettings.h
+++ b/Engine/Source/Runtime/EngineSettings/Classes/GameMapsSettings.h
@@ -121,9 +121,11 @@ class ENGINESETTINGS_API UGameMapsSettings
 
 public:
 
+#if WITH_EDITORONLY_DATA
 	/** If set, this map will be loaded when the Editor starts up. */
 	UPROPERTY(config, EditAnywhere, Category=DefaultMaps, meta=(AllowedClasses="World"))
 	FSoftObjectPath EditorStartupMap;
+#endif
 
 	/** The default options that will be appended to a map being loaded. */
 	UPROPERTY(config, EditAnywhere, Category=DefaultMaps, AdvancedDisplay)
