Removed some debugging code that I accidentally checked in

#rnx

diff --git a/Engine/Source/Runtime/CoreUObject/Private/UObject/Obj.cpp b/Engine/Source/Runtime/CoreUObject/Private/UObject/Obj.cpp
--- a/Engine/Source/Runtime/CoreUObject/Private/UObject/Obj.cpp
+++ b/Engine/Source/Runtime/CoreUObject/Private/UObject/Obj.cpp
@@ -2192,17 +2192,6 @@ void UObject::LoadConfig( UClass* ConfigClass/*=NULL*/, const TCHAR* InFilename/
 	{
 #if WITH_EDITOR
 		FSoftObjectPathSerializationScope SerializationScope(NAME_None, Property->GetFName(), Property->IsEditorOnlyProperty() ? ESoftObjectPathCollectType::EditorOnlyCollect : ESoftObjectPathCollectType::AlwaysCollect, ESoftObjectPathSerializeType::AlwaysSerialize);
-
-		if (Property->IsEditorOnlyProperty() && Property->IsA(USoftObjectProperty::StaticClass()))
-		{
-			static TSet<FName> Encountered;
-			FName FullName = *(GetClass()->GetName() + TEXT("::") + Property->GetName());
-			if (!Encountered.Contains(FullName))
-			{
-				UE_LOG(LogConfig, Display, TEXT("Editor only config property: %s"), *FullName.ToString());
-				Encountered.Add(FullName);
-			}
-		}
 #endif
 
 		if ( !Property->HasAnyPropertyFlags(CPF_Config) )
