Optimize UPackage::Save by removing lots of temporary allocations.
Change visibility of internal UObjectBaseUtility::GetPathName(const UObject* StopOuter, FString& ResultString) overload to public.
Add public overload of UObjectBaseUtility::GetFullName(const UObject* StopOuter, FString& ResultString).

diff --git a/Engine/Source/Runtime/CoreUObject/Private/UObject/SavePackage.cpp b/Engine/Source/Runtime/CoreUObject/Private/UObject/SavePackage.cpp
--- a/Engine/Source/Runtime/CoreUObject/Private/UObject/SavePackage.cpp
+++ b/Engine/Source/Runtime/CoreUObject/Private/UObject/SavePackage.cpp
@@ -1631,6 +1631,8 @@ public:
 			Used.AddZeroed(Orig.Num());
 
 			TMap<FString,int32> OriginalImportIndexes;
+			OriginalImportIndexes.Reserve(Orig.Num());
+			ObjectToFullNameMap.Reserve(Orig.Num());
 			for ( int32 i = 0; i < Orig.Num(); i++ )
 			{
 				FObjectImport& Import = Orig[i];
@@ -1673,12 +1675,13 @@ public:
 		}
 		else
 		{
+			ObjectToFullNameMap.Reserve(Linker->ImportMap.Num());
 			for ( int32 ImportIndex = 0; ImportIndex < Linker->ImportMap.Num(); ImportIndex++ )
 			{
 				const FObjectImport& Import = Linker->ImportMap[ImportIndex];
 				if ( Import.XObject )
 				{
-					ObjectToFullNameMap.Add(Import.XObject, *Import.XObject->GetFullName());
+					ObjectToFullNameMap.Add(Import.XObject, Import.XObject->GetFullName());
 				}
 			}
 		}
@@ -1804,6 +1807,15 @@ public:
 	{
 		bUseFObjectFullName = InbUseFObjectFullName;
 
+		if (bUseFObjectFullName)
+		{
+			ObjectToObjectFullNameMap.Reserve(Linker->ExportMap.Num());
+		}
+		else
+		{
+			ObjectToFullNameMap.Reserve(Linker->ExportMap.Num());
+		}
+
 		int32 SortStartPosition=0;
 		if ( LinkerToConformTo )
 		{
@@ -1815,6 +1827,7 @@ public:
 
 			// Populate object to current index map.
 			TMap<FString,int32> OriginalExportIndexes;
+			OriginalExportIndexes.Reserve(Linker->ExportMap.Num());
 			for( int32 ExportIndex=0; ExportIndex < Linker->ExportMap.Num(); ExportIndex++ )
 			{
 				const FObjectExport& Export = Linker->ExportMap[ExportIndex];
@@ -1918,7 +1931,7 @@ public:
 					}
 					else
 					{
-						ObjectToFullNameMap.Add(Export.Object, *Export.Object->GetFullName());
+						ObjectToFullNameMap.Add(Export.Object, Export.Object->GetFullName());
 					}
 				}
 			}
@@ -4217,21 +4230,12 @@ FSavePackageResultStruct UPackage::Save(UPackage* InOuter, UObject* Base, EObjec
 							UE_LOG(LogSavePackage, Fatal, TEXT("%s"), *FString::Printf( TEXT("Transient object imported: %s"), *Obj->GetFullName() ) );
 						}
 
-						// add the list of dependencies to the dependency map
-						ObjectDependencies.Add(Obj, ImportTagger.Dependencies);
-						NativeObjectDependencies.Add(Obj, ImportTagger.NativeDependencies);
-
 						if (Obj->GetClass() != UObjectRedirector::StaticClass())
 						{
-							for ( auto DepIt = ImportTagger.Dependencies.CreateConstIterator(); DepIt; ++DepIt )
-							{
-								UObject* DependencyObject = *DepIt;
-								if ( DependencyObject )
-								{
-									DependenciesReferencedByNonRedirectors.Add(DependencyObject);
-								}
-							}
+							DependenciesReferencedByNonRedirectors.Append(ImportTagger.Dependencies);
 						}
+						ObjectDependencies.Add(Obj, MoveTemp(ImportTagger.Dependencies));
+						NativeObjectDependencies.Add(Obj, MoveTemp(ImportTagger.NativeDependencies));
 					}
 				}
 				if (PrestreamPackages.Num())
@@ -4801,6 +4805,7 @@ FSavePackageResultStruct UPackage::Save(UPackage* InOuter, UObject* Base, EObjec
 					checkf(SrcDepends,TEXT("Couldn't find dependency map for %s"), *Object->GetFullName());
 
 					// go through each object and...
+					DependIndices.Reserve(SrcDepends->Num());
 					for (int32 DependIndex = 0; DependIndex < SrcDepends->Num(); DependIndex++)
 					{
 						UObject* DependentObject = (*SrcDepends)[DependIndex];
@@ -5244,14 +5249,18 @@ FSavePackageResultStruct UPackage::Save(UPackage* InOuter, UObject* Base, EObjec
 					FStructuredArchive::FStream DepedenciesStream = StructuredArchiveRoot.EnterStream(FIELD_NAME_TEXT("PreloadDependencies"));
 					TArray<UObject*> Subobjects;
 					TArray<UObject*> Deps;
+					TSet<FPackageIndex> SerializationBeforeCreateDependencies;
+					TSet<FPackageIndex> SerializationBeforeSerializationDependencies;
+					TSet<FPackageIndex> CreateBeforeSerializationDependencies;
+					TSet<FPackageIndex> CreateBeforeCreateDependencies;
 					for (int32 i = 0; i < Linker->ExportMap.Num(); i++)
 					{
 						FObjectExport& Export = Linker->ExportMap[i];
 						if (Export.Object)
 						{
 							EDLCookChecker.AddExport(Export.Object);
-							TSet<FPackageIndex> SerializationBeforeCreateDependencies;
 							{
+								SerializationBeforeCreateDependencies.Reset();
 								IncludeIndexAsDependency(SerializationBeforeCreateDependencies, Export.ClassIndex);
 								UObject* CDO = Export.Object->GetArchetype();
 								IncludeObjectAsDependency(1, SerializationBeforeCreateDependencies, CDO, Export.Object, true, false);
@@ -5278,8 +5287,8 @@ FSavePackageResultStruct UPackage::Save(UPackage* InOuter, UObject* Base, EObjec
 									}
 								}
 							}
-							TSet<FPackageIndex> SerializationBeforeSerializationDependencies;
 							{
+								SerializationBeforeSerializationDependencies.Reset();
 								Deps.Reset();
 								Export.Object->GetPreloadDependencies(Deps);
 
@@ -5325,8 +5334,8 @@ FSavePackageResultStruct UPackage::Save(UPackage* InOuter, UObject* Base, EObjec
 								}
 							}
 
-							TSet<FPackageIndex> CreateBeforeSerializationDependencies;
 							{
+								CreateBeforeSerializationDependencies.Reset();
 								UClass* Class = Cast<UClass>(Export.Object);
 								UObject* ClassCDO = Class ? Class->GetDefaultObject() : nullptr;
 								{
@@ -5351,8 +5360,9 @@ FSavePackageResultStruct UPackage::Save(UPackage* InOuter, UObject* Base, EObjec
 									}
 								}
 							}
-							TSet<FPackageIndex> CreateBeforeCreateDependencies;
+
 							{
+								CreateBeforeCreateDependencies.Reset();
 								IncludeIndexAsDependency(CreateBeforeCreateDependencies, Export.OuterIndex);
 								IncludeIndexAsDependency(CreateBeforeCreateDependencies, Export.SuperIndex);
 							}
@@ -5467,6 +5477,7 @@ FSavePackageResultStruct UPackage::Save(UPackage* InOuter, UObject* Base, EObjec
 					FScopedSlowTask ExportScope(Linker->ExportMap.Num());
 
 					// Save exports.
+					FString ObjectName;
 					int32 LastExportSaveStep = 0;
 					for( int32 i=0; i<Linker->ExportMap.Num(); i++ )
 					{
@@ -5484,8 +5495,8 @@ FSavePackageResultStruct UPackage::Save(UPackage* InOuter, UObject* Base, EObjec
 							Linker->CurrentlySavingExport = FPackageIndex::FromExport(i);
 							// UE_LOG(LogSavePackage, Log, TEXT("export %s for %s"), *Export.Object->GetFullName(), *Linker->CookingTarget()->PlatformName());
 
-							//FString ObjectName = Export.Object->GetPathName(InOuter);
-							FString ObjectName = Export.Object->GetPathName(InOuter);
+							ObjectName.Reset();
+							Export.Object->GetPathName(InOuter, ObjectName);
 							FStructuredArchive::FSlot ExportSlot = StructuredArchiveRoot.EnterField(FIELD_NAME(*ObjectName));
 
 #if WITH_EDITOR
diff --git a/Engine/Source/Runtime/CoreUObject/Private/UObject/UObjectBaseUtility.cpp b/Engine/Source/Runtime/CoreUObject/Private/UObject/UObjectBaseUtility.cpp
--- a/Engine/Source/Runtime/CoreUObject/Private/UObject/UObjectBaseUtility.cpp
+++ b/Engine/Source/Runtime/CoreUObject/Private/UObject/UObjectBaseUtility.cpp
@@ -33,26 +33,26 @@ FString UObjectBaseUtility::GetPathName( const UObject* StopOuter/*=NULL*/ ) con
 }
 
 /**
- * Internal version of GetPathName() that eliminates unnecessary copies.
+ * Version of GetPathName() that eliminates unnecessary copies and appends an existing string.
  */
-void UObjectBaseUtility::GetPathName( const UObject* StopOuter, FString& ResultString ) const
+void UObjectBaseUtility::GetPathName(const UObject* StopOuter, FString& ResultString) const
 {
-	if( this != StopOuter && this != NULL )
+	if(this != StopOuter && this != NULL)
 	{
 		UObject* ObjOuter = GetOuter();
-		if (ObjOuter && ObjOuter != StopOuter )
+		if (ObjOuter && ObjOuter != StopOuter)
 		{
-			ObjOuter->GetPathName( StopOuter, ResultString );
+			ObjOuter->GetPathName(StopOuter, ResultString);
 
-			// SUBOBJECT_DELIMITER is used to indicate that this object's outer is not a UPackage
+			// SUBOBJECT_DELIMITER_CHAR is used to indicate that this object's outer is not a UPackage
 			if (ObjOuter->GetClass() != UPackage::StaticClass()
 			&& ObjOuter->GetOuter()->GetClass() == UPackage::StaticClass())
 			{
-				ResultString += SUBOBJECT_DELIMITER;
+				ResultString += SUBOBJECT_DELIMITER_CHAR;
 			}
 			else
 			{
-				ResultString += TEXT(".");
+				ResultString += TEXT('.');
 			}
 		}
 		AppendName(ResultString);
@@ -72,22 +72,29 @@ void UObjectBaseUtility::GetPathName( const UObject* StopOuter, FString& ResultS
  *
  * @note	safe to call on NULL object pointers!
  */
-FString UObjectBaseUtility::GetFullName( const UObject* StopOuter/*=NULL*/ ) const
+FString UObjectBaseUtility::GetFullName(const UObject* StopOuter) const
 {
-	FString Result;  
-	if( this != nullptr )
+	FString Result;
+	Result.Empty(128);
+	GetFullName(StopOuter, Result);
+	return Result;
+}
+
+ /**
+  * Version of GetFullName() that eliminates unnecessary copies and appends an existing string.
+  */
+void UObjectBaseUtility::GetFullName(const UObject* StopOuter, FString& ResultString) const
+{
+	if (this != nullptr)
 	{
-		Result.Empty(128);
-		GetClass()->AppendName(Result);
-		Result += TEXT(" ");
-		GetPathName( StopOuter, Result );
-		// could possibly put a Result.Shrink() here, but this isn't used much in a shipping game
+		GetClass()->AppendName(ResultString);
+		ResultString += TEXT(' ');
+		GetPathName(StopOuter, ResultString);
 	}
 	else
 	{
-		Result += TEXT("None");
-	}
-	return Result;  
+		ResultString += TEXT("None");
+	} 
 }
 
 
diff --git a/Engine/Source/Runtime/CoreUObject/Public/UObject/UObjectBaseUtility.h b/Engine/Source/Runtime/CoreUObject/Public/UObject/UObjectBaseUtility.h
--- a/Engine/Source/Runtime/CoreUObject/Public/UObject/UObjectBaseUtility.h
+++ b/Engine/Source/Runtime/CoreUObject/Public/UObject/UObjectBaseUtility.h
@@ -292,6 +292,11 @@ public:
 	 */
 	FString GetFullName( const UObject* StopOuter=NULL ) const;
 
+	/**
+	 * Version of GetFullName() that eliminates unnecessary copies.
+	 */
+	void GetFullName(const UObject* StopOuter, FString& ResultString) const;
+
 	/**
 	 * Returns the fully qualified pathname for this object, in the format:
 	 * 'Outermost[.Outer].Name'
@@ -303,6 +308,11 @@ public:
 	 */
 	FString GetPathName( const UObject* StopOuter=NULL ) const;
 
+	/**
+	 * Version of GetPathName() that eliminates unnecessary copies.
+	 */
+	void GetPathName(const UObject* StopOuter, FString& ResultString) const;
+
 public:
 	/**
 	* Called after load to determine if the object can be a cluster root
@@ -342,11 +352,6 @@ protected:
 	/** Helper function to create a cluster from UObject */
 	static void CreateClusterFromObject(UObjectBaseUtility* ClusterRootObject, UObjectBaseUtility* ReferencingObject);
 
-	/**
-	 * Internal version of GetPathName() that eliminates lots of copies.
-	 */
-	void GetPathName( const UObject* StopOuter, FString& ResultString ) const;
-
 public:
 	/**
 	 * Walks up the chain of packages until it reaches the top level, which it ignores.
@@ -815,4 +820,4 @@ struct FScopeCycleCounterUObject
 	#ifdef PRAGMA_ENABLE_SHADOW_VARIABLE_WARNINGS
 		PRAGMA_ENABLE_SHADOW_VARIABLE_WARNINGS
 	#endif
-#endif
\ No newline at end of file
+#endif
