Undo my soft reference change because I want to resubmit in a single CL so I can reference it easily in a UDN post

#rnx

diff --git a/Engine/Source/Runtime/CoreUObject/Private/UObject/Obj.cpp b/Engine/Source/Runtime/CoreUObject/Private/UObject/Obj.cpp
--- a/Engine/Source/Runtime/CoreUObject/Private/UObject/Obj.cpp
+++ b/Engine/Source/Runtime/CoreUObject/Private/UObject/Obj.cpp
@@ -2190,10 +2190,6 @@ void UObject::LoadConfig( UClass* ConfigClass/*=NULL*/, const TCHAR* InFilename/
 
 	for ( UProperty* Property = ConfigClass->PropertyLink; Property; Property = Property->PropertyLinkNext )
 	{
-#if WITH_EDITOR
-		FSoftObjectPathSerializationScope SerializationScope(NAME_None, Property->GetFName(), Property->IsEditorOnlyProperty() ? ESoftObjectPathCollectType::EditorOnlyCollect : ESoftObjectPathCollectType::AlwaysCollect, ESoftObjectPathSerializeType::AlwaysSerialize);
-#endif
-
 		if ( !Property->HasAnyPropertyFlags(CPF_Config) )
 		{
 			continue;
diff --git a/Engine/Source/Runtime/EngineSettings/Classes/GameMapsSettings.h b/Engine/Source/Runtime/EngineSettings/Classes/GameMapsSettings.h
--- a/Engine/Source/Runtime/EngineSettings/Classes/GameMapsSettings.h
+++ b/Engine/Source/Runtime/EngineSettings/Classes/GameMapsSettings.h
@@ -121,11 +121,9 @@ class ENGINESETTINGS_API UGameMapsSettings
 
 public:
 
-#if WITH_EDITORONLY_DATA
 	/** If set, this map will be loaded when the Editor starts up. */
 	UPROPERTY(config, EditAnywhere, Category=DefaultMaps, meta=(AllowedClasses="World"))
 	FSoftObjectPath EditorStartupMap;
-#endif
 
 	/** The default options that will be appended to a map being loaded. */
 	UPROPERTY(config, EditAnywhere, Category=DefaultMaps, AdvancedDisplay)
