noMaster fix

diff --git a/Engine/Source/Runtime/CoreUObject/Private/Serialization/LoadTimeTrace.cpp b/Engine/Source/Runtime/CoreUObject/Private/Serialization/LoadTimeTrace.cpp
--- a/Engine/Source/Runtime/CoreUObject/Private/Serialization/LoadTimeTrace.cpp
+++ b/Engine/Source/Runtime/CoreUObject/Private/Serialization/LoadTimeTrace.cpp
@@ -11,6 +11,7 @@
 #include "HAL/PlatformTLS.h"
 #include "UObject/Object.h"
 #include "UObject/UObjectGlobals.h"
+#include "Misc/CommandLine.h"
 
 UE_TRACE_EVENT_BEGIN(LoadTime, StartAsyncLoading, Always)
 	UE_TRACE_EVENT_FIELD(uint64, Cycle)
