Asset loading optimization - avoid checking type of soft object property references in test and shipping configs

diff --git a/Engine/Source/Runtime/CoreUObject/Private/UObject/PropertySoftObjectPtr.cpp b/Engine/Source/Runtime/CoreUObject/Private/UObject/PropertySoftObjectPtr.cpp
--- a/Engine/Source/Runtime/CoreUObject/Private/UObject/PropertySoftObjectPtr.cpp
+++ b/Engine/Source/Runtime/CoreUObject/Private/UObject/PropertySoftObjectPtr.cpp
@@ -54,13 +54,16 @@ void USoftObjectProperty::SerializeItem( FStructuredArchive::FSlot Slot, void* V
 		FSoftObjectPtr OldValue = *(FSoftObjectPtr*)Value;
 		Slot << *(FSoftObjectPtr*)Value;
 
-		if (UnderlyingArchive.IsLoading() || UnderlyingArchive.IsModifyingWeakAndStrongReferences()) 
+		// Check for references to instances of wrong types and null them out 
+#if !(UE_BUILD_TEST || UE_BUILD_SHIPPING) 
+		if (UnderlyingArchive.IsLoading() || UnderlyingArchive.IsModifyingWeakAndStrongReferences())
 		{
 			if (OldValue.GetUniqueID() != ((FSoftObjectPtr*)Value)->GetUniqueID())
 			{
 				CheckValidObject(Value);
 			}
 		}
+#endif
 	}
 	else
 	{
