Script map loading optimization

Will cleanup FScriptMap::FindOrAdd() reliance on FScriptSet::Add() semantics after my vacation. Submitting now to help asset loading measurements while I'm OOO.

diff --git a/Engine/Source/Runtime/Core/Public/Containers/Map.h b/Engine/Source/Runtime/Core/Public/Containers/Map.h
--- a/Engine/Source/Runtime/Core/Public/Containers/Map.h
+++ b/Engine/Source/Runtime/Core/Public/Containers/Map.h
@@ -1614,6 +1614,45 @@ public:
 		);
 	}
 
+	/**
+	 * Adds a new uninitialized value or destroys the existing value for a given key.
+	 *
+	 * No need to rehash after calling. The hash table must be properly hashed before calling.
+	 *
+	 * @return The address to the value, not the pair
+	 **/
+	void* FindOrAdd(
+		const void* Key,
+		const FScriptMapLayout& Layout,
+		TFunctionRef<uint32(const void*)> GetKeyHash,
+		TFunctionRef<bool(const void*, const void*)> KeyEqualityFn,
+		TFunctionRef<void(void*, void*)> ConstructPairFn)
+	{
+		void* OutPair = nullptr;
+		bool bCreatePair = true;
+		const int32 ValueOffset = Layout.ValueOffset;
+				
+		Pairs.Add(
+			Key,
+			Layout.SetLayout,
+			GetKeyHash,
+			KeyEqualityFn,
+			[ConstructPairFn, &bCreatePair, &OutPair, ValueOffset](void* NewPair)
+			{
+				if (bCreatePair)
+				{
+					ConstructPairFn(NewPair, (uint8*)NewPair + ValueOffset);
+				}
+
+				OutPair = NewPair;
+			},
+			[&bCreatePair](void* NewPair) {	bCreatePair = false; }
+		);
+
+		checkSlow(OutPair);
+		return (uint8*)OutPair + ValueOffset;
+	}
+
 private:
 	FScriptSet Pairs;
 
diff --git a/Engine/Source/Runtime/CoreUObject/Private/UObject/PropertyMap.cpp b/Engine/Source/Runtime/CoreUObject/Private/UObject/PropertyMap.cpp
--- a/Engine/Source/Runtime/CoreUObject/Private/UObject/PropertyMap.cpp
+++ b/Engine/Source/Runtime/CoreUObject/Private/UObject/PropertyMap.cpp
@@ -297,10 +297,9 @@ void UMapProperty::SerializeItem(FStructuredArchive::FSlot Slot, void* Value, co
 				KeyProp->SerializeItem(KeysToRemoveArray.EnterElement(), TempKeyStorage);
 
 				// If the key is in the map, remove it
-				int32 Found = MapHelper.FindMapIndexWithKey(TempKeyStorage);
-				if (Found != INDEX_NONE)
+				if (uint8* PairPtr = MapHelper.FindMapPairPtrFromHash(TempKeyStorage))
 				{
-					MapHelper.RemoveAt(Found);
+					MapHelper.RemovePair(PairPtr);
 				}
 			}
 		}
@@ -325,27 +324,15 @@ void UMapProperty::SerializeItem(FStructuredArchive::FSlot Slot, void* Value, co
 				FSerializedPropertyScope SerializedProperty(UnderlyingArchive, KeyProp, this);
 				KeyProp->SerializeItem(EntryRecord.EnterField(FIELD_NAME_TEXT("Key")), TempKeyStorage);
 			}
-			
-			// Add a new default value if the key doesn't currently exist in the map
-			int32 NextPairIndex = MapHelper.FindMapIndexWithKey(TempKeyStorage);
-			if (NextPairIndex == INDEX_NONE)
-			{
-				NextPairIndex = MapHelper.AddDefaultValue_Invalid_NeedsRehash();
-			}
 
-			uint8* NextPairPtr = MapHelper.GetPairPtrWithoutCheck(NextPairIndex);
+			void* ValuePtr = MapHelper.FindOrAdd(TempKeyStorage);
 
-			// Copy over deserialised key from temporary storage
-			KeyProp->CopyCompleteValue_InContainer(NextPairPtr, TempKeyStorage);
-
-			// Deserialize value
+			// Deserialize value into hash map-owned memory
 			{
 				FSerializedPropertyScope SerializedProperty(UnderlyingArchive, ValueProp, this);
-				ValueProp->SerializeItem(EntryRecord.EnterField(FIELD_NAME_TEXT("Value")), NextPairPtr + MapLayout.ValueOffset);
+				ValueProp->SerializeItem(EntryRecord.EnterField(FIELD_NAME_TEXT("Value")), ValuePtr);
 			}
 		}
-
-		MapHelper.Rehash();
 	}
 	else
 	{
diff --git a/Engine/Source/Runtime/CoreUObject/Public/UObject/UnrealType.h b/Engine/Source/Runtime/CoreUObject/Public/UObject/UnrealType.h
--- a/Engine/Source/Runtime/CoreUObject/Public/UObject/UnrealType.h
+++ b/Engine/Source/Runtime/CoreUObject/Public/UObject/UnrealType.h
@@ -3409,6 +3409,50 @@ public:
 		);
 	}
 
+
+	/**
+	 * Finds or adds a new default-constructed value
+	 *
+	 * No need to rehash after calling. The hash table must be properly hashed before calling.
+	 *
+	 * @return The address to the value, not the pair
+	 **/
+	void* FindOrAdd(const void* KeyPtr)
+	{
+		UProperty* LocalKeyPropForCapture = KeyProp;
+		UProperty* LocalValuePropForCapture = ValueProp;
+		return Map->FindOrAdd(
+			KeyPtr,
+			MapLayout,
+			[LocalKeyPropForCapture](const void* ElementKey) { return LocalKeyPropForCapture->GetValueTypeHash(ElementKey); },
+			[LocalKeyPropForCapture](const void* A, const void* B) { return LocalKeyPropForCapture->Identical(A, B); },
+			[LocalKeyPropForCapture, LocalValuePropForCapture, KeyPtr](void* NewElementKey, void* NewElementValue)
+			{
+				if (LocalKeyPropForCapture->PropertyFlags & CPF_ZeroConstructor)
+				{
+					FMemory::Memzero(NewElementKey, LocalKeyPropForCapture->GetSize());
+				}
+				else
+				{
+					LocalKeyPropForCapture->InitializeValue(NewElementKey);
+				}
+
+				LocalKeyPropForCapture->CopySingleValueToScriptVM(NewElementKey, KeyPtr);
+
+				if (LocalValuePropForCapture->PropertyFlags & CPF_ZeroConstructor)
+				{
+					FMemory::Memzero(NewElementValue, LocalValuePropForCapture->GetSize());
+				}
+				else
+				{
+					LocalValuePropForCapture->InitializeValue(NewElementValue);
+				}
+			}
+		);
+
+	}
+
+
 	/** Removes the key and its associated value from the map */
 	bool RemovePair(const void* KeyPtr)
 	{
