Fixed a potential crash when adding tickable game objects in static init

diff --git a/Engine/Source/Runtime/CoreUObject/Private/UObject/UObjectBase.cpp b/Engine/Source/Runtime/CoreUObject/Private/UObject/UObjectBase.cpp
--- a/Engine/Source/Runtime/CoreUObject/Private/UObject/UObjectBase.cpp
+++ b/Engine/Source/Runtime/CoreUObject/Private/UObject/UObjectBase.cpp
@@ -34,8 +34,16 @@ DEFINE_LOG_CATEGORY_STATIC(LogUObjectBootstrap, Display, Display);
 /** Whether uobject system is initialized.												*/
 namespace Internal
 {
-	bool GObjInitialized = false;
+	static bool& GetUObjectSubsystemInitialised()
+	{
+		static bool ObjInitialized = false;
+		return ObjInitialized;
+	}
 };
+bool UObjectInitialized()
+{
+	return Internal::GetUObjectSubsystemInitialised();
+}
 
 
 /** Objects to automatically register once the object system is ready.					*/
@@ -208,7 +216,7 @@ void UObjectBase::CreateStatID() const
  */
 void UObjectBase::DeferredRegister(UClass *UClassStaticClass,const TCHAR* PackageName,const TCHAR* InName)
 {
-	check(Internal::GObjInitialized);
+	check(UObjectInitialized());
 	// Set object properties.
 	UPackage* Package = CreatePackage(nullptr, PackageName);
 	check(Package);
@@ -1122,7 +1130,7 @@ void UObjectBaseInit()
 	InitAsyncThread();
 
 	// Note initialized.
-	Internal::GObjInitialized = true;
+	Internal::GetUObjectSubsystemInitialised() = true;
 
 	UObjectProcessRegistrants();
 }
@@ -1133,7 +1141,7 @@ void UObjectBaseInit()
 void UObjectBaseShutdown()
 {
 	GUObjectArray.ShutdownUObjectArray();
-	Internal::GObjInitialized = false;
+	Internal::GetUObjectSubsystemInitialised() = false;
 }
 
 /**
diff --git a/Engine/Source/Runtime/CoreUObject/Public/UObject/UObjectBase.h b/Engine/Source/Runtime/CoreUObject/Public/UObject/UObjectBase.h
--- a/Engine/Source/Runtime/CoreUObject/Public/UObject/UObjectBase.h
+++ b/Engine/Source/Runtime/CoreUObject/Public/UObject/UObjectBase.h
@@ -272,20 +272,13 @@ private:
 #endif // HACK_HEADER_GENERATOR
 };
 
-namespace Internal
-{
-	/** Internal state indicating whether uobject system is initialized. Do not use that directly. Use
-		UObjectInitialized() instead. */
-	COREUOBJECT_API extern bool GObjInitialized;
-}
-
 /**
  * Checks to see if the UObject subsystem is fully bootstrapped and ready to go.
  * If true, then all objects are registered and auto registration of natives is over, forever.
  *
  * @return true if the UObject subsystem is initialized.
  */
-FORCEINLINE bool UObjectInitialized() { return Internal::GObjInitialized; }
+COREUOBJECT_API bool UObjectInitialized();
 
 /**
  * Force a pending registrant to register now instead of in the natural order
diff --git a/Engine/Source/Runtime/Engine/Private/Tickable.cpp b/Engine/Source/Runtime/Engine/Private/Tickable.cpp
--- a/Engine/Source/Runtime/Engine/Private/Tickable.cpp
+++ b/Engine/Source/Runtime/Engine/Private/Tickable.cpp
@@ -70,7 +70,14 @@ FTickableGameObject::FTickableGameObject()
 {
 	FTickableStatics& Statics = FTickableStatics::Get();
 
-	Statics.NewTickableObjects.Push(this);
+	if (UObjectInitialized())
+	{
+		Statics.NewTickableObjects.Push(this);
+	}
+	else
+	{
+		AddTickableObject(Statics.TickableObjects, this);
+	}
 }
 
 FTickableGameObject::~FTickableGameObject()
