CancelAsyncLoading() is now supported with the Event Driven Loader enabled.

When exiting the application async loading will be cancelled instead of flushed.

diff --git a/Engine/Source/Runtime/CoreUObject/Private/Serialization/AsyncLoading.cpp b/Engine/Source/Runtime/CoreUObject/Private/Serialization/AsyncLoading.cpp
--- a/Engine/Source/Runtime/CoreUObject/Private/Serialization/AsyncLoading.cpp
+++ b/Engine/Source/Runtime/CoreUObject/Private/Serialization/AsyncLoading.cpp
@@ -379,89 +379,6 @@ void FAsyncLoadingThread::InitializeAsyncThread()
 	AsyncThreadReady.Increment();
 }
 
-void FAsyncLoadingThread::CancelAsyncLoadingInternal()
-{
-	// Cancel is not thread safe because the loaded delegates expect to be called on the game thread
-	// EDL does not support this function, but for backward compatible reasons we are letting it run on the async loading thread
-	// If this is enabled for EDL it must be made properly thread safe
-	UE_CLOG(GEventDrivenLoaderEnabled && !IsInGameThread(), LogStreaming, Fatal, TEXT("CancelAsyncLoadingInternal is not thread safe! This must be fixed before being enabled for EDL"));
-
-	{
-		// Packages we haven't yet started processing.
-#if THREADSAFE_UOBJECTS
-		FScopeLock QueueLock(&QueueCritical);
-#endif
-		const EAsyncLoadingResult::Type Result = EAsyncLoadingResult::Canceled;
-		for (FAsyncPackageDesc* PackageDesc : QueuedPackages)
-		{
-			if (PackageDesc->PackageLoadedDelegate.IsValid())
-			{
-				PackageDesc->PackageLoadedDelegate->ExecuteIfBound(PackageDesc->Name, nullptr, Result);
-			}
-
-			delete PackageDesc;
-		}
-		QueuedPackages.Reset();
-	}
-
-	{
-		// Packages we started processing, need to be canceled.
-		// Accessed only in async thread, no need to protect region.
-		// Move first so that we remove the package from these lists BEFORE we delete it otherwise we will assert in the package dtor.
-		TArray<FAsyncPackage*> PackagesToDeleteCopy = MoveTemp(PackagesToDelete);
-
-		// This is accessed on the game thread but it should be blocked at this point
-		TArray<FAsyncPackage*> AsyncPackagesCopy = MoveTemp(AsyncPackages);
-
-		for (FAsyncPackage* Package : PackagesToDeleteCopy)
-		{
-			Package->Cancel();
-			delete Package;
-		}
-
-		for (FAsyncPackage* AsyncPackage : AsyncPackagesCopy)
-		{
-			AsyncPackage->Cancel();
-			delete AsyncPackage;
-		}
-		AsyncPackageNameLookup.Reset();
-	}
-
-	{
-		// Packages that are already loaded. May be halfway through PostLoad
-#if THREADSAFE_UOBJECTS
-		FScopeLock LoadedLock(&LoadedPackagesCritical);
-#endif
-		for (FAsyncPackage* LoadedPackage : LoadedPackages)
-		{
-			LoadedPackage->Cancel();
-			delete LoadedPackage;
-		}
-		LoadedPackages.Reset();
-		LoadedPackagesNameLookup.Reset();
-	}
-	{
-#if THREADSAFE_UOBJECTS
-		FScopeLock LoadedLock(&LoadedPackagesToProcessCritical);
-#endif
-		for (FAsyncPackage* LoadedPackage : LoadedPackagesToProcess)
-		{
-			LoadedPackage->Cancel();
-			delete LoadedPackage;
-		}
-		LoadedPackagesToProcess.Reset();
-		LoadedPackagesToProcessNameLookup.Reset();
-	}
-
-	ExistingAsyncPackagesCounter.Reset();
-	QueuedPackagesCounter.Reset();
-
-	NotifyAsyncLoadingStateHasMaybeChanged();
-
-	// Notify everyone streaming is canceled.
-	CancelLoadingEvent->Trigger();
-}
-
 void FAsyncLoadingThread::QueuePackage(FAsyncPackageDesc& Package)
 {
 	{
@@ -5387,6 +5304,119 @@ void FAsyncLoadingThread::CancelAsyncLoading()
 		// It's possible we haven't been async loading at all in which case the above call would not reset bShouldCancelLoading
 		bShouldCancelLoading = false;
 	}
+
+	// Actually delete all packages and execute delegates
+	FinalizeCancelAsyncLoadingInternal();
+}
+
+void FAsyncLoadingThread::CancelAsyncLoadingInternal()
+{
+	FAsyncLoadingTickScope AsyncTickScope;
+
+	if (GEventDrivenLoaderEnabled)
+	{
+		while (GPrecacheCallbackHandler.AnyIOOutstanding())
+		{
+			GPrecacheCallbackHandler.WaitForIO(10.0f);
+			GPrecacheCallbackHandler.ProcessIncoming();
+		}
+	}
+
+	{
+		// Packages we haven't yet started processing.
+#if THREADSAFE_UOBJECTS
+		FScopeLock QueueLock(&QueueCritical);
+#endif
+		QueuedPackagesToCancel = MoveTemp(QueuedPackages);
+		QueuedPackages.Reset();
+	}
+
+	{
+		// Packages we started processing, need to be canceled.
+		// Accessed only in async thread, no need to protect region.
+		// Move first so that we remove the package from these lists BEFORE we delete it otherwise we will assert in the package dtor.
+		PackagesToCancel.Append(PackagesToDelete);
+
+		// This is accessed on the game thread but it should be blocked at this point
+		PackagesToCancel.Append(AsyncPackagesReadyForTick);
+		PackagesToCancel.Append(AsyncPackages);
+		AsyncPackagesReadyForTick.Reset();
+		AsyncPackages.Reset();
+		AsyncPackageNameLookup.Reset();
+	}
+
+	{
+		// Packages that are already loaded. May be halfway through PostLoad
+#if THREADSAFE_UOBJECTS
+		FScopeLock LoadedLock(&LoadedPackagesCritical);
+#endif
+		PackagesToCancel.Append(LoadedPackages);
+		LoadedPackages.Reset();
+		LoadedPackagesNameLookup.Reset();
+	}
+	{
+#if THREADSAFE_UOBJECTS
+		FScopeLock LoadedLock(&LoadedPackagesToProcessCritical);
+#endif
+		PackagesToCancel.Append(LoadedPackagesToProcess);
+		LoadedPackagesToProcess.Reset();
+		LoadedPackagesToProcessNameLookup.Reset();
+	}
+
+	ExistingAsyncPackagesCounter.Reset();
+	QueuedPackagesCounter.Reset();
+
+	EventQueue.EventQueue.Empty();
+	FAsyncPackage::GetEventGraph().PackagesWithNodes.Empty();
+
+	{
+#if THREADSAFE_UOBJECTS
+		FScopeLock Lock(&PendingRequestsCritical);
+#endif
+		PendingRequests.Empty();
+	}
+
+	NotifyAsyncLoadingStateHasMaybeChanged();
+
+	// Notify everyone streaming is canceled.
+	CancelLoadingEvent->Trigger();
+}
+
+void FAsyncLoadingThread::FinalizeCancelAsyncLoadingInternal()
+{
+	check(IsInGameThread());
+
+#if THREADSAFE_UOBJECTS
+	FScopeLock QueueLock(&QueueCritical);
+	FScopeLock LoadedLock(&LoadedPackagesCritical);
+	FScopeLock LoadedToProcessLock(&LoadedPackagesToProcessCritical);
+#endif
+
+	check(QueuedPackages.Num() == 0);
+	const EAsyncLoadingResult::Type Result = EAsyncLoadingResult::Canceled;
+	for (FAsyncPackageDesc* PackageDescToCancel : QueuedPackagesToCancel)
+	{
+		if (PackageDescToCancel->PackageLoadedDelegate.IsValid())
+		{
+			PackageDescToCancel->PackageLoadedDelegate->ExecuteIfBound(PackageDescToCancel->Name, nullptr, Result);
+		}
+		delete PackageDescToCancel;
+	}
+	QueuedPackagesToCancel.Empty();
+
+	check(PackagesToDelete.Num() == 0);
+	check(AsyncPackages.Num() == 0);
+	check(LoadedPackages.Num() == 0);
+	check(LoadedPackagesToProcess.Num() == 0);
+	for (FAsyncPackage* PackageToCancel : PackagesToCancel)
+	{
+		PackageToCancel->Cancel();
+	}
+	for (FAsyncPackage* PackageToCancel : PackagesToCancel)
+	{
+		delete PackageToCancel;
+	}
+	PackagesToCancel.Empty();
 }
 
 void FAsyncLoadingThread::SuspendLoading()
@@ -7111,29 +7141,69 @@ UPackage* FAsyncPackage::GetLoadedPackage()
 	return LoadedPackage;
 }
 
-void FAsyncPackage::Cancel()
+template <typename T>
+void ClearFlagsAndDissolveClustersFromLoadedObjects(T& LoadedObjects)
 {
-	UE_CLOG(GEventDrivenLoaderEnabled, LogStreaming, Fatal, TEXT("FAsyncPackage::Cancel is not supported with the new loader"));
+	const EObjectFlags ObjectLoadFlags = EObjectFlags(RF_NeedLoad | RF_NeedPostLoad | RF_NeedPostLoadSubobjects | RF_WasLoaded);
+	for (UObject* ObjRef : LoadedObjects)
+	{
+		if (ObjRef)
+		{
+			ObjRef->AtomicallyClearFlags(ObjectLoadFlags);
+			if (ObjRef->HasAnyInternalFlags(EInternalObjectFlags::ClusterRoot))
+			{
+				GUObjectClusters.DissolveCluster(ObjRef);
+			}
+		}
+	}
+}
 
+void FAsyncPackage::Cancel()
+{
 	// Call any completion callbacks specified.
 	bLoadHasFailed = true;
 	const EAsyncLoadingResult::Type Result = EAsyncLoadingResult::Canceled;
 	CallCompletionCallbacks(true, Result);
 	CallCompletionCallbacks(false, Result);
 
+	for (TPair<IAsyncReadRequest*, FExportIORequest> Request : PrecacheRequests)
+	{
+		delete Request.Key;
+	}
+	PrecacheRequests.Empty();
+	ExportIndexToPrecacheRequest.Empty();
+
+	PackagesIMayBeWaitingForBeforePostload.Empty();
+	PackagesIAmWaitingForBeforePostload.Empty();
+	OtherPackagesWaitingForMeBeforePostload.Empty();
+	PackagesWaitingToLinkImports.Empty();
+
+	EventNodeArray.TotalNumberOfNodesAdded = 0;
+	EventNodeArray.TotalNumberOfImportExportNodes = 0;
+	EventNodeArray.Shutdown();
+
+	FUObjectSerializeContext* LoadContext = GetSerializeContext();
+	if (LoadContext)
+	{
+		TArray<UObject*>& ThreadObjLoaded = LoadContext->PRIVATE_GetObjectsLoadedInternalUseOnly();
+		if (ThreadObjLoaded.Num())
+		{
+			PackageObjLoaded.Append(ThreadObjLoaded);
+			ThreadObjLoaded.Reset();
+		}
+	}
+
 	{
 		// Clear load flags from any referenced objects
 		FScopeLock RereferncedObjectsLock(&ReferencedObjectsCritical);
-		const EObjectFlags ObjectLoadFlags = EObjectFlags(RF_NeedLoad | RF_NeedPostLoad | RF_NeedPostLoadSubobjects | RF_WasLoaded);
-		for (UObject* ObjRef : ReferencedObjects)
-		{			
-			if (ObjRef)
-			{
-				ObjRef->AtomicallyClearFlags(ObjectLoadFlags);
-			}
-		}
+		ClearFlagsAndDissolveClustersFromLoadedObjects(ReferencedObjects);
+		ClearFlagsAndDissolveClustersFromLoadedObjects(PackageObjLoaded);
+		ClearFlagsAndDissolveClustersFromLoadedObjects(DeferredFinalizeObjects);
+	
 		// Release references
 		EmptyReferencedObjects();
+		PackageObjLoaded.Empty();
+		DeferredFinalizeObjects.Empty();
 	}
 
 	if (LinkerRoot)
@@ -7148,7 +7218,7 @@ void FAsyncPackage::Cancel()
 			LinkerRoot->bHasBeenFullyLoaded = false;
 			LinkerRoot->Rename(*MakeUniqueObjectName(GetTransientPackage(), UPackage::StaticClass()).ToString(), nullptr, REN_DontCreateRedirectors | REN_DoNotDirty | REN_ForceNoResetLoaders | REN_NonTransactional);
 		}
-		DetachLinker();
+		ResetLoader();
 	}
 	PreLoadIndex = 0;
 	PreLoadSortIndex = 0;
@@ -7278,16 +7348,12 @@ void CancelAsyncLoading()
 	// Cancelling async loading while loading is suspend will result in infinite stall
 	UE_CLOG(FAsyncLoadingThread::Get().IsAsyncLoadingSuspended(), LogStreaming, Fatal, TEXT("Cannot Cancel Async Loading while async loading is suspended."));
 
-	if (GEventDrivenLoaderEnabled)
-	{
-		UE_LOG(LogStreaming, Warning, TEXT("Cannot Cancel Async Loading using the EDL loader. Async loading will be flushed instead."));
-		FlushAsyncLoading();
-	}
-	else
+	FAsyncLoadingThread::Get().CancelAsyncLoading();
+
+	if (!GIsRequestingExit)
 	{
-		FAsyncLoadingThread::Get().CancelAsyncLoading();
+		CollectGarbage(GARBAGE_COLLECTION_KEEPFLAGS, true);
 	}
-	CollectGarbage(GARBAGE_COLLECTION_KEEPFLAGS, false);
 }
 
 float GetAsyncLoadPercentage(const FName& PackageName)
diff --git a/Engine/Source/Runtime/CoreUObject/Private/Serialization/AsyncLoadingThread.h b/Engine/Source/Runtime/CoreUObject/Private/Serialization/AsyncLoadingThread.h
--- a/Engine/Source/Runtime/CoreUObject/Private/Serialization/AsyncLoadingThread.h
+++ b/Engine/Source/Runtime/CoreUObject/Private/Serialization/AsyncLoadingThread.h
@@ -228,6 +228,11 @@ private:
 
 	FThreadSafeCounter AsyncThreadReady;
 
+	/** When cancelling async loading: list of package requests to cancel */
+	TArray<FAsyncPackageDesc*> QueuedPackagesToCancel;
+	/** When cancelling async loading: list of packages to cancel */
+	TSet<FAsyncPackage*> PackagesToCancel;
+
 	/** Async loading thread ID */
 	static uint32 AsyncLoadingThreadID;
 
@@ -577,6 +582,7 @@ private:
 
 	/** Cancels async loading internally */
 	void CancelAsyncLoadingInternal();
+	void FinalizeCancelAsyncLoadingInternal();
 
 	/** Number of times we re-entered the async loading tick, mostly used by singlethreaded ticking. Debug purposes only. */
 	int32 AsyncLoadingTickCounter;
diff --git a/Engine/Source/Runtime/Launch/Private/LaunchEngineLoop.cpp b/Engine/Source/Runtime/Launch/Private/LaunchEngineLoop.cpp
--- a/Engine/Source/Runtime/Launch/Private/LaunchEngineLoop.cpp
+++ b/Engine/Source/Runtime/Launch/Private/LaunchEngineLoop.cpp
@@ -3634,7 +3634,7 @@ void FEngineLoop::Exit()
 
 
 	// Make sure we're not in the middle of loading something.
-	FlushAsyncLoading();
+	CancelAsyncLoading();
 
 	// Block till all outstanding resource streaming requests are fulfilled.
 	if (!IStreamingManager::HasShutdown())
