Added the ability to toggle between flushing and cancelling async loading on exit in project settings

diff --git a/Engine/Config/BaseEngine.ini b/Engine/Config/BaseEngine.ini
--- a/Engine/Config/BaseEngine.ini
+++ b/Engine/Config/BaseEngine.ini
@@ -811,6 +811,7 @@ s.LevelStreamingComponentsRegistrationGranularity = 10
 s.UnregisterComponentsTimeLimit = 1.0
 s.LevelStreamingComponentsUnregistrationGranularity = 5
 s.MaxPackageSummarySize=16384
+s.FlushStreamingOnExit=True
 +FixedBootOrder=/Script/Engine/Default__SoundBase
 +FixedBootOrder=/Script/Engine/Default__MaterialInterface
 +FixedBootOrder=/Script/Engine/Default__DeviceProfileManager
diff --git a/Engine/Source/Runtime/CoreUObject/Private/Serialization/AsyncLoading.cpp b/Engine/Source/Runtime/CoreUObject/Private/Serialization/AsyncLoading.cpp
--- a/Engine/Source/Runtime/CoreUObject/Private/Serialization/AsyncLoading.cpp
+++ b/Engine/Source/Runtime/CoreUObject/Private/Serialization/AsyncLoading.cpp
@@ -134,14 +134,21 @@ struct FAsyncPackageScope
 	}
 };
 
-static int32 GAsyncLoadingThreadEnabled;
-static FAutoConsoleVariableRef CVarAsyncLoadingThreadEnabledg(
+static int32 GAsyncLoadingThreadEnabled = 0;
+static FAutoConsoleVariableRef CVarAsyncLoadingThreadEnabled(
 	TEXT("s.AsyncLoadingThreadEnabled"),
 	GAsyncLoadingThreadEnabled,
 	TEXT("Placeholder console variable, currently not used in runtime."),
 	ECVF_Default
 	);
 
+static int32 GFlushStreamingOnExit = 1;
+static FAutoConsoleVariableRef CFlushStreamingOnExit(
+	TEXT("s.FlushStreamingOnExit"),
+	GFlushStreamingOnExit,
+	TEXT("Placeholder console variable, currently not used in runtime."),
+	ECVF_Default
+);
 
 static int32 GWarnIfTimeLimitExceeded = 0;
 static FAutoConsoleVariableRef CVarWarnIfTimeLimitExceeded(
diff --git a/Engine/Source/Runtime/Engine/Classes/Engine/CoreSettings.h b/Engine/Source/Runtime/Engine/Classes/Engine/CoreSettings.h
--- a/Engine/Source/Runtime/Engine/Classes/Engine/CoreSettings.h
+++ b/Engine/Source/Runtime/Engine/Classes/Engine/CoreSettings.h
@@ -104,6 +104,11 @@ protected:
 		ToolTip = "Enables the event driven loader in cooked builds."))
 	uint32 EventDrivenLoaderEnabled : 1;
 
+	UPROPERTY(EditAnywhere, config, Category = General, meta = (
+		ConsoleVariable = "s.FlushStreamingOnExit", DisplayName = "Flush Streaming when exiting the application",
+		ToolTip = "If enabled, streaming will be flushed when exiting application, otherwise it will be cancelled."))
+	uint32 FlushStreamingOnExit : 1;
+
 	//~ Begin UObject Interface
 	virtual void PostInitProperties() override;
 
diff --git a/Engine/Source/Runtime/Engine/Private/CoreSettings.cpp b/Engine/Source/Runtime/Engine/Private/CoreSettings.cpp
--- a/Engine/Source/Runtime/Engine/Private/CoreSettings.cpp
+++ b/Engine/Source/Runtime/Engine/Private/CoreSettings.cpp
@@ -132,6 +132,7 @@ UStreamingSettings::UStreamingSettings()
 	LevelStreamingUnregisterComponentsTimeLimit = 1.0f;
 	LevelStreamingComponentsUnregistrationGranularity = 5;
 	EventDrivenLoaderEnabled = false;
+	FlushStreamingOnExit = true;
 }
 
 void UStreamingSettings::PostInitProperties()
diff --git a/Engine/Source/Runtime/Launch/Private/LaunchEngineLoop.cpp b/Engine/Source/Runtime/Launch/Private/LaunchEngineLoop.cpp
--- a/Engine/Source/Runtime/Launch/Private/LaunchEngineLoop.cpp
+++ b/Engine/Source/Runtime/Launch/Private/LaunchEngineLoop.cpp
@@ -3633,7 +3633,22 @@ void FEngineLoop::Exit()
 
 
 	// Make sure we're not in the middle of loading something.
-	CancelAsyncLoading();
+	{
+		bool bFlushOnExit = true;
+		if (GConfig)
+		{
+			FBoolConfigValueHelper FlushStreamingOnExitHelper(TEXT("/Script/Engine.StreamingSettings"), TEXT("s.FlushStreamingOnExit"), GEngineIni);
+			bFlushOnExit = FlushStreamingOnExitHelper;			
+		}
+		if (bFlushOnExit)
+		{
+			FlushAsyncLoading();
+		}
+		else
+		{
+			CancelAsyncLoading();
+		}
+	}
 
 	// Block till all outstanding resource streaming requests are fulfilled.
 	if (!IStreamingManager::HasShutdown())
