'if constexpr' support added.
Change over some existing `TEnableIf` overloads to use it.

diff --git a/Engine/Source/Runtime/Core/Public/Android/AndroidPlatform.h b/Engine/Source/Runtime/Core/Public/Android/AndroidPlatform.h
--- a/Engine/Source/Runtime/Core/Public/Android/AndroidPlatform.h
+++ b/Engine/Source/Runtime/Core/Public/Android/AndroidPlatform.h
@@ -80,6 +80,12 @@ typedef FAndroidTypes FPlatformTypes;
 	#define PLATFORM_COMPILER_HAS_DECLTYPE_AUTO 0
 #endif
 
+#if __has_feature(cxx_if_constexpr)
+	#define PLATFORM_COMPILER_HAS_IF_CONSTEXPR 1
+#else
+	#define PLATFORM_COMPILER_HAS_IF_CONSTEXPR 0
+#endif
+
 // some android platform overrides that sub-platforms can disable
 #ifndef USE_ANDROID_JNI
 	#define USE_ANDROID_JNI							1
diff --git a/Engine/Source/Runtime/Core/Public/HAL/Platform.h b/Engine/Source/Runtime/Core/Public/HAL/Platform.h
--- a/Engine/Source/Runtime/Core/Public/HAL/Platform.h
+++ b/Engine/Source/Runtime/Core/Public/HAL/Platform.h
@@ -218,6 +218,9 @@
 #ifndef PLATFORM_COMPILER_HAS_DECLTYPE_AUTO
 	#define PLATFORM_COMPILER_HAS_DECLTYPE_AUTO 1
 #endif
+#ifndef PLATFORM_COMPILER_HAS_IF_CONSTEXPR
+	#define PLATFORM_COMPILER_HAS_IF_CONSTEXPR 1
+#endif
 #ifndef PLATFORM_TCHAR_IS_1_BYTE
 	#define PLATFORM_TCHAR_IS_1_BYTE			0
 #endif
diff --git a/Engine/Source/Runtime/Core/Public/HTML5/HTML5Platform.h b/Engine/Source/Runtime/Core/Public/HTML5/HTML5Platform.h
--- a/Engine/Source/Runtime/Core/Public/HTML5/HTML5Platform.h
+++ b/Engine/Source/Runtime/Core/Public/HTML5/HTML5Platform.h
@@ -74,6 +74,12 @@ typedef FHTML5Types FPlatformTypes;
 	#define PLATFORM_COMPILER_HAS_DECLTYPE_AUTO 0
 #endif
 
+#if __has_feature(cxx_if_constexpr)
+	#define PLATFORM_COMPILER_HAS_IF_CONSTEXPR 1
+#else
+	#define PLATFORM_COMPILER_HAS_IF_CONSTEXPR 0
+#endif
+
 // Function type macros.
 #define FORCEINLINE		inline __attribute__((__always_inline__))					/* Force code to be inline */
 #define FORCENOINLINE	__attribute__((noinline))									/* Force code to NOT be inline */
diff --git a/Engine/Source/Runtime/Core/Public/HoloLens/HoloLensPlatformCompilerSetup.h b/Engine/Source/Runtime/Core/Public/HoloLens/HoloLensPlatformCompilerSetup.h
--- a/Engine/Source/Runtime/Core/Public/HoloLens/HoloLensPlatformCompilerSetup.h
+++ b/Engine/Source/Runtime/Core/Public/HoloLens/HoloLensPlatformCompilerSetup.h
@@ -25,7 +25,7 @@ static_assert(_MSC_VER >= 1900, "Visual Studio 2015 or newer required.");
 #pragma warning (error: 4600 4601 4602 4603 4604      4606 4607      4609 4610 4611 4612 4613 4614 4615 4616 4617 4618      4620 4621 4622 4623 4624 4625 4626 4627 4628 4629 4630 4631 4632 4633 4634 4635 4636 4637 4638 4639 4640 4641 4642 4643 4644 4645 4646 4647 4648 4649 4650      4652 4653 4654 4655 4656 4657 4658 4659 4660 4661 4662 4663 4664 4665 4666 4667 4668 4669 4670 4671 4672 4673 4674 4675 4676 4677 4678 4679 4680 4681 4682 4683 4684 4685 4686 4687 4688 4689 4690 4691      4693 4694 4695 4696 4697 4698 4699)
 #pragma warning (error: 4700      4702 4703 4704 4705 4706 4707 4708 4709 4710 4711 4712 4713 4714 4715 4716 4717 4718 4719 4720 4721 4722 4723 4724 4725 4726 4727 4728 4729      4731 4732 4733 4734 4735 4736 4737      4739 4740 4741 4742 4743 4744 4745      4747 4748 4749 4750 4751 4752 4753 4754 4755 4756 4757 4758 4759 4760 4761 4762 4763 4764 4765 4766      4768 4769 4770 4771 4772 4773      4775 4776 4777 4778 4779 4780 4781 4782 4783 4784 4785 4786 4787 4788 4789 4790 4791 4792 4793 4794 4795 4796 4797 4798 4799)
 #pragma warning (error: 4800 4801 4802 4803 4804 4805 4806 4807 4808 4809 4810 4811 4812 4813 4814 4815 4816 4817 4818 4819      4821 4822 4823 4824 4825      4827 4828 4829 4830 4831 4832 4833 4834 4835 4836           4839 4840 4841 4842 4843 4844 4845 4846 4847 4848 4849 4850 4851 4852 4853 4854 4855 4856 4857 4858 4859 4860 4861 4862 4863 4864 4865 4866 4867      4869 4870 4871 4872 4873 4874 4875 4876 4877 4878 4879 4880 4881 4882 4883 4884 4885 4886 4887 4888 4889 4890 4891 4892 4893 4894 4895 4896 4897 4898 4899)
-#pragma warning (error: 4900 4901 4902 4903 4904 4905 4906 4907 4908 4909 4910 4911 4912 4913 4914 4915 4916 4917 4918 4919 4920 4921 4922 4923 4924 4925 4926 4927 4928 4929 4930 4931 4932 4933 4934 4935 4936 4937 4938 4939 4940 4941 4942 4943 4944 4945 4946 4947 4948 4949 4950 4951 4952 4953 4954 4955 4956 4957 4958 4959 4960 4961      4963 4964 4965 4966 4967 4968 4969 4970 4971 4972 4973 4974 4975 4976 4977 4978 4979 4980 4981 4982 4983 4984 4985                4989 4990 4991 4992 4993 4994 4995      4997 4998 4999)
+#pragma warning (error: 4900 4901 4902 4903 4904 4905 4906 4907 4908 4909 4910 4911 4912 4913 4914 4915 4916 4917 4918 4919 4920 4921 4922 4923 4924 4925 4926 4927 4928 4929 4930 4931 4932 4933 4934 4935 4936 4937 4938 4939 4940 4941 4942 4943 4944 4945 4946 4947 4948 4949 4950 4951 4952 4953 4954 4955 4956 4957 4958 4959 4960 4961      4963 4964 4965 4966 4967 4968 4969 4970 4971 4972 4973 4974 4975 4976 4977 4978 4979 4980 4981 4982 4983      4985                4989 4990 4991 4992 4993 4994 4995      4997 4998 4999)
 
 //
 // Skipped warnings, which are not disabled below. Most are disabled by default. Might be useful to look through, re-enable some and fix the code.
diff --git a/Engine/Source/Runtime/Core/Public/IOS/IOSPlatform.h b/Engine/Source/Runtime/Core/Public/IOS/IOSPlatform.h
--- a/Engine/Source/Runtime/Core/Public/IOS/IOSPlatform.h
+++ b/Engine/Source/Runtime/Core/Public/IOS/IOSPlatform.h
@@ -84,6 +84,12 @@ typedef FIOSPlatformTypes FPlatformTypes;
 	#define PLATFORM_COMPILER_HAS_DECLTYPE_AUTO 0
 #endif
 
+#if __has_feature(cxx_if_constexpr)
+	#define PLATFORM_COMPILER_HAS_IF_CONSTEXPR 1
+#else
+	#define PLATFORM_COMPILER_HAS_IF_CONSTEXPR 0
+#endif
+
 //mallocpoison not safe with aligned ansi allocator.  returns the larger unaligned size during Free() which causes writes off the end of the allocation.
 #define UE_USE_MALLOC_FILL_BYTES 0 
 
diff --git a/Engine/Source/Runtime/Core/Public/Mac/MacPlatform.h b/Engine/Source/Runtime/Core/Public/Mac/MacPlatform.h
--- a/Engine/Source/Runtime/Core/Public/Mac/MacPlatform.h
+++ b/Engine/Source/Runtime/Core/Public/Mac/MacPlatform.h
@@ -60,6 +60,12 @@ typedef FMacPlatformTypes FPlatformTypes;
 	#define PLATFORM_COMPILER_HAS_DECLTYPE_AUTO 0
 #endif
 
+#if __has_feature(cxx_if_constexpr)
+	#define PLATFORM_COMPILER_HAS_IF_CONSTEXPR 1
+#else
+	#define PLATFORM_COMPILER_HAS_IF_CONSTEXPR 0
+#endif
+
 // Function type macros.
 #define VARARGS															/* Functions with variable arguments */
 #define CDECL															/* Standard C function */
diff --git a/Engine/Source/Runtime/Core/Public/Templates/Function.h b/Engine/Source/Runtime/Core/Public/Templates/Function.h
--- a/Engine/Source/Runtime/Core/Public/Templates/Function.h
+++ b/Engine/Source/Runtime/Core/Public/Templates/Function.h
@@ -255,6 +255,26 @@ namespace UE4Function_Private
 	{
 	};
 
+#if PLATFORM_COMPILER_HAS_IF_CONSTEXPR
+
+	template <typename T>
+	FORCEINLINE bool IsBound(const T& Func)
+	{
+		if constexpr (TIsNullableBinding<T>::Value)
+		{
+			// Function pointers, data member pointers, member function pointers and TFunctions
+			// can all be null/unbound, so test them using their boolean state.
+			return !!Func;
+		}
+		else
+		{
+			// We can't tell if any other generic callable can be invoked, so just assume they can be.
+			return true;
+		}
+	}
+
+#else
+
 	template <typename T>
 	FORCEINLINE typename TEnableIf<TIsNullableBinding<T>::Value, bool>::Type IsBound(const T& Func)
 	{
@@ -270,6 +290,8 @@ namespace UE4Function_Private
 		return true;
 	}
 
+#endif
+
 	template <typename FunctorType, bool bUnique, bool bOnHeap>
 	struct TStorageOwnerType;
 
@@ -371,6 +393,38 @@ namespace UE4Function_Private
 		{
 		}
 
+#if PLATFORM_COMPILER_HAS_IF_CONSTEXPR
+
+		template <typename FunctorType>
+		typename TDecay<FunctorType>::Type* Bind(FunctorType&& InFunc)
+		{
+			if (!IsBound(InFunc))
+			{
+				return nullptr;
+			}
+
+			constexpr bool bUseInline = TFUNCTION_USES_INLINE_STORAGE && sizeof(TStorageOwnerTypeT<FunctorType, bUnique, false>) <= TFUNCTION_INLINE_SIZE;
+
+			using OwnedType = TStorageOwnerTypeT<FunctorType, bUnique, !bUseInline>;
+
+			void* NewAlloc;
+			if constexpr (bUseInline)
+			{
+				NewAlloc = &InlineAllocation;
+			}
+			else
+			{
+				NewAlloc = FMemory::Malloc(sizeof(OwnedType), alignof(OwnedType));
+				HeapAllocation = NewAlloc;
+			}
+
+			CA_ASSUME(NewAlloc);
+			auto* NewOwned = new (NewAlloc) OwnedType(Forward<FunctorType>(InFunc));
+			return &NewOwned->Obj;
+		}
+
+#else
+
 		#if TFUNCTION_USES_INLINE_STORAGE
 		template <typename FunctorType>
 		typename TEnableIf<
@@ -417,6 +471,8 @@ namespace UE4Function_Private
 
 			return &NewOwned->Obj;
 		}
+
+#endif
 	};
 
 	template <typename T, bool bOnHeap>
diff --git a/Engine/Source/Runtime/Core/Public/Templates/MemoryOps.h b/Engine/Source/Runtime/Core/Public/Templates/MemoryOps.h
--- a/Engine/Source/Runtime/Core/Public/Templates/MemoryOps.h
+++ b/Engine/Source/Runtime/Core/Public/Templates/MemoryOps.h
@@ -33,246 +33,530 @@ namespace UE4MemoryOps_Private
 	};
 }
 
-/**
- * Default constructs a range of items in memory.
- *
- * @param	Elements	The address of the first memory location to construct at.
- * @param	Count		The number of elements to destruct.
- */
-template <typename ElementType, typename SizeType>
-FORCEINLINE typename TEnableIf<!TIsZeroConstructType<ElementType>::Value>::Type DefaultConstructItems(void* Address, SizeType Count)
-{
-	ElementType* Element = (ElementType*)Address;
-	while (Count)
+#if PLATFORM_COMPILER_HAS_IF_CONSTEXPR
+
+	/**
+	 * Default constructs a range of items in memory.
+	 *
+	 * @param	Elements	The address of the first memory location to construct at.
+	 * @param	Count		The number of elements to destruct.
+	 */
+	template <typename ElementType, typename SizeType>
+	FORCEINLINE void DefaultConstructItems(void* Address, SizeType Count)
 	{
-		new (Element) ElementType;
-		++Element;
-		--Count;
+		if constexpr (!TIsZeroConstructType<ElementType>::Value)
+		{
+			ElementType* Element = (ElementType*)Address;
+			while (Count)
+			{
+				new (Element) ElementType;
+				++Element;
+				--Count;
+			}
+		}
+		else
+		{
+			FMemory::Memset(Address, 0, sizeof(ElementType) * Count);
+		}
 	}
-}
 
+#else
 
-template <typename ElementType, typename SizeType>
-FORCEINLINE typename TEnableIf<TIsZeroConstructType<ElementType>::Value>::Type DefaultConstructItems(void* Elements, SizeType Count)
-{
-	FMemory::Memset(Elements, 0, sizeof(ElementType) * Count);
-}
+	/**
+	 * Default constructs a range of items in memory.
+	 *
+	 * @param	Elements	The address of the first memory location to construct at.
+	 * @param	Count		The number of elements to destruct.
+	 */
+	template <typename ElementType, typename SizeType>
+	FORCEINLINE typename TEnableIf<!TIsZeroConstructType<ElementType>::Value>::Type DefaultConstructItems(void* Address, SizeType Count)
+	{
+		ElementType* Element = (ElementType*)Address;
+		while (Count)
+		{
+			new (Element) ElementType;
+			++Element;
+			--Count;
+		}
+	}
 
 
-/**
- * Destructs a single item in memory.
- *
- * @param	Elements	A pointer to the item to destruct.
- *
- * @note: This function is optimized for values of T, and so will not dynamically dispatch destructor calls if T's destructor is virtual.
- */
-template <typename ElementType>
-FORCEINLINE typename TEnableIf<!TIsTriviallyDestructible<ElementType>::Value>::Type DestructItem(ElementType* Element)
-{
-	// We need a typedef here because VC won't compile the destructor call below if ElementType itself has a member called ElementType
-	typedef ElementType DestructItemsElementTypeTypedef;
+	template <typename ElementType, typename SizeType>
+	FORCEINLINE typename TEnableIf<TIsZeroConstructType<ElementType>::Value>::Type DefaultConstructItems(void* Elements, SizeType Count)
+	{
+		FMemory::Memset(Elements, 0, sizeof(ElementType) * Count);
+	}
 
-	Element->DestructItemsElementTypeTypedef::~DestructItemsElementTypeTypedef();
-}
+#endif
 
+#if PLATFORM_COMPILER_HAS_IF_CONSTEXPR
 
-template <typename ElementType>
-FORCEINLINE typename TEnableIf<TIsTriviallyDestructible<ElementType>::Value>::Type DestructItem(ElementType* Element)
-{
-}
+	/**
+	 * Destructs a single item in memory.
+	 *
+	 * @param	Elements	A pointer to the item to destruct.
+	 *
+	 * @note: This function is optimized for values of T, and so will not dynamically dispatch destructor calls if T's destructor is virtual.
+	 */
+	template <typename ElementType>
+	FORCEINLINE void DestructItem(ElementType* Element)
+	{
+		if constexpr (!TIsTriviallyDestructible<ElementType>::Value)
+		{
+			// We need a typedef here because VC won't compile the destructor call below if ElementType itself has a member called ElementType
+			typedef ElementType DestructItemsElementTypeTypedef;
 
+			Element->DestructItemsElementTypeTypedef::~DestructItemsElementTypeTypedef();
+		}
+	}
 
-/**
- * Destructs a range of items in memory.
- *
- * @param	Elements	A pointer to the first item to destruct.
- * @param	Count		The number of elements to destruct.
- *
- * @note: This function is optimized for values of T, and so will not dynamically dispatch destructor calls if T's destructor is virtual.
- */
-template <typename ElementType, typename SizeType>
-FORCEINLINE typename TEnableIf<!TIsTriviallyDestructible<ElementType>::Value>::Type DestructItems(ElementType* Element, SizeType Count)
-{
-	while (Count)
+#else
+
+	/**
+	 * Destructs a single item in memory.
+	 *
+	 * @param	Elements	A pointer to the item to destruct.
+	 *
+	 * @note: This function is optimized for values of T, and so will not dynamically dispatch destructor calls if T's destructor is virtual.
+	 */
+	template <typename ElementType>
+	FORCEINLINE typename TEnableIf<!TIsTriviallyDestructible<ElementType>::Value>::Type DestructItem(ElementType* Element)
 	{
 		// We need a typedef here because VC won't compile the destructor call below if ElementType itself has a member called ElementType
 		typedef ElementType DestructItemsElementTypeTypedef;
 
 		Element->DestructItemsElementTypeTypedef::~DestructItemsElementTypeTypedef();
-		++Element;
-		--Count;
 	}
-}
 
 
-template <typename ElementType, typename SizeType>
-FORCEINLINE typename TEnableIf<TIsTriviallyDestructible<ElementType>::Value>::Type DestructItems(ElementType* Elements, SizeType Count)
-{
-}
+	template <typename ElementType>
+	FORCEINLINE typename TEnableIf<TIsTriviallyDestructible<ElementType>::Value>::Type DestructItem(ElementType* Element)
+	{
+	}
+#endif
 
+#if PLATFORM_COMPILER_HAS_IF_CONSTEXPR
 
-/**
- * Constructs a range of items into memory from a set of arguments.  The arguments come from an another array.
- *
- * @param	Dest		The memory location to start copying into.
- * @param	Source		A pointer to the first argument to pass to the constructor.
- * @param	Count		The number of elements to copy.
- */
-template <typename DestinationElementType, typename SourceElementType, typename SizeType>
-FORCEINLINE typename TEnableIf<!TIsBitwiseConstructible<DestinationElementType, SourceElementType>::Value>::Type ConstructItems(void* Dest, const SourceElementType* Source, SizeType Count)
-{
-	while (Count)
+	/**
+	 * Destructs a range of items in memory.
+	 *
+	 * @param	Elements	A pointer to the first item to destruct.
+	 * @param	Count		The number of elements to destruct.
+	 *
+	 * @note: This function is optimized for values of T, and so will not dynamically dispatch destructor calls if T's destructor is virtual.
+	 */
+	template <typename ElementType, typename SizeType>
+	FORCEINLINE void DestructItems(ElementType* Element, SizeType Count)
 	{
-		new (Dest) DestinationElementType(*Source);
-		++(DestinationElementType*&)Dest;
-		++Source;
-		--Count;
+		if constexpr (!TIsTriviallyDestructible<ElementType>::Value)
+		{
+			while (Count)
+			{
+				// We need a typedef here because VC won't compile the destructor call below if ElementType itself has a member called ElementType
+				typedef ElementType DestructItemsElementTypeTypedef;
+
+				Element->DestructItemsElementTypeTypedef::~DestructItemsElementTypeTypedef();
+				++Element;
+				--Count;
+			}
+		}
 	}
-}
 
+#else
 
-template <typename DestinationElementType, typename SourceElementType, typename SizeType>
-FORCEINLINE typename TEnableIf<TIsBitwiseConstructible<DestinationElementType, SourceElementType>::Value>::Type ConstructItems(void* Dest, const SourceElementType* Source, SizeType Count)
-{
-	FMemory::Memcpy(Dest, Source, sizeof(SourceElementType) * Count);
-}
+	/**
+	 * Destructs a range of items in memory.
+	 *
+	 * @param	Elements	A pointer to the first item to destruct.
+	 * @param	Count		The number of elements to destruct.
+	 *
+	 * @note: This function is optimized for values of T, and so will not dynamically dispatch destructor calls if T's destructor is virtual.
+	 */
+	template <typename ElementType, typename SizeType>
+	FORCEINLINE typename TEnableIf<!TIsTriviallyDestructible<ElementType>::Value>::Type DestructItems(ElementType* Element, SizeType Count)
+	{
+		while (Count)
+		{
+			// We need a typedef here because VC won't compile the destructor call below if ElementType itself has a member called ElementType
+			typedef ElementType DestructItemsElementTypeTypedef;
 
+			Element->DestructItemsElementTypeTypedef::~DestructItemsElementTypeTypedef();
+			++Element;
+			--Count;
+		}
+	}
 
-/**
- * Copy assigns a range of items.
- *
- * @param	Dest		The memory location to start assigning to.
- * @param	Source		A pointer to the first item to assign.
- * @param	Count		The number of elements to assign.
- */
-template <typename ElementType, typename SizeType>
-FORCEINLINE typename TEnableIf<!TIsTriviallyCopyAssignable<ElementType>::Value>::Type CopyAssignItems(ElementType* Dest, const ElementType* Source, SizeType Count)
-{
-	while (Count)
+
+	template <typename ElementType, typename SizeType>
+	FORCEINLINE typename TEnableIf<TIsTriviallyDestructible<ElementType>::Value>::Type DestructItems(ElementType* Elements, SizeType Count)
 	{
-		*Dest = *Source;
-		++Dest;
-		++Source;
-		--Count;
 	}
-}
 
+#endif
 
-template <typename ElementType, typename SizeType>
-FORCEINLINE typename TEnableIf<TIsTriviallyCopyAssignable<ElementType>::Value>::Type CopyAssignItems(ElementType* Dest, const ElementType* Source, SizeType Count)
-{
-	FMemory::Memcpy(Dest, Source, sizeof(ElementType) * Count);
-}
+#if PLATFORM_COMPILER_HAS_IF_CONSTEXPR
 
+	/**
+	 * Constructs a range of items into memory from a set of arguments.  The arguments come from an another array.
+	 *
+	 * @param	Dest		The memory location to start copying into.
+	 * @param	Source		A pointer to the first argument to pass to the constructor.
+	 * @param	Count		The number of elements to copy.
+	 */
+	template <typename DestinationElementType, typename SourceElementType, typename SizeType>
+	FORCEINLINE void ConstructItems(void* Dest, const SourceElementType* Source, SizeType Count)
+	{
+		if constexpr (!TIsBitwiseConstructible<DestinationElementType, SourceElementType>::Value)
+		{
+			while (Count)
+			{
+				new (Dest) DestinationElementType(*Source);
+				++(DestinationElementType*&)Dest;
+				++Source;
+				--Count;
+			}
+		}
+		else
+		{
+			FMemory::Memcpy(Dest, Source, sizeof(SourceElementType) * Count);
+		}
+	}
 
-/**
- * Relocates a range of items to a new memory location as a new type. This is a so-called 'destructive move' for which
- * there is no single operation in C++ but which can be implemented very efficiently in general.
- *
- * @param	Dest		The memory location to relocate to.
- * @param	Source		A pointer to the first item to relocate.
- * @param	Count		The number of elements to relocate.
- */
-template <typename DestinationElementType, typename SourceElementType, typename SizeType>
-FORCEINLINE typename TEnableIf<!UE4MemoryOps_Private::TCanBitwiseRelocate<DestinationElementType, SourceElementType>::Value>::Type RelocateConstructItems(void* Dest, const SourceElementType* Source, SizeType Count)
-{
-	while (Count)
+#else
+
+	/**
+	 * Constructs a range of items into memory from a set of arguments.  The arguments come from an another array.
+	 *
+	 * @param	Dest		The memory location to start copying into.
+	 * @param	Source		A pointer to the first argument to pass to the constructor.
+	 * @param	Count		The number of elements to copy.
+	 */
+	template <typename DestinationElementType, typename SourceElementType, typename SizeType>
+	FORCEINLINE typename TEnableIf<!TIsBitwiseConstructible<DestinationElementType, SourceElementType>::Value>::Type ConstructItems(void* Dest, const SourceElementType* Source, SizeType Count)
 	{
-		// We need a typedef here because VC won't compile the destructor call below if SourceElementType itself has a member called SourceElementType
-		typedef SourceElementType RelocateConstructItemsElementTypeTypedef;
+		while (Count)
+		{
+			new (Dest) DestinationElementType(*Source);
+			++(DestinationElementType*&)Dest;
+			++Source;
+			--Count;
+		}
+	}
+
 
-		new (Dest) DestinationElementType(*Source);
-		++(DestinationElementType*&)Dest;
-		(Source++)->RelocateConstructItemsElementTypeTypedef::~RelocateConstructItemsElementTypeTypedef();
-		--Count;
+	template <typename DestinationElementType, typename SourceElementType, typename SizeType>
+	FORCEINLINE typename TEnableIf<TIsBitwiseConstructible<DestinationElementType, SourceElementType>::Value>::Type ConstructItems(void* Dest, const SourceElementType* Source, SizeType Count)
+	{
+		FMemory::Memcpy(Dest, Source, sizeof(SourceElementType) * Count);
 	}
-}
 
-template <typename DestinationElementType, typename SourceElementType, typename SizeType>
-FORCEINLINE typename TEnableIf<UE4MemoryOps_Private::TCanBitwiseRelocate<DestinationElementType, SourceElementType>::Value>::Type RelocateConstructItems(void* Dest, const SourceElementType* Source, SizeType Count)
-{
-	/* All existing UE containers seem to assume trivial relocatability (i.e. memcpy'able) of their members,
-	 * so we're going to assume that this is safe here.  However, it's not generally possible to assume this
-	 * in general as objects which contain pointers/references to themselves are not safe to be trivially
-	 * relocated.
+#endif
+
+#if PLATFORM_COMPILER_HAS_IF_CONSTEXPR
+
+	/**
+	 * Copy assigns a range of items.
 	 *
-	 * However, it is not yet possible to automatically infer this at compile time, so we can't enable
-	 * different (i.e. safer) implementations anyway. */
+	 * @param	Dest		The memory location to start assigning to.
+	 * @param	Source		A pointer to the first item to assign.
+	 * @param	Count		The number of elements to assign.
+	 */
+	template <typename ElementType, typename SizeType>
+	FORCEINLINE void CopyAssignItems(ElementType* Dest, const ElementType* Source, SizeType Count)
+	{
+		if constexpr (!TIsTriviallyCopyAssignable<ElementType>::Value)
+		{
+			while (Count)
+			{
+				*Dest = *Source;
+				++Dest;
+				++Source;
+				--Count;
+			}
+		}
+		else
+		{
+			FMemory::Memcpy(Dest, Source, sizeof(ElementType) * Count);
+		}
+	}
 
-	FMemory::Memmove(Dest, Source, sizeof(SourceElementType) * Count);
-}
+#else
 
-/**
- * Move constructs a range of items into memory.
- *
- * @param	Dest		The memory location to start moving into.
- * @param	Source		A pointer to the first item to move from.
- * @param	Count		The number of elements to move.
- */
-template <typename ElementType, typename SizeType>
-FORCEINLINE typename TEnableIf<!TIsTriviallyCopyConstructible<ElementType>::Value>::Type MoveConstructItems(void* Dest, const ElementType* Source, SizeType Count)
-{
-	while (Count)
+	/**
+	 * Copy assigns a range of items.
+	 *
+	 * @param	Dest		The memory location to start assigning to.
+	 * @param	Source		A pointer to the first item to assign.
+	 * @param	Count		The number of elements to assign.
+	 */
+	template <typename ElementType, typename SizeType>
+	FORCEINLINE typename TEnableIf<!TIsTriviallyCopyAssignable<ElementType>::Value>::Type CopyAssignItems(ElementType* Dest, const ElementType* Source, SizeType Count)
+	{
+		while (Count)
+		{
+			*Dest = *Source;
+			++Dest;
+			++Source;
+			--Count;
+		}
+	}
+
+
+	template <typename ElementType, typename SizeType>
+	FORCEINLINE typename TEnableIf<TIsTriviallyCopyAssignable<ElementType>::Value>::Type CopyAssignItems(ElementType* Dest, const ElementType* Source, SizeType Count)
 	{
-		new (Dest) ElementType((ElementType&&)*Source);
-		++(ElementType*&)Dest;
-		++Source;
-		--Count;
+		FMemory::Memcpy(Dest, Source, sizeof(ElementType) * Count);
 	}
-}
 
-template <typename ElementType, typename SizeType>
-FORCEINLINE typename TEnableIf<TIsTriviallyCopyConstructible<ElementType>::Value>::Type MoveConstructItems(void* Dest, const ElementType* Source, SizeType Count)
-{
-	FMemory::Memmove(Dest, Source, sizeof(ElementType) * Count);
-}
+#endif
 
-/**
- * Move assigns a range of items.
- *
- * @param	Dest		The memory location to start move assigning to.
- * @param	Source		A pointer to the first item to move assign.
- * @param	Count		The number of elements to move assign.
- */
-template <typename ElementType, typename SizeType>
-FORCEINLINE typename TEnableIf<!TIsTriviallyCopyAssignable<ElementType>::Value>::Type MoveAssignItems(ElementType* Dest, const ElementType* Source, SizeType Count)
-{
-	while (Count)
+#if PLATFORM_COMPILER_HAS_IF_CONSTEXPR
+
+	/**
+	 * Relocates a range of items to a new memory location as a new type. This is a so-called 'destructive move' for which
+	 * there is no single operation in C++ but which can be implemented very efficiently in general.
+	 *
+	 * @param	Dest		The memory location to relocate to.
+	 * @param	Source		A pointer to the first item to relocate.
+	 * @param	Count		The number of elements to relocate.
+	 */
+	template <typename DestinationElementType, typename SourceElementType, typename SizeType>
+	FORCEINLINE void RelocateConstructItems(void* Dest, const SourceElementType* Source, SizeType Count)
 	{
-		*Dest = (ElementType&&)*Source;
-		++Dest;
-		++Source;
-		--Count;
+		if constexpr (!UE4MemoryOps_Private::TCanBitwiseRelocate<DestinationElementType, SourceElementType>::Value)
+		{
+			while (Count)
+			{
+				// We need a typedef here because VC won't compile the destructor call below if SourceElementType itself has a member called SourceElementType
+				typedef SourceElementType RelocateConstructItemsElementTypeTypedef;
+
+				new (Dest) DestinationElementType(*Source);
+				++(DestinationElementType*&)Dest;
+				(Source++)->RelocateConstructItemsElementTypeTypedef::~RelocateConstructItemsElementTypeTypedef();
+				--Count;
+			}
+		}
+		else
+		{
+			/* All existing UE containers seem to assume trivial relocatability (i.e. memcpy'able) of their members,
+			 * so we're going to assume that this is safe here.  However, it's not generally possible to assume this
+			 * in general as objects which contain pointers/references to themselves are not safe to be trivially
+			 * relocated.
+			 *
+			 * However, it is not yet possible to automatically infer this at compile time, so we can't enable
+			 * different (i.e. safer) implementations anyway. */
+
+			FMemory::Memmove(Dest, Source, sizeof(SourceElementType) * Count);
+		}
 	}
-}
 
-template <typename ElementType, typename SizeType>
-FORCEINLINE typename TEnableIf<TIsTriviallyCopyAssignable<ElementType>::Value>::Type MoveAssignItems(ElementType* Dest, const ElementType* Source, SizeType Count)
-{
-	FMemory::Memmove(Dest, Source, sizeof(ElementType) * Count);
-}
+#else
 
-template <typename ElementType, typename SizeType>
-FORCEINLINE typename TEnableIf<TTypeTraits<ElementType>::IsBytewiseComparable, bool>::Type CompareItems(const ElementType* A, const ElementType* B, SizeType Count)
-{
-	return !FMemory::Memcmp(A, B, sizeof(ElementType) * Count);
-}
+	/**
+	 * Relocates a range of items to a new memory location as a new type. This is a so-called 'destructive move' for which
+	 * there is no single operation in C++ but which can be implemented very efficiently in general.
+	 *
+	 * @param	Dest		The memory location to relocate to.
+	 * @param	Source		A pointer to the first item to relocate.
+	 * @param	Count		The number of elements to relocate.
+	 */
+	template <typename DestinationElementType, typename SourceElementType, typename SizeType>
+	FORCEINLINE typename TEnableIf<!UE4MemoryOps_Private::TCanBitwiseRelocate<DestinationElementType, SourceElementType>::Value>::Type RelocateConstructItems(void* Dest, const SourceElementType* Source, SizeType Count)
+	{
+		while (Count)
+		{
+			// We need a typedef here because VC won't compile the destructor call below if SourceElementType itself has a member called SourceElementType
+			typedef SourceElementType RelocateConstructItemsElementTypeTypedef;
 
+			new (Dest) DestinationElementType(*Source);
+			++(DestinationElementType*&)Dest;
+			(Source++)->RelocateConstructItemsElementTypeTypedef::~RelocateConstructItemsElementTypeTypedef();
+			--Count;
+		}
+	}
 
-template <typename ElementType, typename SizeType>
-FORCEINLINE typename TEnableIf<!TTypeTraits<ElementType>::IsBytewiseComparable, bool>::Type CompareItems(const ElementType* A, const ElementType* B, SizeType Count)
-{
-	while (Count)
+	template <typename DestinationElementType, typename SourceElementType, typename SizeType>
+	FORCEINLINE typename TEnableIf<UE4MemoryOps_Private::TCanBitwiseRelocate<DestinationElementType, SourceElementType>::Value>::Type RelocateConstructItems(void* Dest, const SourceElementType* Source, SizeType Count)
+	{
+		/* All existing UE containers seem to assume trivial relocatability (i.e. memcpy'able) of their members,
+		 * so we're going to assume that this is safe here.  However, it's not generally possible to assume this
+		 * in general as objects which contain pointers/references to themselves are not safe to be trivially
+		 * relocated.
+		 *
+		 * However, it is not yet possible to automatically infer this at compile time, so we can't enable
+		 * different (i.e. safer) implementations anyway. */
+
+		FMemory::Memmove(Dest, Source, sizeof(SourceElementType) * Count);
+	}
+
+#endif
+
+#if PLATFORM_COMPILER_HAS_IF_CONSTEXPR
+
+	/**
+	 * Move constructs a range of items into memory.
+	 *
+	 * @param	Dest		The memory location to start moving into.
+	 * @param	Source		A pointer to the first item to move from.
+	 * @param	Count		The number of elements to move.
+	 */
+	template <typename ElementType, typename SizeType>
+	FORCEINLINE void MoveConstructItems(void* Dest, const ElementType* Source, SizeType Count)
 	{
-		if (!(*A == *B))
+		if constexpr (!TIsTriviallyCopyConstructible<ElementType>::Value)
 		{
-			return false;
+			while (Count)
+			{
+				new (Dest) ElementType((ElementType&&)*Source);
+				++(ElementType*&)Dest;
+				++Source;
+				--Count;
+			}
 		}
+		else
+		{
+			FMemory::Memmove(Dest, Source, sizeof(ElementType) * Count);
+		}
+	}
 
-		++A;
-		++B;
-		--Count;
+#else
+
+	/**
+	 * Move constructs a range of items into memory.
+	 *
+	 * @param	Dest		The memory location to start moving into.
+	 * @param	Source		A pointer to the first item to move from.
+	 * @param	Count		The number of elements to move.
+	 */
+	template <typename ElementType, typename SizeType>
+	FORCEINLINE typename TEnableIf<!TIsTriviallyCopyConstructible<ElementType>::Value>::Type MoveConstructItems(void* Dest, const ElementType* Source, SizeType Count)
+	{
+		while (Count)
+		{
+			new (Dest) ElementType((ElementType&&)*Source);
+			++(ElementType*&)Dest;
+			++Source;
+			--Count;
+		}
 	}
 
-	return true;
-}
+	template <typename ElementType, typename SizeType>
+	FORCEINLINE typename TEnableIf<TIsTriviallyCopyConstructible<ElementType>::Value>::Type MoveConstructItems(void* Dest, const ElementType* Source, SizeType Count)
+	{
+		FMemory::Memmove(Dest, Source, sizeof(ElementType) * Count);
+	}
+
+#endif
+
+#if PLATFORM_COMPILER_HAS_IF_CONSTEXPR
+
+	/**
+	 * Move assigns a range of items.
+	 *
+	 * @param	Dest		The memory location to start move assigning to.
+	 * @param	Source		A pointer to the first item to move assign.
+	 * @param	Count		The number of elements to move assign.
+	 */
+	template <typename ElementType, typename SizeType>
+	FORCEINLINE void MoveAssignItems(ElementType* Dest, const ElementType* Source, SizeType Count)
+	{
+		if constexpr (!TIsTriviallyCopyAssignable<ElementType>::Value)
+		{
+			while (Count)
+			{
+				*Dest = (ElementType&&)*Source;
+				++Dest;
+				++Source;
+				--Count;
+			}
+		}
+		else
+		{
+			FMemory::Memmove(Dest, Source, sizeof(ElementType) * Count);
+		}
+	}
+
+#else
+
+	/**
+	 * Move assigns a range of items.
+	 *
+	 * @param	Dest		The memory location to start move assigning to.
+	 * @param	Source		A pointer to the first item to move assign.
+	 * @param	Count		The number of elements to move assign.
+	 */
+	template <typename ElementType, typename SizeType>
+	FORCEINLINE typename TEnableIf<!TIsTriviallyCopyAssignable<ElementType>::Value>::Type MoveAssignItems(ElementType* Dest, const ElementType* Source, SizeType Count)
+	{
+		while (Count)
+		{
+			*Dest = (ElementType&&)*Source;
+			++Dest;
+			++Source;
+			--Count;
+		}
+	}
+
+	template <typename ElementType, typename SizeType>
+	FORCEINLINE typename TEnableIf<TIsTriviallyCopyAssignable<ElementType>::Value>::Type MoveAssignItems(ElementType* Dest, const ElementType* Source, SizeType Count)
+	{
+		FMemory::Memmove(Dest, Source, sizeof(ElementType) * Count);
+	}
+
+#endif
+
+#if PLATFORM_COMPILER_HAS_IF_CONSTEXPR
+
+	template <typename ElementType, typename SizeType>
+	FORCEINLINE bool CompareItems(const ElementType* A, const ElementType* B, SizeType Count)
+	{
+		if constexpr (TTypeTraits<ElementType>::IsBytewiseComparable)
+		{
+			return !FMemory::Memcmp(A, B, sizeof(ElementType) * Count);
+		}
+		else
+		{
+			while (Count)
+			{
+				if (!(*A == *B))
+				{
+					return false;
+				}
+
+				++A;
+				++B;
+				--Count;
+			}
+
+			return true;
+		}
+	}
+
+#else
+
+	template <typename ElementType, typename SizeType>
+	FORCEINLINE typename TEnableIf<TTypeTraits<ElementType>::IsBytewiseComparable, bool>::Type CompareItems(const ElementType* A, const ElementType* B, SizeType Count)
+	{
+		return !FMemory::Memcmp(A, B, sizeof(ElementType) * Count);
+	}
+
+
+	template <typename ElementType, typename SizeType>
+	FORCEINLINE typename TEnableIf<!TTypeTraits<ElementType>::IsBytewiseComparable, bool>::Type CompareItems(const ElementType* A, const ElementType* B, SizeType Count)
+	{
+		while (Count)
+		{
+			if (!(*A == *B))
+			{
+				return false;
+			}
+
+			++A;
+			++B;
+			--Count;
+		}
+
+		return true;
+	}
+
+#endif
diff --git a/Engine/Source/Runtime/Core/Public/Unix/UnixPlatform.h b/Engine/Source/Runtime/Core/Public/Unix/UnixPlatform.h
--- a/Engine/Source/Runtime/Core/Public/Unix/UnixPlatform.h
+++ b/Engine/Source/Runtime/Core/Public/Unix/UnixPlatform.h
@@ -77,6 +77,12 @@ typedef FUnixPlatformTypes FPlatformTypes;
 	#define PLATFORM_COMPILER_HAS_DECLTYPE_AUTO 0
 #endif
 
+#if __has_feature(cxx_if_constexpr)
+	#define PLATFORM_COMPILER_HAS_IF_CONSTEXPR 1
+#else
+	#define PLATFORM_COMPILER_HAS_IF_CONSTEXPR 0
+#endif
+
 // SOCK_CLOEXEC is available on Unix since 2.6.27
 #if LINUX_VERSION_CODE >= KERNEL_VERSION(2,6,27)
 	#define PLATFORM_HAS_BSD_SOCKET_FEATURE_CLOSE_ON_EXEC	1
diff --git a/Engine/Source/Runtime/Core/Public/Windows/WIndowsPlatform.h b/Engine/Source/Runtime/Core/Public/Windows/WIndowsPlatform.h
--- a/Engine/Source/Runtime/Core/Public/Windows/WIndowsPlatform.h
+++ b/Engine/Source/Runtime/Core/Public/Windows/WIndowsPlatform.h
@@ -68,6 +68,12 @@ typedef FWindowsPlatformTypes FPlatformTypes;
 
 #define PLATFORM_GLOBAL_LOG_CATEGORY						LogWindows
 
+#if _MSC_FULL_VER >= 191125507 && defined(_MSVC_LANG) && _MSVC_LANG >= 201402
+	#define PLATFORM_COMPILER_HAS_IF_CONSTEXPR 1
+#else
+	#define PLATFORM_COMPILER_HAS_IF_CONSTEXPR 0
+#endif
+
 
 // Q: Why is there a __nop() before __debugbreak()?
 // A: VS' debug engine has a bug where it will silently swallow explicit
diff --git a/Engine/Source/Runtime/Core/Public/Windows/WindowsPlatformCompilerSetup.h b/Engine/Source/Runtime/Core/Public/Windows/WindowsPlatformCompilerSetup.h
--- a/Engine/Source/Runtime/Core/Public/Windows/WindowsPlatformCompilerSetup.h
+++ b/Engine/Source/Runtime/Core/Public/Windows/WindowsPlatformCompilerSetup.h
@@ -35,7 +35,7 @@ static_assert(_MSC_VER != 1914 && _MSC_VER != 1915, "Visual Studio 2017 versions
 #pragma warning (error: 4600 4601 4602 4603 4604      4606 4607      4609 4610 4611 4612 4613 4614 4615 4616 4617 4618      4620 4621 4622 4623 4624 4625 4626 4627 4628 4629 4630 4631 4632 4633 4634 4635 4636 4637 4638 4639 4640 4641 4642      4644 4645 4646      4648 4649 4650      4652 4653 4654 4655 4656 4657 4658 4659 4660 4661 4662 4663 4664 4665 4666 4667      4669 4670 4671 4672 4673 4674 4675 4676 4677 4678 4679 4680 4681 4682 4683 4684 4685 4686 4687 4688 4689 4690 4691      4693 4694 4695 4696 4697 4698 4699)
 #pragma warning (error: 4700      4702 4703 4704 4705 4706 4707 4708 4709 4710 4711 4712 4713 4714 4715 4716 4717 4718 4719 4720 4721 4722 4723 4724 4725 4726 4727 4728 4729      4731 4732 4733 4734 4735 4736 4737      4739 4740 4741 4742 4743 4744 4745 4746 4747 4748 4749 4750 4751 4752 4753 4754 4755 4756 4757 4758 4759 4760 4761 4762 4763 4764 4765 4766           4769      4771 4772 4773      4775 4776 4777 4778 4779 4780 4781 4782 4783 4784 4785 4786 4787 4788 4789 4790 4791 4792 4793 4794 4795 4796 4797 4798 4799)
 #pragma warning (error: 4800 4801 4802 4803 4804 4805 4806 4807 4808 4809 4810 4811 4812 4813 4814 4815 4816 4817 4818           4821 4822 4823 4824 4825      4827      4829 4830 4831 4832 4833 4834 4835 4836           4839 4840 4841 4842 4843 4844 4845 4846 4847 4848 4849 4850 4851 4852 4853 4854 4855 4856 4857 4858 4859 4860 4861 4862 4863 4864 4865      4867      4869 4870 4871 4872 4873 4874 4875 4876 4877 4878 4879 4880 4881 4882 4883 4884 4885 4886 4887 4888 4889 4890 4891 4892 4893 4894 4895 4896 4897 4898 4899)
-#pragma warning (error: 4900 4901 4902 4903 4904 4905 4906 4907 4908 4909 4910 4911 4912 4913 4914 4915 4916 4917 4918 4919 4920 4921 4922 4923 4924 4925 4926 4927 4928 4929 4930 4931 4932 4933 4934 4935 4936 4937 4938 4939 4940 4941 4942 4943 4944 4945 4946 4947 4948 4949 4950 4951 4952 4953 4954 4955 4956 4957 4958 4959 4960 4961      4963 4964 4965 4966 4967 4968 4969 4970 4971 4972 4973 4974 4975 4976 4977 4978 4979 4980 4981 4982 4983 4984 4985                4989 4990 4991 4992 4993 4994 4995      4997 4998 4999)
+#pragma warning (error: 4900 4901 4902 4903 4904 4905 4906 4907 4908 4909 4910 4911 4912 4913 4914 4915 4916 4917 4918 4919 4920 4921 4922 4923 4924 4925 4926 4927 4928 4929 4930 4931 4932 4933 4934 4935 4936 4937 4938 4939 4940 4941 4942 4943 4944 4945 4946 4947 4948 4949 4950 4951 4952 4953 4954 4955 4956 4957 4958 4959 4960 4961      4963 4964 4965 4966 4967 4968 4969 4970 4971 4972 4973 4974 4975 4976 4977 4978 4979 4980 4981 4982 4983      4985                4989 4990 4991 4992 4993 4994 4995      4997 4998 4999)
 #pragma warning (error:                                                                                                                                                                                               5038                                                                                                                                                                                                                                                                                                                 )
 
 //
@@ -237,6 +237,7 @@ static_assert(_MSC_VER != 1914 && _MSC_VER != 1915, "Visual Studio 2017 versions
 #pragma warning(default: 4931) // we are assuming the type library was built for number-bit pointers								https://docs.microsoft.com/en-us/cpp/error-messages/compiler-warnings/compiler-warning-level-4-c4931
 #pragma warning(default: 4946) // reinterpret_cast used between related classes: 'class1' and 'class2'								https://docs.microsoft.com/en-us/cpp/error-messages/compiler-warnings/compiler-warning-level-1-c4946
 #pragma warning(default: 5038) // data member 'A::y' will be initialized after data member 'A::x'									https://docs.microsoft.com/en-us/cpp/error-messages/compiler-warnings/c5038
+#pragma warning(disable: 4984) // 'if constexpr' is a C++17 language extension
 
 #pragma warning(default: 4996)
 
diff --git a/Engine/Source/Runtime/CoreUObject/Private/UObject/Class.cpp b/Engine/Source/Runtime/CoreUObject/Private/UObject/Class.cpp
--- a/Engine/Source/Runtime/CoreUObject/Private/UObject/Class.cpp
+++ b/Engine/Source/Runtime/CoreUObject/Private/UObject/Class.cpp
@@ -2487,7 +2487,7 @@ uint32 UScriptStruct::GetStructTypeHash(const void* Src) const
 	// UMapProperty::ConvertFromType).
 
 	UScriptStruct::ICppStructOps* TheCppStructOps = GetCppStructOps();
-	return TheCppStructOps->GetTypeHash(Src);
+	return TheCppStructOps->GetStructTypeHash(Src);
 }
 
 void UScriptStruct::InitializeStruct(void* InDest, int32 ArrayDim) const
diff --git a/Engine/Source/Runtime/CoreUObject/Public/UObject/Class.h b/Engine/Source/Runtime/CoreUObject/Public/UObject/Class.h
--- a/Engine/Source/Runtime/CoreUObject/Public/UObject/Class.h
+++ b/Engine/Source/Runtime/CoreUObject/Public/UObject/Class.h
@@ -600,257 +600,277 @@ struct TStructOpsTypeTraits : public TStructOpsTypeTraitsBase2<CPPSTRUCT>
 };
 
 
-/**
- * Selection of constructor behavior.
- */
-template<class CPPSTRUCT>
-FORCEINLINE typename TEnableIf<!TStructOpsTypeTraits<CPPSTRUCT>::WithNoInitConstructor>::Type ConstructWithNoInitOrNot(void *Data)
-{
-	new (Data) CPPSTRUCT();
-}
+#if !PLATFORM_COMPILER_HAS_IF_CONSTEXPR
 
-template<class CPPSTRUCT>
-FORCEINLINE typename TEnableIf<TStructOpsTypeTraits<CPPSTRUCT>::WithNoInitConstructor>::Type ConstructWithNoInitOrNot(void *Data)
-{
-	new (Data) CPPSTRUCT(ForceInit);
-}
+	/**
+	 * Selection of constructor behavior.
+	 */
+	template<class CPPSTRUCT>
+	FORCEINLINE typename TEnableIf<!TStructOpsTypeTraits<CPPSTRUCT>::WithNoInitConstructor>::Type ConstructWithNoInitOrNot(void *Data)
+	{
+		new (Data) CPPSTRUCT();
+	}
 
+	template<class CPPSTRUCT>
+	FORCEINLINE typename TEnableIf<TStructOpsTypeTraits<CPPSTRUCT>::WithNoInitConstructor>::Type ConstructWithNoInitOrNot(void *Data)
+	{
+		new (Data) CPPSTRUCT(ForceInit);
+	}
 
-/**
- * Selection of Serialize call.
- */
-template<class CPPSTRUCT>
-FORCEINLINE typename TEnableIf<!TStructOpsTypeTraits<CPPSTRUCT>::WithSerializer, bool>::Type SerializeOrNot(FArchive& Ar, CPPSTRUCT *Data)
-{
-	return false;
-}
 
-template<class CPPSTRUCT>
-FORCEINLINE typename TEnableIf<TStructOpsTypeTraits<CPPSTRUCT>::WithSerializer, bool>::Type SerializeOrNot(FArchive& Ar, CPPSTRUCT *Data)
-{
-	return Data->Serialize(Ar);
-}
+	/**
+	 * Selection of Serialize call.
+	 */
+	template<class CPPSTRUCT>
+	FORCEINLINE typename TEnableIf<!TStructOpsTypeTraits<CPPSTRUCT>::WithSerializer, bool>::Type SerializeOrNot(FArchive& Ar, CPPSTRUCT *Data)
+	{
+		return false;
+	}
 
-/**
-* Selection of structured Serialize call.
-*/
-template<class CPPSTRUCT>
-FORCEINLINE typename TEnableIf<!TStructOpsTypeTraits<CPPSTRUCT>::WithStructuredSerializer, bool>::Type SerializeOrNot(FStructuredArchive::FSlot Slot, CPPSTRUCT *Data)
-{
-	return false;
-}
+	template<class CPPSTRUCT>
+	FORCEINLINE typename TEnableIf<TStructOpsTypeTraits<CPPSTRUCT>::WithSerializer, bool>::Type SerializeOrNot(FArchive& Ar, CPPSTRUCT *Data)
+	{
+		return Data->Serialize(Ar);
+	}
 
-template<class CPPSTRUCT>
-FORCEINLINE typename TEnableIf<TStructOpsTypeTraits<CPPSTRUCT>::WithStructuredSerializer, bool>::Type SerializeOrNot(FStructuredArchive::FSlot Slot, CPPSTRUCT *Data)
-{
-	return Data->Serialize(Slot);
-}
+	/**
+	* Selection of structured Serialize call.
+	*/
+	template<class CPPSTRUCT>
+	FORCEINLINE typename TEnableIf<!TStructOpsTypeTraits<CPPSTRUCT>::WithStructuredSerializer, bool>::Type SerializeOrNot(FStructuredArchive::FSlot Slot, CPPSTRUCT *Data)
+	{
+		return false;
+	}
 
+	template<class CPPSTRUCT>
+	FORCEINLINE typename TEnableIf<TStructOpsTypeTraits<CPPSTRUCT>::WithStructuredSerializer, bool>::Type SerializeOrNot(FStructuredArchive::FSlot Slot, CPPSTRUCT *Data)
+	{
+		return Data->Serialize(Slot);
+	}
 
-/**
- * Selection of PostSerialize call.
- */
-template<class CPPSTRUCT>
-FORCEINLINE typename TEnableIf<!TStructOpsTypeTraits<CPPSTRUCT>::WithPostSerialize>::Type PostSerializeOrNot(const FArchive& Ar, CPPSTRUCT *Data)
-{
-}
 
-template<class CPPSTRUCT>
-FORCEINLINE typename TEnableIf<TStructOpsTypeTraits<CPPSTRUCT>::WithPostSerialize>::Type PostSerializeOrNot(const FArchive& Ar, CPPSTRUCT *Data)
-{
-	Data->PostSerialize(Ar);
-}
+	/**
+	 * Selection of PostSerialize call.
+	 */
+	template<class CPPSTRUCT>
+	FORCEINLINE typename TEnableIf<!TStructOpsTypeTraits<CPPSTRUCT>::WithPostSerialize>::Type PostSerializeOrNot(const FArchive& Ar, CPPSTRUCT *Data)
+	{
+	}
 
+	template<class CPPSTRUCT>
+	FORCEINLINE typename TEnableIf<TStructOpsTypeTraits<CPPSTRUCT>::WithPostSerialize>::Type PostSerializeOrNot(const FArchive& Ar, CPPSTRUCT *Data)
+	{
+		Data->PostSerialize(Ar);
+	}
 
-/**
- * Selection of NetSerialize call.
- */
-template<class CPPSTRUCT>
-FORCEINLINE typename TEnableIf<!TStructOpsTypeTraits<CPPSTRUCT>::WithNetSerializer, bool>::Type NetSerializeOrNot(FArchive& Ar, UPackageMap* Map, bool& bOutSuccess, CPPSTRUCT *Data)
-{
-	return false;
-}
 
-template<class CPPSTRUCT>
-FORCEINLINE typename TEnableIf<TStructOpsTypeTraits<CPPSTRUCT>::WithNetSerializer, bool>::Type NetSerializeOrNot(FArchive& Ar, UPackageMap* Map, bool& bOutSuccess, CPPSTRUCT *Data)
-{
-	return Data->NetSerialize(Ar, Map, bOutSuccess);
-}
+	/**
+	 * Selection of NetSerialize call.
+	 */
+	template<class CPPSTRUCT>
+	FORCEINLINE typename TEnableIf<!TStructOpsTypeTraits<CPPSTRUCT>::WithNetSerializer, bool>::Type NetSerializeOrNot(FArchive& Ar, UPackageMap* Map, bool& bOutSuccess, CPPSTRUCT *Data)
+	{
+		return false;
+	}
 
+	template<class CPPSTRUCT>
+	FORCEINLINE typename TEnableIf<TStructOpsTypeTraits<CPPSTRUCT>::WithNetSerializer, bool>::Type NetSerializeOrNot(FArchive& Ar, UPackageMap* Map, bool& bOutSuccess, CPPSTRUCT *Data)
+	{
+		return Data->NetSerialize(Ar, Map, bOutSuccess);
+	}
 
-/**
- * Selection of NetDeltaSerialize call.
- */
-template<class CPPSTRUCT>
-FORCEINLINE typename TEnableIf<!TStructOpsTypeTraits<CPPSTRUCT>::WithNetDeltaSerializer, bool>::Type NetDeltaSerializeOrNot(FNetDeltaSerializeInfo & DeltaParms, CPPSTRUCT *Data)
-{
-	return false;
-}
 
-template<class CPPSTRUCT>
-FORCEINLINE typename TEnableIf<TStructOpsTypeTraits<CPPSTRUCT>::WithNetDeltaSerializer, bool>::Type NetDeltaSerializeOrNot(FNetDeltaSerializeInfo & DeltaParms, CPPSTRUCT *Data)
-{
-	return Data->NetDeltaSerialize(DeltaParms);
-}
+	/**
+	 * Selection of NetDeltaSerialize call.
+	 */
+	template<class CPPSTRUCT>
+	FORCEINLINE typename TEnableIf<!TStructOpsTypeTraits<CPPSTRUCT>::WithNetDeltaSerializer, bool>::Type NetDeltaSerializeOrNot(FNetDeltaSerializeInfo & DeltaParms, CPPSTRUCT *Data)
+	{
+		return false;
+	}
 
+	template<class CPPSTRUCT>
+	FORCEINLINE typename TEnableIf<TStructOpsTypeTraits<CPPSTRUCT>::WithNetDeltaSerializer, bool>::Type NetDeltaSerializeOrNot(FNetDeltaSerializeInfo & DeltaParms, CPPSTRUCT *Data)
+	{
+		return Data->NetDeltaSerialize(DeltaParms);
+	}
 
-/**
- * Selection of PostScriptConstruct call.
- */
-template<class CPPSTRUCT>
-FORCEINLINE typename TEnableIf<!TStructOpsTypeTraits<CPPSTRUCT>::WithPostScriptConstruct>::Type PostScriptConstructOrNot(CPPSTRUCT *Data)
-{
-}
 
-template<class CPPSTRUCT>
-FORCEINLINE typename TEnableIf<TStructOpsTypeTraits<CPPSTRUCT>::WithPostScriptConstruct>::Type PostScriptConstructOrNot(CPPSTRUCT *Data)
-{
-	Data->PostScriptConstruct();
-}
+	/**
+	 * Selection of PostScriptConstruct call.
+	 */
+	template<class CPPSTRUCT>
+	FORCEINLINE typename TEnableIf<!TStructOpsTypeTraits<CPPSTRUCT>::WithPostScriptConstruct>::Type PostScriptConstructOrNot(CPPSTRUCT *Data)
+	{
+	}
 
+	template<class CPPSTRUCT>
+	FORCEINLINE typename TEnableIf<TStructOpsTypeTraits<CPPSTRUCT>::WithPostScriptConstruct>::Type PostScriptConstructOrNot(CPPSTRUCT *Data)
+	{
+		Data->PostScriptConstruct();
+	}
 
-/**
- * Selection of Copy behavior.
- */
-template<class CPPSTRUCT>
-FORCEINLINE typename TEnableIf<!TStructOpsTypeTraits<CPPSTRUCT>::WithCopy, bool>::Type CopyOrNot(CPPSTRUCT* Dest, CPPSTRUCT const* Src, int32 ArrayDim)
-{
-	return false;
-}
 
-template<class CPPSTRUCT>
-FORCEINLINE typename TEnableIf<TStructOpsTypeTraits<CPPSTRUCT>::WithCopy, bool>::Type CopyOrNot(CPPSTRUCT* Dest, CPPSTRUCT const* Src, int32 ArrayDim)
-{
-	static_assert((!TIsPODType<CPPSTRUCT>::Value), "You probably don't want custom copy for a POD type.");
-	for (;ArrayDim;--ArrayDim)
+	/**
+	 * Selection of Copy behavior.
+	 */
+	template<class CPPSTRUCT>
+	FORCEINLINE typename TEnableIf<!TStructOpsTypeTraits<CPPSTRUCT>::WithCopy, bool>::Type CopyOrNot(CPPSTRUCT* Dest, CPPSTRUCT const* Src, int32 ArrayDim)
 	{
-		*Dest++ = *Src++;
+		return false;
 	}
-	return true;
-}
 
+	template<class CPPSTRUCT>
+	FORCEINLINE typename TEnableIf<TStructOpsTypeTraits<CPPSTRUCT>::WithCopy, bool>::Type CopyOrNot(CPPSTRUCT* Dest, CPPSTRUCT const* Src, int32 ArrayDim)
+	{
+		static_assert((!TIsPODType<CPPSTRUCT>::Value), "You probably don't want custom copy for a POD type.");
+		for (; ArrayDim; --ArrayDim)
+		{
+			*Dest++ = *Src++;
+		}
+		return true;
+	}
 
-/**
- * Selection of AddStructReferencedObjects check.
- */
-template<class CPPSTRUCT>
-FORCEINLINE typename TEnableIf<!TStructOpsTypeTraits<CPPSTRUCT>::WithAddStructReferencedObjects>::Type AddStructReferencedObjectsOrNot(const void* A, FReferenceCollector& Collector)
-{
-}
 
-template<class CPPSTRUCT>
-FORCEINLINE typename TEnableIf<TStructOpsTypeTraits<CPPSTRUCT>::WithAddStructReferencedObjects>::Type AddStructReferencedObjectsOrNot(const void* A, FReferenceCollector& Collector)
-{
-	((CPPSTRUCT const*)A)->AddStructReferencedObjects(Collector);
-}
+	/**
+	 * Selection of Identical check.
+	 */
+	template<class CPPSTRUCT>
+	FORCEINLINE typename TEnableIf<TStructOpsTypeTraits<CPPSTRUCT>::WithIdentical && TStructOpsTypeTraits<CPPSTRUCT>::WithIdenticalViaEquality, bool>::Type IdenticalOrNot(const CPPSTRUCT* A, const CPPSTRUCT* B, uint32 PortFlags, bool& bOutResult)
+	{
+		static_assert(sizeof(CPPSTRUCT) == 0, "Should not have both WithIdenticalViaEquality and WithIdentical.");
+	}
 
+	template<class CPPSTRUCT>
+	FORCEINLINE typename TEnableIf<!TStructOpsTypeTraits<CPPSTRUCT>::WithIdentical && !TStructOpsTypeTraits<CPPSTRUCT>::WithIdenticalViaEquality, bool>::Type IdenticalOrNot(const CPPSTRUCT* A, const CPPSTRUCT* B, uint32 PortFlags, bool& bOutResult)
+	{
+		bOutResult = false;
+		return false;
+	}
 
-/**
- * Selection of Identical check.
- */
-template<class CPPSTRUCT>
-FORCEINLINE typename TEnableIf<TStructOpsTypeTraits<CPPSTRUCT>::WithIdentical && TStructOpsTypeTraits<CPPSTRUCT>::WithIdenticalViaEquality, bool>::Type IdenticalOrNot(const CPPSTRUCT* A, const CPPSTRUCT* B, uint32 PortFlags, bool& bOutResult)
-{
-	static_assert(sizeof(CPPSTRUCT) == 0, "Should not have both WithIdenticalViaEquality and WithIdentical.");
-}
+	template<class CPPSTRUCT>
+	FORCEINLINE typename TEnableIf<TStructOpsTypeTraits<CPPSTRUCT>::WithIdentical && !TStructOpsTypeTraits<CPPSTRUCT>::WithIdenticalViaEquality, bool>::Type IdenticalOrNot(const CPPSTRUCT* A, const CPPSTRUCT* B, uint32 PortFlags, bool& bOutResult)
+	{
+		bOutResult = A->Identical(B, PortFlags);
+		return true;
+	}
 
-template<class CPPSTRUCT>
-FORCEINLINE typename TEnableIf<!TStructOpsTypeTraits<CPPSTRUCT>::WithIdentical && !TStructOpsTypeTraits<CPPSTRUCT>::WithIdenticalViaEquality, bool>::Type IdenticalOrNot(const CPPSTRUCT* A, const CPPSTRUCT* B, uint32 PortFlags, bool& bOutResult)
-{
-	bOutResult = false;
-	return false;
-}
+	template<class CPPSTRUCT>
+	FORCEINLINE typename TEnableIf<!TStructOpsTypeTraits<CPPSTRUCT>::WithIdentical && TStructOpsTypeTraits<CPPSTRUCT>::WithIdenticalViaEquality, bool>::Type IdenticalOrNot(const CPPSTRUCT* A, const CPPSTRUCT* B, uint32 PortFlags, bool& bOutResult)
+	{
+		bOutResult = (*A == *B);
+		return true;
+	}
 
-template<class CPPSTRUCT>
-FORCEINLINE typename TEnableIf<TStructOpsTypeTraits<CPPSTRUCT>::WithIdentical && !TStructOpsTypeTraits<CPPSTRUCT>::WithIdenticalViaEquality, bool>::Type IdenticalOrNot(const CPPSTRUCT* A, const CPPSTRUCT* B, uint32 PortFlags, bool& bOutResult)
-{
-	bOutResult = A->Identical(B, PortFlags);
-	return true;
-}
 
-template<class CPPSTRUCT>
-FORCEINLINE typename TEnableIf<!TStructOpsTypeTraits<CPPSTRUCT>::WithIdentical && TStructOpsTypeTraits<CPPSTRUCT>::WithIdenticalViaEquality, bool>::Type IdenticalOrNot(const CPPSTRUCT* A, const CPPSTRUCT* B, uint32 PortFlags, bool& bOutResult)
-{
-	bOutResult = (*A == *B);
-	return true;
-}
+	/**
+	 * Selection of ExportTextItem call.
+	 */
+	template<class CPPSTRUCT>
+	FORCEINLINE typename TEnableIf<!TStructOpsTypeTraits<CPPSTRUCT>::WithExportTextItem, bool>::Type ExportTextItemOrNot(FString& ValueStr, const CPPSTRUCT* PropertyValue, const CPPSTRUCT* DefaultValue, UObject* Parent, int32 PortFlags, UObject* ExportRootScope)
+	{
+		return false;
+	}
 
+	template<class CPPSTRUCT>
+	FORCEINLINE typename TEnableIf<TStructOpsTypeTraits<CPPSTRUCT>::WithExportTextItem, bool>::Type ExportTextItemOrNot(FString& ValueStr, const CPPSTRUCT* PropertyValue, const CPPSTRUCT* DefaultValue, UObject* Parent, int32 PortFlags, UObject* ExportRootScope)
+	{
+		return PropertyValue->ExportTextItem(ValueStr, *DefaultValue, Parent, PortFlags, ExportRootScope);
+	}
 
-/**
- * Selection of ExportTextItem call.
- */
-template<class CPPSTRUCT>
-FORCEINLINE typename TEnableIf<!TStructOpsTypeTraits<CPPSTRUCT>::WithExportTextItem, bool>::Type ExportTextItemOrNot(FString& ValueStr, const CPPSTRUCT* PropertyValue, const CPPSTRUCT* DefaultValue, UObject* Parent, int32 PortFlags, UObject* ExportRootScope)
-{
-	return false;
-}
 
-template<class CPPSTRUCT>
-FORCEINLINE typename TEnableIf<TStructOpsTypeTraits<CPPSTRUCT>::WithExportTextItem, bool>::Type ExportTextItemOrNot(FString& ValueStr, const CPPSTRUCT* PropertyValue, const CPPSTRUCT* DefaultValue, UObject* Parent, int32 PortFlags, UObject* ExportRootScope)
-{
-	return PropertyValue->ExportTextItem(ValueStr, *DefaultValue, Parent, PortFlags, ExportRootScope);
-}
+	/**
+	 * Selection of ImportTextItem call.
+	 */
+	template<class CPPSTRUCT>
+	FORCEINLINE typename TEnableIf<!TStructOpsTypeTraits<CPPSTRUCT>::WithImportTextItem, bool>::Type ImportTextItemOrNot(const TCHAR*& Buffer, CPPSTRUCT* Data, int32 PortFlags, UObject* OwnerObject, FOutputDevice* ErrorText)
+	{
+		return false;
+	}
 
+	template<class CPPSTRUCT>
+	FORCEINLINE typename TEnableIf<TStructOpsTypeTraits<CPPSTRUCT>::WithImportTextItem, bool>::Type ImportTextItemOrNot(const TCHAR*& Buffer, CPPSTRUCT* Data, int32 PortFlags, UObject* OwnerObject, FOutputDevice* ErrorText)
+	{
+		return Data->ImportTextItem(Buffer, PortFlags, OwnerObject, ErrorText);
+	}
 
-/**
- * Selection of ImportTextItem call.
- */
-template<class CPPSTRUCT>
-FORCEINLINE typename TEnableIf<!TStructOpsTypeTraits<CPPSTRUCT>::WithImportTextItem, bool>::Type ImportTextItemOrNot(const TCHAR*& Buffer, CPPSTRUCT* Data, int32 PortFlags, UObject* OwnerObject, FOutputDevice* ErrorText)
-{
-	return false;
-}
 
-template<class CPPSTRUCT>
-FORCEINLINE typename TEnableIf<TStructOpsTypeTraits<CPPSTRUCT>::WithImportTextItem, bool>::Type ImportTextItemOrNot(const TCHAR*& Buffer, CPPSTRUCT* Data, int32 PortFlags, UObject* OwnerObject, FOutputDevice* ErrorText)
-{
-	return Data->ImportTextItem(Buffer, PortFlags, OwnerObject, ErrorText);
-}
+	/**
+	 * Selection of SerializeFromMismatchedTag call.
+	 */
+	template<class CPPSTRUCT>
+	FORCEINLINE typename TEnableIf<!TStructOpsTypeTraits<CPPSTRUCT>::WithSerializeFromMismatchedTag, bool>::Type SerializeFromMismatchedTagOrNot(FPropertyTag const& Tag, FArchive& Ar, CPPSTRUCT *Data)
+	{
+		return false;
+	}
 
+	template<class CPPSTRUCT>
+	FORCEINLINE typename TEnableIf<TStructOpsTypeTraits<CPPSTRUCT>::WithSerializeFromMismatchedTag, bool>::Type SerializeFromMismatchedTagOrNot(FPropertyTag const& Tag, FArchive& Ar, CPPSTRUCT *Data)
+	{
+		return Data->SerializeFromMismatchedTag(Tag, Ar);
+	}
 
-/**
- * Selection of SerializeFromMismatchedTag call.
- */
-template<class CPPSTRUCT>
-FORCEINLINE typename TEnableIf<!TStructOpsTypeTraits<CPPSTRUCT>::WithSerializeFromMismatchedTag, bool>::Type SerializeFromMismatchedTagOrNot(FPropertyTag const& Tag, FArchive& Ar, CPPSTRUCT *Data)
-{
-	return false;
-}
+	template<class CPPSTRUCT>
+	FORCEINLINE typename TEnableIf<!TStructOpsTypeTraits<CPPSTRUCT>::WithStructuredSerializeFromMismatchedTag, bool>::Type StructuredSerializeFromMismatchedTagOrNot(FPropertyTag const& Tag, FStructuredArchive::FSlot Slot, CPPSTRUCT *Data)
+	{
+		return false;
+	}
 
-template<class CPPSTRUCT>
-FORCEINLINE typename TEnableIf<TStructOpsTypeTraits<CPPSTRUCT>::WithSerializeFromMismatchedTag, bool>::Type SerializeFromMismatchedTagOrNot(FPropertyTag const& Tag, FArchive& Ar, CPPSTRUCT *Data)
-{
-	return Data->SerializeFromMismatchedTag(Tag, Ar);
-}
+	template<class CPPSTRUCT>
+	FORCEINLINE typename TEnableIf<TStructOpsTypeTraits<CPPSTRUCT>::WithStructuredSerializeFromMismatchedTag, bool>::Type StructuredSerializeFromMismatchedTagOrNot(FPropertyTag const& Tag, FStructuredArchive::FSlot Slot, CPPSTRUCT *Data)
+	{
+		return Data->SerializeFromMismatchedTag(Tag, Slot);
+	}
 
-template<class CPPSTRUCT>
-FORCEINLINE typename TEnableIf<!TStructOpsTypeTraits<CPPSTRUCT>::WithStructuredSerializeFromMismatchedTag, bool>::Type StructuredSerializeFromMismatchedTagOrNot(FPropertyTag const& Tag, FStructuredArchive::FSlot Slot, CPPSTRUCT *Data)
-{
-	return false;
-}
 
-template<class CPPSTRUCT>
-FORCEINLINE typename TEnableIf<TStructOpsTypeTraits<CPPSTRUCT>::WithStructuredSerializeFromMismatchedTag, bool>::Type StructuredSerializeFromMismatchedTagOrNot(FPropertyTag const& Tag, FStructuredArchive::FSlot Slot, CPPSTRUCT *Data)
-{
-	return Data->SerializeFromMismatchedTag(Tag, Slot);
-}
+	/**
+	 * Selection of GetTypeHash call.
+	 */
+	template<class CPPSTRUCT>
+	FORCEINLINE typename TEnableIf<!TModels<CGetTypeHashable, CPPSTRUCT>::Value, uint32>::Type GetTypeHashOrNot(const CPPSTRUCT *Data)
+	{
+		return 0;
+	}
+
+	template<class CPPSTRUCT>
+	FORCEINLINE typename TEnableIf<TModels<CGetTypeHashable, CPPSTRUCT>::Value, uint32>::Type GetTypeHashOrNot(const CPPSTRUCT *Data)
+	{
+		return GetTypeHash(*Data);
+	}
 
+#endif
 
-/**
- * Selection of GetTypeHash call.
- */
-template<class CPPSTRUCT>
-FORCEINLINE typename TEnableIf<!TModels<CGetTypeHashable, CPPSTRUCT>::Value, uint32>::Type GetTypeHashOrNot(const CPPSTRUCT *Data)
-{
-	return 0;
-}
+#if PLATFORM_COMPILER_HAS_IF_CONSTEXPR
 
-template<class CPPSTRUCT>
-FORCEINLINE typename TEnableIf<TModels<CGetTypeHashable, CPPSTRUCT>::Value, uint32>::Type GetTypeHashOrNot(const CPPSTRUCT *Data)
-{
-	return GetTypeHash(*Data);
-}
+	/**
+	 * Selection of AddStructReferencedObjects check.
+	 */
+	template<class CPPSTRUCT>
+	FORCEINLINE void AddStructReferencedObjectsOrNot(const void* A, FReferenceCollector& Collector)
+	{
+		if constexpr (TStructOpsTypeTraits<CPPSTRUCT>::WithAddStructReferencedObjects)
+		{
+			((CPPSTRUCT const*)A)->AddStructReferencedObjects(Collector);
+		}
+	}
+
+#else
+
+	/**
+	 * Selection of AddStructReferencedObjects check.
+	 */
+	template<class CPPSTRUCT>
+	FORCEINLINE typename TEnableIf<!TStructOpsTypeTraits<CPPSTRUCT>::WithAddStructReferencedObjects>::Type AddStructReferencedObjectsOrNot(const void* A, FReferenceCollector& Collector)
+	{
+	}
 
+	template<class CPPSTRUCT>
+	FORCEINLINE typename TEnableIf<TStructOpsTypeTraits<CPPSTRUCT>::WithAddStructReferencedObjects>::Type AddStructReferencedObjectsOrNot(const void* A, FReferenceCollector& Collector)
+	{
+		((CPPSTRUCT const*)A)->AddStructReferencedObjects(Collector);
+	}
+
+#endif
 
 /**
  * Reflection data for a standalone structure declared in a header or as a user defined struct
@@ -991,7 +1011,7 @@ public:
 		virtual bool HasGetTypeHash() = 0;
 
 		/** Calls GetTypeHash if enabled */
-		virtual uint32 GetTypeHash(const void* Src) = 0;
+		virtual uint32 GetStructTypeHash(const void* Src) = 0;
 
 		/** Returns property flag values that can be computed at compile time */
 		virtual EPropertyFlags GetComputedPropertyFlags() const = 0;
@@ -1028,7 +1048,18 @@ public:
 			check(!TTraits::WithZeroConstructor); // don't call this if we have indicated it is not necessary
 			// that could have been an if statement, but we might as well force optimization above the virtual call
 			// could also not attempt to call the constructor for types where this is not possible, but I didn't do that here
+#if PLATFORM_COMPILER_HAS_IF_CONSTEXPR
+			if constexpr (TStructOpsTypeTraits<CPPSTRUCT>::WithNoInitConstructor)
+			{
+				new (Dest) CPPSTRUCT(ForceInit);
+			}
+			else
+			{
+				new (Dest) CPPSTRUCT();
+			}
+#else
 			ConstructWithNoInitOrNot<CPPSTRUCT>(Dest);
+#endif
 		}
 		virtual bool HasDestructor() override
 		{
@@ -1052,13 +1083,35 @@ public:
 		virtual bool Serialize(FArchive& Ar, void *Data) override
 		{
 			check(TTraits::WithSerializer); // don't call this if we have indicated it is not necessary
+#if PLATFORM_COMPILER_HAS_IF_CONSTEXPR
+			if constexpr (TStructOpsTypeTraits<CPPSTRUCT>::WithSerializer)
+			{
+				return ((CPPSTRUCT*)Data)->Serialize(Ar);
+			}
+			else
+			{
+				return false;
+			}
+#else
 			return SerializeOrNot(Ar, (CPPSTRUCT*)Data);
+#endif
 		}
 		virtual bool Serialize(FStructuredArchive::FSlot Slot, void *Data) override
 		{
 			check(TTraits::WithStructuredSerializer); // don't call this if we have indicated it is not necessary
+#if PLATFORM_COMPILER_HAS_IF_CONSTEXPR
+			if constexpr (TStructOpsTypeTraits<CPPSTRUCT>::WithStructuredSerializer)
+			{
+				return ((CPPSTRUCT*)Data)->Serialize(Slot);
+			}
+			else
+			{
+				return false;
+			}
+#else
 			return SerializeOrNot(Slot, (CPPSTRUCT*)Data);
-		}
+#endif
+			}
 		virtual bool HasPostSerialize() override
 		{
 			return TTraits::WithPostSerialize;
@@ -1066,7 +1119,14 @@ public:
 		virtual void PostSerialize(const FArchive& Ar, void *Data) override
 		{
 			check(TTraits::WithPostSerialize); // don't call this if we have indicated it is not necessary
+#if PLATFORM_COMPILER_HAS_IF_CONSTEXPR
+			if constexpr (TStructOpsTypeTraits<CPPSTRUCT>::WithPostSerialize)
+			{
+				((CPPSTRUCT*)Data)->PostSerialize(Ar);
+			}
+#else
 			PostSerializeOrNot(Ar, (CPPSTRUCT*)Data);
+#endif
 		}
 		virtual bool HasNetSerializer() override
 		{
@@ -1082,11 +1142,33 @@ public:
 		}
 		virtual bool NetSerialize(FArchive& Ar, class UPackageMap* Map, bool& bOutSuccess, void *Data) override
 		{
+#if PLATFORM_COMPILER_HAS_IF_CONSTEXPR
+			if constexpr (TStructOpsTypeTraits<CPPSTRUCT>::WithNetSerializer)
+			{
+				return ((CPPSTRUCT*)Data)->NetSerialize(Ar, Map, bOutSuccess);
+			}
+			else
+			{
+				return false;
+			}
+#else
 			return NetSerializeOrNot(Ar, Map, bOutSuccess, (CPPSTRUCT*)Data);
+#endif
 		}
 		virtual bool NetDeltaSerialize(FNetDeltaSerializeInfo & DeltaParms, void *Data) override
 		{
+#if PLATFORM_COMPILER_HAS_IF_CONSTEXPR
+			if constexpr (TStructOpsTypeTraits<CPPSTRUCT>::WithNetDeltaSerializer)
+			{
+				return ((CPPSTRUCT*)Data)->NetDeltaSerialize(DeltaParms);
+			}
+			else
+			{
+				return false;
+			}
+#else
 			return NetDeltaSerializeOrNot(DeltaParms, (CPPSTRUCT*)Data);
+#endif
 		}
 		virtual bool HasPostScriptConstruct() override
 		{
@@ -1095,7 +1177,14 @@ public:
 		virtual void PostScriptConstruct(void *Data) override
 		{
 			check(TTraits::WithPostScriptConstruct); // don't call this if we have indicated it is not necessary
+#if PLATFORM_COMPILER_HAS_IF_CONSTEXPR
+			if constexpr (TStructOpsTypeTraits<CPPSTRUCT>::WithPostScriptConstruct)
+			{
+				((CPPSTRUCT*)Data)->PostScriptConstruct();
+			}
+#else
 			PostScriptConstructOrNot((CPPSTRUCT*)Data);
+#endif
 		}
 		virtual bool IsPlainOldData() override
 		{
@@ -1107,7 +1196,27 @@ public:
 		}
 		virtual bool Copy(void* Dest, void const* Src, int32 ArrayDim) override
 		{
+#if PLATFORM_COMPILER_HAS_IF_CONSTEXPR
+			if constexpr (TStructOpsTypeTraits<CPPSTRUCT>::WithCopy)
+			{
+				static_assert((!TIsPODType<CPPSTRUCT>::Value), "You probably don't want custom copy for a POD type.");
+
+				CPPSTRUCT* TypedDest = (CPPSTRUCT*)Dest;
+				const CPPSTRUCT* TypedSrc  = (const CPPSTRUCT*)Src;
+
+				for (; ArrayDim; --ArrayDim)
+				{
+					*TypedDest++ = *TypedSrc++;
+				}
+				return true;
+			}
+			else
+			{
+				return false;
+			}
+#else
 			return CopyOrNot((CPPSTRUCT*)Dest, (CPPSTRUCT const*)Src, ArrayDim);
+#endif
 		}
 		virtual bool HasIdentical() override
 		{
@@ -1116,7 +1225,27 @@ public:
 		virtual bool Identical(const void* A, const void* B, uint32 PortFlags, bool& bOutResult) override
 		{
 			check((TTraits::WithIdentical || TTraits::WithIdenticalViaEquality)); // don't call this if we have indicated it is not necessary
+#if PLATFORM_COMPILER_HAS_IF_CONSTEXPR
+			if constexpr (TStructOpsTypeTraits<CPPSTRUCT>::WithIdentical)
+			{
+				static_assert(!TStructOpsTypeTraits<CPPSTRUCT>::WithIdenticalViaEquality, "Should not have both WithIdenticalViaEquality and WithIdentical.");
+
+				bOutResult = ((const CPPSTRUCT*)A)->Identical((const CPPSTRUCT*)B, PortFlags);
+				return true;
+			}
+			else if constexpr (TStructOpsTypeTraits<CPPSTRUCT>::WithIdenticalViaEquality)
+			{
+				bOutResult = (*(const CPPSTRUCT*)A == *(const CPPSTRUCT*)B);
+				return true;
+			}
+			else
+			{
+				bOutResult = false;
+				return false;
+			}
+#else
 			return IdenticalOrNot((const CPPSTRUCT*)A, (const CPPSTRUCT*)B, PortFlags, bOutResult);
+#endif
 		}
 		virtual bool HasExportTextItem() override
 		{
@@ -1125,7 +1254,18 @@ public:
 		virtual bool ExportTextItem(FString& ValueStr, const void* PropertyValue, const void* DefaultValue, class UObject* Parent, int32 PortFlags, class UObject* ExportRootScope) override
 		{
 			check(TTraits::WithExportTextItem); // don't call this if we have indicated it is not necessary
+#if PLATFORM_COMPILER_HAS_IF_CONSTEXPR
+			if constexpr (TStructOpsTypeTraits<CPPSTRUCT>::WithExportTextItem)
+			{
+				return ((const CPPSTRUCT*)PropertyValue)->ExportTextItem(ValueStr, *(const CPPSTRUCT*)DefaultValue, Parent, PortFlags, ExportRootScope);
+			}
+			else
+			{
+				return false;
+			}
+#else
 			return ExportTextItemOrNot(ValueStr, (const CPPSTRUCT*)PropertyValue, (const CPPSTRUCT*)DefaultValue, Parent, PortFlags, ExportRootScope);
+#endif
 		}
 		virtual bool HasImportTextItem() override
 		{
@@ -1134,7 +1274,18 @@ public:
 		virtual bool ImportTextItem(const TCHAR*& Buffer, void* Data, int32 PortFlags, class UObject* OwnerObject, FOutputDevice* ErrorText) override
 		{
 			check(TTraits::WithImportTextItem); // don't call this if we have indicated it is not necessary
+#if PLATFORM_COMPILER_HAS_IF_CONSTEXPR
+			if constexpr (TStructOpsTypeTraits<CPPSTRUCT>::WithImportTextItem)
+			{
+				return ((CPPSTRUCT*)Data)->ImportTextItem(Buffer, PortFlags, OwnerObject, ErrorText);
+			}
+			else
+			{
+				return false;
+			}
+#else
 			return ImportTextItemOrNot(Buffer, (CPPSTRUCT*)Data, PortFlags, OwnerObject, ErrorText);
+#endif
 		}
 		virtual bool HasAddStructReferencedObjects() override
 		{
@@ -1152,7 +1303,18 @@ public:
 		virtual bool SerializeFromMismatchedTag(struct FPropertyTag const& Tag, FArchive& Ar, void *Data) override
 		{
 			check(TTraits::WithSerializeFromMismatchedTag); // don't call this if we have indicated it is not allowed
+#if PLATFORM_COMPILER_HAS_IF_CONSTEXPR
+			if constexpr (TStructOpsTypeTraits<CPPSTRUCT>::WithSerializeFromMismatchedTag)
+			{
+				return ((CPPSTRUCT*)Data)->SerializeFromMismatchedTag(Tag, Ar);
+			}
+			else
+			{
+				return false;
+			}
+#else
 			return SerializeFromMismatchedTagOrNot(Tag, Ar, (CPPSTRUCT*)Data);
+#endif
 		}
 		virtual bool HasStructuredSerializeFromMismatchedTag() override
 		{
@@ -1161,7 +1323,18 @@ public:
 		virtual bool StructuredSerializeFromMismatchedTag(struct FPropertyTag const& Tag, FStructuredArchive::FSlot Slot, void *Data) override
 		{
 			check(TTraits::WithStructuredSerializeFromMismatchedTag); // don't call this if we have indicated it is not allowed
+#if PLATFORM_COMPILER_HAS_IF_CONSTEXPR
+			if constexpr (TStructOpsTypeTraits<CPPSTRUCT>::WithStructuredSerializeFromMismatchedTag)
+			{
+				return ((CPPSTRUCT*)Data)->SerializeFromMismatchedTag(Tag, Slot);
+			}
+			else
+			{
+				return false;
+			}
+#else
 			return StructuredSerializeFromMismatchedTagOrNot(Tag, Slot, (CPPSTRUCT*)Data);
+#endif
 		}
 
 		static_assert(!(TTraits::WithSerializeFromMismatchedTag && TTraits::WithStructuredSerializeFromMismatchedTag), "Structs cannot have both WithSerializeFromMismatchedTag and WithStructuredSerializeFromMismatchedTag set");
@@ -1170,10 +1343,22 @@ public:
 		{
 			return TModels<CGetTypeHashable, CPPSTRUCT>::Value;
 		}
-		uint32 GetTypeHash(const void* Src) override
+		uint32 GetStructTypeHash(const void* Src) override
 		{
 			ensure(HasGetTypeHash());
+
+#if PLATFORM_COMPILER_HAS_IF_CONSTEXPR
+			if constexpr (TModels<CGetTypeHashable, CPPSTRUCT>::Value)
+			{
+				return GetTypeHash(*(const CPPSTRUCT*)Src);
+			}
+			else
+			{
+				return 0;
+			}
+#else
 			return GetTypeHashOrNot((const CPPSTRUCT*)Src);
+#endif
 		}
 		virtual EPropertyFlags GetComputedPropertyFlags() const override
 		{
