FAsyncPackage::EmptyReferencedObjects should be a lot faster now - moved most of its cost to the Async Loading Thread

diff --git a/Engine/Source/Runtime/CoreUObject/Private/Serialization/AsyncLoading.cpp b/Engine/Source/Runtime/CoreUObject/Private/Serialization/AsyncLoading.cpp
--- a/Engine/Source/Runtime/CoreUObject/Private/Serialization/AsyncLoading.cpp
+++ b/Engine/Source/Runtime/CoreUObject/Private/Serialization/AsyncLoading.cpp
@@ -5652,6 +5652,7 @@ void FAsyncPackage::AddObjectReference(UObject* InObject)
 			if (!ReferencedObjects.Contains(InObject))
 			{
 				ReferencedObjects.Add(InObject);
+				InObject->AtomicallyClearInternalFlags(EInternalObjectFlags::Async);
 			}
 		}
 		UE_CLOG(InObject->HasAnyInternalFlags(EInternalObjectFlags::Unreachable), LogStreaming, Fatal, TEXT("Trying to add an unreachable object %s to FAsyncPackage %s referenced objects list."), *InObject->GetFullName(), *GetPackageName().ToString());
@@ -5660,18 +5661,7 @@ void FAsyncPackage::AddObjectReference(UObject* InObject)
 
 void FAsyncPackage::EmptyReferencedObjects()
 {
-	const EInternalObjectFlags AsyncFlags = EInternalObjectFlags::Async | EInternalObjectFlags::AsyncLoading;
 	FScopeLock ReferencedObjectsLock(&ReferencedObjectsCritical);
-	for (UObject* Obj : ReferencedObjects)
-	{
-		if (Obj)
-		{
-			// Temporary fatal messages instead of checks to find the cause for a one-time crash in shipping config
-			UE_CLOG(!Obj->IsValidLowLevelFast(), LogStreaming, Fatal, TEXT("Invalid object in Async Objects Referencer"));
-			Obj->AtomicallyClearInternalFlags(AsyncFlags);
-			check(!Obj->HasAnyInternalFlags(AsyncFlags))
-		}
-	}
 	ReferencedObjects.Reset();
 }
 
@@ -6781,6 +6771,7 @@ EAsyncPackageState::Type FAsyncPackage::PostLoadObjects()
 				{
 					TRACE_LOADTIME_OBJECT_SCOPE(Object, LoadTimeProfilerObjectEventType_PostLoad);
 					Object->ConditionalPostLoad();
+					Object->AtomicallyClearInternalFlags(EInternalObjectFlags::AsyncLoading);
 				}
 				ThreadContext.CurrentlyPostLoadedObjectByALT = nullptr;
 
@@ -7099,6 +7090,17 @@ EAsyncPackageState::Type FAsyncPackage::FinishObjects()
 		LoadContext->DetachFromLinkers();
 	}
 
+	{
+		FScopeLock ReferencedObjectsLock(&ReferencedObjectsCritical);
+		for (UObject* Obj : ReferencedObjects)
+		{
+			if (Obj)
+			{
+				Obj->AtomicallyClearInternalFlags(EInternalObjectFlags::AsyncLoading);
+			}
+		}
+	}
+
 	return EAsyncPackageState::Complete;
 }
 
@@ -7374,6 +7376,16 @@ void CancelAsyncLoading()
 	{
 		CollectGarbage(GARBAGE_COLLECTION_KEEPFLAGS, true);
 	}
+
+	const EInternalObjectFlags AsyncFlags = EInternalObjectFlags::Async | EInternalObjectFlags::AsyncLoading;
+	for (int32 ObjectIndex = 0; ObjectIndex < GUObjectArray.GetObjectArrayNum(); ++ObjectIndex)
+	{
+		FUObjectItem* ObjectItem = &GUObjectArray.GetObjectItemArrayUnsafe()[ObjectIndex];
+		if (UObject* Obj = static_cast<UObject*>(ObjectItem->Object))
+		{
+			check(!Obj->HasAnyInternalFlags(AsyncFlags));
+		}
+	}
 }
 
 float GetAsyncLoadPercentage(const FName& PackageName)
