PR #6092: Fix comments for ByteMask (Contributed by [REDACTED_PERSON])

diff --git a/Engine/Source/Runtime/CoreUObject/Public/UObject/UnrealType.h b/Engine/Source/Runtime/CoreUObject/Public/UObject/UnrealType.h
--- a/Engine/Source/Runtime/CoreUObject/Public/UObject/UnrealType.h
+++ b/Engine/Source/Runtime/CoreUObject/Public/UObject/UnrealType.h
@@ -1727,7 +1727,7 @@ private:
 	uint8 FieldSize;
 	/** Offset from the memeber variable to the byte of the property (0-7). */
 	uint8 ByteOffset;
-	/** Mask of the byte byte with the property value. */
+	/** Mask of the byte with the property value. */
 	uint8 ByteMask;
 	/** Mask of the field with the property value. Either equal to ByteMask or 255 in case of 'bool' type. */
 	uint8 FieldMask;
