AnnotationSparse increases start up times for Linux Servers by ~42% versus the current method. As well as a small increase in memory use:

Sparse Startup Time: 42% faster then Chunked
Sparse Startup Memory Used: 0.08% MB more then Chunked
Sparse Startup Memory Peaked: 0.04% MB more then Chunked

Dense was about the same speed as Sparse but took a bit more memory

diff --git a/Engine/Source/Runtime/CoreUObject/Private/UObject/UObjectLinker.cpp b/Engine/Source/Runtime/CoreUObject/Private/UObject/UObjectLinker.cpp
--- a/Engine/Source/Runtime/CoreUObject/Private/UObject/UObjectLinker.cpp
+++ b/Engine/Source/Runtime/CoreUObject/Private/UObject/UObjectLinker.cpp
@@ -82,10 +82,16 @@ template <> struct TIsPODType<FLinkerIndexPair> { enum { Value = false }; };
  * A: It does not need to be. This is GC-Safe.
  * Objects are detached from their linkers prior to destruction of either the linker or the object
  *
- * NOTE: We're currently using dense annotations for linkers to emphasize speed over memory
+ * NOTE: We're currently using chunked annotations for linkers to emphasize memory
  * usage, but might want to revisit this decision on platforms that are memory limited.
+ *
+ * LINUX SERVER NOTE: For servers we are using Sparse to emphasize speed over memory usage
  */
+#if PLATFORM_LINUX && UE_SERVER
+static FUObjectAnnotationSparse<FLinkerIndexPair, false> LinkerAnnotation;
+#else
 static FUObjectAnnotationChunked<FLinkerIndexPair, false> LinkerAnnotation;
+#endif
 
 /** Remove all annotations on exit. This is to prevent issues with the order of static destruction of singletons. */
 void CleanupLinkerAnnotations()
