Consistent with blueprint classes, for native classes do not overwrite transient instanced objects during property initialization
#rnx

diff --git a/Engine/Source/Runtime/CoreUObject/Private/UObject/UObjectGlobals.cpp b/Engine/Source/Runtime/CoreUObject/Private/UObject/UObjectGlobals.cpp
--- a/Engine/Source/Runtime/CoreUObject/Private/UObject/UObjectGlobals.cpp
+++ b/Engine/Source/Runtime/CoreUObject/Private/UObject/UObjectGlobals.cpp
@@ -2982,7 +2982,14 @@ void FObjectInitializer::InitProperties(UObject* Obj, UClass* DefaultsClass, UOb
 				QUICK_SCOPE_CYCLE_COUNTER(STAT_InitProperties_FromTemplate);
 				for (FProperty* P = Class->PropertyLink; P; P = P->PropertyLinkNext)
 				{
-					P->CopyCompleteValue_InContainer(Obj, DefaultData);
+					bool bIsTransient = P->HasAnyPropertyFlags(CPF_Transient | CPF_DuplicateTransient | CPF_NonPIEDuplicateTransient);
+					if (!bIsTransient || !P->ContainsInstancedObjectProperty())
+					{
+						if (P->IsInContainer(DefaultsClass))
+						{
+							P->CopyCompleteValue_InContainer(Obj, DefaultData);
+						}
+					}
 				}
 			}
 			else
@@ -2991,7 +2998,14 @@ void FObjectInitializer::InitProperties(UObject* Obj, UClass* DefaultsClass, UOb
 				// Copy all properties that require additional initialization (e.g. CPF_Config).
 				for (FProperty* P = Class->PostConstructLink; P; P = P->PostConstructLinkNext)
 				{
-					P->CopyCompleteValue_InContainer(Obj, DefaultData);
+					bool bIsTransient = P->HasAnyPropertyFlags(CPF_Transient | CPF_DuplicateTransient | CPF_NonPIEDuplicateTransient);
+					if (!bIsTransient || !P->ContainsInstancedObjectProperty())
+					{
+						if (P->IsInContainer(DefaultsClass))
+						{
+							P->CopyCompleteValue_InContainer(Obj, DefaultData);
+						}
+					}
 				}
 			}
 		}
@@ -3013,10 +3027,10 @@ void FObjectInitializer::InitProperties(UObject* Obj, UClass* DefaultsClass, UOb
 				bNeedInitialize = InitNonNativeProperty(P, Obj);
 			}
 
-			bool IsTransient = P->HasAnyPropertyFlags(CPF_Transient | CPF_DuplicateTransient | CPF_NonPIEDuplicateTransient);
-			if (!IsTransient || !P->ContainsInstancedObjectProperty())
+			bool bIsTransient = P->HasAnyPropertyFlags(CPF_Transient | CPF_DuplicateTransient | CPF_NonPIEDuplicateTransient);
+			if (!bIsTransient || !P->ContainsInstancedObjectProperty())
 			{
-				if (bCopyTransientsFromClassDefaults && IsTransient)
+				if (bCopyTransientsFromClassDefaults && bIsTransient)
 				{
 					// This is a duplicate. The value for all transient or non-duplicatable properties should be copied
 					// from the source class's defaults.
