Fix for SaveGame properties being skipped in non-editor builds.

diff --git a/Engine/Source/Runtime/CoreUObject/Private/UObject/Class.cpp b/Engine/Source/Runtime/CoreUObject/Private/UObject/Class.cpp
--- a/Engine/Source/Runtime/CoreUObject/Private/UObject/Class.cpp
+++ b/Engine/Source/Runtime/CoreUObject/Private/UObject/Class.cpp
@@ -625,17 +625,17 @@ public:
 //
 // Constructors.
 //
-UStruct::UStruct(EStaticConstructor, int32 InSize, int32 InMinAlignment, EObjectFlags InFlags)
-	: UField(EC_StaticConstructor, InFlags)
-	, SuperStruct(nullptr)
+UStruct::UStruct( EStaticConstructor, int32 InSize, int32 InMinAlignment, EObjectFlags InFlags )
+:	UField			( EC_StaticConstructor, InFlags )
+,	SuperStruct		( nullptr )
 	, Children(nullptr)
 	, ChildProperties(nullptr)
-	, PropertiesSize(InSize)
-	, MinAlignment(InMinAlignment)
-	, PropertyLink(NULL)
-	, RefLink(NULL)
-	, DestructorLink(NULL)
-	, PostConstructLink(NULL)
+,	PropertiesSize	( InSize )
+,	MinAlignment	( InMinAlignment )
+,	PropertyLink	( NULL )
+,	RefLink			( NULL )
+,	DestructorLink	( NULL )
+, PostConstructLink( NULL )
 {
 }
 
@@ -656,13 +656,13 @@ UStruct::UStruct(UStruct* InSuperStruct, SIZE_T ParamsSize, SIZE_T Alignment)
 #endif
 }
 
-UStruct::UStruct(const FObjectInitializer& ObjectInitializer, UStruct* InSuperStruct, SIZE_T ParamsSize, SIZE_T Alignment)
-	: UField(ObjectInitializer)
-	, SuperStruct(InSuperStruct)
+UStruct::UStruct(const FObjectInitializer& ObjectInitializer, UStruct* InSuperStruct, SIZE_T ParamsSize, SIZE_T Alignment )
+:	UField			(ObjectInitializer)
+,   SuperStruct		( InSuperStruct )
 	, Children(nullptr)
 	, ChildProperties(nullptr)
-	, PropertiesSize(ParamsSize ? ParamsSize : (InSuperStruct ? InSuperStruct->GetPropertiesSize() : 0))
-	, MinAlignment(Alignment ? Alignment : (FMath::Max(InSuperStruct ? InSuperStruct->GetMinAlignment() : 1, 1)))
+,	PropertiesSize	( ParamsSize ? ParamsSize : (InSuperStruct ? InSuperStruct->GetPropertiesSize() : 0) )
+,	MinAlignment	( Alignment ? Alignment : (FMath::Max(InSuperStruct ? InSuperStruct->GetMinAlignment() : 1,1)) )
 	, PropertyLink(nullptr)
 	, RefLink(nullptr)
 	, DestructorLink(nullptr)
@@ -1227,7 +1227,7 @@ void UStruct::SerializeVersionedTaggedProperties(FStructuredArchive::FSlot Slot,
 	const bool bArePropertyGuidsAvailable = (UnderlyingArchive.UE4Ver() >= VER_UE4_PROPERTY_GUID_IN_PROPERTY_TAG) && !FPlatformProperties::RequiresCookedData() && ArePropertyGuidsAvailable();
 	const bool bUseRedirects = !FPlatformProperties::RequiresCookedData() || UnderlyingArchive.IsSaveGame();
 
-	if (UnderlyingArchive.IsLoading())
+	if( UnderlyingArchive.IsLoading() )
 	{
 #if WITH_TEXT_ARCHIVE_SUPPORT
 		if (UnderlyingArchive.IsTextFormat())
@@ -1237,159 +1237,160 @@ void UStruct::SerializeVersionedTaggedProperties(FStructuredArchive::FSlot Slot,
 		else
 #endif // WITH_TEXT_ARCHIVE_SUPPORT
 		{
-			// Load tagged properties.
-			FStructuredArchive::FStream PropertiesStream = Slot.EnterStream();
+		// Load tagged properties.
+		FStructuredArchive::FStream PropertiesStream = Slot.EnterStream();
 
-			// This code assumes that properties are loaded in the same order they are saved in. This removes a n^2 search 
-			// and makes it an O(n) when properties are saved in the same order as they are loaded (default case). In the 
-			// case that a property was reordered the code falls back to a slower search.
+		// This code assumes that properties are loaded in the same order they are saved in. This removes a n^2 search 
+		// and makes it an O(n) when properties are saved in the same order as they are loaded (default case). In the 
+		// case that a property was reordered the code falls back to a slower search.
 			FProperty*	Property = PropertyLink;
-			bool		bAdvanceProperty = false;
-			int32		RemainingArrayDim = Property ? Property->ArrayDim : 0;
+		bool		bAdvanceProperty	= false;
+		int32		RemainingArrayDim	= Property ? Property->ArrayDim : 0;
 
-			// Load all stored properties, potentially skipping unknown ones.
-			while (true)
-			{
-				FStructuredArchive::FRecord PropertyRecord = PropertiesStream.EnterElement().EnterRecord();
+		// Load all stored properties, potentially skipping unknown ones.
+		while (true)
+		{
+			FStructuredArchive::FRecord PropertyRecord = PropertiesStream.EnterElement().EnterRecord();
 
-				FPropertyTag Tag;
-				PropertyRecord << SA_VALUE(TEXT("Tag"), Tag);
+			FPropertyTag Tag;
+			PropertyRecord << SA_VALUE(TEXT("Tag"), Tag);
 
-				if (Tag.Name.IsNone())
-				{
-					break;
-				}
+			if (Tag.Name.IsNone())
+			{
+				break;
+			}
 
-				// Move to the next property to be serialized
-				if (bAdvanceProperty && --RemainingArrayDim <= 0)
+			// Move to the next property to be serialized
+			if( bAdvanceProperty && --RemainingArrayDim <= 0 )
+			{
+				Property = Property->PropertyLinkNext;
+				// Skip over properties that don't need to be serialized.
+				while( Property && !Property->ShouldSerializeValue( UnderlyingArchive ) )
 				{
 					Property = Property->PropertyLinkNext;
-					// Skip over properties that don't need to be serialized.
-					while (Property && !Property->ShouldSerializeValue(UnderlyingArchive))
-					{
-						Property = Property->PropertyLinkNext;
-					}
-					RemainingArrayDim = Property ? Property->ArrayDim : 0;
 				}
-				bAdvanceProperty = false;
-
-				// Optionally resolve properties using Guid Property tags in non cooked builds that support it.
-				if (bArePropertyGuidsAvailable && Tag.HasPropertyGuid)
+				RemainingArrayDim = Property ? Property->ArrayDim : 0;
+			}
+			bAdvanceProperty = false;
+			
+			// Optionally resolve properties using Guid Property tags in non cooked builds that support it.
+			if (bArePropertyGuidsAvailable && Tag.HasPropertyGuid)
+			{
+				// Use property guids from blueprint generated classes to redirect serialised data.
+				FName Result = FindPropertyNameFromGuid(Tag.PropertyGuid);
+				if (Result != NAME_None && Tag.Name != Result)
 				{
-					// Use property guids from blueprint generated classes to redirect serialised data.
-					FName Result = FindPropertyNameFromGuid(Tag.PropertyGuid);
-					if (Result != NAME_None && Tag.Name != Result)
-					{
-						Tag.Name = Result;
-					}
+					Tag.Name = Result;
 				}
-				// If this property is not the one we expect (e.g. skipped as it matches the default value), do the brute force search.
-				if (Property == nullptr || Property->GetFName() != Tag.Name)
+			}
+			// If this property is not the one we expect (e.g. skipped as it matches the default value), do the brute force search.
+			if( Property == nullptr || Property->GetFName() != Tag.Name )
+			{
+				// No need to check redirects on platforms where everything is cooked. Always check for save games
+				if (bUseRedirects && !UnderlyingArchive.HasAnyPortFlags(PPF_DuplicateForPIE|PPF_Duplicate))
 				{
-					// No need to check redirects on platforms where everything is cooked. Always check for save games
-					if (bUseRedirects && !UnderlyingArchive.HasAnyPortFlags(PPF_DuplicateForPIE | PPF_Duplicate))
+					for (UStruct* CheckStruct = GetOwnerStruct(); CheckStruct; CheckStruct = CheckStruct->GetSuperStruct())
 					{
-						for (UStruct* CheckStruct = GetOwnerStruct(); CheckStruct; CheckStruct = CheckStruct->GetSuperStruct())
-						{
 							FName NewTagName = FProperty::FindRedirectedPropertyName(CheckStruct, Tag.Name);
-							if (!NewTagName.IsNone())
-							{
-								Tag.Name = NewTagName;
-								break;
-							}
+						if (!NewTagName.IsNone())
+						{
+							Tag.Name = NewTagName;
+							break;
 						}
 					}
+				}
 
 					FProperty* CurrentProperty = Property;
-					// Search forward...
-					for (; Property; Property = Property->PropertyLinkNext)
+				// Search forward...
+				for ( ; Property; Property=Property->PropertyLinkNext )
+				{
+					if( Property->GetFName() == Tag.Name )
 					{
-						if (Property->GetFName() == Tag.Name)
-						{
-							break;
-						}
+						break;
 					}
-					// ... and then search from the beginning till we reach the current property if it's not found.
-					if (Property == nullptr)
+				}
+				// ... and then search from the beginning till we reach the current property if it's not found.
+				if( Property == nullptr )
+				{
+					for( Property = PropertyLink; Property && Property != CurrentProperty; Property = Property->PropertyLinkNext )
 					{
-						for (Property = PropertyLink; Property && Property != CurrentProperty; Property = Property->PropertyLinkNext)
+						if( Property->GetFName() == Tag.Name )
 						{
-							if (Property->GetFName() == Tag.Name)
-							{
-								break;
-							}
-						}
-
-						if (Property == CurrentProperty)
-						{
-							// Property wasn't found.
-							Property = nullptr;
+							break;
 						}
 					}
 
-					RemainingArrayDim = Property ? Property->ArrayDim : 0;
+					if( Property == CurrentProperty )
+					{
+						// Property wasn't found.
+						Property = nullptr;
+					}
 				}
 
-				const int64 StartOfProperty = UnderlyingArchive.Tell();
+				RemainingArrayDim = Property ? Property->ArrayDim : 0;
+			}
+
+			const int64 StartOfProperty = UnderlyingArchive.Tell();
 
-				if (!Property)
-				{
-					Property = CustomFindProperty(Tag.Name);
-				}
+			if (!Property)
+			{
+				Property = CustomFindProperty(Tag.Name);
+			}
 
-				if (Property)
-				{
-					FName PropID = Property->GetID();
+			if (Property)
+			{
+				FName PropID = Property->GetID();
 
-					// Check if this is a struct property and we have a redirector
-					// No need to check redirects on platforms where everything is cooked. Always check for save games
-					if (bUseRedirects)
+				// Check if this is a struct property and we have a redirector
+				// No need to check redirects on platforms where everything is cooked. Always check for save games
+				if (bUseRedirects)
+				{
+					if (Tag.Type == NAME_StructProperty && PropID == NAME_StructProperty)
 					{
-						if (Tag.Type == NAME_StructProperty && PropID == NAME_StructProperty)
-						{
-							const FName NewName = FLinkerLoad::FindNewNameForStruct(Tag.StructName);
+						const FName NewName = FLinkerLoad::FindNewNameForStruct(Tag.StructName);
 							const FName StructName = CastFieldChecked<FStructProperty>(Property)->Struct->GetFName();
-							if (NewName == StructName)
-							{
-								Tag.StructName = NewName;
-							}
+						if (NewName == StructName)
+						{
+							Tag.StructName = NewName;
 						}
-						else if ((PropID == NAME_EnumProperty) && ((Tag.Type == NAME_EnumProperty) || (Tag.Type == NAME_ByteProperty)))
+					}
+					else if ((PropID == NAME_EnumProperty) && ((Tag.Type == NAME_EnumProperty) || (Tag.Type == NAME_ByteProperty)))
+					{
+						const FName NewName = FLinkerLoad::FindNewNameForEnum(Tag.EnumName);
+						if (!NewName.IsNone())
 						{
-							const FName NewName = FLinkerLoad::FindNewNameForEnum(Tag.EnumName);
-							if (!NewName.IsNone())
-							{
-								Tag.EnumName = NewName;
-							}
+							Tag.EnumName = NewName;
 						}
 					}
+				}
 
 #if WITH_EDITOR
-					if (BreakRecursionIfFullyLoad && BreakRecursionIfFullyLoad->HasAllFlags(RF_LoadCompleted))
-					{
-					}
+				if (BreakRecursionIfFullyLoad && BreakRecursionIfFullyLoad->HasAllFlags(RF_LoadCompleted))
+				{
+				}
+				else
 #endif // WITH_EDITOR
-					// editoronly properties should be skipped if we are NOT the editor, or we are 
-					// the editor but are cooking for console (editoronly implies notforconsole)
-					else if ((Property->PropertyFlags & CPF_EditorOnly) && !FPlatformProperties::HasEditorOnlyData() && !GForceLoadEditorOnly)
-					{
-					}
-					// check for valid array index
-					else if (Tag.ArrayIndex >= Property->ArrayDim || Tag.ArrayIndex < 0)
-					{
-						UE_LOG(LogClass, Warning, TEXT("Array bound exceeded (var %s=%d, exceeds %s [0-%d] in package:  %s"),
-							*Tag.Name.ToString(), Tag.ArrayIndex, *GetName(), Property->ArrayDim - 1, *UnderlyingArchive.GetArchiveName());
-					}
-					else if (!Property->ShouldSerializeValue(UnderlyingArchive))
-					{
-						UE_CLOG((UnderlyingArchive.IsPersistent() && FPlatformProperties::RequiresCookedData()), LogClass, Warning, TEXT("Skipping saved property %s of %s since it is no longer serializable for asset:  %s. (Maybe resave asset?)"), *Tag.Name.ToString(), *GetName(), *UnderlyingArchive.GetArchiveName());
-					}
-					else
-					{
-						FStructuredArchive::FSlot ValueSlot = PropertyRecord.EnterField(SA_FIELD_NAME(TEXT("Value")));
+				// editoronly properties should be skipped if we are NOT the editor, or we are 
+				// the editor but are cooking for console (editoronly implies notforconsole)
+				if ((Property->PropertyFlags & CPF_EditorOnly) && !FPlatformProperties::HasEditorOnlyData() && !GForceLoadEditorOnly)
+				{
+				}
+				// check for valid array index
+				else if (Tag.ArrayIndex >= Property->ArrayDim || Tag.ArrayIndex < 0)
+				{
+					UE_LOG(LogClass, Warning, TEXT("Array bound exceeded (var %s=%d, exceeds %s [0-%d] in package:  %s"),
+						*Tag.Name.ToString(), Tag.ArrayIndex, *GetName(), Property->ArrayDim - 1, *UnderlyingArchive.GetArchiveName());
+				}
+				else if (!Property->ShouldSerializeValue(UnderlyingArchive))
+				{
+					UE_CLOG((UnderlyingArchive.IsPersistent() && FPlatformProperties::RequiresCookedData()), LogClass, Warning, TEXT("Skipping saved property %s of %s since it is no longer serializable for asset:  %s. (Maybe resave asset?)"), *Tag.Name.ToString(), *GetName(), *UnderlyingArchive.GetArchiveName());
+				}
+				else
+				{
+					FStructuredArchive::FSlot ValueSlot = PropertyRecord.EnterField(SA_FIELD_NAME(TEXT("Value")));
 
-						switch (Property->ConvertFromType(Tag, ValueSlot, Data, DefaultsStruct))
-						{
+					switch (Property->ConvertFromType(Tag, ValueSlot, Data, DefaultsStruct))
+					{
 						case EConvertFromTypeResult::Converted:
 							bAdvanceProperty = true;
 							break;
@@ -1415,16 +1416,16 @@ void UStruct::SerializeVersionedTaggedProperties(FStructuredArchive::FSlot Slot,
 
 						default:
 							check(false);
-						}
 					}
 				}
+			}
 
 				int64 Loaded = UnderlyingArchive.Tell() - StartOfProperty;
 
-				if (!bAdvanceProperty)
-				{
-					UnderlyingArchive.Seek(StartOfProperty + Tag.Size);
-				}
+			if (!bAdvanceProperty)
+			{
+				UnderlyingArchive.Seek(StartOfProperty + Tag.Size);
+			}
 				else
 				{
 					check(Tag.Size == Loaded);
@@ -1460,7 +1461,7 @@ void UStruct::SerializeVersionedTaggedProperties(FStructuredArchive::FSlot Slot,
 			Property;
 			Property = UnderlyingArchive.ArUseCustomPropertyList ? FCustomPropertyListNode::GetNextPropertyAndAdvance(CustomPropertyNode) : Property->PropertyLinkNext)
 		{
-			if (Property->ShouldSerializeValue(UnderlyingArchive))
+			if( Property->ShouldSerializeValue(UnderlyingArchive) )
 			{
 				const int32 LoopMin = CustomPropertyNode ? CustomPropertyNode->ArrayIndex : 0;
 				const int32 LoopMax = CustomPropertyNode ? LoopMin + 1 : Property->ArrayDim;
@@ -1472,7 +1473,7 @@ void UStruct::SerializeVersionedTaggedProperties(FStructuredArchive::FSlot Slot,
 					StaticArrayContainer.Emplace(PropertiesRecord.EnterArray(SA_FIELD_NAME((*Property->GetName())), NumItems));
 				}
 
-				for (int32 Idx = LoopMin; Idx < LoopMax; Idx++)
+				for( int32 Idx = LoopMin; Idx < LoopMax; Idx++ )
 				{
 					uint8* DataPtr      = Property->ContainerPtrToValuePtr           <uint8>(Data, Idx);
 					uint8* DefaultValue = Property->ContainerPtrToValuePtrForDefaults<uint8>(DefaultsStruct, Defaults, Idx);
@@ -1506,7 +1507,7 @@ void UStruct::SerializeVersionedTaggedProperties(FStructuredArchive::FSlot Slot,
 
 						// if using it, save the current custom property list and switch to its sub property list (in case of UStruct serialization)
 						const FCustomPropertyListNode* SavedCustomPropertyList = nullptr;
-						if (UnderlyingArchive.ArUseCustomPropertyList && CustomPropertyNode)
+						if(UnderlyingArchive.ArUseCustomPropertyList && CustomPropertyNode)
 						{
 							SavedCustomPropertyList = UnderlyingArchive.ArCustomPropertyList;
 							UnderlyingArchive.ArCustomPropertyList = CustomPropertyNode->SubPropertyList;
@@ -1523,7 +1524,7 @@ void UStruct::SerializeVersionedTaggedProperties(FStructuredArchive::FSlot Slot,
 						// set the tag's size
 						Tag.Size = UnderlyingArchive.Tell() - DataOffset;
 
-						if (Tag.Size > 0 && !UnderlyingArchive.IsTextFormat())
+						if ( Tag.Size >  0 && !UnderlyingArchive.IsTextFormat())
 						{
 							// mark our current location
 							DataOffset = UnderlyingArchive.Tell();
@@ -1717,9 +1718,9 @@ void UStruct::SerializeProperties(FArchive& Ar)
 
 
 
-void UStruct::Serialize(FArchive& Ar)
+void UStruct::Serialize( FArchive& Ar )
 {
-	Super::Serialize(Ar);
+	Super::Serialize( Ar );
 
 #if USTRUCT_FAST_ISCHILDOF_IMPL == USTRUCT_ISCHILDOF_STRUCTARRAY
 	UStruct* SuperStructBefore = GetSuperStruct();
@@ -4011,20 +4012,20 @@ void UClass::SetUpRuntimeReplicationData()
 		const bool bIsNativeClass = HasAnyClassFlags(CLASS_Native);
 		if (!bIsNativeClass)
 		{
-			// Sort NetProperties so that their ClassReps are sorted by memory offset
+		// Sort NetProperties so that their ClassReps are sorted by memory offset
 			struct FComparePropertyOffsets
-			{
+		{
 				FORCEINLINE bool operator()(FProperty& A, FProperty& B) const
+			{
+				// Ensure stable sort
+				if ( A.GetOffset_ForGC() == B.GetOffset_ForGC() )
 				{
-					// Ensure stable sort
-					if (A.GetOffset_ForGC() == B.GetOffset_ForGC())
-					{
-						return A.GetName() < B.GetName();
-					}
-
-					return A.GetOffset_ForGC() < B.GetOffset_ForGC();
+					return A.GetName() < B.GetName();
 				}
-			};
+
+				return A.GetOffset_ForGC() < B.GetOffset_ForGC();
+			}
+		};
 
 			Sort(NetProperties.GetData(), NetProperties.Num(), FComparePropertyOffsets());
 		}
@@ -4599,7 +4600,7 @@ void UClass::PurgeClass(bool bRecompilingOnLoad)
 	DestructorLink = NULL;
 	ClassAddReferencedObjects = NULL;
 
-	ScriptObjectReferences.Empty();	
+	ScriptObjectReferences.Empty();
 	PropertyObjectReferences.Empty();
 	UnresolvedScriptProperties.Empty();
 
@@ -4631,7 +4632,7 @@ void UClass::PurgeClass(bool bRecompilingOnLoad)
 			PropertiesPendingDestruction = ChildProperties;
 			ChildProperties = nullptr;
 		}
-	}
+}
 	else
 #endif // WITH_EDITORONLY_DATA
 	{
@@ -4690,9 +4691,9 @@ bool UClass::HasProperty(FProperty* InProperty) const
 	{
 		UClass* PropertiesClass = InProperty->GetOwner<UClass>();
 		if (PropertiesClass)
-		{
-			return IsChildOf(PropertiesClass);
-		}
+	{
+		return IsChildOf(PropertiesClass);
+	}
 	}
 
 	return false;
