Allow coercion between ByteProperty and EnumProperty when getting/setting enum values from BP via get/set editor property

#rnx

diff --git a/Engine/Source/Runtime/CoreUObject/Private/UObject/PropertyAccessUtil.cpp b/Engine/Source/Runtime/CoreUObject/Private/UObject/PropertyAccessUtil.cpp
--- a/Engine/Source/Runtime/CoreUObject/Private/UObject/PropertyAccessUtil.cpp
+++ b/Engine/Source/Runtime/CoreUObject/Private/UObject/PropertyAccessUtil.cpp
@@ -1,14 +1,75 @@
 // Copyright Epic Games, Inc. All Rights Reserved.
 
 #include "UObject/PropertyAccessUtil.h"
+#include "UObject/EnumProperty.h"
 #include "UObject/Object.h"
 #include "UObject/Class.h"
 
 namespace PropertyAccessUtil
 {
 
+const UEnum* GetPropertyEnumType(const FProperty* InProp)
+{
+	if (const FByteProperty* ByteProp = CastField<FByteProperty>(InProp))
+	{
+		return ByteProp->Enum;
+	}
+	if (const FEnumProperty* EnumProp = CastField<FEnumProperty>(InProp))
+	{
+		return EnumProp->GetEnum();
+	}
+	return nullptr;
+}
+
+int64 GetPropertyEnumValue(const FProperty* InProp, const void* InPropValue)
+{
+	if (const FByteProperty* ByteProp = CastField<FByteProperty>(InProp))
+	{
+		return ByteProp->GetSignedIntPropertyValue(InPropValue);
+	}
+	if (const FEnumProperty* EnumProp = CastField<FEnumProperty>(InProp))
+	{
+		EnumProp->GetUnderlyingProperty()->GetSignedIntPropertyValue(InPropValue);
+	}
+	return INDEX_NONE;
+}
+
+bool SetPropertyEnumValue(const FProperty* InProp, void* InPropValue, const int64 InEnumValue)
+{
+	if (const FByteProperty* ByteProp = CastField<FByteProperty>(InProp))
+	{
+		ByteProp->SetPropertyValue(InPropValue, (uint8)InEnumValue);
+		return true;
+	}
+	if (const FEnumProperty* EnumProp = CastField<FEnumProperty>(InProp))
+	{
+		EnumProp->GetUnderlyingProperty()->SetIntPropertyValue(InPropValue, InEnumValue);
+		return true;
+	}
+	return false;
+}
+
 bool ArePropertiesCompatible(const FProperty* InSrcProp, const FProperty* InDestProp)
 {
+	// Enum properties can either be a ByteProperty with an enum set, or an EnumProperty
+	// We allow coercion between these two types if they're using the same enum type
+	if (const UEnum* DestEnumType = GetPropertyEnumType(InDestProp))
+	{
+		const UEnum* SrcEnumType = GetPropertyEnumType(InSrcProp);
+		if (SrcEnumType == DestEnumType)
+		{
+			return true;
+		}
+
+		// Blueprints don't always set the Enum field on the ByteProperty when setting properties, so we also 
+		// allow assigning from a raw ByteProperty (for type safety there we rely on the compiler frontend)
+		const FByteProperty* SrcByteProp = CastField<FByteProperty>(InSrcProp);
+		if (SrcByteProp && !SrcByteProp->Enum && InDestProp->IsA<FEnumProperty>())
+		{
+			return true;
+		}
+	}
+
 	// Compare the classes as these must be an *exact* match as the access is low-level and without property coercion
 	if (InSrcProp->GetClass() != InDestProp->GetClass())
 	{
@@ -52,10 +113,8 @@ bool IsSinglePropertyIdentical(const FProperty* InSrcProp, const void* InSrcValu
 		const bool bDestBoolValue = DestBoolProp->GetPropertyValue(InDestValue);
 		return bSrcBoolValue == bDestBoolValue;
 	}
-	else
-	{
-		return InSrcProp->Identical(InSrcValue, InDestValue);
-	}
+	
+	return InSrcProp->Identical(InSrcValue, InDestValue);
 }
 
 bool IsCompletePropertyIdentical(const FProperty* InSrcProp, const void* InSrcValue, const FProperty* InDestProp, const void* InDestValue)
@@ -77,6 +136,14 @@ bool CopySinglePropertyValue(const FProperty* InSrcProp, const void* InSrcValue,
 		return false;
 	}
 
+	// Enum properties can either be a ByteProperty with an enum set, or an EnumProperty
+	// We allow coercion between these two types as long as they're using the same enum type (as validated by ArePropertiesCompatible)
+	if (const UEnum* EnumType = GetPropertyEnumType(InDestProp))
+	{
+		const int64 SrcEnumValue = GetPropertyEnumValue(InSrcProp, InSrcValue);
+		return SetPropertyEnumValue(InDestProp, InDestValue, SrcEnumValue);
+	}
+
 	if (const FBoolProperty* SrcBoolProp = CastField<FBoolProperty>(InSrcProp))
 	{
 		const FBoolProperty* DestBoolProp = CastFieldChecked<FBoolProperty>(InDestProp);
@@ -84,12 +151,10 @@ bool CopySinglePropertyValue(const FProperty* InSrcProp, const void* InSrcValue,
 		// Bools can be represented as bitfields, we we have to handle the copy a little differently to only extract the bool we want
 		const bool bBoolValue = SrcBoolProp->GetPropertyValue(InSrcValue);
 		DestBoolProp->SetPropertyValue(InDestValue, bBoolValue);
+		return true;
 	}
-	else
-	{
-		InSrcProp->CopySingleValue(InDestValue, InSrcValue);
-	}
-
+	
+	InSrcProp->CopySingleValue(InDestValue, InSrcValue);
 	return true;
 }
 
@@ -100,10 +165,25 @@ bool CopyCompletePropertyValue(const FProperty* InSrcProp, const void* InSrcValu
 		return false;
 	}
 
+	// Enum properties can either be a ByteProperty with an enum set, or an EnumProperty
+	// We allow coercion between these two types as long as they're using the same enum type (as validated by ArePropertiesCompatible)
+	if (const UEnum* EnumType = GetPropertyEnumType(InDestProp))
+	{
+		bool bSuccess = true;
+		for (int32 Idx = 0; Idx < InSrcProp->ArrayDim; ++Idx)
+		{
+			const void* SrcElemValue = static_cast<const uint8*>(InSrcValue) + (InSrcProp->ElementSize * Idx);
+			void* DestElemValue = static_cast<uint8*>(InDestValue) + (InDestProp->ElementSize * Idx);
+
+			const int64 SrcEnumValue = GetPropertyEnumValue(InSrcProp, SrcElemValue);
+			bSuccess &= SetPropertyEnumValue(InDestProp, DestElemValue, SrcEnumValue);
+		}
+		return bSuccess;
+	}
+
 	if (const FBoolProperty* SrcBoolProp = CastField<FBoolProperty>(InSrcProp))
 	{
 		const FBoolProperty* DestBoolProp = CastFieldChecked<FBoolProperty>(InDestProp);
-
 		for (int32 Idx = 0; Idx < InSrcProp->ArrayDim; ++Idx)
 		{
 			const void* SrcElemValue = static_cast<const uint8*>(InSrcValue) + (InSrcProp->ElementSize * Idx);
@@ -113,12 +193,10 @@ bool CopyCompletePropertyValue(const FProperty* InSrcProp, const void* InSrcValu
 			const bool bBoolValue = SrcBoolProp->GetPropertyValue(SrcElemValue);
 			DestBoolProp->SetPropertyValue(DestElemValue, bBoolValue);
 		}
+		return true;
 	}
-	else
-	{
-		InSrcProp->CopyCompleteValue(InDestValue, InSrcValue);
-	}
-
+	
+	InSrcProp->CopyCompleteValue(InDestValue, InSrcValue);
 	return true;
 }
 
diff --git a/Engine/Source/Runtime/Engine/Private/KismetSystemLibrary.cpp b/Engine/Source/Runtime/Engine/Private/KismetSystemLibrary.cpp
--- a/Engine/Source/Runtime/Engine/Private/KismetSystemLibrary.cpp
+++ b/Engine/Source/Runtime/Engine/Private/KismetSystemLibrary.cpp
@@ -2699,7 +2699,7 @@ bool UKismetSystemLibrary::Generic_GetEditorProperty(const UObject* Object, cons
 
 	if (EnumHasAnyFlags(AccessResult, EPropertyAccessResultFlags::ConversionFailed))
 	{
-		FFrame::KismetExecutionMessage(*FString::Printf(TEXT("Property '%s' (%s) on '%s' (%s) tried to read into a property value of the incorrect type (%s)"), *ObjectProp->GetName(), *ObjectProp->GetClass()->GetName(), *Object->GetPathName(), *Object->GetClass()->GetName(), *ValueProp->GetClass()->GetName()), ELogVerbosity::Warning, PropertyGetFailedWarning);
+		FFrame::KismetExecutionMessage(*FString::Printf(TEXT("Property '%s' (%s) on '%s' (%s) tried to get to a property value of the incorrect type (%s)"), *ObjectProp->GetName(), *ObjectProp->GetClass()->GetName(), *Object->GetPathName(), *Object->GetClass()->GetName(), *ValueProp->GetClass()->GetName()), ELogVerbosity::Warning, PropertyGetFailedWarning);
 		return false;
 	}
 
@@ -2798,7 +2798,7 @@ bool UKismetSystemLibrary::Generic_SetEditorProperty(UObject* Object, const FPro
 
 	if (EnumHasAnyFlags(AccessResult, EPropertyAccessResultFlags::ConversionFailed))
 	{
-		FFrame::KismetExecutionMessage(*FString::Printf(TEXT("Property '%s' (%s) on '%s' (%s) tried to write into a property value of the incorrect type (%s)"), *ObjectProp->GetName(), *ObjectProp->GetClass()->GetName(), *Object->GetPathName(), *Object->GetClass()->GetName(), *ValueProp->GetClass()->GetName()), ELogVerbosity::Warning, PropertySetFailedWarning);
+		FFrame::KismetExecutionMessage(*FString::Printf(TEXT("Property '%s' (%s) on '%s' (%s) tried to set from a property value of the incorrect type (%s)"), *ObjectProp->GetName(), *ObjectProp->GetClass()->GetName(), *Object->GetPathName(), *Object->GetClass()->GetName(), *ValueProp->GetClass()->GetName()), ELogVerbosity::Warning, PropertySetFailedWarning);
 		return false;
 	}
 
