Fix a crash when a bulkdata file io request gets serviced immediately by the FPakPrecacher and the caller attempts to use the FBulkDataIORequest in the resulting callback.

Make sure that FBulkDataIORequest::ReadRequest is set before we invoke the callers callback.

#rnx

diff --git a/Engine/Source/Runtime/CoreUObject/Private/Serialization/BulkData.cpp b/Engine/Source/Runtime/CoreUObject/Private/Serialization/BulkData.cpp
--- a/Engine/Source/Runtime/CoreUObject/Private/Serialization/BulkData.cpp
+++ b/Engine/Source/Runtime/CoreUObject/Private/Serialization/BulkData.cpp
@@ -1716,8 +1716,13 @@ bool FBulkDataIORequest::MakeReadRequest(int64 Offset, int64 BytesToRead, EAsync
 	check(ReadRequest == nullptr);
 
 	FBulkDataIORequestCallBack LocalCallback = *CompleteCallback;
-	FAsyncFileCallBack AsyncFileCallBack = [LocalCallback, BytesToRead, this](bool bWasCancelled, IAsyncReadRequest*)
+	FAsyncFileCallBack AsyncFileCallBack = [LocalCallback, BytesToRead, this](bool bWasCancelled, IAsyncReadRequest* InRequest)
 	{
+		// In some cases the call to ReadRequest can invoke the callback immediately (if the requested data is cached 
+		// in the pak file system for example) which means that FBulkDataIORequest::ReadRequest might not actually be
+		// set correctly, so we need to make sure it is assigned before we invoke LocalCallback!
+		ReadRequest = InRequest;
+
 		Size = BytesToRead;
 		LocalCallback(bWasCancelled, this);
 	};
