Changed how the engine detects trashed user structs - instead of looking at field lists, it's now going to be using dedicated functions.

The old way, even though was effective, was not entirely correct and resulted in fallout after FProperty changes.

diff --git a/Engine/Source/Editor/KismetCompiler/Private/UserDefinedStructureCompilerUtils.cpp b/Engine/Source/Editor/KismetCompiler/Private/UserDefinedStructureCompilerUtils.cpp
--- a/Engine/Source/Editor/KismetCompiler/Private/UserDefinedStructureCompilerUtils.cpp
+++ b/Engine/Source/Editor/KismetCompiler/Private/UserDefinedStructureCompilerUtils.cpp
@@ -167,6 +167,7 @@ struct FUserDefinedStructureCompilerInner
 			StructToClean->ScriptObjectReferences.Empty();
 			StructToClean->PropertyLink = nullptr;
 			StructToClean->ErrorMessage.Empty();
+			StructToClean->SetStructTrashed(true);
 		}
 
 		return TransientStruct;
diff --git a/Engine/Source/Editor/PropertyEditor/Private/ObjectPropertyNode.cpp b/Engine/Source/Editor/PropertyEditor/Private/ObjectPropertyNode.cpp
--- a/Engine/Source/Editor/PropertyEditor/Private/ObjectPropertyNode.cpp
+++ b/Engine/Source/Editor/PropertyEditor/Private/ObjectPropertyNode.cpp
@@ -176,7 +176,7 @@ bool FObjectPropertyNode::GetReadAddressUncached(FPropertyNode& InNode,
 	}
 
 	UStruct* OwnerStruct = InItemProperty->GetOwnerStruct();
-	if (!OwnerStruct || OwnerStruct->ChildProperties == nullptr)
+	if (!OwnerStruct || OwnerStruct->IsStructTrashed())
 	{
 		// Verify that the property is not part of an invalid trash class
 		return false;
diff --git a/Engine/Source/Editor/PropertyEditor/Private/PropertyNode.cpp b/Engine/Source/Editor/PropertyEditor/Private/PropertyNode.cpp
--- a/Engine/Source/Editor/PropertyEditor/Private/PropertyNode.cpp
+++ b/Engine/Source/Editor/PropertyEditor/Private/PropertyNode.cpp
@@ -426,7 +426,7 @@ EPropertyDataValidationResult FPropertyNode::EnsureDataIsValid()
 			FProperty* MyProperty = Property.Get();
 			UStruct* OwnerStruct = MyProperty->GetOwnerStruct();
 
-			if (!OwnerStruct || OwnerStruct->ChildProperties == nullptr)
+			if (!OwnerStruct || OwnerStruct->IsStructTrashed())
 			{
 				//verify that the property is not part of an invalid trash class, treat it as an invalid object if it is which will cause a refresh
 				return EPropertyDataValidationResult::ObjectInvalid;
diff --git a/Engine/Source/Editor/PropertyEditor/Private/StructurePropertyNode.h b/Engine/Source/Editor/PropertyEditor/Private/StructurePropertyNode.h
--- a/Engine/Source/Editor/PropertyEditor/Private/StructurePropertyNode.h
+++ b/Engine/Source/Editor/PropertyEditor/Private/StructurePropertyNode.h
@@ -50,7 +50,7 @@ public:
 		}
 
 		UStruct* OwnerStruct = InItemProperty->GetOwnerStruct();
-		if (!OwnerStruct || (OwnerStruct->Children == nullptr && OwnerStruct->ChildProperties == nullptr))
+		if (!OwnerStruct || OwnerStruct->IsStructTrashed())
 		{
 			// Verify that the property is not part of an invalid trash class
 			return false;
diff --git a/Engine/Source/Runtime/CoreUObject/Private/UObject/Class.cpp b/Engine/Source/Runtime/CoreUObject/Private/UObject/Class.cpp
--- a/Engine/Source/Runtime/CoreUObject/Private/UObject/Class.cpp
+++ b/Engine/Source/Runtime/CoreUObject/Private/UObject/Class.cpp
@@ -2873,6 +2873,7 @@ void UScriptStruct::ExportText(FString& ValueStr, const void* Value, const void*
 void UScriptStruct::Link(FArchive& Ar, bool bRelinkExistingProperties)
 {
 	Super::Link(Ar, bRelinkExistingProperties);
+	SetStructTrashed(false);
 	if (!HasDefaults()) // if you have CppStructOps, then that is authoritative, otherwise we look at the properties
 	{
 		StructFlags = EStructFlags(StructFlags | STRUCT_ZeroConstructor | STRUCT_NoDestructor | STRUCT_IsPlainOldData);
@@ -3150,6 +3151,23 @@ void UScriptStruct::DestroyStruct(void* Dest, int32 ArrayDim) const
 	}
 }
 
+bool UScriptStruct::IsStructTrashed() const
+{
+	return !!(StructFlags & STRUCT_Trashed);
+}
+
+void UScriptStruct::SetStructTrashed(bool bIsTrash)
+{
+	if (bIsTrash)
+	{
+		StructFlags = EStructFlags(StructFlags | STRUCT_Trashed);
+	}
+	else
+	{
+		StructFlags = EStructFlags(StructFlags & ~STRUCT_Trashed);
+	}
+}
+
 void UScriptStruct::RecursivelyPreload() {}
 
 FGuid UScriptStruct::GetCustomGuid() const
@@ -4200,6 +4218,11 @@ void UClass::SetSuperStruct(UStruct* NewSuperStruct)
 	HashObject(this);
 }
 
+bool UClass::IsStructTrashed() const
+{
+	return Children == nullptr && ChildProperties == nullptr && ClassDefaultObject == nullptr;
+}
+
 void UClass::Serialize( FArchive& Ar )
 {
 	if ( Ar.IsLoading() || Ar.IsModifyingWeakAndStrongReferences() )
diff --git a/Engine/Source/Runtime/CoreUObject/Public/UObject/Class.h b/Engine/Source/Runtime/CoreUObject/Public/UObject/Class.h
--- a/Engine/Source/Runtime/CoreUObject/Public/UObject/Class.h
+++ b/Engine/Source/Runtime/CoreUObject/Public/UObject/Class.h
@@ -497,6 +497,12 @@ public:
 	/** Returns a human readable string for a given field, overridden for user defined structs */
 	virtual FString GetAuthoredNameForField(const FField* Field) const;
 
+	/** If true, this class has been cleaned and sanitized (trashed) and should not be used */
+	virtual bool IsStructTrashed() const
+	{
+		return false;
+	}
+
 #if WITH_EDITORONLY_DATA
 	/** Try and find boolean metadata with the given key. If not found on this class, work up hierarchy looking for it. */
 	bool GetBoolMetaDataHierarchical(const FName& Key) const;
@@ -617,6 +623,9 @@ enum EStructFlags
 	/** If set, this struct can share net serialization state across connections */
 	STRUCT_NetSharedSerialization = 0x00400000,
 
+	/** If set, this struct has been cleaned and sanitized (trashed) and should not be used */
+	STRUCT_Trashed = 0x00800000,
+
 	/** Struct flags that are automatically inherited */
 	STRUCT_Inherit				= STRUCT_HasInstancedReference|STRUCT_Atomic,
 
@@ -1476,8 +1485,12 @@ public:
 	virtual COREUOBJECT_API void Link(FArchive& Ar, bool bRelinkExistingProperties) override;
 	virtual COREUOBJECT_API void InitializeStruct(void* Dest, int32 ArrayDim = 1) const override;
 	virtual COREUOBJECT_API void DestroyStruct(void* Dest, int32 ArrayDim = 1) const override;
+	virtual COREUOBJECT_API bool IsStructTrashed() const override;
 	// End of UStruct interface.
 
+	/** Sets or unsets the trashed flag on this struct */
+	void COREUOBJECT_API SetStructTrashed(bool bIsTrash);
+
 	/** 
 	 * Stash a CppStructOps for future use 
 	 * @param Target Name of the struct 
@@ -2751,6 +2764,7 @@ public:
 	// UStruct interface.
 	virtual void Link(FArchive& Ar, bool bRelinkExistingProperties) override;
 	virtual void SetSuperStruct(UStruct* NewSuperStruct) override;
+	virtual bool IsStructTrashed() const override;
 	// End of UStruct interface.
 
 #if WITH_EDITOR
