FieldPaths will now store the Property class internally so that they don't resolve to different type poroperties after someone changes property type in Blueprints

diff --git a/Engine/Source/Runtime/CoreUObject/Private/UObject/FieldPath.cpp b/Engine/Source/Runtime/CoreUObject/Private/UObject/FieldPath.cpp
--- a/Engine/Source/Runtime/CoreUObject/Private/UObject/FieldPath.cpp
+++ b/Engine/Source/Runtime/CoreUObject/Private/UObject/FieldPath.cpp
@@ -11,6 +11,7 @@
 
 #if WITH_EDITORONLY_DATA
 FFieldPath::FFieldPath(UField* InField, const FName& InPropertyTypeName)
+	: InitialFieldClass(nullptr)
 {
 	if (InField)
 	{
@@ -34,6 +35,9 @@ void FFieldPath::Generate(FField* InField)
 		}
 		ResolvedFieldOwner = GUObjectArray.ObjectToIndex(InField->GetOwnerUObject());
 		SerialNumber = GlobalSerialNumber;
+#if WITH_EDITORONLY_DATA
+		InitialFieldClass = InField->GetClass();
+#endif // WITH_EDITORONLY_DATA
 	}
 	else
 	{
diff --git a/Engine/Source/Runtime/CoreUObject/Public/UObject/FieldPath.h b/Engine/Source/Runtime/CoreUObject/Public/UObject/FieldPath.h
--- a/Engine/Source/Runtime/CoreUObject/Public/UObject/FieldPath.h
+++ b/Engine/Source/Runtime/CoreUObject/Public/UObject/FieldPath.h
@@ -18,6 +18,10 @@ struct COREUOBJECT_API FFieldPath
 protected:
 	/** Untracked pointer to the resolved property */
 	mutable FField* ResolvedField;
+#if WITH_EDITORONLY_DATA
+	/** In editor builds, store the original class of the resolved property in case it changes after recompiling BPs */
+	mutable FFieldClass* InitialFieldClass;
+#endif
 	/** GC tracked index of property owner UObject */
 	mutable int32 ResolvedFieldOwner;
 	/** Serial number this FFieldPath was last resolved with */
@@ -38,12 +42,18 @@ public:
 
 	FFieldPath()
 		: ResolvedField(nullptr)
+#if WITH_EDITORONLY_DATA
+		, InitialFieldClass(nullptr)
+#endif // WITH_EDITORONLY_DATA
 		, ResolvedFieldOwner(-1)
 		, SerialNumber(-1)
 	{}
 
 	FFieldPath(FField* InField)
 		: ResolvedField(InField)
+#if WITH_EDITORONLY_DATA
+		, InitialFieldClass(nullptr)
+#endif // WITH_EDITORONLY_DATA
 		, ResolvedFieldOwner(-1)
 		, SerialNumber(-1)
 	{
@@ -52,6 +62,9 @@ public:
 
 	FFieldPath(const FFieldPath& Other)
 		: ResolvedField(Other.ResolvedField)
+#if WITH_EDITORONLY_DATA
+		, InitialFieldClass(Other.InitialFieldClass)
+#endif // WITH_EDITORONLY_DATA
 		, ResolvedFieldOwner(Other.ResolvedFieldOwner)
 		, SerialNumber(Other.SerialNumber)
 		, Path(Other.Path)
@@ -76,6 +89,9 @@ public:
 	inline void ClearCachedField()
 	{
 		ResolvedField = nullptr;
+#if WITH_EDITORONLY_DATA
+		InitialFieldClass = nullptr;
+#endif // WITH_EDITORONLY_DATA
 		ResolvedFieldOwner = -1;
 	}
 
@@ -96,11 +112,21 @@ public:
 	{
 		int32 FoundOwner = -1;
 		FField* FoundField = TryToResolvePath(InCurrentStruct, &FoundOwner);
-		if (FoundField && FoundField->IsA(ExpectedClass))
+		if (FoundField && FoundField->IsA(ExpectedClass) 
+#if WITH_EDITORONLY_DATA
+			&& (!InitialFieldClass || FoundField->IsA(InitialFieldClass))
+#endif // WITH_EDITORONLY_DATA
+			)
 		{
 			ResolvedField = FoundField;
 			ResolvedFieldOwner = FoundOwner;
 			SerialNumber = GlobalSerialNumber;
+#if WITH_EDITORONLY_DATA
+			if (!InitialFieldClass)
+			{
+				InitialFieldClass = FoundField->GetClass();
+			}
+#endif // WITH_EDITORONLY_DATA
 		}
 		else
 		{
@@ -253,10 +279,7 @@ public:
 		// one or both of the types have only been forward-declared.
 		static_assert(TPointerIsConvertibleFromTo<OtherPropertyType, PropertyType>::Value, "Unable to convert TFieldPath - types are incompatible");
 
-		ResolvedField = Other.ResolvedField;
-		ResolvedFieldOwner = Other.ResolvedFieldOwner;
-		SerialNumber = Other.SerialNumber;
-		Path = Other.Path;
+		FFieldPath::operator=(Other);
 	}
 
 	/**
