TPropertyIterator properly honors template parameter type, skipping properties of unrelated types
#rnx

diff --git a/Engine/Source/Runtime/CoreUObject/Public/UObject/PropertyIterator.h b/Engine/Source/Runtime/CoreUObject/Public/UObject/PropertyIterator.h
--- a/Engine/Source/Runtime/CoreUObject/Public/UObject/PropertyIterator.h
+++ b/Engine/Source/Runtime/CoreUObject/Public/UObject/PropertyIterator.h
@@ -30,7 +30,7 @@ public:
 			Struct = *ClassIterator;
 			Field = Struct->ChildProperties;
 
-			if (!Field)
+			if (!Field || !Field->IsA<T>())
 			{
 				IterateToNext();
 			}
@@ -73,23 +73,27 @@ public:
 protected:
 	inline void IterateToNext()
 	{
-		FField* NewField = nullptr;
+		FField* NewField = Field;
 		do
 		{
-			if (Field)
+			if (NewField)
 			{
-				NewField = Field->Next;
+				NewField = NewField->Next;
 			}
 			if (!NewField && ClassIterator)
 			{
 				++ClassIterator;
-				Struct = *ClassIterator;
-				if (Struct)
+				if(ClassIterator)
 				{
+					Struct = *ClassIterator;
 					NewField = Struct->ChildProperties;
 				}
+				else
+				{
+					NewField = nullptr;
+				}
 			}
-		} while (ClassIterator && !NewField);
+		} while (ClassIterator && (!NewField || !NewField->IsA<T>()));
 		Field = NewField;
 	}
 };
