UnrealHeaderTool: Fixes for various unsafe implicit type conversion warnings (includes a few CoreUObject fixes as well, but that module is not enabled yet)

diff --git a/Engine/Source/Programs/UnrealHeaderTool/Private/CodeGenerator.cpp b/Engine/Source/Programs/UnrealHeaderTool/Private/CodeGenerator.cpp
--- a/Engine/Source/Programs/UnrealHeaderTool/Private/CodeGenerator.cpp
+++ b/Engine/Source/Programs/UnrealHeaderTool/Private/CodeGenerator.cpp
@@ -307,7 +307,7 @@ namespace
  */
 const TCHAR* FindIdentifierExactMatch(const TCHAR* StringBegin, const TCHAR* StringEnd, const FString& Identifier)
 {
-	int32 StringLen = StringEnd - StringBegin;
+	int32 StringLen = UE_PTRDIFF_TO_INT32(StringEnd - StringBegin);
 
 	// Check for exact match first.
 	if (FCString::Strncmp(StringBegin, *Identifier, StringLen) == 0)
@@ -375,7 +375,7 @@ int32 FindIdentifierExactMatch(const FString& String, const FString& Identifier)
 		return INDEX_NONE;
 	}
 
-	return IdentifierPtr - *String;
+	return UE_PTRDIFF_TO_INT32(IdentifierPtr - *String);
 }
 
 /**
@@ -4999,7 +4999,7 @@ void FNativeClassHeaderGenerator::ExportFunctionThunk(FUHTStringBuilder& RPCWrap
 	const TCHAR* ClassEnd   = ClassRange.End;
 	FString      ClassName  = Function->GetOwnerClass()->GetName();
 
-	FString ClassDefinition(ClassEnd - ClassStart, ClassStart);
+	FString ClassDefinition(UE_PTRDIFF_TO_INT32(ClassEnd - ClassStart), ClassStart);
 
 	bool bHasImplementation = HasIdentifierExactMatch(ClassDefinition, FunctionData.CppImplName);
 	bool bHasValidate = HasIdentifierExactMatch(ClassDefinition, FunctionData.CppValidationImplName);
@@ -5209,10 +5209,10 @@ void FNativeClassHeaderGenerator::ExportNativeFunctions(FOutputDevice& OutGenera
 		{
 			const TCHAR* ClassStart = ClassRange.Start;
 			const TCHAR* ClassEnd   = ClassRange.End;
-			FString ClassDefinition(ClassEnd - ClassStart, ClassStart);
+			FString ClassDefinition(UE_PTRDIFF_TO_INT32(ClassEnd - ClassStart), ClassStart);
 
 			FString FunctionName = Function->GetName();
-			int32 ClassDefinitionStartPosition = ClassStart - *SourceFile.GetContent();
+			int32 ClassDefinitionStartPosition = UE_PTRDIFF_TO_INT32(ClassStart - *SourceFile.GetContent());
 
 			int32 ImplementationPosition = FindIdentifierExactMatch(ClassDefinition, FunctionData.CppImplName);
 			bool bHasImplementation = ImplementationPosition != INDEX_NONE;
diff --git a/Engine/Source/Programs/UnrealHeaderTool/Private/HeaderParser.cpp b/Engine/Source/Programs/UnrealHeaderTool/Private/HeaderParser.cpp
--- a/Engine/Source/Programs/UnrealHeaderTool/Private/HeaderParser.cpp
+++ b/Engine/Source/Programs/UnrealHeaderTool/Private/HeaderParser.cpp
@@ -267,7 +267,7 @@ namespace
 					{
 						FError::Throwf(TEXT("Invalid network identifier %s for function"), IdentifierPtr);
 					}
-					FuncInfo.RPCId = TempInt;
+					FuncInfo.RPCId = (uint16)TempInt;
 				}
 				else if (FCString::Strnicmp(IdentifierPtr, ResponseIdTag, UE_ARRAY_COUNT(ResponseIdTag) - 1) == 0 ||
 					FCString::Strnicmp(IdentifierPtr, JSBridgePriTag, UE_ARRAY_COUNT(JSBridgePriTag) - 1) == 0)
@@ -277,7 +277,7 @@ namespace
 					{
 						FError::Throwf(TEXT("Invalid network identifier %s for function"), IdentifierPtr);
 					}
-					FuncInfo.RPCResponseId = TempInt;
+					FuncInfo.RPCResponseId = (uint16)TempInt;
 				}
 			}
 			else
@@ -6949,7 +6949,7 @@ UDelegateFunction* FHeaderParser::CompileDelegateDeclaration(FClasses& AllClasse
 		ParseParameterList(AllClasses, DelegateSignatureFunction, /*bExpectCommaBeforeName=*/ true);
 
 		// Check the expected versus actual number of parameters
-		int32 ParamCount = FoundParamCount - DelegateParameterCountStrings.GetData() + 1;
+		int32 ParamCount = UE_PTRDIFF_TO_INT32(FoundParamCount - DelegateParameterCountStrings.GetData()) + 1;
 		if (DelegateSignatureFunction->NumParms != ParamCount)
 		{
 			FError::Throwf(TEXT("Expected %d parameters but found %d parameters"), ParamCount, DelegateSignatureFunction->NumParms);
@@ -8760,7 +8760,7 @@ void FHeaderParser::ParseClassName(const TCHAR* Temp, FString& ClassName)
 		++Temp;
 	}
 
-	ClassName = FString(Temp - StringStart, StringStart);
+	ClassName = FString(UE_PTRDIFF_TO_INT32(Temp - StringStart), StringStart);
 	if (ClassName.EndsWith(TEXT("_API"), ESearchCase::CaseSensitive))
 	{
 		// RequiresAPI token for a given module
diff --git a/Engine/Source/Programs/UnrealHeaderTool/Private/ParserHelper.h b/Engine/Source/Programs/UnrealHeaderTool/Private/ParserHelper.h
--- a/Engine/Source/Programs/UnrealHeaderTool/Private/ParserHelper.h
+++ b/Engine/Source/Programs/UnrealHeaderTool/Private/ParserHelper.h
@@ -940,7 +940,7 @@ public:
 	{
 		if( TokenType==TOKEN_Const && Type==CPT_Int64 )
 		{
-			I = Int64;
+			I = (int32)Int64;
 			return 1;
 		}
 		else if( TokenType==TOKEN_Const && Type==CPT_Int )
@@ -958,7 +958,7 @@ public:
 			I = (int32) Float;
 			return 1;
 		}
-		else if (TokenType == TOKEN_Const && Type == CPT_Double && Float == FMath::TruncToInt(Double))
+		else if (TokenType == TOKEN_Const && Type == CPT_Double && Double == FMath::TruncToInt((float)Double))
 		{
 			I = (int32) Double;
 			return 1;
@@ -988,7 +988,7 @@ public:
 			I = (int32) Float;
 			return 1;
 		}
-		else if (TokenType == TOKEN_Const && Type == CPT_Double && Float == FMath::TruncToInt(Double))
+		else if (TokenType == TOKEN_Const && Type == CPT_Double && Double == FMath::TruncToInt((float)Double))
 		{
 			I = (int32) Double;
 			return 1;
@@ -1053,9 +1053,9 @@ struct FFuncInfo
 	/** Endpoint name */
 	FString		EndpointName;
 	/** Identifier for an RPC call to a platform service */
-	int16		RPCId;
+	uint16		RPCId;
 	/** Identifier for an RPC call expecting a response */
-	int16		RPCResponseId;
+	uint16		RPCResponseId;
 	/** Whether this function represents a sealed event */
 	bool		bSealedEvent;
 	/** Delegate macro line in header. */
diff --git a/Engine/Source/Programs/UnrealHeaderTool/Private/StringUtils.cpp b/Engine/Source/Programs/UnrealHeaderTool/Private/StringUtils.cpp
--- a/Engine/Source/Programs/UnrealHeaderTool/Private/StringUtils.cpp
+++ b/Engine/Source/Programs/UnrealHeaderTool/Private/StringUtils.cpp
@@ -107,7 +107,7 @@ FORCEINLINE uint64 GenerateTextHash64(const TCHAR* Str)
 	{
 		const TCHAR* End = FindCrOrNul(Str);
 		
-		if (int32 Len = End - Str)
+		if (int32 Len = UE_PTRDIFF_TO_INT32(End - Str))
 		{
 			Hash = CityHash64WithSeed(reinterpret_cast<const char*>(Str), Len * sizeof(TCHAR), Hash);
 		}
diff --git a/Engine/Source/Programs/UnrealHeaderTool/UnrealHeaderTool.Build.cs b/Engine/Source/Programs/UnrealHeaderTool/UnrealHeaderTool.Build.cs
--- a/Engine/Source/Programs/UnrealHeaderTool/UnrealHeaderTool.Build.cs
+++ b/Engine/Source/Programs/UnrealHeaderTool/UnrealHeaderTool.Build.cs
@@ -27,5 +27,7 @@ public class UnrealHeaderTool : ModuleRules
 			});
 		
 		bEnableExceptions = true;
+
+		UnsafeTypeCastWarningLevel = WarningLevel.Warning;
 	}
 }
diff --git a/Engine/Source/Runtime/CoreUObject/Public/UObject/Stack.h b/Engine/Source/Runtime/CoreUObject/Public/UObject/Stack.h
--- a/Engine/Source/Runtime/CoreUObject/Public/UObject/Stack.h
+++ b/Engine/Source/Runtime/CoreUObject/Public/UObject/Stack.h
@@ -291,10 +291,10 @@ inline VariableSizeType FFrame::ReadVariableSize( FProperty** ExpressionField )
 	FProperty* Property = CastField<FProperty>(Field);
 	if (Property)
 	{
-		Result = Property->GetSize();
+		Result = (VariableSizeType)Property->GetSize();
 	}
 
-	if ( ExpressionField != nullptr )
+	if (ExpressionField != nullptr)
 	{
 		*ExpressionField = Property;
 	}
diff --git a/Engine/Source/Runtime/CoreUObject/Public/UObject/UnrealType.h b/Engine/Source/Runtime/CoreUObject/Public/UObject/UnrealType.h
--- a/Engine/Source/Runtime/CoreUObject/Public/UObject/UnrealType.h
+++ b/Engine/Source/Runtime/CoreUObject/Public/UObject/UnrealType.h
@@ -1194,7 +1194,8 @@ class COREUOBJECT_API FNumericProperty : public FProperty
 	{
 		if (!TIsFloatingPoint<T>::Value)
 		{
-			return CanHoldDoubleValueInternal(Value);
+			//@TODO: FLOATPRECISION: This feels wrong, it might be losing precision before it tests to see if it's going to lose precision...
+			return CanHoldDoubleValueInternal((double)Value);
 		}
 		else if (TIsSigned<T>::Value)
 		{
@@ -1393,7 +1394,7 @@ public:
 				if (!Tag.EnumName.IsNone())
 				{
 					int64 PreviousValue = this->ReadEnumAsInt64(Slot, DefaultsStruct, Tag);
-					this->SetPropertyValue_InContainer(Data, PreviousValue, Tag.ArrayIndex);
+					this->SetPropertyValue_InContainer(Data, (TCppType)PreviousValue, Tag.ArrayIndex);
 				}
 				else
 				{
@@ -1450,17 +1451,17 @@ public:
 	virtual void SetIntPropertyValue(void* Data, uint64 Value) const override
 	{
 		check(TIsIntegral<TCppType>::Value);
-		TTypeFundamentals::SetPropertyValue(Data, Value);
+		TTypeFundamentals::SetPropertyValue(Data, (TCppType)Value);
 	}
 	virtual void SetIntPropertyValue(void* Data, int64 Value) const override
 	{
 		check(TIsIntegral<TCppType>::Value);
-		TTypeFundamentals::SetPropertyValue(Data, Value);
+		TTypeFundamentals::SetPropertyValue(Data, (TCppType)Value);
 	}
 	virtual void SetFloatingPointPropertyValue(void* Data, double Value) const override
 	{
 		check(TIsFloatingPoint<TCppType>::Value);
-		TTypeFundamentals::SetPropertyValue(Data, Value);
+		TTypeFundamentals::SetPropertyValue(Data, (TCppType)Value);
 	}
 	virtual void SetNumericPropertyValueFromString(void* Data, TCHAR const* Value) const override
 	{
@@ -1473,17 +1474,17 @@ public:
 	virtual int64 GetSignedIntPropertyValue(void const* Data) const override
 	{
 		check(TIsIntegral<TCppType>::Value);
-		return TTypeFundamentals::GetPropertyValue(Data);
+		return (int64)TTypeFundamentals::GetPropertyValue(Data);
 	}
 	virtual uint64 GetUnsignedIntPropertyValue(void const* Data) const override
 	{
 		check(TIsIntegral<TCppType>::Value);
-		return TTypeFundamentals::GetPropertyValue(Data);
+		return (uint64)TTypeFundamentals::GetPropertyValue(Data);
 	}
 	virtual double GetFloatingPointPropertyValue(void const* Data) const override
 	{
 		check(TIsFloatingPoint<TCppType>::Value);
-		return TTypeFundamentals::GetPropertyValue(Data);
+		return (double)TTypeFundamentals::GetPropertyValue(Data);
 	}
 	// End of FNumericProperty interface
 
@@ -3752,7 +3753,7 @@ public:
 			[LocalKeyPropForCapture](const void* A, const void* B) { return LocalKeyPropForCapture->Identical(A, B); }
 			))
 		{
-			int32 Idx = (Entry - (uint8*)Map->GetData(0, MapLayout)) / MapLayout.SetLayout.Size;
+			int32 Idx = (int32)((Entry - (uint8*)Map->GetData(0, MapLayout)) / MapLayout.SetLayout.Size);
 			RemoveAt(Idx);
 			return true;
 		}
