Added solution to use a macro in the warning message about not registering a replicated property in GetLifetimeReplicatedProps()

diff --git a/Engine/Source/Runtime/CoreUObject/Private/UObject/Class.cpp b/Engine/Source/Runtime/CoreUObject/Private/UObject/Class.cpp
--- a/Engine/Source/Runtime/CoreUObject/Private/UObject/Class.cpp
+++ b/Engine/Source/Runtime/CoreUObject/Private/UObject/Class.cpp
@@ -4149,7 +4149,7 @@ void UClass::ValidateRuntimeReplicationData()
 				}
 			}
 
-			UE_LOG(LogClass, Warning, TEXT("Property %s::%s (SourceClass: %s) was not registered in GetLifetimeReplicatedProps. This property will not be replicated."),
+			UE_LOG(LogClass, Warning, TEXT("Property %s::%s (SourceClass: %s) was not registered in GetLifetimeReplicatedProps. This property will not be replicated. Use DISABLE_REPLICATED_PROPERTY if not replicating was intentional."),
 				*GetName(), *RepProp->GetName(), *RepProp->GetOwnerClass()->GetName());
 		}
 	}
