Optimize -DiffOnly cooks by optimizing FindAssetInPackage

The cost of FindAssetInPackage during a diff-only cook of a large title was reduced by ~99.6%. This reduced the cost of the overall diff-only cook by ~95%.

diff --git a/Engine/Source/Runtime/CoreUObject/Private/UObject/SavePackage.cpp b/Engine/Source/Runtime/CoreUObject/Private/UObject/SavePackage.cpp
--- a/Engine/Source/Runtime/CoreUObject/Private/UObject/SavePackage.cpp
+++ b/Engine/Source/Runtime/CoreUObject/Private/UObject/SavePackage.cpp
@@ -3793,15 +3793,15 @@ FSavePackageResultStruct UPackage::Save(UPackage* InOuter, UObject* Base, EObjec
 					// Finds the asset object within a package
 					auto FindAssetInPackage = [](UPackage* Package) -> UObject*
 					{
-						for (UObject* Object : TObjectRange<UObject>())
-						{
-							if (Object->GetOuter() == Package && Object->IsAsset())
+						UObject* Asset = nullptr;
+						ForEachObjectWithOuter(Package, [&Asset](UObject* Object)
 							{
-								return Object;
-							}
-						}
-
-						return nullptr;
+								if (!Asset && Object->IsAsset())
+								{
+									Asset = Object;
+								}
+							}, /*bIncludeNestedObjects*/ false);
+						return Asset;
 					};
 
 					if (TargetPlatform != nullptr && (SaveFlags & SAVE_DiffCallstack))
