Fix ancient copy/paste bug in CoreRedirects. Caught by forcing all mem slack to 0 (so array is always reallocated on add) and running -stompmalloc

diff --git a/Engine/Source/Runtime/CoreUObject/Private/UObject/CoreRedirects.cpp b/Engine/Source/Runtime/CoreUObject/Private/UObject/CoreRedirects.cpp
--- a/Engine/Source/Runtime/CoreUObject/Private/UObject/CoreRedirects.cpp
+++ b/Engine/Source/Runtime/CoreUObject/Private/UObject/CoreRedirects.cpp
@@ -1454,7 +1454,7 @@ static void RegisterNativeRedirects46(TArray<FCoreRedirect>& Redirects)
 	EPhysicsType->ValueChanges.Add(TEXT("PhysType_Unfixed"), TEXT("PhysType_Simulated"));
 
 	FCoreRedirect* ESceneTextureId = ENUM_REDIRECT("ESceneTextureId", "ESceneTextureId");
-	EPhysicsType->ValueChanges.Add(TEXT("PPI_LightingModel"), TEXT("PPI_ShadingModelColor"));
+	ESceneTextureId->ValueChanges.Add(TEXT("PPI_LightingModel"), TEXT("PPI_ShadingModelColor"));
 
 	// 4.5
 
