Opening up FArchiveStackTrace to dump all USTRUCTs on package difference and not just guids and vectors. Trying to track some some FGameplayTagContainer differences, but that's not a core type.

diff --git a/Engine/Source/Runtime/CoreUObject/Private/Serialization/ArchiveStackTrace.cpp b/Engine/Source/Runtime/CoreUObject/Private/Serialization/ArchiveStackTrace.cpp
--- a/Engine/Source/Runtime/CoreUObject/Private/Serialization/ArchiveStackTrace.cpp
+++ b/Engine/Source/Runtime/CoreUObject/Private/Serialization/ArchiveStackTrace.cpp
@@ -492,11 +492,7 @@ namespace
 
 		if (FStructProperty* StructProp = CastField<FStructProperty>(Prop))
 		{
-			if (StructProp->Struct == TBaseStructure<FVector>::Get()
-				|| StructProp->Struct == TBaseStructure<FGuid>::Get())
-			{
-				return true;
-			}
+			return true;
 		}
 
 		return false;
