Back out revision 3 from ArchiveStackTrace.cpp. SerializedProp isn't even saved in the cases where gameplay tag container differs so there's nothing to output.

diff --git a/Engine/Source/Runtime/CoreUObject/Private/Serialization/ArchiveStackTrace.cpp b/Engine/Source/Runtime/CoreUObject/Private/Serialization/ArchiveStackTrace.cpp
--- a/Engine/Source/Runtime/CoreUObject/Private/Serialization/ArchiveStackTrace.cpp
+++ b/Engine/Source/Runtime/CoreUObject/Private/Serialization/ArchiveStackTrace.cpp
@@ -492,7 +492,11 @@ namespace
 
 		if (FStructProperty* StructProp = CastField<FStructProperty>(Prop))
 		{
-			return true;
+			if (StructProp->Struct == TBaseStructure<FVector>::Get()
+				|| StructProp->Struct == TBaseStructure<FGuid>::Get())
+			{
+				return true;
+			}
 		}
 
 		return false;
