Fixup some UHT Generation issues around NetProperties, as well as net property ordering issues.

diff --git a/Engine/Source/Programs/UnrealHeaderTool/Private/CodeGenerator.cpp b/Engine/Source/Programs/UnrealHeaderTool/Private/CodeGenerator.cpp
--- a/Engine/Source/Programs/UnrealHeaderTool/Private/CodeGenerator.cpp
+++ b/Engine/Source/Programs/UnrealHeaderTool/Private/CodeGenerator.cpp
@@ -163,36 +163,35 @@ namespace
 	/** Whether or not the given class has any replicated properties. */
 	static bool ClassHasReplicatedProperties(UClass* Class)
 	{
-		for (TFieldIterator<FProperty> It(Class, EFieldIteratorFlags::ExcludeSuper); It; ++It)
+		if (!Class->HasAnyClassFlags(CLASS_ReplicationDataIsSetUp))
 		{
-			if ((It->PropertyFlags & CPF_Net) != 0)
+			for (TFieldIterator<FProperty> It(Class, EFieldIteratorFlags::ExcludeSuper); It; ++It)
 			{
-				return true;
+				if ((It->PropertyFlags & CPF_Net) != 0)
+				{
+					return true;
+				}
 			}
 		}
 
-		return false;
+		return Class->FirstOwnedClassRep < Class->ClassReps.Num();
 	}
 
-	static void ExportNetData(FOutputDevice& Out, UClass* Class, UClass* ParentRepClass, const TCHAR* API)
+	static void ExportNetData(FOutputDevice& Out, UClass* Class, const TCHAR* API)
 	{
 		const TArray<FRepRecord>& ClassReps = Class->ClassReps;
-		
-		const FString NetfieldStart = ParentRepClass
-				? FString::Printf(TEXT("(uint16)%s%s::ENetFields_Private::NETFIELD_REP_END + (uint16)1"), ParentRepClass->GetPrefixCPP(), *ParentRepClass->GetName())
-				: FString(TEXT("0"));
 
 		FUHTStringBuilder NetFieldBuilder;
 		NetFieldBuilder.Logf(TEXT(""
 		"\tenum class ENetFields_Private : uint16\r\n"
 		"\t{\r\n"
-		"\t\tNETFIELD_REP_START=%s,\r\n"), *NetfieldStart);
+		"\t\tNETFIELD_REP_START=(uint16)Super::ENetFields_Private::NETFIELD_REP_END + (uint16)1,\r\n"));
 
 		FUHTStringBuilder ArrayDimBuilder;
 
 		bool bAnyStaticArrays = false;
 		bool bIsFirst = true;
-		for (int32 ClassRepIndex = ClassReps.Num() - Class->NetProperties.Num() - Class->NetFields.Num(); ClassRepIndex < ClassReps.Num(); ++ClassRepIndex)
+		for (int32 ClassRepIndex = Class->FirstOwnedClassRep; ClassRepIndex < ClassReps.Num(); ++ClassRepIndex)
 		{
 			const FRepRecord& ClassRep = ClassReps[ClassRepIndex];
 			const FString PropertyName = ClassRep.Property->GetName();
@@ -276,21 +275,19 @@ namespace
 			}
 		}
 
-		UClass* ParentRepClass = nullptr;
-		Class->SetUpUhtReplicationData(&ParentRepClass);
 
-		ExportNetData(Writer, Class, ParentRepClass, API);
+		ExportNetData(Writer, Class, API);
 		
-		// If we don't have any Super Classes that have replicated properties,
-		// then we're the base most replicated class, and we should add in our interface stuff.
-		if (!ParentRepClass)
+		// If this class has replicated properties and it owns the first one, that means
+		// it's the base most replicated class. In that case, go ahead and add our interface macro.
+		if (Class->ClassReps.Num() > 0 && Class->FirstOwnedClassRep == 0)
 		{
 			OutFlags |= FNativeClassHeaderGenerator::EExportClassOutFlags::NeedsPushModelHeaders;
 			Writer.Logf(TEXT(
 				"private:\r\n"
-				"\tREPLICATED_BASE_CLASS(%s)\r\n"
+				"\tREPLICATED_BASE_CLASS(%s%s)\r\n"
 				"public:\r\n"
-			), *Class->GetName());
+			), Class->GetPrefixCPP(), *Class->GetName());
 		}
 	}
 }
@@ -2532,7 +2529,7 @@ void FNativeClassHeaderGenerator::ExportNativeGeneratedInitCode(FOutputDevice& O
 		FUHTStringBuilder ValidationBuilder;
 		ValidationBuilder.Log(TEXT("\t\tconst bool bIsValid = true"));
 
-		for (int32 i = Class->ClassReps.Num() - Class->NetFields.Num(); i < Class->ClassReps.Num(); ++i)
+		for (int32 i = Class->FirstOwnedClassRep; i < Class->ClassReps.Num(); ++i)
 		{
 			const FProperty* const Property = Class->ClassReps[i].Property;
 			const FString PropertyName = Property->GetName();
@@ -5726,7 +5723,7 @@ FNativeClassHeaderGenerator::FNativeClassHeaderGenerator(
 		{
 			if (ClassHasReplicatedProperties(Class))
 			{
-				Class->SetUpUhtReplicationData(nullptr);
+				Class->SetUpUhtReplicationData();
 			}
 		}
 	}
diff --git a/Engine/Source/Programs/UnrealHeaderTool/Private/HeaderParser.cpp b/Engine/Source/Programs/UnrealHeaderTool/Private/HeaderParser.cpp
--- a/Engine/Source/Programs/UnrealHeaderTool/Private/HeaderParser.cpp
+++ b/Engine/Source/Programs/UnrealHeaderTool/Private/HeaderParser.cpp
@@ -6143,7 +6143,7 @@ UClass* FHeaderParser::CompileClassDeclaration(FClasses& AllClasses)
 	check(Class->Children == NULL);
 	check(Class->Next == NULL);
 	check(Class->NetFields.Num() == 0);
-	check(Class->NetProperties.Num() == 0);
+	check(Class->FirstOwnedClassRep == 0);
 
 	// Make sure our parent classes is parsed.
 	for (UClass* Temp = Class->GetSuperClass(); Temp; Temp = Temp->GetSuperClass())
diff --git a/Engine/Source/Runtime/CoreUObject/Private/UObject/Class.cpp b/Engine/Source/Runtime/CoreUObject/Private/UObject/Class.cpp
--- a/Engine/Source/Runtime/CoreUObject/Private/UObject/Class.cpp
+++ b/Engine/Source/Runtime/CoreUObject/Private/UObject/Class.cpp
@@ -3878,7 +3878,7 @@ void UClass::FinishDestroy()
 	//warning: Must be emptied explicitly in order for intrinsic classes
 	// to not show memory leakage on exit.
 	NetFields.Empty();
-	NetProperties.Empty();
+	ClassReps.Empty();
 
 	ClassDefaultObject = nullptr;
 
@@ -3945,17 +3945,21 @@ static int32 GValidateReplicatedProperties = 1;
 
 static FAutoConsoleVariable CVarValidateReplicatedPropertyRegistration(TEXT("net.ValidateReplicatedPropertyRegistration"), GValidateReplicatedProperties, TEXT("Warns if replicated properties were not registered in GetLifetimeReplicatedProps."));
 
-void UClass::SetUpUhtReplicationData(UClass** OutSuperClassWithReplicatedData)
-{
 #if HACK_HEADER_GENERATOR
+void UClass::SetUpUhtReplicationData()
+{
 	if (!HasAnyClassFlags(CLASS_ReplicationDataIsSetUp) && PropertyLink != NULL)
 	{
-		NetFields.Empty();
         ClassReps.Empty();
 		if (UClass* SuperClass = GetSuperClass())
 		{
-			SuperClass->SetUpUhtReplicationData(nullptr);
+			SuperClass->SetUpUhtReplicationData();
 			ClassReps = SuperClass->ClassReps;
+			FirstOwnedClassRep = ClassReps.Num();
+		}
+		else
+		{
+			FirstOwnedClassRep = 0;
 		}
 
 		for (TFieldIterator<FProperty> It(this, EFieldIteratorFlags::ExcludeSuper); It; ++It)
@@ -3963,65 +3967,50 @@ void UClass::SetUpUhtReplicationData(UClass** OutSuperClassWithReplicatedData)
 			if (It->PropertyFlags & CPF_Net)
 			{
 				It->RepIndex = ClassReps.Num();
-				NetProperties.Add(*It);
 				new (ClassReps) FRepRecord(*It, 0);
 			}
 		}
 
 		ClassFlags |= CLASS_ReplicationDataIsSetUp;
 		ClassReps.Shrink();
-		NetFields.Shrink();
 	}
-
-	if (OutSuperClassWithReplicatedData)
-	{
-		UClass* SuperClass = GetSuperClass();
-		for (; SuperClass != nullptr && SuperClass->NetFields.Num() == 0; SuperClass = SuperClass->GetSuperClass()) {}
-		*OutSuperClassWithReplicatedData = SuperClass;
-	}
-#endif
 }
+#endif
 
 void UClass::SetUpRuntimeReplicationData()
 {
 	if (!HasAnyClassFlags(CLASS_ReplicationDataIsSetUp) && PropertyLink != NULL)
 	{
 		NetFields.Empty();
-		NetProperties.Empty();
 
 		if (UClass* SuperClass = GetSuperClass())
 		{
 			SuperClass->SetUpRuntimeReplicationData();
 			ClassReps = SuperClass->ClassReps;
+			FirstOwnedClassRep = ClassReps.Num();
 		}
 		else
 		{
 			ClassReps.Empty();
+			FirstOwnedClassRep = 0;
 		}
 
-		TArray<FProperty*> NetPropsForClassRep;		// Track properties so me can ensure they are sorted by offsets at the end
-
+		// Track properties so me can ensure they are sorted by offsets at the end
+		TArray<FProperty*> NetProperties;
 		for (TFieldIterator<FField> It(this, EFieldIteratorFlags::ExcludeSuper); It; ++It)
 		{
-			FProperty* Prop = CastField<FProperty>(*It);
-			if (Prop)
+			if (FProperty* Prop = CastField<FProperty>(*It))
 			{
-				if (Prop->PropertyFlags & CPF_Net)
+				if ((Prop->PropertyFlags & CPF_Net) && Prop->GetOwner<UObject>() == this)
 				{
 					NetProperties.Add(Prop);
-
-					if (Prop->GetOwner<UObject>() == this)
-					{
-						NetPropsForClassRep.Add(Prop);
-					}
 				}
 			}
-			}
+		}
 
-		for( TFieldIterator<UField> It(this,EFieldIteratorFlags::ExcludeSuper); It; ++It )
+		for(TFieldIterator<UField> It(this,EFieldIteratorFlags::ExcludeSuper); It; ++It)
 		{
-			UFunction* Func = Cast<UFunction>(*It);
-			if (Func)
+			if (UFunction * Func = Cast<UFunction>(*It))
 			{
 				// When loading reflection data (e.g. from blueprints), we may have references to placeholder functions, or reflection data 
 				// in children may be out of date. In that case we cannot enforce this check, but that is ok because reflection data will
@@ -4038,30 +4027,31 @@ void UClass::SetUpRuntimeReplicationData()
 		const bool bIsNativeClass = HasAnyClassFlags(CLASS_Native);
 		if (!bIsNativeClass)
 		{
-		// Sort NetProperties so that their ClassReps are sorted by memory offset
+			// Sort NetProperties so that their ClassReps are sorted by memory offset
 			struct FComparePropertyOffsets
-		{
-				FORCEINLINE bool operator()(FProperty& A, FProperty& B) const
 			{
-				// Ensure stable sort
-				if ( A.GetOffset_ForGC() == B.GetOffset_ForGC() )
+				FORCEINLINE bool operator()(FProperty& A, FProperty& B) const
 				{
-					return A.GetName() < B.GetName();
-				}
+					// Ensure stable sort
+					if ( A.GetOffset_ForGC() == B.GetOffset_ForGC() )
+					{
+						return A.GetName() < B.GetName();
+					}
 
-				return A.GetOffset_ForGC() < B.GetOffset_ForGC();
-			}
-		};
+					return A.GetOffset_ForGC() < B.GetOffset_ForGC();
+				}
+			};
 
 			Sort(NetProperties.GetData(), NetProperties.Num(), FComparePropertyOffsets());
 		}
 
-		for (int32 i = 0; i < NetPropsForClassRep.Num(); i++)
+		ClassReps.Reserve(ClassReps.Num() + NetProperties.Num());
+		for (int32 i = 0; i < NetProperties.Num(); i++)
 		{
-			NetPropsForClassRep[i]->RepIndex = ClassReps.Num();
-			for (int32 j = 0; j < NetPropsForClassRep[i]->ArrayDim; j++)
+			NetProperties[i]->RepIndex = ClassReps.Num();
+			for (int32 j = 0; j < NetProperties[i]->ArrayDim; j++)
 			{
-				new(ClassReps)FRepRecord(NetPropsForClassRep[i], j);
+				new(ClassReps)FRepRecord(NetProperties[i], j);
 			}
 		}
 
@@ -4071,7 +4061,6 @@ void UClass::SetUpRuntimeReplicationData()
 		}
 
 		NetFields.Shrink();
-		NetProperties.Shrink();
 
 		struct FCompareUFieldNames
 		{
@@ -4603,7 +4592,6 @@ void UClass::PurgeClass(bool bRecompilingOnLoad)
 	ClassUnique = 0;
 	ClassReps.Empty();
 	NetFields.Empty();
-	NetProperties.Empty();
 
 #if WITH_EDITOR
 	if (!bRecompilingOnLoad)
diff --git a/Engine/Source/Runtime/CoreUObject/Private/UObject/CoreNet.cpp b/Engine/Source/Runtime/CoreUObject/Private/UObject/CoreNet.cpp
--- a/Engine/Source/Runtime/CoreUObject/Private/UObject/CoreNet.cpp
+++ b/Engine/Source/Runtime/CoreUObject/Private/UObject/CoreNet.cpp
@@ -177,13 +177,13 @@ const FClassNetCache* FClassNetCacheMgr::GetClassNetCache( UClass* Class )
 
 		Result->Fields.Empty(Class->NetFields.Num());
 
-		TArray< FProperty* > Properties;
-		Properties.Empty( Class->NetProperties.Num() );
+		TArray<FProperty*> Properties;
+		Properties.Empty(Class->ClassReps.Num() - Class->FirstOwnedClassRep);
 
-		for (int32 i = 0; i < Class->NetProperties.Num(); i++)
+		for (int32 i = Class->FirstOwnedClassRep; i < Class->ClassReps.Num(); i++)
 		{
 			// Add each net field to cache, and assign index/checksum
-			FProperty* Property = CastField< FProperty >(Class->NetProperties[i]);
+			FProperty* Property = Class->ClassReps[i].Property;
 			check(Property);
 			Properties.Add(Property);
 
diff --git a/Engine/Source/Runtime/CoreUObject/Public/UObject/Class.h b/Engine/Source/Runtime/CoreUObject/Public/UObject/Class.h
--- a/Engine/Source/Runtime/CoreUObject/Public/UObject/Class.h
+++ b/Engine/Source/Runtime/CoreUObject/Public/UObject/Class.h
@@ -2554,8 +2554,8 @@ public:
 	/** List of network relevant fields (functions) */
 	TArray<UField*> NetFields;
 
-	/** List of network relevant fields (properties) */
-	TArray<FProperty*> NetProperties;
+	/** Index of the first ClassRep that belongs to this class. Anything before that was defined by / belongs to parent classes. */
+	int32 FirstOwnedClassRep = 0;
 
 #if WITH_EDITOR || HACK_HEADER_GENERATOR 
 	// Editor only properties
@@ -3117,7 +3117,10 @@ public:
 	 * Also happens after blueprint compiliation.
 	 */
 	void SetUpRuntimeReplicationData();
-	void SetUpUhtReplicationData(UClass** OutSuperClassWithReplicatedData);
+
+#if HACK_HEADER_GENERATOR
+	void SetUpUhtReplicationData();
+#endif  // HACK_HEADER_GENERATOR
 
 	/**
 	 * Helper function for determining if the given class is compatible with structured archive serialization
diff --git a/Engine/Source/Runtime/CoreUObject/Public/UObject/Object.h b/Engine/Source/Runtime/CoreUObject/Public/UObject/Object.h
--- a/Engine/Source/Runtime/CoreUObject/Public/UObject/Object.h
+++ b/Engine/Source/Runtime/CoreUObject/Public/UObject/Object.h
@@ -1457,6 +1457,12 @@ private:
 
 public:
 
+	enum class ENetFields_Private : uint16
+	{
+		NETFIELD_REP_START = 0,
+		NETFIELD_REP_END = 0
+	};
+
 	virtual void ValidateGeneratedRepEnums(const TArray<struct FRepRecord>& ClassReps) const {}	
 
 private:
