Added COREUOBJECT_API to FJsonArchiveInputFormatter so it can be used outside of CoreUObject

diff --git a/Engine/Source/Runtime/CoreUObject/Public/Serialization/Formatters/JsonArchiveInputFormatter.h b/Engine/Source/Runtime/CoreUObject/Public/Serialization/Formatters/JsonArchiveInputFormatter.h
--- a/Engine/Source/Runtime/CoreUObject/Public/Serialization/Formatters/JsonArchiveInputFormatter.h
+++ b/Engine/Source/Runtime/CoreUObject/Public/Serialization/Formatters/JsonArchiveInputFormatter.h
@@ -10,7 +10,7 @@ class FJsonValue;
 
 #if WITH_TEXT_ARCHIVE_SUPPORT
 
-class FJsonArchiveInputFormatter final : public FStructuredArchiveFormatter
+class COREUOBJECT_API FJsonArchiveInputFormatter final : public FStructuredArchiveFormatter
 {
 public:
 	FJsonArchiveInputFormatter(FArchive& InInner, TFunction<UObject* (const FString&)> InResolveObjectName = nullptr);
