Add support for memory mapped files to BulKData2 which will fix a crash when memory mapped animations have been enabled.

- Attempting to memory map inline data will give a log error(as the original implementation did but will not attempt to actually memory map it as it is very unlikely that inline data would have the correct alignment requirements for the operation to succeed so to keep the code simple we just serialize the inline data which was the existing fallback behaviour anyway.
- Added the memory mapping implementation as it's own method (MemoryMapBulkData) since we only need to support the non-inline case.
- Added a flag (bShouldForceLoad) that when true will force the BulkData2 data to be streamed in immediately. Currently used as a fallback if the memory mapping fails (all callers assume 100% success) but I made it generic enough that we could use it for something else if the use case occurs.
- This is a fairly quick fix to unblock testing, future work includes:
  - Reduce the added memory bloat (adding 2 additional pointers to BulkData2)
  - Implement copy/assignment of memory mapped BulkData (based on the above work)

#rnx

diff --git a/Engine/Source/Runtime/CoreUObject/Private/Serialization/BulkData2.cpp b/Engine/Source/Runtime/CoreUObject/Private/Serialization/BulkData2.cpp
--- a/Engine/Source/Runtime/CoreUObject/Private/Serialization/BulkData2.cpp
+++ b/Engine/Source/Runtime/CoreUObject/Private/Serialization/BulkData2.cpp
@@ -12,6 +12,7 @@
 #include "UObject/Object.h"
 #include "UObject/Package.h"
 #include "IO/IoDispatcher.h"
+#include "Async/MappedFileHandle.h"
 
 DEFINE_LOG_CATEGORY_STATIC(LogBulkDataRuntime, Log, All);
 
@@ -469,6 +470,8 @@ private:
 FBulkDataBase::FBulkDataBase(FBulkDataBase&& Other)
 	: ChunkID(Other.ChunkID) // Copies the entire union
 	, DataBuffer(Other.DataBuffer)
+	, MappedHandle(Other.MappedHandle)
+	, MappedRegion(Other.MappedRegion)
 	, BulkDataFlags(Other.BulkDataFlags)
 {
 	check(Other.LockStatus == LOCKSTATUS_Unlocked); // Make sure that the other object wasn't inuse
@@ -477,7 +480,11 @@ FBulkDataBase::FBulkDataBase(FBulkDataBase&& Other)
 	{
 		Other.Fallback.Token = InvalidToken; // Prevent the other object from unregistering the token
 	}	
-	Other.DataBuffer = nullptr;	// Prevent the other object from deleting our data
+
+	// Prevent the other object from deleting our data
+	Other.DataBuffer = nullptr;	
+	Other.MappedHandle = nullptr;
+	Other.MappedRegion = nullptr;
 }
 
 FBulkDataBase& FBulkDataBase::operator=(const FBulkDataBase& Other)
@@ -502,6 +509,8 @@ FBulkDataBase& FBulkDataBase::operator=(const FBulkDataBase& Other)
 		FMemory::Memcpy(DataBuffer, DataBuffer, DataSize);
 	}
 
+	checkf(Other.MappedHandle == nullptr && Other.MappedRegion == nullptr, TEXT("Attempting to make a copy of a memory mapped bulkdata object, this is not currently supported!"));
+
 	return *this;
 }
 
@@ -606,8 +615,15 @@ void FBulkDataBase::Serialize(FArchive& Ar, UObject* Owner, int32 /*Index*/, boo
 		FArchive* CacheableArchive = Ar.GetCacheableArchive();
 		if (Ar.IsAllowingLazyLoading() && CacheableArchive != nullptr)
 		{
+			// Some failed paths require us to load the data before we return from ::Serialize but it is not
+			// safe to do so until the end of this method. By setting this flag to true we can indicate that 
+			// the load is required.
+			bool bShouldForceStreamData = false;
+
 			if (IsInlined())
 			{
+				UE_CLOG(bAttemptFileMapping, LogSerialization, Error, TEXT("Attempt to file map inline bulk data, this will almost certainly fail due to alignment requirements. Package '%s'"), *PackageName.ToString());
+
 				// Inline data is already in the archive so serialize it immediately
 				AllocateData(BulkDataSize);
 				SerializeBulkData(Ar, DataBuffer, BulkDataSize);
@@ -666,6 +682,16 @@ void FBulkDataBase::Serialize(FArchive& Ar, UObject* Owner, int32 /*Index*/, boo
 				{
 					BulkDataOffsetInFile += Linker->Summary.BulkDataStartOffset;
 				}
+
+				if (bAttemptFileMapping)
+				{
+					check(Filename != nullptr);
+					FString MemoryMappedFilename = ConvertFilenameFromFlags(*Filename);
+					if (!MemoryMapBulkData(MemoryMappedFilename, BulkDataOffsetInFile, BulkDataSize))
+					{
+						bShouldForceStreamData = true; // Signal we want to force the BulkData to load
+					}
+				}
 			}
 
 			// If we are not using the FIoDispatcher then we need to make sure we can retrieve the filename later
@@ -674,6 +700,12 @@ void FBulkDataBase::Serialize(FArchive& Ar, UObject* Owner, int32 /*Index*/, boo
 				check(Filename != nullptr);
 				Fallback.Token = FileTokenSystem::RegisterFileToken( PackageName, *Filename, BulkDataOffsetInFile);
 			}
+
+			// Check to see if we need to load data immediately
+			if (bShouldForceStreamData)
+			{
+				ForceBulkDataResident();
+			}
 		}
 		else
 		{
@@ -1005,10 +1037,18 @@ FOwnedBulkDataPtr* FBulkDataBase::StealFileMapping()
 {
 	check(LockStatus == LOCKSTATUS_Unlocked);
 
-	FOwnedBulkDataPtr* Result = new FOwnedBulkDataPtr(DataBuffer);
-
-	// We are giving ownership of the data to the caller so we just null out the pointer.
-	DataBuffer = nullptr;
+	FOwnedBulkDataPtr* Result;
+	if (MappedHandle == nullptr || MappedRegion == nullptr)
+	{
+		Result = new FOwnedBulkDataPtr(DataBuffer);
+		DataBuffer = nullptr;
+	}
+	else
+	{
+		Result = new FOwnedBulkDataPtr(MappedHandle, MappedRegion);
+		MappedHandle = nullptr;
+		MappedRegion = nullptr;
+	}
 
 	return Result;
 }
@@ -1162,6 +1202,31 @@ void FBulkDataBase::SerializeBulkData(FArchive& Ar, void* DstBuffer, int64 DataL
 	}
 }
 
+bool FBulkDataBase::MemoryMapBulkData(const FString& Filename, int64 OffsetInBulkData, int64 BytesToRead)
+{
+	check(!MappedHandle && !MappedRegion);
+
+	MappedHandle = FPlatformFileManager::Get().GetPlatformFile().OpenMapped(*Filename);
+
+	if (MappedHandle == nullptr)
+	{
+		return false;
+	}
+
+	MappedRegion = MappedHandle->MapRegion(OffsetInBulkData, BytesToRead, true);
+	if (MappedRegion == nullptr)
+	{
+		delete MappedHandle;
+		MappedHandle = nullptr;
+		return false;
+	}
+
+	checkf(MappedRegion->GetMappedSize() == BytesToRead, TEXT("Mapped size (%lld) is different to the requested size (%lld)!"), MappedRegion->GetMappedSize(), BytesToRead);
+	checkf(IsAligned(MappedRegion->GetMappedPtr(), FPlatformProperties::GetMemoryMappingAlignment()), TEXT("Memory mapped file has the wrong alignment!"));
+
+	return true;
+}
+
 void FBulkDataBase::AllocateData(SIZE_T SizeInBytes)
 {
 	DataBuffer = FMemory::Realloc(DataBuffer, SizeInBytes, DEFAULT_ALIGNMENT);
@@ -1171,6 +1236,12 @@ void FBulkDataBase::FreeData()
 {
 	FMemory::Free(DataBuffer);
 	DataBuffer = nullptr;
+
+	delete MappedRegion;
+	MappedRegion = nullptr;
+
+	delete MappedHandle;
+	MappedHandle = nullptr;
 }
 
 FString FBulkDataBase::ConvertFilenameFromFlags(const FString& Filename) const
diff --git a/Engine/Source/Runtime/CoreUObject/Public/Serialization/BulkData2.h b/Engine/Source/Runtime/CoreUObject/Public/Serialization/BulkData2.h
--- a/Engine/Source/Runtime/CoreUObject/Public/Serialization/BulkData2.h
+++ b/Engine/Source/Runtime/CoreUObject/Public/Serialization/BulkData2.h
@@ -15,6 +15,8 @@
 #define BULKDATA_NOT_IMPLEMENTED_FOR_RUNTIME PLATFORM_BREAK()
 
 struct FOwnedBulkDataPtr;
+class IMappedFileHandle;
+class IMappedFileRegion;
 
 /**
  * Represents an IO request from the BulkData streaming API.
@@ -140,6 +142,8 @@ private:
 	void SerializeDuplicateData(FArchive& Ar, uint32& OutBulkDataFlags, int64& OutBulkDataSizeOnDisk, int64& OutBulkDataOffsetInFile);
 	void SerializeBulkData(FArchive& Ar, void* DstBuffer, int64 DataLength);
 
+	bool MemoryMapBulkData(const FString& Filename, int64 OffsetInBulkData, int64 BytesToRead);
+
 	void AllocateData(SIZE_T SizeInBytes);
 	void FreeData();
 
@@ -164,6 +168,10 @@ private:
 	}; // Note that the union will end up being 16 bytes with padding
 	
 	void* DataBuffer = nullptr;
+
+	IMappedFileHandle* MappedHandle = nullptr; // Temp
+	IMappedFileRegion* MappedRegion = nullptr; // Temp
+
 	uint32 BulkDataFlags = 0;
 	mutable uint8 LockStatus = 0; // Mutable so that the read only lock can be const
 };
