Add some helper functions to FPackageName for determining if a file extension is a valid text map/asset/either

diff --git a/Engine/Source/Runtime/CoreUObject/Private/Misc/PackageName.cpp b/Engine/Source/Runtime/CoreUObject/Private/Misc/PackageName.cpp
--- a/Engine/Source/Runtime/CoreUObject/Private/Misc/PackageName.cpp
+++ b/Engine/Source/Runtime/CoreUObject/Private/Misc/PackageName.cpp
@@ -1237,15 +1237,61 @@ FString FPackageName::PackageFromPath(const TCHAR* InPathName)
 	}
 }
 
+bool FPackageName::IsTextPackageExtension(const TCHAR* Ext)
+{
+	return IsTextAssetPackageExtension(Ext) || IsTextMapPackageExtension(Ext);
+}
+
+bool FPackageName::IsTextAssetPackageExtension(const TCHAR* Ext)
+{
+	if (*Ext != TEXT('.'))
+	{
+		return (TextAssetPackageExtension.EndsWith(Ext));
+	}
+	else
+	{
+		return (TextAssetPackageExtension == Ext);
+	}
+}
+
+bool FPackageName::IsTextMapPackageExtension(const TCHAR* Ext)
+{
+	if (*Ext != TEXT('.'))
+	{
+		return (TextMapPackageExtension.EndsWith(Ext));
+	}
+	else
+	{
+		return (TextMapPackageExtension == Ext);
+	}
+}
+
 bool FPackageName::IsPackageExtension( const TCHAR* Ext )
+{
+	return IsAssetPackageExtension(Ext) || IsMapPackageExtension(Ext);
+}
+
+bool FPackageName::IsAssetPackageExtension(const TCHAR* Ext)
+{
+	if (*Ext != TEXT('.'))
+	{
+		return (AssetPackageExtension.EndsWith(Ext));
+	}
+	else
+	{
+		return (AssetPackageExtension == Ext);
+	}
+}
+
+bool FPackageName::IsMapPackageExtension(const TCHAR* Ext)
 {
 	if (*Ext != TEXT('.'))
 	{
-		return AssetPackageExtension.EndsWith(Ext) || MapPackageExtension.EndsWith(Ext);
+		return (MapPackageExtension.EndsWith(Ext));
 	}
 	else
 	{
-		return AssetPackageExtension == Ext || MapPackageExtension == Ext;
+		return (MapPackageExtension == Ext);
 	}
 }
 
diff --git a/Engine/Source/Runtime/CoreUObject/Public/Misc/PackageName.h b/Engine/Source/Runtime/CoreUObject/Public/Misc/PackageName.h
--- a/Engine/Source/Runtime/CoreUObject/Public/Misc/PackageName.h
+++ b/Engine/Source/Runtime/CoreUObject/Public/Misc/PackageName.h
@@ -312,15 +312,60 @@ public:
 		return TextMapPackageExtension;
 	}
 
+	/**
+	 * Returns whether the passed in extension is a valid text package
+	 * extension. Extensions with and without trailing dots are supported.
+	 *
+	 * @param	Extension to test.
+	 * @return	True if Ext is either an text asset or a text map extension, otherwise false
+	 */
+	static bool IsTextPackageExtension(const TCHAR* Ext);
+
+	/**
+	 * Returns whether the passed in extension is a valid text asset package
+	 * extension. Extensions with and without trailing dots are supported.
+	 *
+	 * @param	Extension to test.
+	 * @return	True if Ext is a text asset extension, otherwise false
+	 */
+	static bool IsTextAssetPackageExtension(const TCHAR* Ext);
+
+	/**
+	 * Returns whether the passed in extension is a valid text map package
+	 * extension. Extensions with and without trailing dots are supported.
+	 *
+	 * @param	Extension to test.
+	 * @return	True if Ext is a text map extension, otherwise false
+	 */
+	static bool IsTextMapPackageExtension(const TCHAR* Ext);
+
 	/** 
-	 * Returns whether the passed in extension is a valid package 
+	 * Returns whether the passed in extension is a valid binary package
 	 * extension. Extensions with and without trailing dots are supported.
 	 *
 	 * @param	Extension to test. 
-	 * @return	True if Ext is either an asset or a map extension, otherwise false
+	 * @return	True if Ext is either a binary  asset or map extension, otherwise false
 	 */
 	static bool IsPackageExtension(const TCHAR* Ext);
 
+	/**
+	 * Returns whether the passed in extension is a valid binary asset package
+	 * extension. Extensions with and without trailing dots are supported.
+	 *
+	 * @param	Extension to test.
+	 * @return	True if Ext is a binary asset extension, otherwise false
+	 */
+	static bool IsAssetPackageExtension(const TCHAR* Ext);
+
+	/**
+	 * Returns whether the passed in extension is a valid binary map package
+	 * extension. Extensions with and without trailing dots are supported.
+	 *
+	 * @param	Extension to test.
+	 * @return	True if Ext is a binary asset extension, otherwise false
+	 */
+	static bool IsMapPackageExtension(const TCHAR* Ext);
+
 	/** 
 	 * Returns whether the passed in filename ends with any of the known
 	 * package extensions.
