FReferenceFinder will not use multithreading during initial load to avoid asserts when assembling token streams for classes that haven't had it assembled yet.

diff --git a/Engine/Source/Runtime/CoreUObject/Private/UObject/ReferencerFinder.cpp b/Engine/Source/Runtime/CoreUObject/Private/UObject/ReferencerFinder.cpp
--- a/Engine/Source/Runtime/CoreUObject/Private/UObject/ReferencerFinder.cpp
+++ b/Engine/Source/Runtime/CoreUObject/Private/UObject/ReferencerFinder.cpp
@@ -94,7 +94,7 @@ TArray<UObject*> FReferencerFinder::GetAllReferencers(const TSet<UObject*>& Refe
 				FScopeLock ResultLock(&ResultCritical);
 				Ret.Append(ThreadResult.Array());
 			}
-		});
+		}, GIsInitialLoad ? EParallelForFlags::ForceSingleThread : EParallelForFlags::None);
 
 		UnlockUObjectHashTables();
 	}
