Simplification of DiffMap checking code.

#rnx

diff --git a/Engine/Source/Runtime/CoreUObject/Private/Serialization/ArchiveStackTrace.cpp b/Engine/Source/Runtime/CoreUObject/Private/Serialization/ArchiveStackTrace.cpp
--- a/Engine/Source/Runtime/CoreUObject/Private/Serialization/ArchiveStackTrace.cpp
+++ b/Engine/Source/Runtime/CoreUObject/Private/Serialization/ArchiveStackTrace.cpp
@@ -281,7 +281,7 @@ void FArchiveStackTrace::Serialize(void* InData, int64 Num)
 		}
 
 		// Walk the stack and dump it to the allocated memory.
-		bool bShouldCollectCallstack = bCollectCallstacks && (DiffMap == nullptr || IsInDiffMap(CurrentOffset)) && !GIgnoreDiffManager.ShouldIgnoreDiff();
+		bool bShouldCollectCallstack = bCollectCallstacks && ShouldLogOffset(CurrentOffset) && !GIgnoreDiffManager.ShouldIgnoreDiff();
 		if (bShouldCollectCallstack)
 		{
 			StackTrace[0] = '\0';
@@ -532,7 +532,7 @@ void FArchiveStackTrace::CompareWithInternal(const FPackageData& SourcePackage,
 
 		if (SourcePackage.Data[SourceAbsoluteOffset] != DestPackage.Data[DestAbsoluteOffset])
 		{
-			if (DiffMap == nullptr || IsInDiffMap(DestAbsoluteOffset))
+			if (ShouldLogOffset(DestAbsoluteOffset))
 			{
 				int32 DifferenceCallstackoffsetIndex = GetCallstackAtOffset(DestAbsoluteOffset, FMath::Max(LastDifferenceCallstackOffsetIndex, 0));
 				if (DifferenceCallstackoffsetIndex >= 0 && DifferenceCallstackoffsetIndex != LastDifferenceCallstackOffsetIndex)
diff --git a/Engine/Source/Runtime/CoreUObject/Public/Serialization/ArchiveStackTrace.h b/Engine/Source/Runtime/CoreUObject/Public/Serialization/ArchiveStackTrace.h
--- a/Engine/Source/Runtime/CoreUObject/Public/Serialization/ArchiveStackTrace.h
+++ b/Engine/Source/Runtime/CoreUObject/Public/Serialization/ArchiveStackTrace.h
@@ -174,8 +174,13 @@ class COREUOBJECT_API FArchiveStackTrace : public FLargeMemoryWriter
 		}
 	}
 
-	bool IsInDiffMap(int64 InOffset) const
+	bool ShouldLogOffset(int64 InOffset) const
 	{
+		if (!DiffMap)
+		{
+			return true;
+		}
+
 		for (const FArchiveDiffInfo& Diff : *DiffMap)
 		{
 			if (Diff.Offset <= InOffset && InOffset < (Diff.Offset + Diff.Size))
