Simple refactoring to reduce indentation.

#rnx

diff --git a/Engine/Source/Runtime/CoreUObject/Private/Serialization/ArchiveStackTrace.cpp b/Engine/Source/Runtime/CoreUObject/Private/Serialization/ArchiveStackTrace.cpp
--- a/Engine/Source/Runtime/CoreUObject/Private/Serialization/ArchiveStackTrace.cpp
+++ b/Engine/Source/Runtime/CoreUObject/Private/Serialization/ArchiveStackTrace.cpp
@@ -530,123 +530,127 @@ void FArchiveStackTrace::CompareWithInternal(const FPackageData& SourcePackage,
 		const int64 SourceAbsoluteOffset = LocalOffset + SourcePackage.StartOffset;
 		const int64 DestAbsoluteOffset = LocalOffset + DestPackage.StartOffset;
 
-		if (SourcePackage.Data[SourceAbsoluteOffset] != DestPackage.Data[DestAbsoluteOffset])
+		const uint8 SourceByte = SourcePackage.Data[SourceAbsoluteOffset];
+		const uint8 DestByte   = DestPackage  .Data[DestAbsoluteOffset];
+		if (SourceByte == DestByte)
 		{
-			if (ShouldLogOffset(DestAbsoluteOffset))
+			continue;
+		}
+
+		if (ShouldLogOffset(DestAbsoluteOffset))
+		{
+			int32 DifferenceCallstackoffsetIndex = GetCallstackAtOffset(DestAbsoluteOffset, FMath::Max(LastDifferenceCallstackOffsetIndex, 0));
+			if (DifferenceCallstackoffsetIndex >= 0 && DifferenceCallstackoffsetIndex != LastDifferenceCallstackOffsetIndex)
 			{
-				int32 DifferenceCallstackoffsetIndex = GetCallstackAtOffset(DestAbsoluteOffset, FMath::Max(LastDifferenceCallstackOffsetIndex, 0));
-				if (DifferenceCallstackoffsetIndex >= 0 && DifferenceCallstackoffsetIndex != LastDifferenceCallstackOffsetIndex)
+				const FCallstactAtOffset& CallstackAtOffset = CallstackAtOffsetMap[DifferenceCallstackoffsetIndex];
+				const FCallstackData& DifferenceCallstackData = UniqueCallstacks[CallstackAtOffset.Callstack];
+				FString DifferenceCallstackDataText = DifferenceCallstackData.ToString(CallstackCutoffText);
+				if (LastDifferenceCallstackDataText.Compare(DifferenceCallstackDataText, ESearchCase::CaseSensitive) != 0)
 				{
-					const FCallstactAtOffset& CallstackAtOffset = CallstackAtOffsetMap[DifferenceCallstackoffsetIndex];
-					const FCallstackData& DifferenceCallstackData = UniqueCallstacks[CallstackAtOffset.Callstack];
-					FString DifferenceCallstackDataText = DifferenceCallstackData.ToString(CallstackCutoffText);
-					if (LastDifferenceCallstackDataText.Compare(DifferenceCallstackDataText, ESearchCase::CaseSensitive) != 0)
+					if (!CallstackAtOffset.bIgnore && (MaxDiffsToLog < 0 || InOutDiffsLogged < MaxDiffsToLog))
 					{
-						if (!CallstackAtOffset.bIgnore && (MaxDiffsToLog < 0 || InOutDiffsLogged < MaxDiffsToLog))
+						FString BeforePropertyVal;
+						FString AfterPropertyVal;
+						if (FProperty* SerProp = DifferenceCallstackData.SerializedProp)
 						{
-							FString BeforePropertyVal;
-							FString AfterPropertyVal;
-							if (FProperty* SerProp = DifferenceCallstackData.SerializedProp)
+							if (SourceSize == DestSize && ShouldDumpPropertyValueState(SerProp))
 							{
-								if (SourceSize == DestSize && ShouldDumpPropertyValueState(SerProp))
+								// Walk backwards until we find a callstack which wasn't from the given property
+								int32 OffsetX = DestAbsoluteOffset;
+								for (;;)
 								{
-									// Walk backwards until we find a callstack which wasn't from the given property
-									int32 OffsetX = DestAbsoluteOffset;
-									for (;;)
+									if (OffsetX == 0)
 									{
-										if (OffsetX == 0)
-										{
-											break;
-										}
-
-										int32 CallstackIndex = GetCallstackAtOffset(OffsetX - 1, 0);
+										break;
+									}
 
-										const FCallstactAtOffset& PreviousCallstack = CallstackAtOffsetMap[CallstackIndex];
-										if (UniqueCallstacks[PreviousCallstack.Callstack].SerializedProp != SerProp)
-										{
-											break;
-										}
+									int32 CallstackIndex = GetCallstackAtOffset(OffsetX - 1, 0);
 
-										--OffsetX;
+									const FCallstactAtOffset& PreviousCallstack = CallstackAtOffsetMap[CallstackIndex];
+									if (UniqueCallstacks[PreviousCallstack.Callstack].SerializedProp != SerProp)
+									{
+										break;
 									}
 
-									FPropertyTempVal SourceVal(SerProp);
-									FPropertyTempVal DestVal  (SerProp);
+									--OffsetX;
+								}
 
-									FStaticMemoryReader SourceReader(&SourcePackage.Data[SourceAbsoluteOffset - (DestAbsoluteOffset - OffsetX)], SourcePackage.Size - SourceAbsoluteOffset);
-									FStaticMemoryReader DestReader(&DestPackage.Data[OffsetX], DestPackage.Size - DestAbsoluteOffset);
+								FPropertyTempVal SourceVal(SerProp);
+								FPropertyTempVal DestVal  (SerProp);
 
-									SourceVal.Serialize(SourceReader);
-									DestVal  .Serialize(DestReader);
+								FStaticMemoryReader SourceReader(&SourcePackage.Data[SourceAbsoluteOffset - (DestAbsoluteOffset - OffsetX)], SourcePackage.Size - SourceAbsoluteOffset);
+								FStaticMemoryReader DestReader(&DestPackage.Data[OffsetX], DestPackage.Size - DestAbsoluteOffset);
 
-									if (!SourceReader.ArIsError && !DestReader.ArIsError)
-									{
-										SourceVal.ExportText(BeforePropertyVal);
-										DestVal  .ExportText(AfterPropertyVal);
-									}
+								SourceVal.Serialize(SourceReader);
+								DestVal  .Serialize(DestReader);
+
+								if (!SourceReader.ArIsError && !DestReader.ArIsError)
+								{
+									SourceVal.ExportText(BeforePropertyVal);
+									DestVal  .ExportText(AfterPropertyVal);
 								}
 							}
+						}
 
-							FString DiffValues;
-							if (BeforePropertyVal != AfterPropertyVal)
-							{
-								DiffValues = FString::Printf(TEXT("\r\n%sBefore: %s\r\n%sAfter:  %s"), Indent, *BeforePropertyVal, Indent, *AfterPropertyVal);
-							}
+						FString DiffValues;
+						if (BeforePropertyVal != AfterPropertyVal)
+						{
+							DiffValues = FString::Printf(TEXT("\r\n%sBefore: %s\r\n%sAfter:  %s"), Indent, *BeforePropertyVal, Indent, *AfterPropertyVal);
+						}
 
-							FString DebugDataStackText;
+						FString DebugDataStackText;
 #if WITH_EDITOR
-							//check for a debug data stack as part of the unique stack entry, and log it out if we find it.
-							FString FullStackText = DifferenceCallstackData.Callstack;
-							int32 DebugDataIndex = FullStackText.Find(ANSI_TO_TCHAR(DebugDataStackMarker), ESearchCase::CaseSensitive);
-							if (DebugDataIndex > 0)
-							{
-								DebugDataStackText = FString::Printf(TEXT("\r\n%s"), FDiffFormatHelper::Get().Indent) + FullStackText.RightChop(DebugDataIndex + 2);
-							}
-#endif
-							UE_LOG(
-								LogArchiveDiff,
-								Warning,
-								TEXT("%s: Difference at offset %lld%s (absolute offset: %lld), callstack:%s%s%s%s%s"),
-								AssetFilename,
-								CallstackAtOffset.Offset - DestPackage.StartOffset,
-								DestAbsoluteOffset > CallstackAtOffset.Offset ? *FString::Printf(TEXT("(+%lld)"), DestAbsoluteOffset - CallstackAtOffset.Offset) : TEXT(""),
-								DestAbsoluteOffset,
-								LineTerminator,
-								LineTerminator,
-								*DifferenceCallstackDataText,
-								*DiffValues,
-								*DebugDataStackText
-							);
-							InOutDiffsLogged++;
-							NumDiffsLoggedLocal++;
-						}
-						else if (FirstUnreportedDiffIndex == -1)
+						//check for a debug data stack as part of the unique stack entry, and log it out if we find it.
+						FString FullStackText = DifferenceCallstackData.Callstack;
+						int32 DebugDataIndex = FullStackText.Find(ANSI_TO_TCHAR(DebugDataStackMarker), ESearchCase::CaseSensitive);
+						if (DebugDataIndex > 0)
 						{
-							FirstUnreportedDiffIndex = DestAbsoluteOffset;
+							DebugDataStackText = FString::Printf(TEXT("\r\n%s"), FDiffFormatHelper::Get().Indent) + FullStackText.RightChop(DebugDataIndex + 2);
 						}
-						LastDifferenceCallstackDataText = MoveTemp(DifferenceCallstackDataText);
-						OutStats.FindOrAdd(AssetClass).NumDiffs++;
-						NumDiffsLocal++;
+#endif
+						UE_LOG(
+							LogArchiveDiff,
+							Warning,
+							TEXT("%s: Difference at offset %lld%s (absolute offset: %lld), callstack:%s%s%s%s%s"),
+							AssetFilename,
+							CallstackAtOffset.Offset - DestPackage.StartOffset,
+							DestAbsoluteOffset > CallstackAtOffset.Offset ? *FString::Printf(TEXT("(+%lld)"), DestAbsoluteOffset - CallstackAtOffset.Offset) : TEXT(""),
+							DestAbsoluteOffset,
+							LineTerminator,
+							LineTerminator,
+							*DifferenceCallstackDataText,
+							*DiffValues,
+							*DebugDataStackText
+						);
+						InOutDiffsLogged++;
+						NumDiffsLoggedLocal++;
 					}
+					else if (FirstUnreportedDiffIndex == -1)
+					{
+						FirstUnreportedDiffIndex = DestAbsoluteOffset;
+					}
+					LastDifferenceCallstackDataText = MoveTemp(DifferenceCallstackDataText);
+					OutStats.FindOrAdd(AssetClass).NumDiffs++;
+					NumDiffsLocal++;
 				}
-				else if (DifferenceCallstackoffsetIndex < 0)
-				{
-					UE_LOG(LogArchiveDiff, Warning, TEXT("%s: Difference at offset %lld (absolute offset: %lld), unknown callstack"), AssetFilename, LocalOffset, DestAbsoluteOffset);
-				}
-				LastDifferenceCallstackOffsetIndex = DifferenceCallstackoffsetIndex;
 			}
-			else
+			else if (DifferenceCallstackoffsetIndex < 0)
 			{
-				// Each byte will count as a difference but without callstack data there's no way around it
-				OutStats.FindOrAdd(AssetClass).NumDiffs++;
-				NumDiffsLocal++;
-				if (FirstUnreportedDiffIndex == -1)
-				{
-					FirstUnreportedDiffIndex = DestAbsoluteOffset;
-				}
+				UE_LOG(LogArchiveDiff, Warning, TEXT("%s: Difference at offset %lld (absolute offset: %lld), unknown callstack"), AssetFilename, LocalOffset, DestAbsoluteOffset);
+			}
+			LastDifferenceCallstackOffsetIndex = DifferenceCallstackoffsetIndex;
+		}
+		else
+		{
+			// Each byte will count as a difference but without callstack data there's no way around it
+			OutStats.FindOrAdd(AssetClass).NumDiffs++;
+			NumDiffsLocal++;
+			if (FirstUnreportedDiffIndex == -1)
+			{
+				FirstUnreportedDiffIndex = DestAbsoluteOffset;
 			}
-			OutStats.FindOrAdd(AssetClass).DiffSize++;
 		}
+		OutStats.FindOrAdd(AssetClass).DiffSize++;
 	}
 
 	if (MaxDiffsToLog >= 0 && NumDiffsLocal > NumDiffsLoggedLocal)
