Use of ON_SCOPE_EXIT to clarify intended code flow.

#rnx

diff --git a/Engine/Source/Runtime/CoreUObject/Private/Serialization/ArchiveStackTrace.cpp b/Engine/Source/Runtime/CoreUObject/Private/Serialization/ArchiveStackTrace.cpp
--- a/Engine/Source/Runtime/CoreUObject/Private/Serialization/ArchiveStackTrace.cpp
+++ b/Engine/Source/Runtime/CoreUObject/Private/Serialization/ArchiveStackTrace.cpp
@@ -20,6 +20,7 @@
 #include "Misc/PackageName.h"
 #include "Templates/UniquePtr.h"
 #include "UObject/UObjectGlobals.h"
+#include "Misc/ScopeExit.h"
 
 DEFINE_LOG_CATEGORY_STATIC(LogArchiveDiff, Log, All);
 
@@ -540,6 +541,11 @@ void FArchiveStackTrace::CompareWithInternal(const FPackageData& SourcePackage,
 		if (ShouldLogOffset(DestAbsoluteOffset))
 		{
 			int32 DifferenceCallstackoffsetIndex = GetCallstackAtOffset(DestAbsoluteOffset, FMath::Max(LastDifferenceCallstackOffsetIndex, 0));
+			ON_SCOPE_EXIT
+			{
+				LastDifferenceCallstackOffsetIndex = DifferenceCallstackoffsetIndex;
+			};
+
 			if (DifferenceCallstackoffsetIndex >= 0 && DifferenceCallstackoffsetIndex != LastDifferenceCallstackOffsetIndex)
 			{
 				const FCallstactAtOffset& CallstackAtOffset = CallstackAtOffsetMap[DifferenceCallstackoffsetIndex];
@@ -638,7 +644,6 @@ void FArchiveStackTrace::CompareWithInternal(const FPackageData& SourcePackage,
 			{
 				UE_LOG(LogArchiveDiff, Warning, TEXT("%s: Difference at offset %lld (absolute offset: %lld), unknown callstack"), AssetFilename, LocalOffset, DestAbsoluteOffset);
 			}
-			LastDifferenceCallstackOffsetIndex = DifferenceCallstackoffsetIndex;
 		}
 		else
 		{
