LastDifferenceCallstackDataText updated on scope exit, no matter how that happens.

#rnx

diff --git a/Engine/Source/Runtime/CoreUObject/Private/Serialization/ArchiveStackTrace.cpp b/Engine/Source/Runtime/CoreUObject/Private/Serialization/ArchiveStackTrace.cpp
--- a/Engine/Source/Runtime/CoreUObject/Private/Serialization/ArchiveStackTrace.cpp
+++ b/Engine/Source/Runtime/CoreUObject/Private/Serialization/ArchiveStackTrace.cpp
@@ -564,6 +564,10 @@ void FArchiveStackTrace::CompareWithInternal(const FPackageData& SourcePackage,
 			{
 				continue;
 			}
+			ON_SCOPE_EXIT
+			{
+				LastDifferenceCallstackDataText = MoveTemp(DifferenceCallstackDataText);
+			};
 
 			if (!CallstackAtOffset.bIgnore && (MaxDiffsToLog < 0 || InOutDiffsLogged < MaxDiffsToLog))
 			{
@@ -647,7 +651,6 @@ void FArchiveStackTrace::CompareWithInternal(const FPackageData& SourcePackage,
 			{
 				FirstUnreportedDiffIndex = DestAbsoluteOffset;
 			}
-			LastDifferenceCallstackDataText = MoveTemp(DifferenceCallstackDataText);
 			OutStats.FindOrAdd(AssetClass).NumDiffs++;
 			NumDiffsLocal++;
 		}
