Scoped approach to incrementing log counts.

#rnx

diff --git a/Engine/Source/Runtime/CoreUObject/Private/Serialization/ArchiveStackTrace.cpp b/Engine/Source/Runtime/CoreUObject/Private/Serialization/ArchiveStackTrace.cpp
--- a/Engine/Source/Runtime/CoreUObject/Private/Serialization/ArchiveStackTrace.cpp
+++ b/Engine/Source/Runtime/CoreUObject/Private/Serialization/ArchiveStackTrace.cpp
@@ -538,6 +538,16 @@ void FArchiveStackTrace::CompareWithInternal(const FPackageData& SourcePackage,
 			continue;
 		}
 
+		bool bDifferenceLogged = false;
+		ON_SCOPE_EXIT
+		{
+			if (bDifferenceLogged)
+			{
+				InOutDiffsLogged++;
+				NumDiffsLoggedLocal++;
+			}
+		};
+
 		if (ShouldLogOffset(DestAbsoluteOffset))
 		{
 			int32 DifferenceCallstackoffsetIndex = GetCallstackAtOffset(DestAbsoluteOffset, FMath::Max(LastDifferenceCallstackOffsetIndex, 0));
@@ -644,8 +654,7 @@ void FArchiveStackTrace::CompareWithInternal(const FPackageData& SourcePackage,
 					*DiffValues,
 					*DebugDataStackText
 				);
-				InOutDiffsLogged++;
-				NumDiffsLoggedLocal++;
+				bDifferenceLogged = true;
 			}
 			else if (FirstUnreportedDiffIndex == -1)
 			{
