FProperty related performance optimizations (regression fixes):

- optimized TFieldIterator when not excluding deprecated members
- added support for fast but unsafe (under general conditions) functions to get FProperty owners for BP script VM
- FFieldPath will no longer need unnecessarily to resolve itself when initialized through copy constructor or assignment operator

diff --git a/Engine/Source/Runtime/CoreUObject/Private/UObject/ScriptCore.cpp b/Engine/Source/Runtime/CoreUObject/Private/UObject/ScriptCore.cpp
--- a/Engine/Source/Runtime/CoreUObject/Private/UObject/ScriptCore.cpp
+++ b/Engine/Source/Runtime/CoreUObject/Private/UObject/ScriptCore.cpp
@@ -670,7 +670,7 @@ void UObject::SkipFunction(FFrame& Stack, RESULT_DECL, UFunction* Function)
 	// allocate temporary memory on the stack for evaluating parameters
 	uint8* Frame = (uint8*)FMemory_Alloca(Function->PropertiesSize);
 	FMemory::Memzero(Frame, Function->PropertiesSize);
-	for (FProperty* Property = static_cast<FProperty*>(Function->ChildProperties); *Stack.Code != EX_EndFunctionParms; Property = static_cast<FProperty*>(Property->Next))
+	for (FProperty* Property = (FProperty*)(Function->ChildProperties); *Stack.Code != EX_EndFunctionParms; Property = (FProperty*)(Property->Next))
 	{
 		Stack.MostRecentPropertyAddress = NULL;
 		// evaluate the expression into our temporary memory space
@@ -753,7 +753,7 @@ void ProcessScriptFunction(UObject* Context, UFunction* Function, FFrame& Stack,
 	NewStack.Locals = FrameMemory;
 	FOutParmRec** LastOut = &NewStack.OutParms;
 		
-	for (FProperty* Property = static_cast<FProperty*>(Function->ChildProperties); *Stack.Code != EX_EndFunctionParms; Property = static_cast<FProperty*>(Property->Next))
+	for (FProperty* Property = (FProperty*)(Function->ChildProperties); *Stack.Code != EX_EndFunctionParms; Property = (FProperty*)(Property->Next))
 	{
 		checkfSlow(Property, TEXT("NULL Property in Function %s"), *Function->GetPathName()); 
 
@@ -813,7 +813,7 @@ void ProcessScriptFunction(UObject* Context, UFunction* Function, FFrame& Stack,
 	if (!bUsePersistentFrame)
 	{
 		// Initialize any local struct properties with defaults
-		for (FProperty* LocalProp = Function->FirstPropertyToInit; LocalProp != NULL; LocalProp = static_cast<FProperty*>(LocalProp->Next))
+		for (FProperty* LocalProp = Function->FirstPropertyToInit; LocalProp != NULL; LocalProp = (FProperty*)(LocalProp->Next))
 		{
 			LocalProp->InitializeValue_InContainer(NewStack.Locals);
 		}
@@ -1866,7 +1866,7 @@ void UObject::ProcessEvent( UFunction* Function, void* Parms )
 		if ( Function->HasAnyFunctionFlags(FUNC_HasOutParms) )
 		{
 			FOutParmRec** LastOut = &NewStack.OutParms;
-			for ( FProperty* Property = static_cast<FProperty*>(Function->ChildProperties); Property && (Property->PropertyFlags&(CPF_Parm)) == CPF_Parm; Property = static_cast<FProperty*>(Property->Next) )
+			for ( FProperty* Property = (FProperty*)(Function->ChildProperties); Property && (Property->PropertyFlags&(CPF_Parm)) == CPF_Parm; Property = (FProperty*)Property->Next )
 			{
 				// this is used for optional parameters - the destination address for out parameter values is the address of the calling function
 				// so we'll need to know which address to use if we need to evaluate the default parm value expression located in the new function's
@@ -1902,7 +1902,7 @@ void UObject::ProcessEvent( UFunction* Function, void* Parms )
 
 		if (!bUsePersistentFrame)
 		{
-			for (FProperty* LocalProp = Function->FirstPropertyToInit; LocalProp != NULL; LocalProp = static_cast<FProperty*>(LocalProp->Next))
+			for (FProperty* LocalProp = Function->FirstPropertyToInit; LocalProp != NULL; LocalProp = (FProperty*)LocalProp->Next)
 			{
 				LocalProp->InitializeValue_InContainer(NewStack.Locals);
 			}
@@ -1977,7 +1977,7 @@ DEFINE_FUNCTION(UObject::execInstanceVariable)
 	FProperty* VarProperty = (FProperty*)Stack.ReadObject();
 	Stack.MostRecentProperty = VarProperty;
 
-	if (VarProperty == nullptr || !P_THIS->IsA(VarProperty->GetOwner<UClass>()))
+	if (VarProperty == nullptr || !P_THIS->IsA((UClass*)VarProperty->InternalGetOwnerAsUObjectUnsafe()))
 	{
 		FBlueprintExceptionInfo ExceptionInfo(EBlueprintExceptionType::AccessViolation, FText::Format(LOCTEXT("MissingProperty", "Attempted to access missing property '{0}'. If this is a packaged/cooked build, are you attempting to use an editor-only property?"), FText::FromString(GetNameSafe(VarProperty))));
 		FBlueprintCoreDelegates::ThrowScriptException(P_THIS, Stack, ExceptionInfo);
@@ -2043,7 +2043,7 @@ DEFINE_FUNCTION(UObject::execDefaultVariable)
 		// @todo - allow access to archetype properties through object references?
 	}
 
-	if (VarProperty == nullptr || (DefaultObject && !DefaultObject->IsA(VarProperty->GetOwner<UClass>())))
+	if (VarProperty == nullptr || (DefaultObject && !DefaultObject->IsA((UClass*)VarProperty->InternalGetOwnerAsUObjectUnsafe())))
 	{
 		FBlueprintExceptionInfo ExceptionInfo(EBlueprintExceptionType::AccessViolation, LOCTEXT("MissingPropertyDefaultObject", "Attempted to access a missing property on a CDO. If this is a packaged/cooked build, are you attempting to use an editor-only property?"));
 		FBlueprintCoreDelegates::ThrowScriptException(P_THIS, Stack, ExceptionInfo);
@@ -2697,7 +2697,7 @@ DEFINE_FUNCTION(UObject::execLetBool)
 	Stack.Step( Stack.Object, &NewValue );
 	if( BoolAddr )
 	{
-		checkSlow(dynamic_cast<FBoolProperty*>(BoolProperty));
+		checkSlow(CastField<FBoolProperty>(BoolProperty));
 		BoolProperty->SetPropertyValue( BoolAddr, NewValue );
 	}
 }
@@ -2730,7 +2730,7 @@ DEFINE_FUNCTION(UObject::execLetMulticastDelegate)
 	Stack.MostRecentProperty = NULL;
 	Stack.Step( Stack.Object, NULL ); // Variable.
 
-	FMulticastDelegateProperty* DelegateProp = CastField<FMulticastDelegateProperty>(Stack.MostRecentProperty);
+	FMulticastDelegateProperty* DelegateProp = CastFieldCheckedNullAllowed<FMulticastDelegateProperty>(Stack.MostRecentProperty);
 	void* DelegateAddr = Stack.MostRecentPropertyAddress;
 	FMulticastScriptDelegate Delegate;
 	Stack.Step( Stack.Object, &Delegate );
@@ -2933,13 +2933,13 @@ public:
 		Stack.MostRecentPropertyAddress = NULL;
 		Stack.MostRecentProperty = NULL;
 		Stack.Step( Stack.Object, NULL );
-		FMulticastDelegateProperty* DelegateProp = CastField<FMulticastDelegateProperty>(Stack.MostRecentProperty);
+		FMulticastDelegateProperty* DelegateProp = CastFieldCheckedNullAllowed<FMulticastDelegateProperty>(Stack.MostRecentProperty);
 		const FMulticastScriptDelegate* DelegateAddr = (DelegateProp ? DelegateProp->GetMulticastDelegate(Stack.MostRecentPropertyAddress) : nullptr);
 
 		//Fill parameters
 		uint8* Parameters = (uint8*)FMemory_Alloca(SignatureFunction->ParmsSize);
 		FMemory::Memzero(Parameters, SignatureFunction->ParmsSize);
-		for (FProperty* Property = static_cast<FProperty*>(SignatureFunction->ChildProperties); *Stack.Code != EX_EndFunctionParms; Property = static_cast<FProperty*>(Property->Next))
+		for (FProperty* Property = (FProperty*)SignatureFunction->ChildProperties; *Stack.Code != EX_EndFunctionParms; Property = (FProperty*)Property->Next)
 		{
 			Stack.MostRecentPropertyAddress = NULL;
 			if (Property->PropertyFlags & CPF_OutParm)
@@ -2989,7 +2989,7 @@ DEFINE_FUNCTION(UObject::execAddMulticastDelegate)
 	Stack.MostRecentProperty = NULL;
 	Stack.Step( Stack.Object, NULL ); // Variable.
 
-	FMulticastDelegateProperty* DelegateProp = CastField<FMulticastDelegateProperty>(Stack.MostRecentProperty);
+	FMulticastDelegateProperty* DelegateProp = CastFieldCheckedNullAllowed<FMulticastDelegateProperty>(Stack.MostRecentProperty);
 	void* DelegateAddr = Stack.MostRecentPropertyAddress;
 
 	FScriptDelegate Delegate;
@@ -3009,7 +3009,7 @@ DEFINE_FUNCTION(UObject::execRemoveMulticastDelegate)
 	Stack.MostRecentProperty = NULL;
 	Stack.Step( Stack.Object, NULL ); // Variable.
 
-	FMulticastDelegateProperty* DelegateProp = CastField<FMulticastDelegateProperty>(Stack.MostRecentProperty);
+	FMulticastDelegateProperty* DelegateProp = CastFieldCheckedNullAllowed<FMulticastDelegateProperty>(Stack.MostRecentProperty);
 	void* DelegateAddr = Stack.MostRecentPropertyAddress;
 
 	FScriptDelegate Delegate;
@@ -3029,7 +3029,7 @@ DEFINE_FUNCTION(UObject::execClearMulticastDelegate)
 	Stack.MostRecentProperty = NULL;
 	Stack.Step( Stack.Object, NULL );
 
-	FMulticastDelegateProperty* DelegateProp = CastField<FMulticastDelegateProperty>(Stack.MostRecentProperty);
+	FMulticastDelegateProperty* DelegateProp = CastFieldCheckedNullAllowed<FMulticastDelegateProperty>(Stack.MostRecentProperty);
 	void* DelegateAddr = Stack.MostRecentPropertyAddress;
 
 	if (DelegateProp && DelegateAddr)
@@ -3393,7 +3393,6 @@ IMPLEMENT_VM_FUNCTION( EX_SetMap, execSetMap );
 DEFINE_FUNCTION(UObject::execArrayConst)
 {
 	FProperty* InnerProperty = CastFieldChecked<FProperty>((FField*)Stack.ReadObject());
-	check(InnerProperty);
 	int32 Num = Stack.ReadInt<int32>();
 	check(RESULT_PARAM);
 	FScriptArrayHelper ArrayHelper = FScriptArrayHelper::CreateHelperFormInnerProperty(InnerProperty, RESULT_PARAM);
@@ -3414,7 +3413,6 @@ IMPLEMENT_VM_FUNCTION(EX_ArrayConst, execArrayConst);
 DEFINE_FUNCTION(UObject::execSetConst)
 {
 	FProperty* InnerProperty = CastFieldChecked<FProperty>((FField*)Stack.ReadObject());
-	check(InnerProperty);
 	int32 Num = Stack.ReadInt<int32>();
 	check(RESULT_PARAM);
 
@@ -3436,7 +3434,6 @@ DEFINE_FUNCTION(UObject::execMapConst)
 {
 	FProperty* KeyProperty = CastFieldChecked<FProperty>((FField*)Stack.ReadObject());
 	FProperty* ValProperty = CastFieldChecked<FProperty>((FField*)Stack.ReadObject());
-	check(KeyProperty && ValProperty);;
 	int32 Num = Stack.ReadInt<int32>();
 	check(RESULT_PARAM);
 
diff --git a/Engine/Source/Runtime/CoreUObject/Public/UObject/Class.h b/Engine/Source/Runtime/CoreUObject/Public/UObject/Class.h
--- a/Engine/Source/Runtime/CoreUObject/Public/UObject/Class.h
+++ b/Engine/Source/Runtime/CoreUObject/Public/UObject/Class.h
@@ -66,6 +66,7 @@ class COREUOBJECT_API UField : public UObject
 	DECLARE_CASTED_CLASS_INTRINSIC(UField, UObject, CLASS_Abstract, TEXT("/Script/CoreUObject"), CASTCLASS_UField)
 
 	typedef UField BaseFieldClass;
+	typedef UClass FieldTypeClass;
 
 	/** Next Field in the linked list */
 	UField*			Next;
diff --git a/Engine/Source/Runtime/CoreUObject/Public/UObject/Field.h b/Engine/Source/Runtime/CoreUObject/Public/UObject/Field.h
--- a/Engine/Source/Runtime/CoreUObject/Public/UObject/Field.h
+++ b/Engine/Source/Runtime/CoreUObject/Public/UObject/Field.h
@@ -65,23 +65,31 @@ public:
 	explicit FFieldClass(const TCHAR* InCPPName, uint64 InId, uint64 InCastFlags, FFieldClass* InSuperClass, FField* (*ConstructFnPtr)(const FFieldVariant&, const FName&, EObjectFlags));
 	~FFieldClass();
 
-	FString GetName() const
+	inline FString GetName() const
 	{
 		return Name.ToString();
 	}
-	FName GetFName() const
+	inline FName GetFName() const
 	{
 		return Name;
 	}
-	uint64 GetId() const
+	inline uint64 GetId() const
 	{
 		return Id;
 	}
-	uint64 GetCastFlags() const
+	inline uint64 GetCastFlags() const
 	{
 		return CastFlags;
 	}
-	bool IsChildOf(const FFieldClass* InClass) const
+	inline bool HasAnyCastFlags(const uint64 InCastFlags) const
+	{
+		return !!(CastFlags & InCastFlags);
+	}
+	inline bool HasAllCastFlags(const uint64 InCastFlags) const
+	{
+		return (CastFlags & InCastFlags) == InCastFlags;
+	}
+	inline bool IsChildOf(const FFieldClass* InClass) const
 	{
 		return !!(CastFlags & InClass->GetId());
 	}
@@ -291,6 +299,11 @@ public:
 	{
 		return Container.Field;
 	}
+	/** FOR INTERNAL USE ONLY: Function that returns the owner as UObject without checking if it's actually a UObject */
+	FORCEINLINE UObject* ToUObjectUnsafe() const
+	{
+		return Container.Object;
+	}
 
 	void* GetRawPointer() const
 	{
@@ -342,7 +355,7 @@ public:
 	typedef FField Super;
 	typedef FField ThisClass;
 	typedef FField BaseFieldClass;	
-	typedef FFieldClass ObjectClassType;
+	typedef FFieldClass FieldTypeClass;
 
 	static FFieldClass* StaticClass();
 
@@ -456,16 +469,16 @@ public:
 		return ((GetFlags() & FlagsToCheck) == FlagsToCheck);
 	}
 
-	FFieldClass* GetClass() const
+	inline FFieldClass* GetClass() const
 	{
 		return ClassPrivate;
 	}
-	uint64 GetCastFlags() const
+	inline uint64 GetCastFlags() const
 	{
 		return GetClass()->GetCastFlags();
 	}
 
-	bool IsA(const FFieldClass* FieldType) const
+	inline bool IsA(const FFieldClass* FieldType) const
 	{
 		check(FieldType);
 		return !!(GetCastFlags() & FieldType->GetId());
@@ -477,12 +490,12 @@ public:
 		return !!(GetCastFlags() & T::StaticClassCastFlagsPrivate());
 	}
 
-	bool HasAnyCastFlags(const uint64 InCastFlags) const
+	inline bool HasAnyCastFlags(const uint64 InCastFlags) const
 	{
 		return !!(GetCastFlags() & InCastFlags);
 	}
 
-	bool HasAllCastFlags(const uint64 InCastFlags) const
+	inline bool HasAllCastFlags(const uint64 InCastFlags) const
 	{
 		return (GetCastFlags() & InCastFlags) == InCastFlags;
 	}
@@ -510,6 +523,12 @@ public:
 		return TempOuter.ToUObject();
 	}
 
+	/** Internal function for quickly getting the owner of this object as UObject. FOR INTERNAL USE ONLY */
+	FORCEINLINE UObject* InternalGetOwnerAsUObjectUnsafe() const
+	{
+		return Owner.ToUObjectUnsafe();
+	}
+
 	/** Goes up the outer chain to look for a UClass */
 	UClass* GetOwnerClass() const;
 
@@ -791,9 +810,9 @@ FORCEINLINE FieldType* CastFieldChecked(FField* Src)
 FUNCTION_NON_NULL_RETURN_END
 {
 #if !DO_CHECK
-	return static_cast<FieldType*>(Src);
+	return (FieldType*)Src;
 #else
-	FieldType* CastResult = Src && Src->HasAnyCastFlags(FieldType::StaticClassCastFlagsPrivate()) ? static_cast<FieldType*>(Src) : nullptr;
+	FieldType* CastResult = Src && Src->HasAnyCastFlags(FieldType::StaticClassCastFlagsPrivate()) ? (FieldType*)Src : nullptr;
 	checkf(CastResult, TEXT("CastFieldChecked failed with 0x%016llx"), (int64)(PTRINT)Src);
 	return CastResult;
 #endif // !DO_CHECK
@@ -805,14 +824,42 @@ FORCEINLINE const FieldType* CastFieldChecked(const FField* Src)
 FUNCTION_NON_NULL_RETURN_END
 {
 #if !DO_CHECK
-	return static_cast<const FieldType*>(Src);
+	return (const FieldType*)Src;
 #else
-	const FieldType* CastResult = Src && Src->HasAnyCastFlags(FieldType::StaticClassCastFlagsPrivate()) ? static_cast<const FieldType*>(Src) : nullptr;
+	const FieldType* CastResult = Src && Src->HasAnyCastFlags(FieldType::StaticClassCastFlagsPrivate()) ? (const FieldType*)Src : nullptr;
 	checkf(CastResult, TEXT("CastFieldChecked failed with 0x%016llx"), (int64)(PTRINT)Src);
 	return CastResult;
 #endif // !DO_CHECK
 }
 
+template<typename FieldType>
+FUNCTION_NON_NULL_RETURN_START
+FORCEINLINE FieldType* CastFieldCheckedNullAllowed(FField* Src)
+FUNCTION_NON_NULL_RETURN_END
+{
+#if !DO_CHECK
+	return (FieldType*)Src;
+#else
+	FieldType* CastResult = Src && Src->HasAnyCastFlags(FieldType::StaticClassCastFlagsPrivate()) ? (FieldType*)Src : nullptr;
+	checkf(CastResult || !Src, TEXT("CastFieldCheckedNullAllowed failed with 0x%016llx"), (int64)(PTRINT)Src);
+	return CastResult;
+#endif // !DO_CHECK
+}
+
+template<typename FieldType>
+FUNCTION_NON_NULL_RETURN_START
+FORCEINLINE const FieldType* CastFieldCheckedNullAllowed(const FField* Src)
+FUNCTION_NON_NULL_RETURN_END
+{
+#if !DO_CHECK
+	return (const FieldType*)Src;
+#else
+	const FieldType* CastResult = Src && Src->HasAnyCastFlags(FieldType::StaticClassCastFlagsPrivate()) ? (const FieldType*)Src : nullptr;
+	checkf(CastResult || !Src, TEXT("CastFieldCheckedNullAllowed failed with 0x%016llx"), (int64)(PTRINT)Src);
+	return CastResult;
+#endif // !DO_CHECK
+}
+
 /**
  * Helper function for serializing FField to an archive. This function fully serializes the field and its properties.
  */
diff --git a/Engine/Source/Runtime/CoreUObject/Public/UObject/FieldPath.h b/Engine/Source/Runtime/CoreUObject/Public/UObject/FieldPath.h
--- a/Engine/Source/Runtime/CoreUObject/Public/UObject/FieldPath.h
+++ b/Engine/Source/Runtime/CoreUObject/Public/UObject/FieldPath.h
@@ -202,8 +202,21 @@ private:
 public:
 	TFieldPath()
 	{}
-	TFieldPath(const TFieldPath&) = default;
-	TFieldPath& operator=(const TFieldPath&) = default;
+	FORCEINLINE TFieldPath(const TFieldPath& Other)
+	{
+		//  First refresh the serial number from the other path
+		ResolvedField = Other.Get();
+		// Now that the Other path is refreshed, we can copy from it
+		FFieldPath::operator=(Other);
+	}
+	FORCEINLINE TFieldPath& operator=(const TFieldPath& Other)
+	{
+		//  First refresh the serial number from the other path
+		ResolvedField = Other.Get();
+		// Now that the Other path is refreshed, we can copy from it
+		FFieldPath::operator=(Other);
+		return *this;
+	}
 
 	/**
 	* Construct from a null pointer
@@ -280,6 +293,10 @@ public:
 		// one or both of the types have only been forward-declared.
 		static_assert(TPointerIsConvertibleFromTo<OtherPropertyType, PropertyType>::Value, "Unable to convert TFieldPath - types are incompatible");
 
+		// First make sure the Other path has the serial number up to date, otherwise we'll keep having to
+		// reevealuate this path because it gets the serial number copied from the Other path
+		ResolvedField = Other.Get();
+		// Now that the Other path is refreshed, we can copy from it
 		FFieldPath::operator=(Other);
 	}
 
diff --git a/Engine/Source/Runtime/CoreUObject/Public/UObject/UnrealType.h b/Engine/Source/Runtime/CoreUObject/Public/UObject/UnrealType.h
--- a/Engine/Source/Runtime/CoreUObject/Public/UObject/UnrealType.h
+++ b/Engine/Source/Runtime/CoreUObject/Public/UObject/UnrealType.h
@@ -5282,10 +5282,12 @@ protected:
 		{
 			while (CurrentField)
 			{
-				if (CurrentField->HasAllCastFlags(T::StaticClassCastFlags()) &&
+				typename T::FieldTypeClass* FieldClass = CurrentField->GetClass();
+
+				if (FieldClass->HasAllCastFlags(T::StaticClassCastFlags()) &&
 					(
 						   bIncludeDeprecated
-						|| !CurrentField->HasAllCastFlags(CASTCLASS_FProperty)
+						|| !FieldClass->HasAllCastFlags(CASTCLASS_FProperty)
 						|| !((FProperty*)CurrentField)->HasAllPropertyFlags(CPF_Deprecated)
 					)
 				)
