CIS fix

diff --git a/Engine/Source/Runtime/CoreUObject/Public/UObject/FieldPath.h b/Engine/Source/Runtime/CoreUObject/Public/UObject/FieldPath.h
--- a/Engine/Source/Runtime/CoreUObject/Public/UObject/FieldPath.h
+++ b/Engine/Source/Runtime/CoreUObject/Public/UObject/FieldPath.h
@@ -205,14 +205,14 @@ public:
 	FORCEINLINE TFieldPath(const TFieldPath& Other)
 	{
 		//  First refresh the serial number from the other path
-		ResolvedField = Other.Get();
+		Other.Get();
 		// Now that the Other path is refreshed, we can copy from it
 		FFieldPath::operator=(Other);
 	}
 	FORCEINLINE TFieldPath& operator=(const TFieldPath& Other)
 	{
 		//  First refresh the serial number from the other path
-		ResolvedField = Other.Get();
+		Other.Get();
 		// Now that the Other path is refreshed, we can copy from it
 		FFieldPath::operator=(Other);
 		return *this;
@@ -295,7 +295,7 @@ public:
 
 		// First make sure the Other path has the serial number up to date, otherwise we'll keep having to
 		// reevealuate this path because it gets the serial number copied from the Other path
-		ResolvedField = Other.Get();
+		Other.Get();
 		// Now that the Other path is refreshed, we can copy from it
 		FFieldPath::operator=(Other);
 	}
