Allow TStrongObjectPtr to be created during static init

This defers allocating the ReferenceCollector until we need to store a non-null object pointer.

diff --git a/Engine/Source/Runtime/CoreUObject/Public/UObject/StrongObjectPtr.h b/Engine/Source/Runtime/CoreUObject/Public/UObject/StrongObjectPtr.h
--- a/Engine/Source/Runtime/CoreUObject/Public/UObject/StrongObjectPtr.h
+++ b/Engine/Source/Runtime/CoreUObject/Public/UObject/StrongObjectPtr.h
@@ -5,14 +5,14 @@
 #include "UObject/GCObject.h"
 #include "Templates/EnableIf.h"
 #include "Templates/PointerIsConvertibleFromTo.h"
-#include "Templates/UniqueObj.h"
+#include "Templates/UniquePtr.h"
 
 namespace UE4StrongObjectPtr_Private
 {
 	class FInternalReferenceCollector : public FGCObject
 	{
 	public:
-		FInternalReferenceCollector(const volatile UObject* InObject = nullptr)
+		explicit FInternalReferenceCollector(const volatile UObject* InObject)
 			: Object(InObject)
 		{
 			check(IsInGameThread());
@@ -45,6 +45,11 @@ namespace UE4StrongObjectPtr_Private
 			Collector.AddReferencedObject(Object);
 		}
 
+		virtual FString GetReferencerName() const
+		{
+			return TEXT("TStrongObjectPtr");
+		}
+
 	private:
 		const volatile UObject* Object;
 	};
@@ -59,7 +64,6 @@ class TStrongObjectPtr
 {
 public:
 	TStrongObjectPtr(TStrongObjectPtr&& InOther) = default;
-	TStrongObjectPtr(const TStrongObjectPtr& InOther) = default;
 	TStrongObjectPtr& operator=(TStrongObjectPtr&& InOther) = default;
 	~TStrongObjectPtr() = default;
 
@@ -69,9 +73,14 @@ public:
 	}
 
 	FORCEINLINE_DEBUGGABLE explicit TStrongObjectPtr(ObjectType* InObject)
-		: ReferenceCollector(InObject)
 	{
 		static_assert(TPointerIsConvertibleFromTo<ObjectType, const volatile UObject>::Value, "TStrongObjectPtr can only be constructed with UObject types");
+		Reset(InObject);
+	}
+
+	FORCEINLINE_DEBUGGABLE TStrongObjectPtr(const TStrongObjectPtr& InOther)
+	{
+		Reset(InOther.Get());
 	}
 
 	template <
@@ -79,14 +88,13 @@ public:
 		typename = typename TEnableIf<TPointerIsConvertibleFromTo<OtherObjectType, ObjectType>::Value>::Type
 	>
 	FORCEINLINE_DEBUGGABLE TStrongObjectPtr(const TStrongObjectPtr<OtherObjectType>& InOther)
-		: ReferenceCollector(InOther.Get())
 	{
+		Reset(InOther.Get());
 	}
 
 	FORCEINLINE_DEBUGGABLE TStrongObjectPtr& operator=(const TStrongObjectPtr& InOther)
 	{
-		// TUniqueObj is not assignable so we need to implement this instead of defaulting it.
-		ReferenceCollector->Set(InOther.Get());
+		Reset(InOther.Get());
 		return *this;
 	}
 
@@ -96,7 +104,7 @@ public:
 		TStrongObjectPtr&
 	>::Type operator=(const TStrongObjectPtr<OtherObjectType>& InOther)
 	{
-		ReferenceCollector->Set(InOther.Get());
+		Reset(InOther.Get());
 		return *this;
 	}
 
@@ -114,22 +122,31 @@ public:
 
 	FORCEINLINE_DEBUGGABLE bool IsValid() const
 	{
-		return ReferenceCollector->IsValid();
+		return ReferenceCollector && ReferenceCollector->IsValid();
 	}
 
 	FORCEINLINE_DEBUGGABLE explicit operator bool() const
 	{
-		return ReferenceCollector->IsValid();
+		return IsValid();
 	}
 
 	FORCEINLINE_DEBUGGABLE ObjectType* Get() const
 	{
-		return ReferenceCollector->GetAs<ObjectType>();
+		return ReferenceCollector
+			? ReferenceCollector->GetAs<ObjectType>()
+			: nullptr;
 	}
 
 	FORCEINLINE_DEBUGGABLE void Reset(ObjectType* InNewObject = nullptr)
 	{
-		ReferenceCollector->Set(InNewObject);
+		if (ReferenceCollector)
+		{
+			ReferenceCollector->Set(InNewObject);
+		}
+		else if (InNewObject)
+		{
+			ReferenceCollector = MakeUnique<UE4StrongObjectPtr_Private::FInternalReferenceCollector>(InNewObject);
+		}
 	}
 
 	FORCEINLINE_DEBUGGABLE friend uint32 GetTypeHash(const TStrongObjectPtr& InStrongObjectPtr)
@@ -138,7 +155,7 @@ public:
 	}
 
 private:
-	TUniqueObj<UE4StrongObjectPtr_Private::FInternalReferenceCollector> ReferenceCollector;
+	TUniquePtr<UE4StrongObjectPtr_Private::FInternalReferenceCollector> ReferenceCollector;
 };
 
 template <typename LHSObjectType, typename RHSObjectType>
