Fixing CleanAndSanitizeStruct not destroying struct properties

diff --git a/Engine/Source/Editor/KismetCompiler/Private/UserDefinedStructureCompilerUtils.cpp b/Engine/Source/Editor/KismetCompiler/Private/UserDefinedStructureCompilerUtils.cpp
--- a/Engine/Source/Editor/KismetCompiler/Private/UserDefinedStructureCompilerUtils.cpp
+++ b/Engine/Source/Editor/KismetCompiler/Private/UserDefinedStructureCompilerUtils.cpp
@@ -159,14 +159,12 @@ struct FUserDefinedStructureCompilerInner
 
 			StructToClean->SetSuperStruct(nullptr);
 			StructToClean->Children = nullptr;
+			StructToClean->DestroyChildPropertiesAndResetPropertyLinks();
 			StructToClean->Script.Empty();
 			StructToClean->MinAlignment = 0;
-			StructToClean->RefLink = nullptr;
-			StructToClean->PropertyLink = nullptr;
-			StructToClean->DestructorLink = nullptr;
 			StructToClean->ScriptObjectReferences.Empty();
-			StructToClean->PropertyLink = nullptr;
 			StructToClean->ErrorMessage.Empty();
+			StructToClean->SetStructTrashed(true);
 		}
 
 		return TransientStruct;
diff --git a/Engine/Source/Runtime/CoreUObject/Private/UObject/Class.cpp b/Engine/Source/Runtime/CoreUObject/Private/UObject/Class.cpp
--- a/Engine/Source/Runtime/CoreUObject/Private/UObject/Class.cpp
+++ b/Engine/Source/Runtime/CoreUObject/Private/UObject/Class.cpp
@@ -1367,10 +1367,11 @@ void UStruct::SerializeVersionedTaggedProperties(FStructuredArchive::FSlot Slot,
 					if (BreakRecursionIfFullyLoad && BreakRecursionIfFullyLoad->HasAllFlags(RF_LoadCompleted))
 					{
 					}
+					else
 #endif // WITH_EDITOR
 					// editoronly properties should be skipped if we are NOT the editor, or we are 
 					// the editor but are cooking for console (editoronly implies notforconsole)
-					else if ((Property->PropertyFlags & CPF_EditorOnly) && !FPlatformProperties::HasEditorOnlyData() && !GForceLoadEditorOnly)
+					if ((Property->PropertyFlags & CPF_EditorOnly) && !FPlatformProperties::HasEditorOnlyData() && !GForceLoadEditorOnly)
 					{
 					}
 					// check for valid array index
@@ -1572,6 +1573,15 @@ inline void DestroyPropertyLinkedList(FField*& PropertiesToDestroy)
 	PropertiesToDestroy = nullptr;
 }
 
+void UStruct::DestroyChildPropertiesAndResetPropertyLinks()
+{
+	DestroyPropertyLinkedList(ChildProperties);
+	PropertyLink = nullptr;
+	RefLink = nullptr;
+	DestructorLink = nullptr;
+	PostConstructLink = nullptr;
+}
+
 UStruct::~UStruct()
 {
 	// Destroy all properties owned by this struct
@@ -2871,6 +2881,7 @@ void UScriptStruct::ExportText(FString& ValueStr, const void* Value, const void*
 void UScriptStruct::Link(FArchive& Ar, bool bRelinkExistingProperties)
 {
 	Super::Link(Ar, bRelinkExistingProperties);
+	SetStructTrashed(false);
 	if (!HasDefaults()) // if you have CppStructOps, then that is authoritative, otherwise we look at the properties
 	{
 		StructFlags = EStructFlags(StructFlags | STRUCT_ZeroConstructor | STRUCT_NoDestructor | STRUCT_IsPlainOldData);
@@ -3148,6 +3159,23 @@ void UScriptStruct::DestroyStruct(void* Dest, int32 ArrayDim) const
 	}
 }
 
+bool UScriptStruct::IsStructTrashed() const
+{
+	return !!(StructFlags & STRUCT_Trashed);
+}
+ 
+void UScriptStruct::SetStructTrashed(bool bIsTrash)
+{
+	if (bIsTrash)
+	{
+		StructFlags = EStructFlags(StructFlags | STRUCT_Trashed);
+	}
+	else
+	{
+		StructFlags = EStructFlags(StructFlags & ~STRUCT_Trashed);
+	}
+}
+
 void UScriptStruct::RecursivelyPreload() {}
 
 FGuid UScriptStruct::GetCustomGuid() const
@@ -4198,6 +4226,11 @@ void UClass::SetSuperStruct(UStruct* NewSuperStruct)
 	HashObject(this);
 }
 
+bool UClass::IsStructTrashed() const
+{
+	return Children == nullptr && ChildProperties == nullptr && ClassDefaultObject == nullptr;
+}
+
 void UClass::Serialize( FArchive& Ar )
 {
 	if ( Ar.IsLoading() || Ar.IsModifyingWeakAndStrongReferences() )
diff --git a/Engine/Source/Runtime/CoreUObject/Public/UObject/Class.h b/Engine/Source/Runtime/CoreUObject/Public/UObject/Class.h
--- a/Engine/Source/Runtime/CoreUObject/Public/UObject/Class.h
+++ b/Engine/Source/Runtime/CoreUObject/Public/UObject/Class.h
@@ -497,6 +497,15 @@ public:
 	/** Returns a human readable string for a given field, overridden for user defined structs */
 	virtual FString GetAuthoredNameForField(const FField* Field) const;
 
+	/** If true, this class has been cleaned and sanitized (trashed) and should not be used */
+	virtual bool IsStructTrashed() const
+	{
+		return false;
+	}
+
+	/** Destroys all properties owned by this struct */
+	void DestroyChildPropertiesAndResetPropertyLinks();
+
 #if WITH_EDITORONLY_DATA
 	/** Try and find boolean metadata with the given key. If not found on this class, work up hierarchy looking for it. */
 	bool GetBoolMetaDataHierarchical(const FName& Key) const;
@@ -617,6 +626,9 @@ enum EStructFlags
 	/** If set, this struct can share net serialization state across connections */
 	STRUCT_NetSharedSerialization = 0x00400000,
 
+	/** If set, this struct has been cleaned and sanitized (trashed) and should not be used */
+	STRUCT_Trashed = 0x00800000,
+
 	/** Struct flags that are automatically inherited */
 	STRUCT_Inherit				= STRUCT_HasInstancedReference|STRUCT_Atomic,
 
@@ -1476,8 +1488,12 @@ public:
 	virtual COREUOBJECT_API void Link(FArchive& Ar, bool bRelinkExistingProperties) override;
 	virtual COREUOBJECT_API void InitializeStruct(void* Dest, int32 ArrayDim = 1) const override;
 	virtual COREUOBJECT_API void DestroyStruct(void* Dest, int32 ArrayDim = 1) const override;
+	virtual COREUOBJECT_API bool IsStructTrashed() const override;
 	// End of UStruct interface.
 
+	/** Sets or unsets the trashed flag on this struct */
+	void COREUOBJECT_API SetStructTrashed(bool bIsTrash);
+
 	/** 
 	 * Stash a CppStructOps for future use 
 	 * @param Target Name of the struct 
@@ -2751,6 +2767,7 @@ public:
 	// UStruct interface.
 	virtual void Link(FArchive& Ar, bool bRelinkExistingProperties) override;
 	virtual void SetSuperStruct(UStruct* NewSuperStruct) override;
+	virtual bool IsStructTrashed() const override;
 	// End of UStruct interface.
 
 #if WITH_EDITOR
