Changed the description for the BlueprintProtected metadata specifier to be inline with how it actually works.

diff --git a/Engine/Source/Runtime/CoreUObject/Public/UObject/ObjectMacros.h b/Engine/Source/Runtime/CoreUObject/Public/UObject/ObjectMacros.h
--- a/Engine/Source/Runtime/CoreUObject/Public/UObject/ObjectMacros.h
+++ b/Engine/Source/Runtime/CoreUObject/Public/UObject/ObjectMacros.h
@@ -1290,7 +1290,7 @@ namespace UM
 		/// [FunctionMetadata] This function is an internal implementation detail, used to implement another function or node.  It is never directly exposed in a graph.
 		BlueprintInternalUseOnly,
 
-		/// [FunctionMetadata] This function can only be called on 'this' in a blueprint. It cannot be called on another instance.
+		/// [FunctionMetadata] This function is only accessible from within its class and derived classes.
 		BlueprintProtected,
 
 		/// [FunctionMetadata] Used for BlueprintCallable functions that have a WorldContext pin to indicate that the function can be called even if the class does not implement the virtual function GetWorld().
