Make properties inside sparse class data respect the AssetRegistrySearchable flag.

diff --git a/Engine/Source/Runtime/CoreUObject/Private/UObject/Obj.cpp b/Engine/Source/Runtime/CoreUObject/Private/UObject/Obj.cpp
--- a/Engine/Source/Runtime/CoreUObject/Private/UObject/Obj.cpp
+++ b/Engine/Source/Runtime/CoreUObject/Private/UObject/Obj.cpp
@@ -1684,87 +1684,99 @@ FString GetConfigFilename( UObject* SourceObject )
 	else
 #endif
 	{
-		// otherwise look at the class to get the config name
-		return SourceObject->GetClass()->GetConfigName();
+	// otherwise look at the class to get the config name
+	return SourceObject->GetClass()->GetConfigName();
 	}
 }
 
-void UObject::FAssetRegistryTag::GetAssetRegistryTagsFromSearchableProperties(const UObject* Object, TArray<FAssetRegistryTag>& OutTags)
+void GetAssetRegistryTagFromProperty(const void* BaseMemoryLocation, const UObject* OwnerObject, UProperty* Prop, TSet<FName>& OutFoundSpecialStructs, TArray<UObject::FAssetRegistryTag>& OutTags)
 {
-	TSet<FName> FoundSpecialStructs;
+	FName TagName;
+	UObject::FAssetRegistryTag::ETagType TagType = UObject::FAssetRegistryTag::ETagType::TT_Alphabetical;
+	UStructProperty* StructProp = Cast<UStructProperty>(Prop);
 
-	check(nullptr != Object);
-	for( TFieldIterator<UProperty> FieldIt( Object->GetClass() ); FieldIt; ++FieldIt )
+	if (StructProp && StructProp->Struct && UObject::FAssetRegistryTag::IsUniqueAssetRegistryTagStruct(StructProp->Struct->GetFName(), TagType))
 	{
-		FName TagName;
-		FAssetRegistryTag::ETagType TagType = TT_Alphabetical;
-		UStructProperty* StructProp = Cast<UStructProperty>(*FieldIt);
+		// Special unique structure type
+		TagName = StructProp->Struct->GetFName();
 
-		if (StructProp && StructProp->Struct && IsUniqueAssetRegistryTagStruct(StructProp->Struct->GetFName(), TagType))
+		if (OutFoundSpecialStructs.Contains(TagName))
 		{
-			// Special unique structure type
-			TagName = StructProp->Struct->GetFName();
-			
-			if (FoundSpecialStructs.Contains(TagName))
-			{
-				UE_LOG(LogObj, Error, TEXT("Object %s has more than one unique asset registry struct %s!"), *Object->GetPathName(), *TagName.ToString());
-			}
-			else
-			{
-				FoundSpecialStructs.Add(TagName);
-			}
+			UE_LOG(LogObj, Error, TEXT("Object %s has more than one unique asset registry struct %s!"), *OwnerObject->GetPathName(), *TagName.ToString());
 		}
-		else if (FieldIt->HasAnyPropertyFlags(CPF_AssetRegistrySearchable))
+		else
 		{
-			TagName = FieldIt->GetFName();
-
-			UClass* Class = FieldIt->GetClass();
-			if (Class->IsChildOf(UIntProperty::StaticClass()) ||
-				Class->IsChildOf(UFloatProperty::StaticClass()) ||
-				Class->IsChildOf(UDoubleProperty::StaticClass()))
-			{
-				// ints and floats are always numerical
-				TagType = FAssetRegistryTag::TT_Numerical;
-			}
-			else if (Class->IsChildOf(UByteProperty::StaticClass()))
-			{
-				// bytes are numerical, enums are alphabetical
-				UByteProperty* ByteProp = static_cast<UByteProperty*>(*FieldIt);
-				if (ByteProp->Enum)
-				{
-					TagType = FAssetRegistryTag::TT_Alphabetical;
-				}
-				else
-				{
-					TagType = FAssetRegistryTag::TT_Numerical;
-				}
-			}
-			else if ( Class->IsChildOf(UEnumProperty::StaticClass()) )
-			{
-				// enums are alphabetical
-				TagType = FAssetRegistryTag::TT_Alphabetical;
-			}
-			else if (Class->IsChildOf(UArrayProperty::StaticClass()) || Class->IsChildOf(UMapProperty::StaticClass()) || Class->IsChildOf(USetProperty::StaticClass())
-				|| Class->IsChildOf(UStructProperty::StaticClass()) || Class->IsChildOf(UObjectPropertyBase::StaticClass()))
+			OutFoundSpecialStructs.Add(TagName);
+		}
+	}
+	else if (Prop->HasAnyPropertyFlags(CPF_AssetRegistrySearchable))
+	{
+		UClass* Class = Prop->GetClass();
+		if (Class->IsChildOf(UIntProperty::StaticClass()) ||
+			Class->IsChildOf(UFloatProperty::StaticClass()) ||
+			Class->IsChildOf(UDoubleProperty::StaticClass()))
+		{
+			// ints and floats are always numerical
+			TagType = UObject::FAssetRegistryTag::ETagType::TT_Numerical;
+		}
+		else if (Class->IsChildOf(UByteProperty::StaticClass()))
+		{
+			// bytes are numerical, enums are alphabetical
+			UByteProperty* ByteProp = static_cast<UByteProperty*>(Prop);
+			if (ByteProp->Enum)
 			{
-				// Arrays/maps/sets/structs/objects are hidden, it is often too much information to display and sort
-				TagType = FAssetRegistryTag::TT_Hidden;
+				TagType = UObject::FAssetRegistryTag::ETagType::TT_Alphabetical;
 			}
 			else
 			{
-				// All other types are alphabetical
-				TagType = FAssetRegistryTag::TT_Alphabetical;
+				TagType = UObject::FAssetRegistryTag::ETagType::TT_Numerical;
 			}
 		}
-
-		if (TagName != NAME_None)
+		else if (Class->IsChildOf(UEnumProperty::StaticClass()))
 		{
+			// enums are alphabetical
+			TagType = UObject::FAssetRegistryTag::ETagType::TT_Alphabetical;
+		}
+		else if (Class->IsChildOf(UArrayProperty::StaticClass()) || Class->IsChildOf(UMapProperty::StaticClass()) || Class->IsChildOf(USetProperty::StaticClass())
+			|| Class->IsChildOf(UStructProperty::StaticClass()) || Class->IsChildOf(UObjectPropertyBase::StaticClass()))
+		{
+			// Arrays/maps/sets/structs/objects are hidden, it is often too much information to display and sort
+			TagType = UObject::FAssetRegistryTag::ETagType::TT_Hidden;
+		}
+		else
+		{
+			// All other types are alphabetical
+			TagType = UObject::FAssetRegistryTag::ETagType::TT_Alphabetical;
+		}
+	}
 
-			FString PropertyStr;
-			const uint8* PropertyAddr = FieldIt->ContainerPtrToValuePtr<uint8>(Object);
-			FieldIt->ExportTextItem( PropertyStr, PropertyAddr, PropertyAddr, nullptr, PPF_None );
+	if (TagName != NAME_None)
+	{
+		FString PropertyStr;
+		const uint8* PropertyAddr = Prop->ContainerPtrToValuePtr<uint8>(BaseMemoryLocation);
+		Prop->ExportTextItem(PropertyStr, PropertyAddr, PropertyAddr, nullptr, PPF_None);
 
-			OutTags.Add( FAssetRegistryTag(TagName, PropertyStr, TagType) );
+		OutTags.Add(UObject::FAssetRegistryTag(TagName, PropertyStr, TagType));
+	}
+}
+
+void UObject::FAssetRegistryTag::GetAssetRegistryTagsFromSearchableProperties(const UObject* Object, TArray<FAssetRegistryTag>& OutTags)
+{
+	TSet<FName> FoundSpecialStructs;
+
+	check(nullptr != Object);
+	for (TFieldIterator<UProperty> FieldIt( Object->GetClass() ); FieldIt; ++FieldIt)
+	{
+		GetAssetRegistryTagFromProperty(Object, Object, Cast<UProperty>(*FieldIt), FoundSpecialStructs, OutTags);
+	}
+
+	UScriptStruct* SparseClassDataStruct = Object->GetClass()->GetSparseClassDataStruct();
+	if (SparseClassDataStruct)
+	{
+		void* SparseClassData = Object->GetClass()->GetOrCreateSparseClassData();
+		for (TFieldIterator<UProperty> FieldIt(SparseClassDataStruct); FieldIt; ++FieldIt)
+		{
+			GetAssetRegistryTagFromProperty(SparseClassData, Object, Cast<UProperty>(*FieldIt), FoundSpecialStructs, OutTags);
 		}
 	}
 }
