Added a check to make sure that no one attempts to access sparse class data while the CDO is still under construction.

diff --git a/Engine/Source/Runtime/CoreUObject/Private/UObject/Class.cpp b/Engine/Source/Runtime/CoreUObject/Private/UObject/Class.cpp
--- a/Engine/Source/Runtime/CoreUObject/Private/UObject/Class.cpp
+++ b/Engine/Source/Runtime/CoreUObject/Private/UObject/Class.cpp
@@ -4245,7 +4245,9 @@ void* UClass::CreateSparseClassData()
 {
 	check(SparseClassData == nullptr);
 
-	SparseClassData = GetDefaultObject()->CreateSparseClassData();
+	UObject* DefaultObj = GetDefaultObject();
+	checkf(!DefaultObj->HasAnyFlags(RF_NeedInitialization), TEXT("You cannot create sparse class data until the CDO has been initialized. If you are using sparse class data inside the constructor consider moving it to PostConstructInit."));
+	SparseClassData = DefaultObj->CreateSparseClassData();
 
 	if (SparseClassData)
 	{
