FProperty fixes
#rnx

diff --git a/Engine/Source/Runtime/CoreUObject/Private/UObject/Obj.cpp b/Engine/Source/Runtime/CoreUObject/Private/UObject/Obj.cpp
--- a/Engine/Source/Runtime/CoreUObject/Private/UObject/Obj.cpp
+++ b/Engine/Source/Runtime/CoreUObject/Private/UObject/Obj.cpp
@@ -1759,18 +1759,18 @@ void UObject::FAssetRegistryTag::GetAssetRegistryTagsFromSearchableProperties(co
 	TSet<FName> FoundSpecialStructs;
 
 	check(nullptr != Object);
-	for (TFieldIterator<UProperty> FieldIt( Object->GetClass() ); FieldIt; ++FieldIt)
+	for (TFieldIterator<FProperty> FieldIt( Object->GetClass() ); FieldIt; ++FieldIt)
 	{
-		GetAssetRegistryTagFromProperty(Object, Object, Cast<UProperty>(*FieldIt), FoundSpecialStructs, OutTags);
+		GetAssetRegistryTagFromProperty(Object, Object, CastField<FProperty>(*FieldIt), FoundSpecialStructs, OutTags);
 	}
 
 	UScriptStruct* SparseClassDataStruct = Object->GetClass()->GetSparseClassDataStruct();
 	if (SparseClassDataStruct)
 	{
 		void* SparseClassData = Object->GetClass()->GetOrCreateSparseClassData();
-		for (TFieldIterator<UProperty> FieldIt(SparseClassDataStruct); FieldIt; ++FieldIt)
+		for (TFieldIterator<FProperty> FieldIt(SparseClassDataStruct); FieldIt; ++FieldIt)
 		{
-			GetAssetRegistryTagFromProperty(SparseClassData, Object, Cast<UProperty>(*FieldIt), FoundSpecialStructs, OutTags);
+			GetAssetRegistryTagFromProperty(SparseClassData, Object, CastField<FProperty>(*FieldIt), FoundSpecialStructs, OutTags);
 		}
 	}
 }
