Merge fix
#rnx

diff --git a/Engine/Source/Runtime/CoreUObject/Private/UObject/Class.cpp b/Engine/Source/Runtime/CoreUObject/Private/UObject/Class.cpp
--- a/Engine/Source/Runtime/CoreUObject/Private/UObject/Class.cpp
+++ b/Engine/Source/Runtime/CoreUObject/Private/UObject/Class.cpp
@@ -4820,21 +4820,12 @@ void* UClass::CreateSparseClassData()
 {
 	check(SparseClassData == nullptr);
 
->>>> ORIGINAL //UE4/Main/Engine/Source/Runtime/CoreUObject/Private/UObject/Class.cpp#115
-	SparseClassData = GetDefaultObject()->CreateSparseClassData();
-
-==== THEIRS //UE4/Main/Engine/Source/Runtime/CoreUObject/Private/UObject/Class.cpp#123
 	if (SparseClassDataStruct)
 	{
 		SparseClassData = FMemory::Malloc(SparseClassDataStruct->GetStructureSize(), SparseClassDataStruct->GetMinAlignment());
 		SparseClassDataStruct->GetCppStructOps()->Construct(SparseClassData);
 	}
-==== YOURS //Dan.Oconnor_sea-wd-virtual/Engine/Source/Runtime/CoreUObject/Private/UObject/Class.cpp
-	UObject* DefaultObj = GetDefaultObject();
-	checkf(!DefaultObj->HasAnyFlags(RF_NeedInitialization), TEXT("You cannot create sparse class data until the CDO has been initialized. If you are using sparse class data inside the constructor consider moving it to PostConstructInit."));
-	SparseClassData = DefaultObj->CreateSparseClassData();
 
-<<<<
 	if (SparseClassData)
 	{
 		// initialize per class data from the archetype if we have one
