Core: Added the name of the failing struct to the bad alignment check in UScriptStruct::InitializeStruct

diff --git a/Engine/Source/Runtime/CoreUObject/Private/UObject/Class.cpp b/Engine/Source/Runtime/CoreUObject/Private/UObject/Class.cpp
--- a/Engine/Source/Runtime/CoreUObject/Private/UObject/Class.cpp
+++ b/Engine/Source/Runtime/CoreUObject/Private/UObject/Class.cpp
@@ -3011,8 +3011,9 @@ void UScriptStruct::InitializeStruct(void* InDest, int32 ArrayDim) const
 			{
 				void* PropertyDest = Dest + ArrayIndex * Stride;
 				checkf(IsAligned(PropertyDest, TheCppStructOps->GetAlignment()),
-					TEXT("Destination address for property does not match requirement of %d byte alignment"), 
-					TheCppStructOps->GetAlignment());
+					TEXT("Destination address for property does not match requirement of %d byte alignment for %s"), 
+					TheCppStructOps->GetAlignment(),
+					*GetPathNameSafe(this));
 				TheCppStructOps->Construct(PropertyDest);
 			}
 		}
