Fixed a bug that was causing properties that were marked as asset registry searchable to not be included with asset data in the asset registry.

diff --git a/Engine/Source/Runtime/CoreUObject/Private/UObject/Obj.cpp b/Engine/Source/Runtime/CoreUObject/Private/UObject/Obj.cpp
--- a/Engine/Source/Runtime/CoreUObject/Private/UObject/Obj.cpp
+++ b/Engine/Source/Runtime/CoreUObject/Private/UObject/Obj.cpp
@@ -1706,6 +1706,8 @@ void GetAssetRegistryTagFromProperty(const void* BaseMemoryLocation, const UObje
 	}
 	else if (Prop->HasAnyPropertyFlags(CPF_AssetRegistrySearchable))
 	{
+		TagName = Prop->GetFName();
+
 		if (Prop->IsA(FIntProperty::StaticClass()) ||
 			Prop->IsA(FFloatProperty::StaticClass()) ||
 			Prop->IsA(FDoubleProperty::StaticClass()))
