Explicitly support move vs. const-ref for AddAnnotation and change calling code to use MoveTemp where possible

diff --git a/Engine/Source/Editor/BlueprintGraph/Private/K2Node.cpp b/Engine/Source/Editor/BlueprintGraph/Private/K2Node.cpp
--- a/Engine/Source/Editor/BlueprintGraph/Private/K2Node.cpp
+++ b/Engine/Source/Editor/BlueprintGraph/Private/K2Node.cpp
@@ -1334,7 +1334,7 @@ FOnUserDefinedPinRenamed& UK2Node::OnUserDefinedPinRenamed()
 	{
 		FPinRenamedAnnotation NewEvent;
 		NewEvent.bIsDefault = false;
-		GOnUserDefinedPinRenamedAnnotation.AddAnnotation(this, NewEvent);
+		GOnUserDefinedPinRenamedAnnotation.AddAnnotation(this, MoveTemp(NewEvent));
 	}
 
 	return Map.FindChecked(this).PinRenamedEvent;
diff --git a/Engine/Source/Runtime/CoreUObject/Private/UObject/UObjectMarks.cpp b/Engine/Source/Runtime/CoreUObject/Private/UObject/UObjectMarks.cpp
--- a/Engine/Source/Runtime/CoreUObject/Private/UObject/UObjectMarks.cpp
+++ b/Engine/Source/Runtime/CoreUObject/Private/UObject/UObjectMarks.cpp
@@ -33,7 +33,7 @@ template<typename TAnnotation>
 class FUObjectAnnotationSparseNoSync : public FUObjectArray::FUObjectDeleteListener
 {
 public:
-	virtual void NotifyUObjectDeleted(const UObjectBase *Object, int32 Index) override
+	virtual void NotifyUObjectDeleted(const UObjectBase* Object, int32 Index) override
 	{
 		RemoveAnnotation(Object);
 	}
@@ -54,11 +54,13 @@ public:
 		RemoveAllAnnotations();
 	}
 
-	void AddAnnotation(const UObjectBase* Object,TAnnotation Annotation)
+private:
+	template<typename T>
+	void AddAnnotationInternal(const UObjectBase* Object, T&& Annotation)
 	{
 		check(Object);
 		AnnotationCacheKey = Object;
-		AnnotationCacheValue = Annotation;
+		AnnotationCacheValue = Forward<T>(Annotation);
 		if (Annotation.IsDefault())
 		{
 			RemoveAnnotation(Object); // adding the default annotation is the same as removing an annotation
@@ -70,7 +72,19 @@ public:
 				// we are adding the first one, so if we are auto removing or verifying removal, register now
 				GUObjectArray.AddUObjectDeleteListener(this);
 			}
-			AnnotationMap.Add(AnnotationCacheKey,AnnotationCacheValue);
+			AnnotationMap.Add(AnnotationCacheKey, AnnotationCacheValue);
+		}
+	}
+
+public:
+	void AddAnnotation(const UObjectBase* Object, TAnnotation&& Annotation)
+	{
+		AddAnnotationInternal(Object, MoveTemp(Annotation));
+	}
+
+	void AddAnnotation(const UObjectBase* Object, const TAnnotation& Annotation)
+	{
+		AddAnnotationInternal(Object, Annotation);
 		}
 	}
 
diff --git a/Engine/Source/Runtime/CoreUObject/Public/UObject/UObjectAnnotation.h b/Engine/Source/Runtime/CoreUObject/Public/UObject/UObjectAnnotation.h
--- a/Engine/Source/Runtime/CoreUObject/Public/UObject/UObjectAnnotation.h
+++ b/Engine/Source/Runtime/CoreUObject/Public/UObject/UObjectAnnotation.h
@@ -75,18 +75,14 @@ public:
 		RemoveAllAnnotations();		
 	}
 
-	/**
-	 * Add an annotation to the annotation list. If the Annotation is the default, then the annotation is removed.
-	 *
-	 * @param Object		Object to annotate.
-	 * @param Annotation	Annotation to associate with Object.
-	 */
-	void AddAnnotation(const UObjectBase *Object,TAnnotation Annotation)
+private:
+	template<typename T>
+	void AddAnnotationInternal(const UObjectBase* Object, T&& Annotation)
 	{
 		check(Object);
 		FScopeLock AnnotationMapLock(&AnnotationMapCritical);
 		AnnotationCacheKey = Object;
-		AnnotationCacheValue = MoveTemp(Annotation);
+		AnnotationCacheValue = Forward<T>(Annotation);
 		if (AnnotationCacheValue.IsDefault())
 		{
 			RemoveAnnotation(Object); // adding the default annotation is the same as removing an annotation
@@ -103,8 +99,24 @@ public:
 					GUObjectArray.AddUObjectDeleteListener(this);
 				}
 			}
-			AnnotationMap.Add(AnnotationCacheKey,AnnotationCacheValue);
-		}
+			AnnotationMap.Add(AnnotationCacheKey, AnnotationCacheValue);
+	}
+
+public:
+	/**
+	 * Add an annotation to the annotation list. If the Annotation is the default, then the annotation is removed.
+	 *
+	 * @param Object        Object to annotate.
+	 * @param Annotation    Annotation to associate with Object.
+	 */
+	void AddAnnotation(const UObjectBase* Object, TAnnotation&& Annotation)
+	{
+		AddAnnotationInternal(Object, MoveTemp(Annotation));
+	}
+
+	void AddAnnotation(const UObjectBase* Object, const TAnnotation& Annotation)
+	{
+		AddAnnotationInternal(Object, Annotation);
 	}
 	/**
 	 * Removes an annotation from the annotation list and returns the annotation if it had one 
@@ -290,20 +302,16 @@ public:
 	 * @param Annotation	Annotation to find
 	 * @return				Object associated with this annotation or NULL if none found
 	 */
-	UObject *Find(TAnnotation Annotation)
+	UObject* Find(const TAnnotation& Annotation)
 	{
 		FScopeLock InverseAnnotationMapLock(&InverseAnnotationMapCritical);
 		checkSlow(!Annotation.IsDefault()); // it is not legal to search for the default annotation
-		return (UObject *)InverseAnnotationMap.FindRef(Annotation);
+		return (UObject*)InverseAnnotationMap.FindRef(Annotation);
 	}
 
-	/**
-	 * Add an annotation to the annotation list. If the Annotation is the default, then the annotation is removed.
-	 *
-	 * @param Object		Object to annotate.
-	 * @param Annotation	Annotation to associate with Object.
-	 */
-	void AddAnnotation(const UObjectBase *Object,TAnnotation Annotation)
+private:
+	template<T> 
+	void AddAnnotationInternal(const UObjectBase* Object, T&& Annotation)
 	{
 		FScopeLock InverseAnnotationMapLock(&InverseAnnotationMapCritical);
 		if (Annotation.IsDefault())
@@ -320,9 +328,29 @@ public:
 			// should not exist in the mapping; we require uniqueness
 			int32 NumRemoved = InverseAnnotationMap.Remove(Annotation);
 			checkSlow(NumRemoved == 0);
-			InverseAnnotationMap.Add(Annotation, Object);
+			InverseAnnotationMap.Add(Forward<T>(Annotation), Object);
 		}
 	}
+
+
+public:
+
+	/**
+	 * Add an annotation to the annotation list. If the Annotation is the default, then the annotation is removed.
+	 *
+	 * @param Object        Object to annotate.
+	 * @param Annotation    Annotation to associate with Object.
+	 */
+	void AddAnnotation(const UObjectBase* Object, const TAnnotation& Annotation)
+	{
+		AddAnnotationInternal(Annotation)
+	}
+
+	void AddAnnotation(const UObjectBase* Object, TAnnotation&& Annotation)
+	{
+		AddAnnotationInternal(MoveTemp(Annotation))
+	}
+
 	/**
 	 * Removes an annotation from the annotation list. 
 	 *
@@ -596,7 +624,8 @@ class FUObjectAnnotationChunked : public FUObjectArray::FUObjectDeleteListener
 	/**
 	* Adds a new annotation for the specified index
 	**/
-	void AddAnnotationInternal(int32 Index, TAnnotation Annotation)
+	template<typename T>
+	void AddAnnotationInternal(int32 Index, T&& Annotation)
 	{
 		check(Index >= 0);
 		if (Annotation.IsDefault())
@@ -617,7 +646,7 @@ class FUObjectAnnotationChunked : public FUObjectArray::FUObjectDeleteListener
 			}
 
 			TAnnotation& NewAnnotation = AllocateAnnotation(Index);
-			NewAnnotation = Annotation;
+			NewAnnotation = Forward<T>(Annotation);
 		}
 	}
 
@@ -642,32 +671,39 @@ public:
 	/**
 	 * Add an annotation to the annotation list. If the Annotation is the default, then the annotation is removed.
 	 *
-	 * @param Object		Object to annotate.
-	 * @param Annotation	Annotation to associate with Object.
+	 * @param Object        Object to annotate.
+	 * @param Annotation    Annotation to associate with Object.
 	 */
-	void AddAnnotation(const UObjectBase *Object, TAnnotation Annotation)
+	void AddAnnotation(const UObjectBase *Object, const TAnnotation& Annotation)
 	{
 		check(Object);
 		AddAnnotation(GUObjectArray.ObjectToIndex(Object), Annotation);
 	}
+
+	void AddAnnotation(const UObjectBase* Object, TAnnotation&& Annotation)
+	{
+		check(Object);
+		AddAnnotation(GUObjectArray.ObjectToIndex(Object), MoveTemp(Annotation));
+	}
+
 	/**
 	 * Add an annotation to the annotation list. If the Annotation is the default, then the annotation is removed.
 	 *
-	 * @param Index			Index of object to annotate.
-	 * @param Annotation	Annotation to associate with Object.
+	 * @param Index         Index of object to annotate.
+	 * @param Annotation    Annotation to associate with Object.
 	 */
-	void AddAnnotation(int32 Index, TAnnotation Annotation)
+	void AddAnnotation(int32 Index, const TAnnotation& Annotation)
 	{
 		FRWScopeLock AnnotationArrayLock(AnnotationArrayCritical, SLT_Write);
 		AddAnnotationInternal(Index, Annotation);
 	}
 
-	/**
-	 * Add an annotation to the annotation list. If the Annotation is the default, then the annotation is removed.
-	 *
-	 * @param Object		Object to annotate.
-	 * @param Annotation	Annotation to associate with Object.
-	 */
+	void AddAnnotation(int32 Index, TAnnotation&& Annotation)
+	{
+		FRWScopeLock AnnotationArrayLock(AnnotationArrayCritical, SLT_Write);
+		AddAnnotationInternal(Index, MoveTemp(Annotation));
+	}
+
 	TAnnotation& AddOrGetAnnotation(const UObjectBase *Object, TFunctionRef<TAnnotation()> NewAnnotationFn)
 	{
 		check(Object);
@@ -947,28 +983,43 @@ public:
 	/**
 	 * Add an annotation to the annotation list. If the Annotation is the default, then the annotation is removed.
 	 *
-	 * @param Object		Object to annotate.
-	 * @param Annotation	Annotation to associate with Object.
+	 * @param Object        Object to annotate.
+	 * @param Annotation    Annotation to associate with Object.
 	 */
-	void AddAnnotation(const UObjectBase *Object,TAnnotation Annotation)
+	void AddAnnotation(const UObjectBase* Object, const TAnnotation& Annotation)
+	{
+		check(Object);
+		AddAnnotation(GUObjectArray.ObjectToIndex(Object), Annotation);
+	}
+
+	void AddAnnotation(const UObjectBase* Object, TAnnotation&& Annotation)
 	{
 		check(Object);
-		AddAnnotation(GUObjectArray.ObjectToIndex(Object),Annotation);
+		AddAnnotation(GUObjectArray.ObjectToIndex(Object), MoveTemp(Annotation));
 	}
+
 	/**
 	 * Add an annotation to the annotation list. If the Annotation is the default, then the annotation is removed.
 	 *
-	 * @param Index			Index of object to annotate.
-	 * @param Annotation	Annotation to associate with Object.
+	 * @param Index         Index of object to annotate.
+	 * @param Annotation    Annotation to associate with Object.
 	 */
-	void AddAnnotation(int32 Index, TAnnotation Annotation)
+	void AddAnnotation(int32 Index, const TAnnotation& Annotation)
 	{
 		FRWScopeLock AnnotationArrayLock(AnnotationArrayCritical, SLT_Write);
 		AddAnnotationInternal(Index, Annotation);
 	}
 
+	void AddAnnotation(int32 Index, TAnnotation&& Annotation)
+	{
+		FRWScopeLock AnnotationArrayLock(AnnotationArrayCritical, SLT_Write);
+		AddAnnotationInternal(Index, MoveTemp(Annotation));
+	}
+
+
 private:
-	void AddAnnotationInternal(int32 Index,TAnnotation Annotation)
+	template<typename T>
+	void AddAnnotationInternal(int32 Index, T&& Annotation)
 	{
 		check(Index >= 0);
 		if (Annotation.IsDefault())
@@ -996,7 +1047,7 @@ private:
 					new (AnnotationArray.GetData() + Start++) TAnnotation();
 				}
 			}
-			AnnotationArray[Index] = Annotation;
+			AnnotationArray[Index] = Forward<T>(Annotation);
 		}
 	}
 
diff --git a/Engine/Source/Runtime/Engine/Private/Components/LightComponent.cpp b/Engine/Source/Runtime/Engine/Private/Components/LightComponent.cpp
--- a/Engine/Source/Runtime/Engine/Private/Components/LightComponent.cpp
+++ b/Engine/Source/Runtime/Engine/Private/Components/LightComponent.cpp
@@ -526,7 +526,7 @@ void ULightComponent::Serialize(FArchive& Ar)
 			FLightComponentLegacyMapBuildData LegacyLightData;
 			LegacyLightData.Id = LightGuid;
 			LegacyLightData.Data = LegacyData;
-			GLightComponentsWithLegacyBuildData.AddAnnotation(this, LegacyLightData);
+			GLightComponentsWithLegacyBuildData.AddAnnotation(this, MoveTemp(LegacyLightData));
 		}
 	}
 
diff --git a/Engine/Source/Runtime/Engine/Private/Components/ModelComponent.cpp b/Engine/Source/Runtime/Engine/Private/Components/ModelComponent.cpp
--- a/Engine/Source/Runtime/Engine/Private/Components/ModelComponent.cpp
+++ b/Engine/Source/Runtime/Engine/Private/Components/ModelComponent.cpp
@@ -270,7 +270,7 @@ void UModelComponent::Serialize(FArchive& Ar)
 
 		if (LegacyComponentData.Data.Num() > 0)
 		{
-			GComponentsWithLegacyLightmaps.AddAnnotation(this, LegacyComponentData);
+			GComponentsWithLegacyLightmaps.AddAnnotation(this, MoveTemp(LegacyComponentData));
 		}
 	}
 
diff --git a/Engine/Source/Runtime/Engine/Private/Components/ReflectionCaptureComponent.cpp b/Engine/Source/Runtime/Engine/Private/Components/ReflectionCaptureComponent.cpp
--- a/Engine/Source/Runtime/Engine/Private/Components/ReflectionCaptureComponent.cpp
+++ b/Engine/Source/Runtime/Engine/Private/Components/ReflectionCaptureComponent.cpp
@@ -906,7 +906,7 @@ void UReflectionCaptureComponent::SerializeLegacyData(FArchive& Ar)
 					FReflectionCaptureMapBuildLegacyData LegacyComponentData;
 					LegacyComponentData.Id = MapBuildDataId;
 					LegacyComponentData.MapBuildData = LegacyMapBuildData;
-					GReflectionCapturesWithLegacyBuildData.AddAnnotation(this, LegacyComponentData);
+					GReflectionCapturesWithLegacyBuildData.AddAnnotation(this, MoveTemp(LegacyComponentData));
 				}
 			}
 		}
diff --git a/Engine/Source/Runtime/Engine/Private/Components/StaticMeshComponent.cpp b/Engine/Source/Runtime/Engine/Private/Components/StaticMeshComponent.cpp
--- a/Engine/Source/Runtime/Engine/Private/Components/StaticMeshComponent.cpp
+++ b/Engine/Source/Runtime/Engine/Private/Components/StaticMeshComponent.cpp
@@ -338,7 +338,7 @@ void UStaticMeshComponent::Serialize(FArchive& Ar)
 			}
 		}
 
-		GComponentsWithLegacyLightmaps.AddAnnotation(this, LegacyComponentData);
+		GComponentsWithLegacyLightmaps.AddAnnotation(this, MoveTemp(LegacyComponentData));
 	}
 
 	if (Ar.UE4Ver() < VER_UE4_AUTO_WELDING)
diff --git a/Engine/Source/Runtime/Engine/Private/Level.cpp b/Engine/Source/Runtime/Engine/Private/Level.cpp
--- a/Engine/Source/Runtime/Engine/Private/Level.cpp
+++ b/Engine/Source/Runtime/Engine/Private/Level.cpp
@@ -472,7 +472,7 @@ void ULevel::Serialize( FArchive& Ar )
 		FLevelLegacyMapBuildData LegacyLevelData;
 		LegacyLevelData.Id = LevelBuildDataId;
 		LegacyLevelData.Data = LegacyData;
-		GLevelsWithLegacyBuildData.AddAnnotation(this, LegacyLevelData);
+		GLevelsWithLegacyBuildData.AddAnnotation(this, MoveTemp(LegacyLevelData));
 	}
 
 	Ar << PrecomputedVisibilityHandler;
diff --git a/Engine/Source/Runtime/Engine/Private/LocalPlayer.cpp b/Engine/Source/Runtime/Engine/Private/LocalPlayer.cpp
--- a/Engine/Source/Runtime/Engine/Private/LocalPlayer.cpp
+++ b/Engine/Source/Runtime/Engine/Private/LocalPlayer.cpp
@@ -546,11 +546,12 @@ public:
 		if (bShouldLockView)
 		{
 			PlayerState.bLocked = true;
-			PlayerStates.AddAnnotation(Player, PlayerState);
 
 			// Also copy to the clipboard.
 			FString ViewPointString = ViewPointToString(PlayerState.ViewPoint);
 			FPlatformApplicationMisc::ClipboardCopy(*ViewPointString);
+
+			PlayerStates.AddAnnotation(Player, MoveTemp(PlayerState));
 		}
 
 		if (bPrintHelp)
diff --git a/Engine/Source/Runtime/Engine/Private/Materials/Material.cpp b/Engine/Source/Runtime/Engine/Private/Materials/Material.cpp
--- a/Engine/Source/Runtime/Engine/Private/Materials/Material.cpp
+++ b/Engine/Source/Runtime/Engine/Private/Materials/Material.cpp
@@ -1392,7 +1392,7 @@ void UMaterial::MarkUsageFlagDirty(EMaterialUsage Usage, bool CurrentValue, bool
 	{
 		FMaterialsWithDirtyUsageFlags Annotation = GMaterialsWithDirtyUsageFlags.GetAnnotation(this);
 		Annotation.MarkUsageFlagDirty(Usage);
-		GMaterialsWithDirtyUsageFlags.AddAnnotation(this,Annotation);
+		GMaterialsWithDirtyUsageFlags.AddAnnotation(this, MoveTemp(Annotation));
 	}
 #endif
 }
diff --git a/Engine/Source/Runtime/Engine/Private/ObjectTrace.cpp b/Engine/Source/Runtime/Engine/Private/ObjectTrace.cpp
--- a/Engine/Source/Runtime/Engine/Private/ObjectTrace.cpp
+++ b/Engine/Source/Runtime/Engine/Private/ObjectTrace.cpp
@@ -99,7 +99,7 @@ uint64 FObjectTrace::GetObjectId(const UObject* InObject)
 		if(Annotation.Id == 0)
 		{
 			Annotation.Id = CurrentId++;
-			GObjectTraceAnnotations.AddAnnotation(InObjectInner, Annotation);
+			GObjectTraceAnnotations.AddAnnotation(InObjectInner, MoveTemp(Annotation));
 		}
 
 		return Annotation.Id;
@@ -136,7 +136,7 @@ void FObjectTrace::OutputClass(const UClass* InClass)
 	}
 
 	Annotation.bTraced = true;
-	GObjectTraceAnnotations.AddAnnotation(InClass, Annotation);
+	GObjectTraceAnnotations.AddAnnotation(InClass, MoveTemp(Annotation));
 
 	int32 ClassNameStringLength = InClass->GetFName().GetStringLength() + 1;
 	int32 ClassFullNameStringLength = InClass->GetPathName().Len() + 1;
@@ -175,7 +175,7 @@ void FObjectTrace::OutputObject(const UObject* InObject)
 	}
 
 	Annotation.bTraced = true;
-	GObjectTraceAnnotations.AddAnnotation(InObject, Annotation);
+	GObjectTraceAnnotations.AddAnnotation(InObject, MoveTemp(Annotation));
 
 	if(!ShouldTraceObjectsWorld(InObject))
 	{
diff --git a/Engine/Source/Runtime/Landscape/Private/Landscape.cpp b/Engine/Source/Runtime/Landscape/Private/Landscape.cpp
--- a/Engine/Source/Runtime/Landscape/Private/Landscape.cpp
+++ b/Engine/Source/Runtime/Landscape/Private/Landscape.cpp
@@ -447,7 +447,7 @@ void ULandscapeComponent::Serialize(FArchive& Ar)
 
 		FMeshMapBuildLegacyData LegacyComponentData;
 		LegacyComponentData.Data.Emplace(MapBuildDataId, LegacyMapBuildData);
-		GComponentsWithLegacyLightmaps.AddAnnotation(this, LegacyComponentData);
+		GComponentsWithLegacyLightmaps.AddAnnotation(this, MoveTemp(LegacyComponentData));
 	}
 
 	if (Ar.IsLoading() && Ar.CustomVer(FFortniteMainBranchObjectVersion::GUID) < FFortniteMainBranchObjectVersion::NewLandscapeMaterialPerLOD)
diff --git a/Engine/Source/Runtime/MovieScene/Private/MovieSceneSpawnableAnnotation.cpp b/Engine/Source/Runtime/MovieScene/Private/MovieSceneSpawnableAnnotation.cpp
--- a/Engine/Source/Runtime/MovieScene/Private/MovieSceneSpawnableAnnotation.cpp
+++ b/Engine/Source/Runtime/MovieScene/Private/MovieSceneSpawnableAnnotation.cpp
@@ -17,7 +17,7 @@ void FMovieSceneSpawnableAnnotation::Add(UObject* SpawnedObject, const FGuid& Ob
 		Annotation.ObjectBindingID = ObjectBindingID;
 		Annotation.OriginatingSequence = InOriginatingSequence;
 
-		SpawnedObjectAnnotation.AddAnnotation(SpawnedObject, Annotation);
+		SpawnedObjectAnnotation.AddAnnotation(SpawnedObject, MoveTemp(Annotation));
 	}
 }
 
