Fix compile errors
#rnx

diff --git a/Engine/Source/Runtime/CoreUObject/Private/UObject/UObjectMarks.cpp b/Engine/Source/Runtime/CoreUObject/Private/UObject/UObjectMarks.cpp
--- a/Engine/Source/Runtime/CoreUObject/Private/UObject/UObjectMarks.cpp
+++ b/Engine/Source/Runtime/CoreUObject/Private/UObject/UObjectMarks.cpp
@@ -85,7 +85,6 @@ public:
 	void AddAnnotation(const UObjectBase* Object, const TAnnotation& Annotation)
 	{
 		AddAnnotationInternal(Object, Annotation);
-		}
 	}
 
 	void RemoveAnnotation(const UObjectBase* Object)
diff --git a/Engine/Source/Runtime/CoreUObject/Public/UObject/LazyObjectPtr.h b/Engine/Source/Runtime/CoreUObject/Public/UObject/LazyObjectPtr.h
--- a/Engine/Source/Runtime/CoreUObject/Public/UObject/LazyObjectPtr.h
+++ b/Engine/Source/Runtime/CoreUObject/Public/UObject/LazyObjectPtr.h
@@ -67,7 +67,7 @@ struct COREUOBJECT_API FUniqueObjectGuid
 	}
 
 	/** Returns true is this is the default value */
-	FORCEINLINE bool IsDefault()
+	FORCEINLINE bool IsDefault() const
 	{
 		// A default GUID is 0,0,0,0 and this is "invalid"
 		return !IsValid(); 
diff --git a/Engine/Source/Runtime/CoreUObject/Public/UObject/UObjectAnnotation.h b/Engine/Source/Runtime/CoreUObject/Public/UObject/UObjectAnnotation.h
--- a/Engine/Source/Runtime/CoreUObject/Public/UObject/UObjectAnnotation.h
+++ b/Engine/Source/Runtime/CoreUObject/Public/UObject/UObjectAnnotation.h
@@ -100,6 +100,7 @@ private:
 				}
 			}
 			AnnotationMap.Add(AnnotationCacheKey, AnnotationCacheValue);
+		}
 	}
 
 public:
@@ -310,7 +311,7 @@ public:
 	}
 
 private:
-	template<T> 
+	template<typename T> 
 	void AddAnnotationInternal(const UObjectBase* Object, T&& Annotation)
 	{
 		FScopeLock InverseAnnotationMapLock(&InverseAnnotationMapCritical);
@@ -343,12 +344,12 @@ public:
 	 */
 	void AddAnnotation(const UObjectBase* Object, const TAnnotation& Annotation)
 	{
-		AddAnnotationInternal(Annotation)
+		AddAnnotationInternal(Object, Annotation);
 	}
 
 	void AddAnnotation(const UObjectBase* Object, TAnnotation&& Annotation)
 	{
-		AddAnnotationInternal(MoveTemp(Annotation))
+		AddAnnotationInternal(Object, MoveTemp(Annotation));
 	}
 
 	/**
