New NDC line which reports the name of the class of the asset containing the differences.

diff --git a/Engine/Source/Runtime/CoreUObject/Private/Serialization/ArchiveStackTrace.cpp b/Engine/Source/Runtime/CoreUObject/Private/Serialization/ArchiveStackTrace.cpp
--- a/Engine/Source/Runtime/CoreUObject/Private/Serialization/ArchiveStackTrace.cpp
+++ b/Engine/Source/Runtime/CoreUObject/Private/Serialization/ArchiveStackTrace.cpp
@@ -705,6 +705,7 @@ void FArchiveStackTrace::CompareWith(const TCHAR* InFilename, const int64 TotalH
 		DestPackage.StartOffset = 0;
 
 		UE_LOG(LogArchiveDiff, Display, TEXT("Comparing: %s"), *GetArchiveName());
+		UE_LOG(LogArchiveDiff, Display, TEXT("Asset class: %s"), *AssetClass.ToString());
 
 		int32 NumLoggedDiffs = 0;
 
