Fixing a potential crash when destroying GC clusters

diff --git a/Engine/Source/Runtime/CoreUObject/Private/UObject/GarbageCollection.cpp b/Engine/Source/Runtime/CoreUObject/Private/UObject/GarbageCollection.cpp
--- a/Engine/Source/Runtime/CoreUObject/Private/UObject/GarbageCollection.cpp
+++ b/Engine/Source/Runtime/CoreUObject/Private/UObject/GarbageCollection.cpp
@@ -1111,8 +1111,13 @@ public:
 			ClustersToDissolveList.PopAll(ClustersToDissolve);
 			for (FUObjectItem* ObjectItem : ClustersToDissolve)
 			{
-				GUObjectClusters.DissolveClusterAndMarkObjectsAsUnreachable(ObjectItem);
-				GUObjectClusters.SetClustersNeedDissolving();
+				// Check if the object is still a cluster root - it's possible one of the previous
+				// DissolveClusterAndMarkObjectsAsUnreachable calls already dissolved its cluster
+				if (ObjectItem->HasAnyFlags(EInternalObjectFlags::ClusterRoot))
+				{
+					GUObjectClusters.DissolveClusterAndMarkObjectsAsUnreachable(ObjectItem);
+					GUObjectClusters.SetClustersNeedDissolving();
+				}
 			}
 		}
 
