Optimizations to tagged property serialization and numeric property conversion

Makes tagged property serialization about 20% faster during initial loading into game of an internal project on Test Client Win64 using TBB malloc, 600ms -> 500ms.

* Compiler failed to compile out CLOG condition when converting from and to same numeric type, likely since a constant variable was passed by const reference. Added specialization for non-converting case.
* Avoid excessive Ename -> FName lookup during numeric conversion
* Inline FPropertyTag default constructor to avoid call overhead
* Avoid attempting type conversions when not needed.
* Cache conditions
* Use slightly faster FName::IsNone() method
* Removed paranoid / useless FName::IsValid() check

diff --git a/Engine/Source/Runtime/CoreUObject/Private/UObject/Class.cpp b/Engine/Source/Runtime/CoreUObject/Private/UObject/Class.cpp
--- a/Engine/Source/Runtime/CoreUObject/Private/UObject/Class.cpp
+++ b/Engine/Source/Runtime/CoreUObject/Private/UObject/Class.cpp
@@ -977,6 +977,7 @@ void UStruct::SerializeTaggedProperties(FStructuredArchive::FSlot Slot, uint8* D
 
 	// Determine if this struct supports optional property guid's (UBlueprintGeneratedClasses Only)
 	const bool bArePropertyGuidsAvailable = (UnderlyingArchive.UE4Ver() >= VER_UE4_PROPERTY_GUID_IN_PROPERTY_TAG) && !FPlatformProperties::RequiresCookedData() && ArePropertyGuidsAvailable();
+	const bool bUseRedirects = !FPlatformProperties::RequiresCookedData() || UnderlyingArchive.IsSaveGame();
 
 	if( UnderlyingArchive.IsLoading() )
 	{
@@ -998,15 +999,10 @@ void UStruct::SerializeTaggedProperties(FStructuredArchive::FSlot Slot, uint8* D
 			FPropertyTag Tag;
 			PropertyRecord << SA_VALUE(TEXT("Tag"), Tag);
 
-			if( Tag.Name == NAME_None )
+			if (Tag.Name.IsNone())
 			{
 				break;
 			}
-			if (!Tag.Name.IsValid())
-			{
-				UE_LOG(LogClass, Warning, TEXT("Invalid tag name: struct '%s', archive '%s'"), *GetName(), *UnderlyingArchive.GetArchiveName());
-				break;
-			}
 
 			// Move to the next property to be serialized
 			if( bAdvanceProperty && --RemainingArrayDim <= 0 )
@@ -1035,26 +1031,15 @@ void UStruct::SerializeTaggedProperties(FStructuredArchive::FSlot Slot, uint8* D
 			if( Property == nullptr || Property->GetFName() != Tag.Name )
 			{
 				// No need to check redirects on platforms where everything is cooked. Always check for save games
-				if ((!FPlatformProperties::RequiresCookedData() || UnderlyingArchive.IsSaveGame()) && !UnderlyingArchive.HasAnyPortFlags(PPF_DuplicateForPIE|PPF_Duplicate))
+				if (bUseRedirects && !UnderlyingArchive.HasAnyPortFlags(PPF_DuplicateForPIE|PPF_Duplicate))
 				{
-					FName EachName = GetFName();
-					FName PackageName = GetOutermost()->GetFName();
-					// Search the current class first, then work up the class hierarchy to see if theres a match for our fixup.
-					UStruct* Owner = GetOwnerStruct();
-					if( Owner )
+					for (UStruct* CheckStruct = GetOwnerStruct(); CheckStruct; CheckStruct = CheckStruct->GetSuperStruct())
 					{
-						UStruct* CheckStruct = Owner;
-						while(CheckStruct)
+						FName NewTagName = UProperty::FindRedirectedPropertyName(CheckStruct, Tag.Name);
+						if (!NewTagName.IsNone())
 						{
-							FName NewTagName = UProperty::FindRedirectedPropertyName(CheckStruct, Tag.Name);
-
-							if (NewTagName != NAME_None)
-							{
-								Tag.Name = NewTagName;
-								break;
-							}
-
-							CheckStruct = CheckStruct->GetSuperStruct();
+							Tag.Name = NewTagName;
+							break;
 						}
 					}
 				}
@@ -1089,93 +1074,92 @@ void UStruct::SerializeTaggedProperties(FStructuredArchive::FSlot Slot, uint8* D
 				RemainingArrayDim = Property ? Property->ArrayDim : 0;
 			}
 
+			const int64 StartOfProperty = UnderlyingArchive.Tell();
+
 			if (!Property)
 			{
 				Property = CustomFindProperty(Tag.Name);
 			}
 
-			FName PropID = Property ? Property->GetID() : NAME_None;
-			FName ArrayInnerID = NAME_None;
-
-			// Check if this is a struct property and we have a redirector
-			// No need to check redirects on platforms where everything is cooked. Always check for save games
-			if (!FPlatformProperties::RequiresCookedData() || UnderlyingArchive.IsSaveGame())
+			if (Property)
 			{
-				if (Tag.Type == NAME_StructProperty && PropID == NAME_StructProperty)
+				FName PropID = Property->GetID();
+
+				// Check if this is a struct property and we have a redirector
+				// No need to check redirects on platforms where everything is cooked. Always check for save games
+				if (bUseRedirects)
 				{
-					const FName NewName = FLinkerLoad::FindNewNameForStruct(Tag.StructName);
-					const FName StructName = CastChecked<UStructProperty>(Property)->Struct->GetFName();
-					if (NewName == StructName)
+					if (Tag.Type == NAME_StructProperty && PropID == NAME_StructProperty)
 					{
-						Tag.StructName = NewName;
+						const FName NewName = FLinkerLoad::FindNewNameForStruct(Tag.StructName);
+						const FName StructName = CastChecked<UStructProperty>(Property)->Struct->GetFName();
+						if (NewName == StructName)
+						{
+							Tag.StructName = NewName;
+						}
 					}
-				}
-				else if ((PropID == NAME_EnumProperty) && ((Tag.Type == NAME_EnumProperty) || (Tag.Type == NAME_ByteProperty)))
-				{
-					const FName NewName = FLinkerLoad::FindNewNameForEnum(Tag.EnumName);
-					if (!NewName.IsNone())
+					else if ((PropID == NAME_EnumProperty) && ((Tag.Type == NAME_EnumProperty) || (Tag.Type == NAME_ByteProperty)))
 					{
-						Tag.EnumName = NewName;
+						const FName NewName = FLinkerLoad::FindNewNameForEnum(Tag.EnumName);
+						if (!NewName.IsNone())
+						{
+							Tag.EnumName = NewName;
+						}
 					}
 				}
-			}
 
-			const int64 StartOfProperty = UnderlyingArchive.Tell();
-			if( !Property )
-			{
-				//UE_LOG(LogClass, Warning, TEXT("Property %s of %s not found for package:  %s"), *Tag.Name.ToString(), *GetFullName(), *UnderlyingArchive.GetArchiveName() );
-			}
 #if WITH_EDITOR
-			else if (BreakRecursionIfFullyLoad && BreakRecursionIfFullyLoad->HasAllFlags(RF_LoadCompleted))
-			{
-			}
+				if (BreakRecursionIfFullyLoad && BreakRecursionIfFullyLoad->HasAllFlags(RF_LoadCompleted))
+				{
+				}
 #endif // WITH_EDITOR
-			// editoronly properties should be skipped if we are NOT the editor, or we are 
-			// the editor but are cooking for console (editoronly implies notforconsole)
-			else if ((Property->PropertyFlags & CPF_EditorOnly) && !FPlatformProperties::HasEditorOnlyData() && !GForceLoadEditorOnly)
-			{
-			}
-			// check for valid array index
-			else if( Tag.ArrayIndex >= Property->ArrayDim || Tag.ArrayIndex < 0 )
-			{
-				UE_LOG(LogClass, Warning, TEXT("Array bound exceeded (var %s=%d, exceeds %s [0-%d] in package:  %s"),
-					*Tag.Name.ToString(), Tag.ArrayIndex, *GetName(), Property->ArrayDim-1, *UnderlyingArchive.GetArchiveName());
-			}
-			else if( !Property->ShouldSerializeValue(UnderlyingArchive) )
-			{
-				UE_CLOG((UnderlyingArchive.IsPersistent() && FPlatformProperties::RequiresCookedData()), LogClass, Warning, TEXT("Skipping saved property %s of %s since it is no longer serializable for asset:  %s. (Maybe resave asset?)"), *Tag.Name.ToString(), *GetName(), *UnderlyingArchive.GetArchiveName() );
-			}
-			else
-			{
-				FStructuredArchive::FSlot ValueSlot = PropertyRecord.EnterField(SA_FIELD_NAME(TEXT("Value")));
-
-				switch (Property->ConvertFromType(Tag, ValueSlot, Data, DefaultsStruct))
+				// editoronly properties should be skipped if we are NOT the editor, or we are 
+				// the editor but are cooking for console (editoronly implies notforconsole)
+				else if ((Property->PropertyFlags & CPF_EditorOnly) && !FPlatformProperties::HasEditorOnlyData() && !GForceLoadEditorOnly)
 				{
-					case EConvertFromTypeResult::Converted:
-						bAdvanceProperty = true;
-						break;
+				}
+				// check for valid array index
+				else if (Tag.ArrayIndex >= Property->ArrayDim || Tag.ArrayIndex < 0)
+				{
+					UE_LOG(LogClass, Warning, TEXT("Array bound exceeded (var %s=%d, exceeds %s [0-%d] in package:  %s"),
+						*Tag.Name.ToString(), Tag.ArrayIndex, *GetName(), Property->ArrayDim - 1, *UnderlyingArchive.GetArchiveName());
+				}
+				else if (!Property->ShouldSerializeValue(UnderlyingArchive))
+				{
+					UE_CLOG((UnderlyingArchive.IsPersistent() && FPlatformProperties::RequiresCookedData()), LogClass, Warning, TEXT("Skipping saved property %s of %s since it is no longer serializable for asset:  %s. (Maybe resave asset?)"), *Tag.Name.ToString(), *GetName(), *UnderlyingArchive.GetArchiveName());
+				}
+				else
+				{
+					FStructuredArchive::FSlot ValueSlot = PropertyRecord.EnterField(SA_FIELD_NAME(TEXT("Value")));
 
-					case EConvertFromTypeResult::UseSerializeItem:
-						if (Tag.Type != PropID)
-						{
-							UE_LOG(LogClass, Warning, TEXT("Type mismatch in %s of %s - Previous (%s) Current(%s) for package:  %s"), *Tag.Name.ToString(), *GetName(), *Tag.Type.ToString(), *PropID.ToString(), *UnderlyingArchive.GetArchiveName() );
-						}
-						else
+					if (Tag.Type == PropID)
+					{
+						uint8* DestAddress = Property->ContainerPtrToValuePtr<uint8>(Data, Tag.ArrayIndex);
+						uint8* DefaultsFromParent = Property->ContainerPtrToValuePtrForDefaults<uint8>(DefaultsStruct, Defaults, Tag.ArrayIndex);
+
+						// This property is ok.
+						Tag.SerializeTaggedProperty(ValueSlot, Property, DestAddress, DefaultsFromParent);
+						bAdvanceProperty = !UnderlyingArchive.IsCriticalError();
+					}
+					else
+					{
+						switch (Property->ConvertFromType(Tag, ValueSlot, Data, DefaultsStruct))
 						{
-							uint8* DestAddress = Property->ContainerPtrToValuePtr<uint8>(Data, Tag.ArrayIndex);
-							uint8* DefaultsFromParent = Property->ContainerPtrToValuePtrForDefaults<uint8>(DefaultsStruct, Defaults, Tag.ArrayIndex);
+						case EConvertFromTypeResult::Converted:
+							bAdvanceProperty = true;
+							break;
 
-							// This property is ok.
-							Tag.SerializeTaggedProperty(ValueSlot, Property, DestAddress, DefaultsFromParent);
-							bAdvanceProperty = !UnderlyingArchive.IsCriticalError();
-						}
-						break;
+						case EConvertFromTypeResult::UseSerializeItem:
+							UE_LOG(LogClass, Warning, TEXT("Type mismatch in %s of %s - Previous (%s) Current(%s) for package:  %s"), *Tag.Name.ToString(), *GetName(), *Tag.Type.ToString(), *PropID.ToString(), *UnderlyingArchive.GetArchiveName());
+							break;
 
-					case EConvertFromTypeResult::CannotConvert:
-						break;
+						case EConvertFromTypeResult::CannotConvert:
+							break;
 
-					default:
-						check(false);
+						default:
+							check(false);
+						}
+					}
 				}
 			}
 
diff --git a/Engine/Source/Runtime/CoreUObject/Private/UObject/PropertyTag.cpp b/Engine/Source/Runtime/CoreUObject/Private/UObject/PropertyTag.cpp
--- a/Engine/Source/Runtime/CoreUObject/Private/UObject/PropertyTag.cpp
+++ b/Engine/Source/Runtime/CoreUObject/Private/UObject/PropertyTag.cpp
@@ -12,36 +12,11 @@
 FPropertyTag
 -----------------------------------------------------------------------------*/
 
-// Constructors.
-FPropertyTag::FPropertyTag()
-	: Prop      (nullptr)
-	, Type      (NAME_None)
-	, BoolVal   (0)
-	, Name      (NAME_None)
-	, StructName(NAME_None)
-	, EnumName  (NAME_None)
-	, InnerType (NAME_None)
-	, ValueType	(NAME_None)
-	, Size      (0)
-	, ArrayIndex(INDEX_NONE)
-	, SizeOffset(INDEX_NONE)
-	, HasPropertyGuid(0)
-{
-}
-
 FPropertyTag::FPropertyTag( FArchive& InSaveAr, UProperty* Property, int32 InIndex, uint8* Value, uint8* Defaults )
 	: Prop      (Property)
 	, Type      (Property->GetID())
-	, BoolVal   (0)
 	, Name      (Property->GetFName())
-	, StructName(NAME_None)
-	, EnumName	(NAME_None)
-	, InnerType	(NAME_None)
-	, ValueType	(NAME_None)
-	, Size		(0)
 	, ArrayIndex(InIndex)
-	, SizeOffset(INDEX_NONE)
-	, HasPropertyGuid(0)
 {
 	if (Property)
 	{
diff --git a/Engine/Source/Runtime/CoreUObject/Public/UObject/PropertyTag.h b/Engine/Source/Runtime/CoreUObject/Public/UObject/PropertyTag.h
--- a/Engine/Source/Runtime/CoreUObject/Public/UObject/PropertyTag.h
+++ b/Engine/Source/Runtime/CoreUObject/Public/UObject/PropertyTag.h
@@ -15,25 +15,25 @@
 struct FPropertyTag
 {
 	// Transient.
-	UProperty* Prop;
+	UProperty* Prop = nullptr;
 
 	// Variables.
 	FName	Type;		// Type of property
-	uint8	BoolVal;	// a boolean property's value (never need to serialize data for bool properties except here)
+	uint8	BoolVal = 0;// a boolean property's value (never need to serialize data for bool properties except here)
 	FName	Name;		// Name of property.
 	FName	StructName;	// Struct name if UStructProperty.
 	FName	EnumName;	// Enum name if UByteProperty or UEnumProperty
 	FName	InnerType;	// Inner type if UArrayProperty, USetProperty, or UMapProperty
 	FName	ValueType;	// Value type if UMapPropery
-	int32	Size;       // Property size.
-	int32	ArrayIndex;	// Index if an array; else 0.
-	int64	SizeOffset;	// location in stream of tag size member
+	int32	Size = 0;   // Property size.
+	int32	ArrayIndex = INDEX_NONE; // Index if an array; else 0.
+	int64	SizeOffset = INDEX_NONE; // location in stream of tag size member
 	FGuid	StructGuid;
-	uint8	HasPropertyGuid;
+	uint8	HasPropertyGuid = 0;
 	FGuid	PropertyGuid;
 
 	// Constructors.
-	FPropertyTag();
+	FPropertyTag() {}
 	FPropertyTag( FArchive& InSaveAr, UProperty* Property, int32 InIndex, uint8* Value, uint8* Defaults );
 
 	// Set optional property guid
diff --git a/Engine/Source/Runtime/CoreUObject/Public/UObject/UnrealType.h b/Engine/Source/Runtime/CoreUObject/Public/UObject/UnrealType.h
--- a/Engine/Source/Runtime/CoreUObject/Public/UObject/UnrealType.h
+++ b/Engine/Source/Runtime/CoreUObject/Public/UObject/UnrealType.h
@@ -1284,19 +1284,26 @@ public:
 
 protected:
 	template <typename OldNumericType>
-	void ConvertFromArithmeticValue(FStructuredArchive::FSlot Slot, void* Obj, const FPropertyTag& Tag)
+	FORCEINLINE void ConvertFromArithmeticValue(FStructuredArchive::FSlot Slot, void* Obj, const FPropertyTag& Tag)
+	{
+		ConvertFromArithmeticValueImpl<OldNumericType>(*this, Slot, Obj, Tag);
+	}
+
+private:
+	template <typename OldNumericType>
+	static void ConvertFromArithmeticValueImpl(const TProperty_Numeric& This, FStructuredArchive::FSlot Slot, void* Obj, const FPropertyTag& Tag)
 	{
 		OldNumericType OldValue;
 		Slot << OldValue;
 		TCppType NewValue = (TCppType)OldValue;
-		this->SetPropertyValue_InContainer(Obj, NewValue, Tag.ArrayIndex);
+		This.SetPropertyValue_InContainer(Obj, NewValue, Tag.ArrayIndex);
 
 		UE_CLOG(
 			((TIsSigned<OldNumericType>::Value || TIsFloatingPoint<OldNumericType>::Value) && (!TIsSigned<TCppType>::Value && !TIsFloatingPoint<TCppType>::Value) && OldValue < 0) || ((OldNumericType)NewValue != OldValue),
 			LogClass,
 			Warning,
 			TEXT("Potential data loss during conversion of integer property %s of %s - was (%s) now (%s) - for package: %s"),
-			*this->GetName(),
+			*This.GetName(),
 			*Slot.GetUnderlyingArchive().GetArchiveName(),
 			*LexToString(OldValue),
 			*LexToString(NewValue),
@@ -1304,69 +1311,82 @@ protected:
 			);
 	}
 
+	template <>
+	FORCEINLINE static void ConvertFromArithmeticValueImpl<TCppType>(const TProperty_Numeric& This, FStructuredArchive::FSlot Slot, void* Obj, const FPropertyTag& Tag)
+	{
+		TCppType Value;
+		Slot << Value;
+		This.SetPropertyValue_InContainer(Obj, Value, Tag.ArrayIndex);
+	}
 public:
 	virtual EConvertFromTypeResult ConvertFromType(const FPropertyTag& Tag, FStructuredArchive::FSlot Slot, uint8* Data, UStruct* DefaultsStruct) override
 	{
-		if (Tag.Type == NAME_Int8Property)
-		{
-			ConvertFromArithmeticValue<int8>(Slot, Data, Tag);
-		}
-		else if (Tag.Type == NAME_Int16Property)
-		{
-			ConvertFromArithmeticValue<int16>(Slot, Data, Tag);
-		}
-		else if (Tag.Type == NAME_IntProperty)
-		{
-			ConvertFromArithmeticValue<int32>(Slot, Data, Tag);
-		}
-		else if (Tag.Type == NAME_Int64Property)
+		if (const EName * TagType = Tag.Type.ToEName())
 		{
-			ConvertFromArithmeticValue<int64>(Slot, Data, Tag);
-		}
-		else if (Tag.Type == NAME_ByteProperty)
-		{
-			if (Tag.EnumName != NAME_None)
+			switch (*TagType)
+			{
+			case NAME_Int8Property:
+				ConvertFromArithmeticValue<int8>(Slot, Data, Tag);
+				return EConvertFromTypeResult::Converted;
+
+			case NAME_Int16Property:
+				ConvertFromArithmeticValue<int16>(Slot, Data, Tag);
+				return EConvertFromTypeResult::Converted;
+
+			case NAME_IntProperty:
+				ConvertFromArithmeticValue<int32>(Slot, Data, Tag);
+				return EConvertFromTypeResult::Converted;
+
+			case NAME_Int64Property:
+				ConvertFromArithmeticValue<int64>(Slot, Data, Tag);
+				return EConvertFromTypeResult::Converted;
+
+			case NAME_ByteProperty:
+				if (!Tag.EnumName.IsNone())
+				{
+					int64 PreviousValue = this->ReadEnumAsInt64(Slot, DefaultsStruct, Tag);
+					this->SetPropertyValue_InContainer(Data, PreviousValue, Tag.ArrayIndex);
+				}
+				else
+				{
+					ConvertFromArithmeticValue<int8>(Slot, Data, Tag);
+				}
+				return EConvertFromTypeResult::Converted;
+
+			case NAME_EnumProperty:
 			{
 				int64 PreviousValue = this->ReadEnumAsInt64(Slot, DefaultsStruct, Tag);
-				this->SetPropertyValue_InContainer(Data, PreviousValue, Tag.ArrayIndex);
+				this->SetPropertyValue_InContainer(Data, (TCppType)PreviousValue, Tag.ArrayIndex);
+				return EConvertFromTypeResult::Converted;
 			}
-			else
-			{
-				ConvertFromArithmeticValue<int8>(Slot, Data, Tag);
+
+			case NAME_UInt16Property:
+				ConvertFromArithmeticValue<uint16>(Slot, Data, Tag);
+				return EConvertFromTypeResult::Converted;
+
+			case NAME_UInt32Property:
+				ConvertFromArithmeticValue<uint32>(Slot, Data, Tag);
+				return EConvertFromTypeResult::Converted;
+
+			case NAME_UInt64Property:
+				ConvertFromArithmeticValue<uint64>(Slot, Data, Tag);
+				return EConvertFromTypeResult::Converted;
+
+			case NAME_FloatProperty:
+				ConvertFromArithmeticValue<float>(Slot, Data, Tag);
+				return EConvertFromTypeResult::Converted;
+
+			case NAME_DoubleProperty:
+				ConvertFromArithmeticValue<double>(Slot, Data, Tag);
+				return EConvertFromTypeResult::Converted;
+
+			default:
+				// We didn't convert it
+				break;
 			}
 		}
-		else if (Tag.Type == NAME_EnumProperty)
-		{
-			int64 PreviousValue = this->ReadEnumAsInt64(Slot, DefaultsStruct, Tag);
-			this->SetPropertyValue_InContainer(Data, (TCppType)PreviousValue, Tag.ArrayIndex);
-		}
-		else if (Tag.Type == NAME_UInt16Property)
-		{
-			ConvertFromArithmeticValue<uint16>(Slot, Data, Tag);
-		}
-		else if (Tag.Type == NAME_UInt32Property)
-		{
-			ConvertFromArithmeticValue<uint32>(Slot, Data, Tag);
-		}
-		else if (Tag.Type == NAME_UInt64Property)
-		{
-			ConvertFromArithmeticValue<uint64>(Slot, Data, Tag);
-		}
-		else if (Tag.Type == NAME_FloatProperty)
-		{
-			ConvertFromArithmeticValue<float>(Slot, Data, Tag);
-		}
-		else if (Tag.Type == NAME_DoubleProperty)
-		{
-			ConvertFromArithmeticValue<double>(Slot, Data, Tag);
-		}
-		else
-		{
-			// We didn't convert it
-			return EConvertFromTypeResult::UseSerializeItem;
-		}
 
-		return EConvertFromTypeResult::Converted;
+		return EConvertFromTypeResult::UseSerializeItem;
 	}
 	// End of UProperty interface
 
