Don't clear the AsyncLoading flag before objects have been post loaded

diff --git a/Engine/Source/Runtime/CoreUObject/Private/Serialization/AsyncLoading.cpp b/Engine/Source/Runtime/CoreUObject/Private/Serialization/AsyncLoading.cpp
--- a/Engine/Source/Runtime/CoreUObject/Private/Serialization/AsyncLoading.cpp
+++ b/Engine/Source/Runtime/CoreUObject/Private/Serialization/AsyncLoading.cpp
@@ -7106,7 +7106,7 @@ EAsyncPackageState::Type FAsyncPackage::FinishObjects()
 		FScopeLock ReferencedObjectsLock(&ReferencedObjectsCritical);
 		for (UObject* Obj : ReferencedObjects)
 		{
-			if (Obj)
+			if (Obj && !Obj->HasAnyFlags(RF_NeedPostLoad | RF_NeedPostLoadSubobjects))
 			{
 				Obj->AtomicallyClearInternalFlags(EInternalObjectFlags::AsyncLoading);
 			}
