Fix for VS2019 strict CIS error in AsyncLoading.cpp

diff --git a/Engine/Source/Runtime/CoreUObject/Private/Serialization/AsyncLoading.cpp b/Engine/Source/Runtime/CoreUObject/Private/Serialization/AsyncLoading.cpp
--- a/Engine/Source/Runtime/CoreUObject/Private/Serialization/AsyncLoading.cpp
+++ b/Engine/Source/Runtime/CoreUObject/Private/Serialization/AsyncLoading.cpp
@@ -7178,7 +7178,7 @@ UPackage* FAsyncPackage::GetLoadedPackage()
 template <typename T>
 void ClearFlagsAndDissolveClustersFromLoadedObjects(T& LoadedObjects)
 {
-	const EObjectFlags ObjectLoadFlags = EObjectFlags(RF_NeedLoad | RF_NeedPostLoad | RF_NeedPostLoadSubobjects | RF_WasLoaded);
+	const EObjectFlags ObjectLoadFlags = RF_NeedLoad | RF_NeedPostLoad | RF_NeedPostLoadSubobjects | RF_WasLoaded;
 	for (UObject* ObjRef : LoadedObjects)
 	{
 		if (ObjRef)
