Fix clang link error

diff --git a/Engine/Source/Runtime/CoreUObject/Public/UObject/UnrealType.h b/Engine/Source/Runtime/CoreUObject/Public/UObject/UnrealType.h
--- a/Engine/Source/Runtime/CoreUObject/Public/UObject/UnrealType.h
+++ b/Engine/Source/Runtime/CoreUObject/Public/UObject/UnrealType.h
@@ -1312,7 +1312,7 @@ private:
 	}
 
 	template <>
-	FORCEINLINE static void ConvertFromArithmeticValueImpl<TCppType>(const TProperty_Numeric& This, FStructuredArchive::FSlot Slot, void* Obj, const FPropertyTag& Tag)
+	inline static void ConvertFromArithmeticValueImpl<TCppType>(const TProperty_Numeric& This, FStructuredArchive::FSlot Slot, void* Obj, const FPropertyTag& Tag)
 	{
 		TCppType Value;
 		Slot << Value;
