Fix ClangGame link error, removed superfluous COREUOBJECT_API and improved const correctness

diff --git a/Engine/Source/Runtime/CoreUObject/Public/UObject/UnrealType.h b/Engine/Source/Runtime/CoreUObject/Public/UObject/UnrealType.h
--- a/Engine/Source/Runtime/CoreUObject/Public/UObject/UnrealType.h
+++ b/Engine/Source/Runtime/CoreUObject/Public/UObject/UnrealType.h
@@ -1246,9 +1246,8 @@ private:
 };
 
 template<typename InTCppType>
-class COREUOBJECT_API TProperty_Numeric : public TProperty_WithEqualityAndSerializer<InTCppType, UNumericProperty>
+class TProperty_Numeric : public TProperty_WithEqualityAndSerializer<InTCppType, UNumericProperty>
 {
-
 public:
 	typedef TProperty_WithEqualityAndSerializer<InTCppType, UNumericProperty> Super;
 	typedef InTCppType TCppType;
@@ -1284,40 +1283,47 @@ public:
 
 protected:
 	template <typename OldNumericType>
-	FORCEINLINE void ConvertFromArithmeticValue(FStructuredArchive::FSlot Slot, void* Obj, const FPropertyTag& Tag)
+	FORCEINLINE void ConvertFromArithmeticValue(FStructuredArchive::FSlot Slot, void* Obj, const FPropertyTag& Tag) const
 	{
-		ConvertFromArithmeticValueImpl<OldNumericType>(*this, Slot, Obj, Tag);
+		TConvertAndSet<OldNumericType, TCppType>(*this, Slot, Obj, Tag);
 	}
 
 private:
-	template <typename OldNumericType>
-	static void ConvertFromArithmeticValueImpl(const TProperty_Numeric& This, FStructuredArchive::FSlot Slot, void* Obj, const FPropertyTag& Tag)
-	{
-		OldNumericType OldValue;
-		Slot << OldValue;
-		TCppType NewValue = (TCppType)OldValue;
-		This.SetPropertyValue_InContainer(Obj, NewValue, Tag.ArrayIndex);
-
-		UE_CLOG(
-			((TIsSigned<OldNumericType>::Value || TIsFloatingPoint<OldNumericType>::Value) && (!TIsSigned<TCppType>::Value && !TIsFloatingPoint<TCppType>::Value) && OldValue < 0) || ((OldNumericType)NewValue != OldValue),
-			LogClass,
-			Warning,
-			TEXT("Potential data loss during conversion of integer property %s of %s - was (%s) now (%s) - for package: %s"),
-			*This.GetName(),
-			*Slot.GetUnderlyingArchive().GetArchiveName(),
-			*LexToString(OldValue),
-			*LexToString(NewValue),
-			*Slot.GetUnderlyingArchive().GetArchiveName()
+	template <typename FromType, typename ToType>
+	struct TConvertAndSet
+	{
+		TConvertAndSet(const TProperty_Numeric& Property, FStructuredArchive::FSlot Slot, void* Obj, const FPropertyTag& Tag)
+		{
+			FromType OldValue;
+			Slot << OldValue;
+			ToType NewValue = (ToType)OldValue;
+			Property.SetPropertyValue_InContainer(Obj, NewValue, Tag.ArrayIndex);
+
+			UE_CLOG(
+				((TIsSigned<FromType>::Value || TIsFloatingPoint<FromType>::Value) && (!TIsSigned<ToType>::Value && !TIsFloatingPoint<ToType>::Value) && OldValue < 0) || ((FromType)NewValue != OldValue),
+				LogClass,
+				Warning,
+				TEXT("Potential data loss during conversion of integer property %s of %s - was (%s) now (%s) - for package: %s"),
+				*Property.GetName(),
+				*Slot.GetUnderlyingArchive().GetArchiveName(),
+				*LexToString(OldValue),
+				*LexToString(NewValue),
+				*Slot.GetUnderlyingArchive().GetArchiveName()
 			);
-	}
+		}
+	};
 
-	template <>
-	inline static void ConvertFromArithmeticValueImpl<TCppType>(const TProperty_Numeric& This, FStructuredArchive::FSlot Slot, void* Obj, const FPropertyTag& Tag)
+	template <typename SameType>
+	struct TConvertAndSet<SameType, SameType>
 	{
-		TCppType Value;
-		Slot << Value;
-		This.SetPropertyValue_InContainer(Obj, Value, Tag.ArrayIndex);
-	}
+		FORCEINLINE TConvertAndSet(const TProperty_Numeric& Property, FStructuredArchive::FSlot Slot, void* Obj, const FPropertyTag& Tag)
+		{
+			SameType Value;
+			Slot << Value;
+			Property.SetPropertyValue_InContainer(Obj, Value, Tag.ArrayIndex);
+		}
+	};
+
 public:
 	virtual EConvertFromTypeResult ConvertFromType(const FPropertyTag& Tag, FStructuredArchive::FSlot Slot, uint8* Data, UStruct* DefaultsStruct) override
 	{
