TBA: Change world level info serialization to "inject" itself into the structured archive root record, so that when no level info exists we can choose to not enter a slot at all

#rnx

diff --git a/Engine/Source/Runtime/CoreUObject/Private/UObject/SavePackage.cpp b/Engine/Source/Runtime/CoreUObject/Private/UObject/SavePackage.cpp
--- a/Engine/Source/Runtime/CoreUObject/Private/UObject/SavePackage.cpp
+++ b/Engine/Source/Runtime/CoreUObject/Private/UObject/SavePackage.cpp
@@ -5120,7 +5120,7 @@ FSavePackageResultStruct UPackage::Save(UPackage* InOuter, UObject* Base, EObjec
 					UPackage::SaveAssetRegistryData(InOuter, Linker.Get(), StructuredArchiveRoot.EnterField(SA_FIELD_NAME(TEXT("AssetRegistry"))));
 
 					// Save level information used by World browser
-					UPackage::SaveWorldLevelInfo(InOuter, Linker.Get(), StructuredArchiveRoot.EnterField(SA_FIELD_NAME(TEXT("WorldLevelInfo"))));
+					UPackage::SaveWorldLevelInfo(InOuter, Linker.Get(), StructuredArchiveRoot);
 				}
 
 
@@ -6492,19 +6492,14 @@ void UPackage::SaveAssetRegistryData(UPackage* InOuter, FLinkerSave* Linker, FSt
 	}
 }
 
-void UPackage::SaveWorldLevelInfo(UPackage* InOuter, FLinkerSave* Linker, FStructuredArchive::FSlot Slot)
+void UPackage::SaveWorldLevelInfo(UPackage* InOuter, FLinkerSave* Linker, FStructuredArchive::FRecord Record)
 {
 	Linker->Summary.WorldTileInfoDataOffset = 0;
 	
 	if(InOuter->WorldTileInfo.IsValid())
 	{
 		Linker->Summary.WorldTileInfoDataOffset = Linker->Tell();
-		Slot << *(InOuter->WorldTileInfo);
-	}
-	else
-	{
-		// Make empty record, because structured archive API doesn't allow us to leave a slot empty
-		Slot.EnterRecord();
+		Record << SA_VALUE(TEXT("WorldLevelInfo"), *(InOuter->WorldTileInfo));
 	}
 }
 
diff --git a/Engine/Source/Runtime/CoreUObject/Public/UObject/Package.h b/Engine/Source/Runtime/CoreUObject/Public/UObject/Package.h
--- a/Engine/Source/Runtime/CoreUObject/Public/UObject/Package.h
+++ b/Engine/Source/Runtime/CoreUObject/Public/UObject/Package.h
@@ -592,8 +592,9 @@ public:
 	*
 	* @param	InOuter							the outer to use for the new package
 	* @param	Linker							linker we're currently saving with
+	* @param	Record							the structured archive record we should add the world level info to
 	*/
-	static void SaveWorldLevelInfo(UPackage* InOuter, FLinkerSave* Linker, FStructuredArchive::FSlot Slot);
+	static void SaveWorldLevelInfo(UPackage* InOuter, FLinkerSave* Linker, FStructuredArchive::FRecord Record);
 
 	/**
 	* Determines if a package contains no more assets.
