TBA: Don't serialize dependsmap, searchable names or asset registry data to text files

#rnx

diff --git a/Engine/Source/Runtime/CoreUObject/Private/UObject/SavePackage.cpp b/Engine/Source/Runtime/CoreUObject/Private/UObject/SavePackage.cpp
--- a/Engine/Source/Runtime/CoreUObject/Private/UObject/SavePackage.cpp
+++ b/Engine/Source/Runtime/CoreUObject/Private/UObject/SavePackage.cpp
@@ -5047,29 +5047,32 @@ FSavePackageResultStruct UPackage::Save(UPackage* InOuter, UObject* Base, EObjec
 				}
 				SlowTask.EnterProgressFrame();
 
-				FStructuredArchive::FStream DependsStream = StructuredArchiveRoot.EnterStream(SA_FIELD_NAME(TEXT("DependsMap")));
-				if (Linker->IsCooking())
+				if (!bTextFormat)
 				{
+					FStructuredArchive::FStream DependsStream = StructuredArchiveRoot.EnterStream(SA_FIELD_NAME(TEXT("DependsMap")));
+					if (Linker->IsCooking())
+					{
 #if WITH_EDITOR
-					FArchiveStackTraceIgnoreScope IgnoreSummaryDiffsScope(DiffSettings.bIgnoreHeaderDiffs);
+						FArchiveStackTraceIgnoreScope IgnoreSummaryDiffsScope(DiffSettings.bIgnoreHeaderDiffs);
 #endif // WITH_EDITOR
-					//@todo optimization, this should just be stripped entirely from cooked packages
-					TArray<FPackageIndex> Depends; // empty array
-					Linker->Summary.DependsOffset = Linker->Tell();
-					for (int32 i = 0; i < Linker->ExportMap.Num(); i++)
-					{
-						DependsStream.EnterElement() << Depends;
+						//@todo optimization, this should just be stripped entirely from cooked packages
+						TArray<FPackageIndex> Depends; // empty array
+						Linker->Summary.DependsOffset = Linker->Tell();
+						for (int32 i = 0; i < Linker->ExportMap.Num(); i++)
+						{
+							DependsStream.EnterElement() << Depends;
+						}
 					}
-				}
-				else
-				{
-					// save depends map (no need for later patching)
-					check(Linker->DependsMap.Num() == Linker->ExportMap.Num());
-					Linker->Summary.DependsOffset = Linker->Tell();
-					for (int32 i = 0; i < Linker->ExportMap.Num(); i++)
+					else
 					{
-						TArray<FPackageIndex>& Depends = Linker->DependsMap[i];
-						DependsStream.EnterElement() << Depends;
+						// save depends map (no need for later patching)
+						check(Linker->DependsMap.Num() == Linker->ExportMap.Num());
+						Linker->Summary.DependsOffset = Linker->Tell();
+						for (int32 i = 0; i < Linker->ExportMap.Num(); i++)
+						{
+							TArray<FPackageIndex>& Depends = Linker->DependsMap[i];
+							DependsStream.EnterElement() << Depends;
+						}
 					}
 				}
 
@@ -5086,6 +5089,7 @@ FSavePackageResultStruct UPackage::Save(UPackage* InOuter, UObject* Base, EObjec
 					// Save soft package references
 					Linker->Summary.SoftPackageReferencesOffset = Linker->Tell();
 					Linker->Summary.SoftPackageReferencesCount = Linker->SoftPackageReferenceList.Num();
+					if (!bTextFormat)
 					{
 #if WITH_EDITOR
 						FArchive::FScopeSetDebugSerializationFlags S(*Linker, DSF_IgnoreDiff, true);
@@ -5116,8 +5120,11 @@ FSavePackageResultStruct UPackage::Save(UPackage* InOuter, UObject* Base, EObjec
 					// Save thumbnails
 					UPackage::SaveThumbnails(InOuter, Linker.Get(), StructuredArchiveRoot.EnterField(SA_FIELD_NAME(TEXT("Thumbnails"))));
 
-					// Save asset registry data so the editor can search for information about assets in this package
-					UPackage::SaveAssetRegistryData(InOuter, Linker.Get(), StructuredArchiveRoot.EnterField(SA_FIELD_NAME(TEXT("AssetRegistry"))));
+					if (!bTextFormat)
+					{	
+						// Save asset registry data so the editor can search for information about assets in this package
+						UPackage::SaveAssetRegistryData(InOuter, Linker.Get(), StructuredArchiveRoot.EnterField(SA_FIELD_NAME(TEXT("AssetRegistry"))));
+					}
 
 					// Save level information used by World browser
 					UPackage::SaveWorldLevelInfo(InOuter, Linker.Get(), StructuredArchiveRoot);
@@ -5868,7 +5875,7 @@ FSavePackageResultStruct UPackage::Save(UPackage* InOuter, UObject* Base, EObjec
 
 					if (!bTextFormat)
 					{
-					Linker->Seek(Linker->Summary.ImportOffset);
+						Linker->Seek(Linker->Summary.ImportOffset);
 					}
 
 					int32 NumImports = Linker->ImportMap.Num();
