Fixed cook EngineTest bug by reverting a small part of the recent tagged property serializations

diff --git a/Engine/Source/Runtime/CoreUObject/Private/UObject/Class.cpp b/Engine/Source/Runtime/CoreUObject/Private/UObject/Class.cpp
--- a/Engine/Source/Runtime/CoreUObject/Private/UObject/Class.cpp
+++ b/Engine/Source/Runtime/CoreUObject/Private/UObject/Class.cpp
@@ -1132,25 +1132,26 @@ void UStruct::SerializeTaggedProperties(FStructuredArchive::FSlot Slot, uint8* D
 				{
 					FStructuredArchive::FSlot ValueSlot = PropertyRecord.EnterField(SA_FIELD_NAME(TEXT("Value")));
 
-					if (Tag.Type == PropID)
+					switch (Property->ConvertFromType(Tag, ValueSlot, Data, DefaultsStruct))
 					{
-						uint8* DestAddress = Property->ContainerPtrToValuePtr<uint8>(Data, Tag.ArrayIndex);
-						uint8* DefaultsFromParent = Property->ContainerPtrToValuePtrForDefaults<uint8>(DefaultsStruct, Defaults, Tag.ArrayIndex);
-
-						// This property is ok.
-						Tag.SerializeTaggedProperty(ValueSlot, Property, DestAddress, DefaultsFromParent);
-						bAdvanceProperty = !UnderlyingArchive.IsCriticalError();
-					}
-					else
-					{
-						switch (Property->ConvertFromType(Tag, ValueSlot, Data, DefaultsStruct))
-						{
 						case EConvertFromTypeResult::Converted:
 							bAdvanceProperty = true;
 							break;
 
 						case EConvertFromTypeResult::UseSerializeItem:
-							UE_LOG(LogClass, Warning, TEXT("Type mismatch in %s of %s - Previous (%s) Current(%s) for package:  %s"), *Tag.Name.ToString(), *GetName(), *Tag.Type.ToString(), *PropID.ToString(), *UnderlyingArchive.GetArchiveName());
+							if (Tag.Type != PropID)
+							{
+								UE_LOG(LogClass, Warning, TEXT("Type mismatch in %s of %s - Previous (%s) Current(%s) for package:  %s"), *Tag.Name.ToString(), *GetName(), *Tag.Type.ToString(), *PropID.ToString(), *UnderlyingArchive.GetArchiveName());
+							}
+							else
+							{
+								uint8* DestAddress = Property->ContainerPtrToValuePtr<uint8>(Data, Tag.ArrayIndex);
+								uint8* DefaultsFromParent = Property->ContainerPtrToValuePtrForDefaults<uint8>(DefaultsStruct, Defaults, Tag.ArrayIndex);
+
+								// This property is ok.
+								Tag.SerializeTaggedProperty(ValueSlot, Property, DestAddress, DefaultsFromParent);
+								bAdvanceProperty = !UnderlyingArchive.IsCriticalError();
+							}
 							break;
 
 						case EConvertFromTypeResult::CannotConvert:
@@ -1158,7 +1159,6 @@ void UStruct::SerializeTaggedProperties(FStructuredArchive::FSlot Slot, uint8* D
 
 						default:
 							check(false);
-						}
 					}
 				}
 			}
