Fix compile error due to a missing include.

diff --git a/Engine/Source/Runtime/CoreUObject/Private/Serialization/AsyncLoading.cpp b/Engine/Source/Runtime/CoreUObject/Private/Serialization/AsyncLoading.cpp
--- a/Engine/Source/Runtime/CoreUObject/Private/Serialization/AsyncLoading.cpp
+++ b/Engine/Source/Runtime/CoreUObject/Private/Serialization/AsyncLoading.cpp
@@ -39,6 +39,7 @@
 #include "Async/TaskGraphInterfaces.h"
 #include "Blueprint/BlueprintSupport.h"
 #include "HAL/LowLevelMemTracker.h"
+#include "HAL/LowLevelMemStats.h"
 #include "ProfilingDebugging/CsvProfiler.h"
 #include "UObject/GCScopeLock.h"
 #include "ProfilingDebugging/MiscTrace.h"
