Added a function FPackageName::SplitFullObjectPath() which will split the supplied full path into its constituent pieces

diff --git a/Engine/Source/Runtime/CoreUObject/Private/Misc/PackageName.cpp b/Engine/Source/Runtime/CoreUObject/Private/Misc/PackageName.cpp
--- a/Engine/Source/Runtime/CoreUObject/Private/Misc/PackageName.cpp
+++ b/Engine/Source/Runtime/CoreUObject/Private/Misc/PackageName.cpp
@@ -500,6 +500,56 @@ bool FPackageName::SplitLongPackageName(const FString& InLongPackageName, FStrin
 	return true;
 }
 
+void FPackageName::SplitFullObjectPath(const FString& InFullObjectPath, FString& OutClassName, FString& OutPackageName, FString& OutObjectName, FString& OutSubObjectName)
+{
+	FString Sanitized = InFullObjectPath.TrimStartAndEnd();
+
+	int32 StartOfClass = 0;
+	int32 StartOfPackage = 0;
+	int32 StartOfObject = 0;
+	int32 StartOfSubObject = 0;
+	int32 EndOfClass = 0;
+	int32 EndOfPackage = 0;
+	int32 EndOfObject = 0;
+	int32 EndOfSubObject = 0;
+
+	int32 TempIndex = INDEX_NONE;
+	if (Sanitized.FindChar(' ', TempIndex))
+	{
+		EndOfClass = TempIndex;
+		StartOfPackage = EndOfPackage = EndOfClass + 1;
+	}
+
+	if (Sanitized.FindChar('.', TempIndex))
+	{
+		EndOfPackage = TempIndex;
+		StartOfObject = EndOfObject = EndOfPackage + 1;
+
+		if (Sanitized.FindChar(':', TempIndex))
+		{
+			EndOfObject = TempIndex;
+			if (TempIndex < InFullObjectPath.Len())
+			{
+				StartOfSubObject = TempIndex + 1;
+				EndOfSubObject = InFullObjectPath.Len();
+			}
+		}
+		else
+		{
+			EndOfObject = Sanitized.Len();
+		}
+	}
+	else
+	{
+		EndOfPackage = Sanitized.Len();
+	}
+
+	OutClassName = Sanitized.Mid(StartOfClass, EndOfClass - StartOfClass);
+	OutPackageName = Sanitized.Mid(StartOfPackage, EndOfPackage - StartOfPackage);
+	OutObjectName = Sanitized.Mid(StartOfObject, EndOfObject - StartOfObject);
+	OutSubObjectName = Sanitized.Mid(StartOfSubObject, EndOfSubObject - StartOfSubObject);
+}
+
 FString FPackageName::GetLongPackageAssetName(const FString& InLongPackageName)
 {
 	return GetShortName(InLongPackageName);
diff --git a/Engine/Source/Runtime/CoreUObject/Public/Misc/PackageName.h b/Engine/Source/Runtime/CoreUObject/Public/Misc/PackageName.h
--- a/Engine/Source/Runtime/CoreUObject/Public/Misc/PackageName.h
+++ b/Engine/Source/Runtime/CoreUObject/Public/Misc/PackageName.h
@@ -103,6 +103,17 @@ public:
 	 */
 	static bool SplitLongPackageName(const FString& InLongPackageName, FString& OutPackageRoot, FString& OutPackagePath, FString& OutPackageName, const bool bStripRootLeadingSlash = false);
 
+	/**
+	 * Split a full object path (Class /Path/To/A/Package.Object:SubObject) into its constituent pieces
+	 *  
+	 * @param InFullObjectPath			Full object path we want to split
+	 * @param OutClassName				The extracted class name (Class)
+	 * @param OutPackageName			The extracted package name (/Path/To/A/Package)
+	 * @param OutObjectName				The extracted object name (Object)
+	 * @param OutSubObjectName			The extracted subobject name (SubObject)
+	 */
+	static void SplitFullObjectPath(const FString& InFullObjectPath, FString& OutClassName, FString& OutPackageName, FString& OutObjectName, FString& OutSubObjectName);
+
 	/** 
 	 * Returns true if the path starts with a valid root (i.e. /Game/, /Engine/, etc) and contains no illegal characters.
 	 *
