Exposed FDeferredMessageLog for public use (UDN request)

diff --git a/Engine/Source/Runtime/CoreUObject/Private/Serialization/DeferredMessageLog.h b/Engine/Source/Runtime/CoreUObject/Private/Serialization/DeferredMessageLog.h
deleted file mode 100644
diff --git a/Engine/Source/Runtime/CoreUObject/Public/Serialization/DeferredMessageLog.h b/Engine/Source/Runtime/CoreUObject/Public/Serialization/DeferredMessageLog.h
new file mode 100644
--- /dev/null
+++ b/Engine/Source/Runtime/CoreUObject/Public/Serialization/DeferredMessageLog.h
@@ -0,0 +1,38 @@
+// Copyright 1998-2019 Epic Games, Inc. All Rights Reserved.
+
+/*=============================================================================
+	DeferredMessgaeLog.h: Unreal async loading log.
+=============================================================================*/
+
+#pragma once
+
+#include "CoreMinimal.h"
+
+class Error;
+class FTokenizedMessage;
+
+/**
+ * Thread safe proxy for the FMessageLog while performing async loading.
+ * Makes sure the messages does not get added to the log until async loading is 
+ * finished to prevent modules from being loaded outside of game thread.
+ * Also makes sure the messages are added to the message queue in a thread-safe way.
+*/
+class COREUOBJECT_API FDeferredMessageLog
+{
+	FName LogCategory;
+
+	static TMap<FName, TArray<TSharedRef<FTokenizedMessage>>*> Messages;
+	static FCriticalSection MessagesCritical;
+
+	void AddMessage(TSharedRef<FTokenizedMessage>& Message);
+
+public:
+	FDeferredMessageLog(const FName& InLogCategory);
+	
+	TSharedRef<FTokenizedMessage> Info(const FText& Message);
+	TSharedRef<FTokenizedMessage> Warning(const FText& Message);
+	TSharedRef<FTokenizedMessage> Error(const FText& Message);
+
+	static void Flush();
+	static void Cleanup();
+};
