Fixing CIS after the last big merge from Main

diff --git a/Engine/Source/Runtime/Core/Private/HAL/MallocAnsi.cpp b/Engine/Source/Runtime/Core/Private/HAL/MallocAnsi.cpp
--- a/Engine/Source/Runtime/Core/Private/HAL/MallocAnsi.cpp
+++ b/Engine/Source/Runtime/Core/Private/HAL/MallocAnsi.cpp
@@ -150,8 +150,6 @@ FMallocAnsi::FMallocAnsi()
 
 void* FMallocAnsi::Malloc( SIZE_T Size, uint32 Alignment )
 {
-	IncrementTotalMallocCalls();
-
 #if !UE_BUILD_SHIPPING
 	uint64 LocalMaxSingleAlloc = MaxSingleAlloc.Load(EMemoryOrder::Relaxed);
 	if (LocalMaxSingleAlloc != 0 && Size > LocalMaxSingleAlloc)
diff --git a/Engine/Source/Runtime/CoreUObject/Private/Serialization/AsyncLoading.cpp b/Engine/Source/Runtime/CoreUObject/Private/Serialization/AsyncLoading.cpp
--- a/Engine/Source/Runtime/CoreUObject/Private/Serialization/AsyncLoading.cpp
+++ b/Engine/Source/Runtime/CoreUObject/Private/Serialization/AsyncLoading.cpp
@@ -5339,10 +5339,12 @@ void FAsyncLoadingThread::CancelAsyncLoadingInternal()
 
 	if (GEventDrivenLoaderEnabled)
 	{
-		while (GPrecacheCallbackHandler.AnyIOOutstanding())
+		FPrecacheCallbackHandler* PrecacheHandler = FCoreUObjectModule::GetGlobalPrecacheHandler();
+
+		while (PrecacheHandler->AnyIOOutstanding())
 		{
-			GPrecacheCallbackHandler.WaitForIO(10.0f);
-			GPrecacheCallbackHandler.ProcessIncoming();
+			PrecacheHandler->WaitForIO(10.0f);
+			PrecacheHandler->ProcessIncoming();
 		}
 	}
 
