TBA: Remove import and export table from text asset, then reconstruct on load

#rnx

diff --git a/Engine/Source/Runtime/CoreUObject/Private/UObject/LinkerLoad.cpp b/Engine/Source/Runtime/CoreUObject/Private/UObject/LinkerLoad.cpp
--- a/Engine/Source/Runtime/CoreUObject/Private/UObject/LinkerLoad.cpp
+++ b/Engine/Source/Runtime/CoreUObject/Private/UObject/LinkerLoad.cpp
@@ -731,6 +731,13 @@ FLinkerLoad::ELinkerStatus FLinkerLoad::Tick( float InTimeLimit, bool bInUseTime
 				Status = SerializeExportMap();
 			}
 
+			// Reconstruct the import and export maps for text assets
+			if (Status == LINKER_Loaded)
+			{
+				SCOPED_LOADTIMER(LinkerLoad_ReconstructImportAndExportMap);
+				Status = ReconstructImportAndExportMap();
+			}
+
 			// Fix up import map for backward compatible serialization.
 			if( Status == LINKER_Loaded )
 			{	
@@ -828,6 +835,7 @@ FLinkerLoad::FLinkerLoad(UPackage* InParent, const TCHAR* InFilename, uint32 InL
 , DependsMapIndex(0)
 , ExportHashIndex(0)
 , bHasSerializedPackageFileSummary(false)
+, bHasReconstructedImportAndExportMap(false)
 , bHasSerializedPreloadDependencies(false)
 , bHasFixedUpImportMap(false)
 , bHasFoundExistingExports(false)
@@ -990,29 +998,18 @@ FLinkerLoad::ELinkerStatus FLinkerLoad::CreateLoader(
 			{
 				INC_DWORD_STAT(STAT_TextAssetLinkerCount);
 				Loader = IFileManager::Get().CreateFileReader(*Filename);
-				StructuredArchiveFormatter = new FJsonArchiveInputFormatter(*this, [this](const FString& InName)
+				StructuredArchiveFormatter = new FJsonArchiveInputFormatter(*this, [this](const FString& InFullPath)
 				{
-					int32 EndOfClassNameIndex = INDEX_NONE;
-					InName.FindChar(' ', EndOfClassNameIndex);
-					
-					FName FullObjectPath = *InName.Right(InName.Len() - EndOfClassNameIndex - 1);
-					UObject* Result = nullptr;
-					const FPackageIndex* ExportIndex = ObjectNameToPackageExportIndex.Find(FullObjectPath);
-					if (ExportIndex)
+					FPackageIndex Index = FindOrCreateImportOrExport(InFullPath);
+					if (Index.IsImport())
 					{
-						Result = CreateExport(ExportIndex->ToExport());
+						return CreateImport(Index.ToImport());
 					}
 					else
-						{
-						const FPackageIndex* ImportIndex = ObjectNameToPackageImportIndex.Find(FullObjectPath);
-
-						if (ImportIndex)
-						{
-							Result = CreateImport(ImportIndex->ToImport());
-						}
+					{
+						check(Index.IsExport());
+						return CreateExport(Index.ToExport());
 					}
-
-					return Result;
 				});
 			}
 			else
@@ -1497,7 +1494,12 @@ FLinkerLoad::ELinkerStatus FLinkerLoad::SerializeImportMap()
 {
 	DECLARE_SCOPE_CYCLE_COUNTER( TEXT( "FLinkerLoad::SerializeImportMap" ), STAT_LinkerLoad_SerializeImportMap, STATGROUP_LinkerLoad );
 
-	if( !IsTextFormat() && ImportMapIndex == 0 && Summary.ImportCount > 0 )
+	if (IsTextFormat())
+	{
+		return LINKER_Loaded;
+	}
+
+	if(ImportMapIndex == 0 && Summary.ImportCount > 0 )
 	{
 		Seek( Summary.ImportOffset );
 	}
@@ -1703,11 +1705,15 @@ FLinkerLoad::ELinkerStatus FLinkerLoad::SerializeExportMap()
 {
 	DECLARE_SCOPE_CYCLE_COUNTER( TEXT( "FLinkerLoad::SerializeExportMap" ), STAT_LinkerLoad_SerializeExportMap, STATGROUP_LinkerLoad );
 
-	if( !IsTextFormat() && ExportMapIndex == 0 && Summary.ExportCount > 0 )
+	if (IsTextFormat())
 	{
-		Seek( Summary.ExportOffset );
+		return LINKER_Loaded;
 	}
 
+	if(ExportMapIndex == 0 && Summary.ExportCount > 0)
+	{
+		Seek( Summary.ExportOffset );
+	}
 
 	FStructuredArchive::FStream Stream = StructuredArchiveRootRecord->EnterStream(SA_FIELD_NAME(TEXT("ExportTable")));
 
@@ -1724,6 +1730,124 @@ FLinkerLoad::ELinkerStatus FLinkerLoad::SerializeExportMap()
 	return ((ExportMapIndex == Summary.ExportCount) && !IsTimeLimitExceeded( TEXT("serializing export map") )) ? LINKER_Loaded : LINKER_TimedOut;
 }
 
+FPackageIndex FLinkerLoad::FindOrCreateImport(const FName InObjectName, const FName InClassName, const FName InClassPackageName)
+{
+	for (int32 ImportIndex = 0; ImportIndex < ImportMap.Num(); ++ImportIndex)
+	{
+		FObjectImport& ExistingImport = ImportMap[ImportIndex];
+		if (ExistingImport.ObjectName == InObjectName && ExistingImport.ClassPackage == InClassPackageName && ExistingImport.ClassName == InClassName)
+		{
+			return FPackageIndex::FromImport(ImportIndex);
+		}
+	}
+
+	FObjectImport& NewImport = ImportMap.Emplace_GetRef();
+	NewImport.ObjectName = InObjectName;
+	NewImport.ClassName = InClassName;
+	NewImport.ClassPackage = InClassPackageName;
+	NewImport.SourceIndex = INDEX_NONE;
+	NewImport.XObject = nullptr;
+	NewImport.SourceLinker = nullptr;
+
+	Summary.ImportCount++;
+
+	return FPackageIndex::FromImport(ImportMap.Num() - 1);
+}
+
+FPackageIndex FLinkerLoad::FindOrCreateImportOrExport(const FString& InFullPath)
+{
+	if (InFullPath.Len() == 0)
+	{
+		return FPackageIndex();
+	}
+
+	FString Class, Package, Object, SubObject;
+	FPackageName::SplitFullObjectPath(InFullPath, Class, Package, Object, SubObject);
+	bool bIsExport = (Package == LinkerRoot->GetName());
+
+	FName ObjectName = *Object;
+
+	if (bIsExport)
+	{
+		for (int32 ExportIndex = 0; ExportIndex < ExportMap.Num(); ++ExportIndex)
+		{
+			if (ExportMap[ExportIndex].ObjectName == ObjectName)
+			{
+				return FPackageIndex::FromExport(ExportIndex);
+			}
+		}
+
+		return FPackageIndex();
+	}
+	else
+	{
+		FName ClassName = *Class;
+		FName PackageName = *Package;
+
+		if (UClass* ObjectClass = FindObjectFast<UClass>(nullptr, ClassName, false, true))
+		{
+			FName ClassPackageName = *ObjectClass->GetOuterUPackage()->GetPathName();
+
+			FPackageIndex ImportOuterIndex = FindOrCreateImport(PackageName, NAME_Package, FName(TEXT("/Script/CoreUObject")));
+			FPackageIndex ImportIndex = FindOrCreateImport(ObjectName, ClassName, ClassPackageName);
+			ImportMap[ImportIndex.ToImport()].OuterIndex = ImportOuterIndex;
+
+			return ImportIndex;
+		}
+		else
+		{
+			UE_LOG(LogLinker, Warning, TEXT("Failed to find class '%s' while trying to resolve full path '%s'"), *ClassName.ToString(), *InFullPath);
+		}
+	}
+
+	return FPackageIndex();
+}
+
+FLinkerLoad::ELinkerStatus FLinkerLoad::ReconstructImportAndExportMap()
+{
+	DECLARE_SCOPE_CYCLE_COUNTER(TEXT("FLinkerLoad::ReconstructImportAndExportMap"), STAT_LinkerLoad_ReconstructImportAndExportMap, STATGROUP_LinkerLoad);
+
+	if (!bHasReconstructedImportAndExportMap && IsTextFormat())
+	{
+		int32 NumExports = 0;
+		FStructuredArchive::FMap PackageExports = StructuredArchiveRootRecord->EnterMap(SA_FIELD_NAME(TEXT("Exports")), NumExports);
+		ExportMap.Reserve(NumExports);
+
+		TArray<FObjectTextExport> ExportRecords;
+		ExportRecords.Reserve(NumExports);
+		ExportMap.SetNum(NumExports);
+
+		Summary.ExportCount = ExportMap.Num();
+		Summary.ImportCount = 0;
+
+		for (int32 ExportIndex = 0; ExportIndex < NumExports; ++ExportIndex)
+		{
+			FObjectTextExport& TextExport = ExportRecords.Emplace_GetRef(ExportMap[ExportIndex], nullptr);
+			FString ObjectName;
+			PackageExports.EnterElement(ObjectName) << TextExport;
+			ExportMap[ExportIndex].ObjectName = *ObjectName;
+		}
+		
+		// Now pass over all the exports and rebuild the export/import records
+		for (int32 ExportIndex = 0; ExportIndex < NumExports; ++ExportIndex)
+		{
+			FObjectTextExport& TextExport = ExportRecords[ExportIndex];
+			
+			TextExport.Export.ThisIndex = FPackageIndex::FromExport(ExportIndex);
+			TextExport.Export.ClassIndex = FindOrCreateImportOrExport(TextExport.ClassName);
+			TextExport.Export.SuperIndex = FindOrCreateImportOrExport(TextExport.SuperStructName);
+			TextExport.Export.OuterIndex = FindOrCreateImportOrExport(TextExport.OuterName);
+		}
+
+		bHasReconstructedImportAndExportMap = true;
+		return LINKER_Loaded;
+	}
+	else
+	{
+		return LINKER_Loaded;
+	}
+}
+
 /**
  * Serializes the depends map.
  */
@@ -1829,10 +1953,27 @@ FLinkerLoad::ELinkerStatus FLinkerLoad::SerializeThumbnails( bool bForceEnableIn
 		return LINKER_Loaded;
 	}
 
-	FStructuredArchive::FRecord Record = StructuredArchiveRootRecord->EnterRecord(SA_FIELD_NAME(TEXT("Thumbnails")));
+	TOptional<FStructuredArchive::FSlot> ThumbnailsSlot;
 
-	if((Summary.ThumbnailTableOffset > 0) || IsTextFormat())
+	if (IsTextFormat())
 	{
+		ThumbnailsSlot = StructuredArchiveRootRecord->TryEnterField(SA_FIELD_NAME(TEXT("Thumbnails")), false);
+		if (!ThumbnailsSlot.IsSet())
+		{
+			return LINKER_Loaded;
+		}
+	}
+	else
+	{
+		if(Summary.ThumbnailTableOffset > 0)
+		{
+			ThumbnailsSlot = StructuredArchiveRootRecord->EnterField(SA_FIELD_NAME(TEXT("Thumbnails")));
+		}
+	}
+
+	if (ThumbnailsSlot.IsSet())
+	{
+		FStructuredArchive::FRecord Record = ThumbnailsSlot->EnterRecord();
 		TOptional<FStructuredArchive::FSlot> IndexSlot;
 
 		if (IsTextFormat())
@@ -2013,69 +2154,6 @@ FLinkerLoad::ELinkerStatus FLinkerLoad::FinalizeCreation()
 		}
 #endif
 
-		if (IsTextFormat())
-		{
-			ObjectNameToPackageExportIndex.Empty(ExportMap.Num());
-			ObjectNameToPackageImportIndex.Empty(ImportMap.Num());
-
-			// Build full export map name lookup
-			for (int32 ExportIndex = 0; ExportIndex < ExportMap.Num(); ++ExportIndex)
-			{
-				FString FullName;
-				FPackageIndex CurIndex = FPackageIndex::FromExport(ExportIndex);
-				int32 PartCount = 0;
-				do
-				{
-					FObjectExport& Export = ExportMap[CurIndex.ToExport()];
-					if (FullName.Len() > 0)
-					{
-						FullName = TEXT(".") + FullName;
-					}
-					FullName = Export.ObjectName.ToString() + FullName;
-					CurIndex = Export.OuterIndex;
-				} while (!CurIndex.IsNull());
-				
-				// In some cases, the export in this package can be pathed into a different package (map BuiltInData stuff mainly) so we only want to attach
-				// this package path to the name if it doesn't look like a long package name
-				if (FullName[0] != '/')
-				{
-				int32 DotIndex = INDEX_NONE;
-				if (FullName.FindChar('.', DotIndex))
-				{
-					FullName[DotIndex] = ':';
-				}
-
-				FullName = LinkerRoot->GetName() + TEXT(".") + FullName;
-				}
-
-				check(ObjectNameToPackageExportIndex.Find(*FullName) == nullptr);
-				ObjectNameToPackageExportIndex.Add(*FullName, FPackageIndex::FromExport(ExportIndex));
-			}
-
-			// Same for import map
-			ObjectNameToPackageImportIndex.Reserve(ImportMap.Num());
-			for (int32 ImportIndex = 0; ImportIndex < ImportMap.Num(); ++ImportIndex)
-			{
-				FPackageIndex CurIndex = FPackageIndex::FromImport(ImportIndex);
-				FString FullName;
-				do 
-				{
-					FObjectImport& Import = ImportMap[CurIndex.ToImport()];
-					if (FullName.Len() > 0)
-					{
-						FullName = TEXT(".") + FullName;
-					}
-					FullName = Import.ObjectName.ToString() + FullName;
-					CurIndex = Import.OuterIndex;
-				}
-				while (!CurIndex.IsNull());
-
-				FName FullFName(*FullName);
-				check(!ObjectNameToPackageImportIndex.Contains(FullFName));
-				ObjectNameToPackageImportIndex.Add(FullFName, FPackageIndex::FromImport(ImportIndex));
-			}
-		}
-
 		// Add this linker to the object manager's linker array.
 		FLinkerManager::Get().AddLoader(this);
 
@@ -3632,17 +3710,11 @@ void FLinkerLoad::Preload( UObject* Object )
 #if WITH_EDITOR && WITH_TEXT_ARCHIVE_SUPPORT
 						if (IsTextFormat())
 						{
-							const FName* FullObjectPath = ObjectNameToPackageExportIndex.FindKey(Export.ThisIndex);
-							check(FullObjectPath);
-							FString ObjectPath = FullObjectPath->ToString();
-							FString PackagePath = LinkerRoot->GetPathName(nullptr);
-							FString ObjectName = ObjectPath;
-							if (ObjectPath.StartsWith(PackagePath + TEXT(".")))
-							{
-								ObjectName = ObjectPath.Right(ObjectPath.Len() - PackagePath.Len() - 1);
-							}
+							FStructuredArchive::FSlot ExportSlot = StructuredArchiveRootRecord->EnterRecord(SA_FIELD_NAME(TEXT("Exports"))).EnterField(SA_FIELD_NAME(*Export.ObjectName.ToString()));
 
-							FStructuredArchive::FSlot ExportSlot = StructuredArchiveRootRecord->EnterField(SA_FIELD_NAME(*ObjectName));
+							// TODO BEFORE CHECKIN: This attribute read only exists because we crash if we attempt to serialize the value of an attributed field when we don't read any attributes first
+							FString Temp;
+							ExportSlot << SA_ATTRIBUTE(TEXT("Class"), Temp);
 
 							if (bClassSupportsTextFormat)
 							{
diff --git a/Engine/Source/Runtime/CoreUObject/Private/UObject/ObjectResource.cpp b/Engine/Source/Runtime/CoreUObject/Private/UObject/ObjectResource.cpp
--- a/Engine/Source/Runtime/CoreUObject/Private/UObject/ObjectResource.cpp
+++ b/Engine/Source/Runtime/CoreUObject/Private/UObject/ObjectResource.cpp
@@ -171,6 +171,69 @@ void operator<<(FStructuredArchive::FSlot Slot, FObjectExport& E)
 	}	
 }
 
+/*-----------------------------------------------------------------------------
+	FObjectTextExport.
+-----------------------------------------------------------------------------*/
+
+void operator<<(FStructuredArchive::FSlot Slot, FObjectTextExport& E)
+{
+	FArchive& BaseArchive = Slot.GetUnderlyingArchive();
+	FString ClassName, OuterName, SuperStructName;
+	check(BaseArchive.IsTextFormat());
+
+	if (BaseArchive.IsSaving())
+	{
+		check(E.Export.Object);
+
+		UClass* ObjClass = E.Export.Object->GetClass();
+		if (ObjClass != UClass::StaticClass())
+		{
+			ClassName = ObjClass->GetFullName();
+		}
+
+		if (E.Export.Object->GetOuter() != E.Outer)
+		{
+			OuterName = E.Export.Object->GetOuter() ? E.Export.Object->GetOuter()->GetFullName() : FString();
+		}
+
+		if (UStruct* Struct = Cast<UStruct>(E.Export.Object))
+		{
+			if (Struct->GetSuperStruct() != nullptr)
+			{
+				SuperStructName = Struct->GetSuperStruct()->GetFullName();
+			}
+		}
+	}
+
+	Slot << SA_ATTRIBUTE(TEXT("Class"), ClassName);
+	Slot << SA_DEFAULTED_ATTRIBUTE(TEXT("Outer"), OuterName, FString());
+	Slot << SA_DEFAULTED_ATTRIBUTE(TEXT("SuperStruct"), SuperStructName, FString());
+
+	if (BaseArchive.IsLoading())
+	{
+		E.ClassName = ClassName;
+		E.OuterName = OuterName;
+		E.SuperStructName = SuperStructName;
+	}
+
+	uint32 Save = E.Export.ObjectFlags & RF_Load;
+	Slot << SA_DEFAULTED_ATTRIBUTE(TEXT("ObjectFlags"), Save, 0);
+	if (BaseArchive.IsLoading())
+	{
+		E.Export.ObjectFlags = EObjectFlags(Save & RF_Load);
+	}
+
+	Slot << SA_DEFAULTED_ATTRIBUTE(TEXT("bForcedExport"), E.Export.bForcedExport, false);
+	Slot << SA_DEFAULTED_ATTRIBUTE(TEXT("bNotForClient"), E.Export.bNotForClient, false);
+	Slot << SA_DEFAULTED_ATTRIBUTE(TEXT("bNotForServer"), E.Export.bNotForServer, false);
+
+	Slot << SA_DEFAULTED_ATTRIBUTE(TEXT("PackageGuid"), E.Export.PackageGuid, FGuid());
+	Slot << SA_DEFAULTED_ATTRIBUTE(TEXT("PackageFlags"), E.Export.PackageFlags, 0);
+
+	Slot << SA_DEFAULTED_ATTRIBUTE(TEXT("bNotAlwaysLoadedForEditorGame"), E.Export.bNotAlwaysLoadedForEditorGame, false);
+	Slot << SA_DEFAULTED_ATTRIBUTE(TEXT("bIsAsset"), E.Export.bIsAsset, false);
+}
+
 /*-----------------------------------------------------------------------------
 	FObjectImport.
 -----------------------------------------------------------------------------*/
diff --git a/Engine/Source/Runtime/CoreUObject/Private/UObject/SavePackage.cpp b/Engine/Source/Runtime/CoreUObject/Private/UObject/SavePackage.cpp
--- a/Engine/Source/Runtime/CoreUObject/Private/UObject/SavePackage.cpp
+++ b/Engine/Source/Runtime/CoreUObject/Private/UObject/SavePackage.cpp
@@ -5517,6 +5517,8 @@ FSavePackageResultStruct UPackage::Save(UPackage* InOuter, UObject* Base, EObjec
 #endif
 					FScopedSlowTask ExportScope(Linker->ExportMap.Num());
 
+					FStructuredArchive::FRecord ExportsRecord = StructuredArchiveRoot.EnterRecord(SA_FIELD_NAME(TEXT("Exports")));
+
 					// Save exports.
 					FString ObjectName;
 					int32 LastExportSaveStep = 0;
@@ -5529,7 +5531,7 @@ FSavePackageResultStruct UPackage::Save(UPackage* InOuter, UObject* Base, EObjec
 						ExportScope.EnterProgressFrame();
 
 						FObjectExport& Export = Linker->ExportMap[i];
-						if( Export.Object )
+						if (Export.Object)
 						{
 							// Save the object data.
 							Export.SerialOffset = Linker->Tell();
@@ -5538,7 +5540,13 @@ FSavePackageResultStruct UPackage::Save(UPackage* InOuter, UObject* Base, EObjec
 
 							ObjectName.Reset();
 							Export.Object->GetPathName(InOuter, ObjectName);
-							FStructuredArchive::FSlot ExportSlot = StructuredArchiveRoot.EnterField(SA_FIELD_NAME(*ObjectName));
+							FStructuredArchive::FSlot ExportSlot = ExportsRecord.EnterField(SA_FIELD_NAME(*ObjectName));
+
+							if (bTextFormat)
+							{
+								FObjectTextExport ObjectTextExport(Export, InOuter);
+								ExportSlot << ObjectTextExport;
+							}
 
 #if WITH_EDITOR
 							bool bSupportsText = UClass::IsSafeToSerializeToStructuredArchives(Export.Object->GetClass());
diff --git a/Engine/Source/Runtime/CoreUObject/Public/UObject/LinkerLoad.h b/Engine/Source/Runtime/CoreUObject/Public/UObject/LinkerLoad.h
--- a/Engine/Source/Runtime/CoreUObject/Public/UObject/LinkerLoad.h
+++ b/Engine/Source/Runtime/CoreUObject/Public/UObject/LinkerLoad.h
@@ -143,10 +143,6 @@ private:
 	FArchiveFormatterType* StructuredArchiveFormatter;
 	TOptional<FStructuredArchive::FRecord> StructuredArchiveRootRecord;
 
-	/** A map of full object path name to package index. Used with text assets to resolve incoming string names to an export */
-	TMap<FName, FPackageIndex> ObjectNameToPackageImportIndex;
-	TMap<FName, FPackageIndex> ObjectNameToPackageExportIndex;
-
 	/** The archive that actually reads the raw data from disk.																*/
 	FArchive*				Loader;
 
@@ -277,6 +273,8 @@ private:
 
 	/** Whether we already serialized the package file summary.																*/
 	bool					bHasSerializedPackageFileSummary;
+	/** Whether we have already reconstructed the import/export tables for a text asset */
+	bool					bHasReconstructedImportAndExportMap;
 	/** Whether we already serialized preload dependencies.																*/
 	bool					bHasSerializedPreloadDependencies;
 	/** Whether we already fixed up import map.																				*/
@@ -931,6 +929,11 @@ private:
 	 */
 	ELinkerStatus SerializeExportMap();
 
+	/**
+	 * Create an import and export table when loading a text asset.
+	 */
+	ELinkerStatus ReconstructImportAndExportMap();
+
 	ELinkerStatus SerializeDependsMap();
 
 	ELinkerStatus SerializePreloadDependencies();
@@ -938,6 +941,12 @@ private:
 	/** Sets the basic linker archive info */
 	void ResetStatusInfo();
 
+	/** For a given full object path, find or create the associated import or export table record. Used when loading text assets which only store object paths */
+	FPackageIndex FindOrCreateImportOrExport(const FString& InFullPath);
+
+	/** For the given object and class info, find or create an associated import record. Used when loading text assets which only store object paths */
+	FPackageIndex FindOrCreateImport(const FName InObjectName, const FName InClassName, const FName InClassPackageName);
+
 public:
 	/**
 	 * Serializes the gatherable text data container.
diff --git a/Engine/Source/Runtime/CoreUObject/Public/UObject/ObjectResource.h b/Engine/Source/Runtime/CoreUObject/Public/UObject/ObjectResource.h
--- a/Engine/Source/Runtime/CoreUObject/Public/UObject/ObjectResource.h
+++ b/Engine/Source/Runtime/CoreUObject/Public/UObject/ObjectResource.h
@@ -365,7 +365,43 @@ struct FObjectExport : public FObjectResource
 	/** I/O functions */
 	friend COREUOBJECT_API FArchive& operator<<(FArchive& Ar, FObjectExport& E);
 	friend COREUOBJECT_API void operator<<(FStructuredArchive::FSlot Slot, FObjectExport& E);
+};
+
+/*-----------------------------------------------------------------------------
+	FObjectTextExport
+-----------------------------------------------------------------------------*/
+
+/**
+ * Simple wrapper around a FObjectExport which does the text asset specific serialization of export data
+ */
+struct FObjectTextExport
+{
+	/**
+	 * Constructor
+	 */
+	FObjectTextExport(FObjectExport& InExport, UObject* InOuter)
+		: Export(InExport)
+		, Outer(InOuter)
+	{
+	}
+
+	/** The export object that we are wrapping */
+	FObjectExport& Export;
+
+	/** The outer that this export lives inside */
+	UObject* Outer;
+
+	/** String full object path for this export's class */
+	FString ClassName;
+
+	/** String full object path for this export's superstruct, if applicable */
+	FString SuperStructName;
+
+	/** String full object path for this export's outer, if applicable (i.e. if it's not the package itself) */
+	FString OuterName;
 
+	/** Serializer */
+	friend COREUOBJECT_API void operator<<(FStructuredArchive::FSlot Slot, FObjectTextExport& E);
 };
 
 /*-----------------------------------------------------------------------------
