GitHub #6135 : Update Alignment comments

diff --git a/Engine/Source/Runtime/CoreUObject/Public/UObject/Class.h b/Engine/Source/Runtime/CoreUObject/Public/UObject/Class.h
--- a/Engine/Source/Runtime/CoreUObject/Public/UObject/Class.h
+++ b/Engine/Source/Runtime/CoreUObject/Public/UObject/Class.h
@@ -914,7 +914,7 @@ public:
 		{
 			return Size;
 		}
-		/** return the ALIGNOF() of this structure **/
+		/** return the alignof() of this structure **/
 		FORCEINLINE int32 GetAlignment()
 		{
 			return Alignment;
@@ -1029,7 +1029,7 @@ public:
 	private:
 		/** sizeof() of the structure **/
 		const int32 Size;
-		/** ALIGNOF() of the structure **/
+		/** alignof() of the structure **/
 		const int32 Alignment;
 	};
 
