Changed all package save Conform arguments to accept a non-existent type, to trigger an error in case any code actually passes in a non-null argument

diff --git a/Engine/Source/Editor/UnrealEd/Classes/Editor/EditorEngine.h b/Engine/Source/Editor/UnrealEd/Classes/Editor/EditorEngine.h
--- a/Engine/Source/Editor/UnrealEd/Classes/Editor/EditorEngine.h
+++ b/Engine/Source/Editor/UnrealEd/Classes/Editor/EditorEngine.h
@@ -2291,12 +2291,12 @@ public:
 
 	/** The editor wrapper for UPackage::SavePackage. Auto-adds files to source control when necessary */
 	bool SavePackage( UPackage* InOuter, UObject* Base, EObjectFlags TopLevelFlags, const TCHAR* Filename, 
-		FOutputDevice* Error=GError, FLinkerLoad* Conform=NULL, bool bForceByteSwapping=false, bool bWarnOfLongFilename=true, 
+		FOutputDevice* Error=GError, FLinkerNull* Conform=NULL, bool bForceByteSwapping=false, bool bWarnOfLongFilename=true, 
 		uint32 SaveFlags=SAVE_None, const class ITargetPlatform* TargetPlatform = NULL, const FDateTime& FinalTimeStamp = FDateTime::MinValue(), bool bSlowTask = true );
 
 	/** The editor wrapper for UPackage::Save. Auto-adds files to source control when necessary */
 	FSavePackageResultStruct Save(UPackage* InOuter, UObject* Base, EObjectFlags TopLevelFlags, const TCHAR* Filename,
-		FOutputDevice* Error = GError, FLinkerLoad* Conform = NULL, bool bForceByteSwapping = false, bool bWarnOfLongFilename = true,
+		FOutputDevice* Error = GError, FLinkerNull* Conform = NULL, bool bForceByteSwapping = false, bool bWarnOfLongFilename = true,
 		uint32 SaveFlags = SAVE_None, const class ITargetPlatform* TargetPlatform = NULL, const FDateTime& FinalTimeStamp = FDateTime::MinValue(), 
 		bool bSlowTask = true, class FArchiveDiffMap* InOutDiffMap = nullptr);
 
diff --git a/Engine/Source/Editor/UnrealEd/Private/Commandlets/PackageUtilities.cpp b/Engine/Source/Editor/UnrealEd/Private/Commandlets/PackageUtilities.cpp
--- a/Engine/Source/Editor/UnrealEd/Private/Commandlets/PackageUtilities.cpp
+++ b/Engine/Source/Editor/UnrealEd/Private/Commandlets/PackageUtilities.cpp
@@ -279,7 +279,7 @@ bool NormalizePackageNames( TArray<FString> PackageNames, TArray<FString>& Packa
 
 * @return true if successful
 */
-bool SavePackageHelper(UPackage* Package, FString Filename, EObjectFlags KeepObjectFlags, FOutputDevice* ErrorDevice, FLinkerLoad* LinkerToConformAgainst, ESaveFlags SaveFlags)
+bool SavePackageHelper(UPackage* Package, FString Filename, EObjectFlags KeepObjectFlags, FOutputDevice* ErrorDevice, FLinkerNull* LinkerToConformAgainst, ESaveFlags SaveFlags)
 {
 	// look for a world object in the package (if there is one, there's a map)
 	UWorld* World = UWorld::FindWorldInPackage(Package);
diff --git a/Engine/Source/Editor/UnrealEd/Private/EditorEngine.cpp b/Engine/Source/Editor/UnrealEd/Private/EditorEngine.cpp
--- a/Engine/Source/Editor/UnrealEd/Private/EditorEngine.cpp
+++ b/Engine/Source/Editor/UnrealEd/Private/EditorEngine.cpp
@@ -4335,7 +4335,7 @@ void UEditorEngine::CleanupPhysicsSceneThatWasInitializedForSave(UWorld* World,
 }
 
 FSavePackageResultStruct UEditorEngine::Save( UPackage* InOuter, UObject* InBase, EObjectFlags TopLevelFlags, const TCHAR* Filename,
-				 FOutputDevice* Error, FLinkerLoad* Conform, bool bForceByteSwapping, bool bWarnOfLongFilename, 
+				 FOutputDevice* Error, FLinkerNull* Conform, bool bForceByteSwapping, bool bWarnOfLongFilename, 
 				 uint32 SaveFlags, const class ITargetPlatform* TargetPlatform, const FDateTime& FinalTimeStamp, bool bSlowTask, FArchiveDiffMap* InOutDiffMap)
 {
 	FScopedSlowTask SlowTask(100, FText(), bSlowTask);
@@ -4432,7 +4432,7 @@ FSavePackageResultStruct UEditorEngine::Save( UPackage* InOuter, UObject* InBase
 }
 
 bool UEditorEngine::SavePackage(UPackage* InOuter, UObject* InBase, EObjectFlags TopLevelFlags, const TCHAR* Filename,
-	FOutputDevice* Error, FLinkerLoad* Conform, bool bForceByteSwapping, bool bWarnOfLongFilename,
+	FOutputDevice* Error, FLinkerNull* Conform, bool bForceByteSwapping, bool bWarnOfLongFilename,
 	uint32 SaveFlags, const class ITargetPlatform* TargetPlatform, const FDateTime& FinalTimeStamp, bool bSlowTask)
 {
 	// Workaround to avoid function signature change while keeping both bool and ESavePackageResult versions of SavePackage
diff --git a/Engine/Source/Editor/UnrealEd/Public/PackageHelperFunctions.h b/Engine/Source/Editor/UnrealEd/Public/PackageHelperFunctions.h
--- a/Engine/Source/Editor/UnrealEd/Public/PackageHelperFunctions.h
+++ b/Engine/Source/Editor/UnrealEd/Public/PackageHelperFunctions.h
@@ -73,7 +73,7 @@ bool UNREALED_API NormalizePackageNames( TArray<FString> PackageNames, TArray<FS
  *
  * @return true if successful
  */
-bool UNREALED_API SavePackageHelper(UPackage* Package, FString Filename,  EObjectFlags KeepObjectFlags = RF_Standalone, FOutputDevice* ErrorDevice=GWarn, FLinkerLoad* LinkerToConformAgainst=NULL, ESaveFlags SaveFlags = SAVE_None);
+bool UNREALED_API SavePackageHelper(UPackage* Package, FString Filename,  EObjectFlags KeepObjectFlags = RF_Standalone, FOutputDevice* ErrorDevice=GWarn, FLinkerNull* LinkerToConformAgainst=NULL, ESaveFlags SaveFlags = SAVE_None);
 
 
 /**
diff --git a/Engine/Source/Runtime/CoreUObject/Private/UObject/SavePackage.cpp b/Engine/Source/Runtime/CoreUObject/Private/UObject/SavePackage.cpp
--- a/Engine/Source/Runtime/CoreUObject/Private/UObject/SavePackage.cpp
+++ b/Engine/Source/Runtime/CoreUObject/Private/UObject/SavePackage.cpp
@@ -3555,12 +3555,14 @@ void AddFileToHash(FString const &Filename, FMD5 &Hash)
 }
 
 FSavePackageResultStruct UPackage::Save(UPackage* InOuter, UObject* Base, EObjectFlags TopLevelFlags, const TCHAR* Filename,
-	FOutputDevice* Error, FLinkerLoad* Conform, bool bForceByteSwapping, bool bWarnOfLongFilename, uint32 SaveFlags, 
+	FOutputDevice* Error, FLinkerNull* ConformNO, bool bForceByteSwapping, bool bWarnOfLongFilename, uint32 SaveFlags, 
 	const class ITargetPlatform* TargetPlatform, const FDateTime&  FinalTimeStamp, bool bSlowTask, FArchiveDiffMap* InOutDiffMap)
 {
 	COOK_STAT(FScopedDurationTimer FuncSaveTimer(SavePackageStats::SavePackageTimeSec));
 	COOK_STAT(SavePackageStats::NumPackagesSaved++);
 
+	FLinkerLoad* Conform = nullptr;
+
 #if WITH_EDITOR
 	TMap<UObject*, UObject*> ReplacedImportOuters;
 #endif //WITH_EDITOR
@@ -6309,7 +6311,7 @@ FSavePackageResultStruct UPackage::Save(UPackage* InOuter, UObject* Base, EObjec
 }
 
 bool UPackage::SavePackage(UPackage* InOuter, UObject* Base, EObjectFlags TopLevelFlags, const TCHAR* Filename,
-	FOutputDevice* Error, FLinkerLoad* Conform, bool bForceByteSwapping, bool bWarnOfLongFilename, uint32 SaveFlags,
+	FOutputDevice* Error, FLinkerNull* Conform, bool bForceByteSwapping, bool bWarnOfLongFilename, uint32 SaveFlags,
 	const class ITargetPlatform* TargetPlatform, const FDateTime&  FinalTimeStamp, bool bSlowTask)
 {
 	const FSavePackageResultStruct Result = Save(InOuter, Base, TopLevelFlags, Filename, Error, Conform, bForceByteSwapping,
diff --git a/Engine/Source/Runtime/CoreUObject/Public/UObject/Package.h b/Engine/Source/Runtime/CoreUObject/Public/UObject/Package.h
--- a/Engine/Source/Runtime/CoreUObject/Public/UObject/Package.h
+++ b/Engine/Source/Runtime/CoreUObject/Public/UObject/Package.h
@@ -16,6 +16,10 @@
 
 class Error;
 
+// This is a dummy type which is not implemented anywhere. It's only 
+// used to flag a deprecated Conform argument to package save functions.
+class FLinkerNull;
+
 /**
 * Represents the result of saving a package
 */
@@ -540,7 +544,7 @@ public:
 	* @return	FSavePackageResultStruct enum value with the result of saving a package as well as extra data
 	*/
 	static FSavePackageResultStruct Save(UPackage* InOuter, UObject* Base, EObjectFlags TopLevelFlags, const TCHAR* Filename,
-		FOutputDevice* Error=GError, FLinkerLoad* Conform=NULL, bool bForceByteSwapping=false, bool bWarnOfLongFilename=true, 
+		FOutputDevice* Error=GError, FLinkerNull* Conform=NULL, bool bForceByteSwapping=false, bool bWarnOfLongFilename=true, 
 		uint32 SaveFlags=SAVE_None, const class ITargetPlatform* TargetPlatform = NULL, const FDateTime& FinalTimeStamp = FDateTime::MinValue(), 
 		bool bSlowTask = true, class FArchiveDiffMap* InOutDiffMap = nullptr);
 
@@ -565,7 +569,7 @@ public:
 	* @return	true if the package was saved successfully.
 	*/
 	static bool SavePackage(UPackage* InOuter, UObject* Base, EObjectFlags TopLevelFlags, const TCHAR* Filename,
-		FOutputDevice* Error = GError, FLinkerLoad* Conform = NULL, bool bForceByteSwapping = false, bool bWarnOfLongFilename = true,
+		FOutputDevice* Error = GError, FLinkerNull* Conform = NULL, bool bForceByteSwapping = false, bool bWarnOfLongFilename = true,
 		uint32 SaveFlags = SAVE_None, const class ITargetPlatform* TargetPlatform = NULL, const FDateTime& FinalTimeStamp = FDateTime::MinValue(), bool bSlowTask = true);
 
 	/** Wait for any SAVE_Async file writes to complete **/
