Removed unused  MAX_MERGED_COMPRESSION_CHUNKSIZE, ENABLE_PACKAGE_CLASS_SERIALIZATION_TIMES, ENABLE_TAGEXPORTS_CLASS_PRESAVE_TIMES

diff --git a/Engine/Source/Runtime/CoreUObject/Private/UObject/SavePackage.cpp b/Engine/Source/Runtime/CoreUObject/Private/UObject/SavePackage.cpp
--- a/Engine/Source/Runtime/CoreUObject/Private/UObject/SavePackage.cpp
+++ b/Engine/Source/Runtime/CoreUObject/Private/UObject/SavePackage.cpp
@@ -60,7 +60,6 @@
 
 DEFINE_LOG_CATEGORY_STATIC(LogSavePackage, Log, All);
 
-static const int32 MAX_MERGED_COMPRESSION_CHUNKSIZE = 1024 * 1024;
 static const FName WorldClassName = FName("World");
 static const FName PrestreamPackageClassName = FName("PrestreamPackage");
 static FCriticalSection InitializeCoreClassesCritSec;
@@ -68,13 +67,6 @@ static FCriticalSection InitializeCoreClassesCritSec;
 #define VALIDATE_INITIALIZECORECLASSES 0
 #define EXPORT_SORTING_DETAILED_LOGGING 0
 
-// Uncomment this code to measure UObject::Serialize time taken per UClass type during save.
-// This code tracks each type in a hash, so is too heavyweight to leave on in general,
-// but can be really useful for telling what classes are causing the most save time.
-#define ENABLE_PACKAGE_CLASS_SERIALIZATION_TIMES 0
-// uncomment this code to measure UObject::PreSave time taken per uclass type during save
-#define ENABLE_TAGEXPORTS_CLASS_PRESAVE_TIMES 0
-
 #if ENABLE_COOK_STATS
 #include "ProfilingDebugging/ScopedTimers.h"
 
@@ -1175,11 +1167,6 @@ void FArchiveSaveTagExports::ProcessTaggedObjects()
 			}
 			// In the CDO case the above would serialize most of the references, including transient properties
 			// but we still want to serialize the object using the normal path to collect all custom versions it might be using.
-#if ENABLE_PACKAGE_CLASS_SERIALIZATION_TIMES
-			auto& TimingInfo = SavePackageStats::TagExportSerializeTimes.FindOrAdd(Obj->GetClass()->GetFName());
-			TimingInfo.Value++;
-			FScopedDurationTimer SerializeTimer(TimingInfo.Key);
-#endif
 			Obj->Serialize(*this);
 		}
 
@@ -2887,11 +2874,6 @@ struct FPackageExportTagger
 		{
 			if ( bRoutePresave )
 			{
-#if ENABLE_TAGEXPORTS_CLASS_PRESAVE_TIMES
-				auto& TimingInfo = SavePackageStats::ClassPreSaveTimes.FindOrAdd(Base->GetClass()->GetFName());
-				TimingInfo.Value++;
-				FScopedDurationTimer SerializeTimer(TimingInfo.Key);
-#endif
 				if (bIsCooking && Base->HasAnyFlags(RF_ClassDefaultObject | RF_ArchetypeObject))
 				{
 					FArchiveObjectCrc32NonEditorProperties CrcArchive;
@@ -2949,11 +2931,6 @@ struct FPackageExportTagger
 				for(int32 Index = 0; Index < TagExpObjects.Num(); Index++)
 				{
 					UObject* Obj = TagExpObjects[Index];
-#if ENABLE_TAGEXPORTS_CLASS_PRESAVE_TIMES
-					auto& TimingInfo = SavePackageStats::ClassPreSaveTimes.FindOrAdd(Obj->GetClass()->GetFName());
-					TimingInfo.Value++;
-					FScopedDurationTimer SerializeTimer(TimingInfo.Key);
-#endif
 					check(Obj->HasAnyMarks(OBJECTMARK_TagExp));
 					//@warning: Objects created from within PreSave will NOT have PreSave called on them!!!
 					if (bIsCooking && Obj->HasAnyFlags(RF_ClassDefaultObject | RF_ArchetypeObject))
@@ -5570,11 +5547,6 @@ FSavePackageResultStruct UPackage::Save(UPackage* InOuter, UObject* Base, EObjec
 							}
 							else
 							{
-#if ENABLE_PACKAGE_CLASS_SERIALIZATION_TIMES
-								auto& TimingInfo = SavePackageStats::PackageClassSerializeTimes.FindOrAdd(Export.Object->GetClass()->GetFName());
-								TimingInfo.Value++;
-								FScopedDurationTimer SerializeTimer(TimingInfo.Key);
-#endif
 								TGuardValue<UObject*> GuardSerializedObject(SaveContext->SerializedObject, Export.Object);
 
 								if (bSupportsText)
