Fix UHT

diff --git a/Engine/Source/Runtime/CoreUObject/Private/UObject/Class.cpp b/Engine/Source/Runtime/CoreUObject/Private/UObject/Class.cpp
--- a/Engine/Source/Runtime/CoreUObject/Private/UObject/Class.cpp
+++ b/Engine/Source/Runtime/CoreUObject/Private/UObject/Class.cpp
@@ -2009,9 +2009,10 @@ void UScriptStruct::PrepareCppStructOps()
 			bPrepareCppStructOpsCompleted = true;
 			return;
 		}
-
+#if !HACK_HEADER_GENERATOR
 		StructFlags = EStructFlags(StructFlags | STRUCT_Native);
-
+#endif
+		
 #if !(UE_BUILD_SHIPPING || UE_BUILD_TEST)
 		// test that the constructor is initializing everything
 		if (!CppStructOps->HasZeroConstructor()
