TBA: Fixes for loading text assets without import/export tables
 - Don't write import table, exporttable, searchable names or soft asset references to text file
 - Remove a dummy attribute read from export preload which was only there as a workaround for a now-fixed crash in the attribute reading code

diff --git a/Engine/Source/Runtime/CoreUObject/Private/UObject/LinkerLoad.cpp b/Engine/Source/Runtime/CoreUObject/Private/UObject/LinkerLoad.cpp
--- a/Engine/Source/Runtime/CoreUObject/Private/UObject/LinkerLoad.cpp
+++ b/Engine/Source/Runtime/CoreUObject/Private/UObject/LinkerLoad.cpp
@@ -731,12 +731,14 @@ FLinkerLoad::ELinkerStatus FLinkerLoad::Tick( float InTimeLimit, bool bInUseTime
 				Status = SerializeExportMap();
 			}
 
+#if WITH_TEXT_ARCHIVE_SUPPORT
 			// Reconstruct the import and export maps for text assets
 			if (Status == LINKER_Loaded)
 			{
 				SCOPED_LOADTIMER(LinkerLoad_ReconstructImportAndExportMap);
 				Status = ReconstructImportAndExportMap();
 			}
+#endif
 
 			// Fix up import map for backward compatible serialization.
 			if( Status == LINKER_Loaded )
@@ -1730,6 +1732,8 @@ FLinkerLoad::ELinkerStatus FLinkerLoad::SerializeExportMap()
 	return ((ExportMapIndex == Summary.ExportCount) && !IsTimeLimitExceeded( TEXT("serializing export map") )) ? LINKER_Loaded : LINKER_TimedOut;
 }
 
+#if WITH_TEXT_ARCHIVE_SUPPORT
+
 FPackageIndex FLinkerLoad::FindOrCreateImport(const FName InObjectName, const FName InClassName, const FName InClassPackageName)
 {
 	for (int32 ImportIndex = 0; ImportIndex < ImportMap.Num(); ++ImportIndex)
@@ -1754,6 +1758,20 @@ FPackageIndex FLinkerLoad::FindOrCreateImport(const FName InObjectName, const FN
 	return FPackageIndex::FromImport(ImportMap.Num() - 1);
 }
 
+FString ExtractObjectName(const FString& InFullPath)
+{
+	FString ObjectName = InFullPath;
+	int32 LastDot, LastSemi;
+	InFullPath.FindLastChar('.', LastDot);
+	InFullPath.FindLastChar(':', LastSemi);
+	int32 StartOfObjectName = FMath::Max(LastDot, LastSemi);
+	if (StartOfObjectName != INDEX_NONE)
+	{
+		return ObjectName.Right(ObjectName.Len() - StartOfObjectName - 1);
+	}
+	return ObjectName;
+}
+
 FPackageIndex FLinkerLoad::FindOrCreateImportOrExport(const FString& InFullPath)
 {
 	if (InFullPath.Len() == 0)
@@ -1763,15 +1781,15 @@ FPackageIndex FLinkerLoad::FindOrCreateImportOrExport(const FString& InFullPath)
 
 	FString Class, Package, Object, SubObject;
 	FPackageName::SplitFullObjectPath(InFullPath, Class, Package, Object, SubObject);
-	bool bIsExport = (Package == LinkerRoot->GetName());
+	FName ObjectName = *(Object + (SubObject.Len() ? TEXT(":") : TEXT("")) + SubObject);
 
-	FName ObjectName = *Object;
+	bool bIsExport = Package == LinkerRoot->GetName();
 
 	if (bIsExport)
 	{
 		for (int32 ExportIndex = 0; ExportIndex < ExportMap.Num(); ++ExportIndex)
 		{
-			if (ExportMap[ExportIndex].ObjectName == ObjectName)
+			if (OriginalExportNames[ExportIndex] == ObjectName)
 			{
 				return FPackageIndex::FromExport(ExportIndex);
 			}
@@ -1811,11 +1829,11 @@ FLinkerLoad::ELinkerStatus FLinkerLoad::ReconstructImportAndExportMap()
 	{
 		int32 NumExports = 0;
 		FStructuredArchive::FMap PackageExports = StructuredArchiveRootRecord->EnterMap(SA_FIELD_NAME(TEXT("Exports")), NumExports);
-		ExportMap.Reserve(NumExports);
 
 		TArray<FObjectTextExport> ExportRecords;
 		ExportRecords.Reserve(NumExports);
 		ExportMap.SetNum(NumExports);
+		OriginalExportNames.SetNum(NumExports);
 
 		Summary.ExportCount = ExportMap.Num();
 		Summary.ImportCount = 0;
@@ -1823,9 +1841,10 @@ FLinkerLoad::ELinkerStatus FLinkerLoad::ReconstructImportAndExportMap()
 		for (int32 ExportIndex = 0; ExportIndex < NumExports; ++ExportIndex)
 		{
 			FObjectTextExport& TextExport = ExportRecords.Emplace_GetRef(ExportMap[ExportIndex], nullptr);
-			FString ObjectName;
-			PackageExports.EnterElement(ObjectName) << TextExport;
-			ExportMap[ExportIndex].ObjectName = *ObjectName;
+			FString ExportName;
+			PackageExports.EnterElement(ExportName) << TextExport;
+			OriginalExportNames[ExportIndex] = *ExportName;
+			ExportMap[ExportIndex].ObjectName = *ExtractObjectName(OriginalExportNames[ExportIndex].ToString());
 		}
 		
 		// Now pass over all the exports and rebuild the export/import records
@@ -1848,6 +1867,8 @@ FLinkerLoad::ELinkerStatus FLinkerLoad::ReconstructImportAndExportMap()
 	}
 }
 
+#endif // WITH_TEXT_ARCHIVE_SUPPORT
+
 /**
  * Serializes the depends map.
  */
@@ -3669,8 +3690,7 @@ void FLinkerLoad::Preload( UObject* Object )
 #if WITH_EDITOR && WITH_TEXT_ARCHIVE_SUPPORT
 						if (IsTextFormat())
 						{
-							FString ObjectName = Object->GetPathName(Object->GetOutermost());
-							FStructuredArchive::FSlot ExportSlot = StructuredArchiveRootRecord->EnterField(SA_FIELD_NAME(*ObjectName));
+							FStructuredArchive::FSlot ExportSlot = StructuredArchiveRootRecord->EnterField(SA_FIELD_NAME(*OriginalExportNames[Export.ThisIndex.ToExport()].ToString()));
 
 							if (bClassSupportsTextFormat)
 							{
@@ -3710,11 +3730,7 @@ void FLinkerLoad::Preload( UObject* Object )
 #if WITH_EDITOR && WITH_TEXT_ARCHIVE_SUPPORT
 						if (IsTextFormat())
 						{
-							FStructuredArchive::FSlot ExportSlot = StructuredArchiveRootRecord->EnterRecord(SA_FIELD_NAME(TEXT("Exports"))).EnterField(SA_FIELD_NAME(*Export.ObjectName.ToString()));
-
-							// TODO BEFORE CHECKIN: This attribute read only exists because we crash if we attempt to serialize the value of an attributed field when we don't read any attributes first
-							FString Temp;
-							ExportSlot << SA_ATTRIBUTE(TEXT("Class"), Temp);
+							FStructuredArchive::FSlot ExportSlot = StructuredArchiveRootRecord->EnterRecord(SA_FIELD_NAME(TEXT("Exports"))).EnterField(SA_FIELD_NAME(*OriginalExportNames[Export.ThisIndex.ToExport()].ToString()));
 
 							if (bClassSupportsTextFormat)
 							{
diff --git a/Engine/Source/Runtime/CoreUObject/Private/UObject/SavePackage.cpp b/Engine/Source/Runtime/CoreUObject/Private/UObject/SavePackage.cpp
--- a/Engine/Source/Runtime/CoreUObject/Private/UObject/SavePackage.cpp
+++ b/Engine/Source/Runtime/CoreUObject/Private/UObject/SavePackage.cpp
@@ -5080,11 +5080,11 @@ FSavePackageResultStruct UPackage::Save(UPackage* InOuter, UObject* Base, EObjec
 						{
 							SoftReferenceStream.EnterElement() << SoftPackageName;
 						}
-					}
 
-					// Save searchable names map
-					Linker->Summary.SearchableNamesOffset = Linker->Tell();
-					Linker->SerializeSearchableNamesMap(StructuredArchiveRoot.EnterField(SA_FIELD_NAME(TEXT("SearchableNames"))));
+						// Save searchable names map
+						Linker->Summary.SearchableNamesOffset = Linker->Tell();
+						Linker->SerializeSearchableNamesMap(StructuredArchiveRoot.EnterField(SA_FIELD_NAME(TEXT("SearchableNames"))));
+					}
 				}
 				else
 				{
@@ -5501,7 +5501,6 @@ FSavePackageResultStruct UPackage::Save(UPackage* InOuter, UObject* Base, EObjec
 					FStructuredArchive::FRecord ExportsRecord = StructuredArchiveRoot.EnterRecord(SA_FIELD_NAME(TEXT("Exports")));
 
 					// Save exports.
-					FString ObjectName;
 					int32 LastExportSaveStep = 0;
 					for( int32 i=0; i<Linker->ExportMap.Num(); i++ )
 					{
@@ -5519,8 +5518,7 @@ FSavePackageResultStruct UPackage::Save(UPackage* InOuter, UObject* Base, EObjec
 							Linker->CurrentlySavingExport = FPackageIndex::FromExport(i);
 							// UE_LOG(LogSavePackage, Log, TEXT("export %s for %s"), *Export.Object->GetFullName(), *Linker->CookingTarget()->PlatformName());
 
-							ObjectName.Reset();
-							Export.Object->GetPathName(InOuter, ObjectName);
+							FString ObjectName = Export.Object->GetPathName(InOuter);
 							FStructuredArchive::FSlot ExportSlot = ExportsRecord.EnterField(SA_FIELD_NAME(*ObjectName));
 
 							if (bTextFormat)
@@ -5860,88 +5858,85 @@ FSavePackageResultStruct UPackage::Save(UPackage* InOuter, UObject* Base, EObjec
 					if (!bTextFormat)
 					{
 						Linker->Seek(Linker->Summary.ImportOffset);
-					}
 
-					int32 NumImports = Linker->ImportMap.Num();
-					FStructuredArchive::FStream ImportTableStream = StructuredArchiveRoot.EnterStream(SA_FIELD_NAME(TEXT("ImportTable")));
+						int32 NumImports = Linker->ImportMap.Num();
+						FStructuredArchive::FStream ImportTableStream = StructuredArchiveRoot.EnterStream(SA_FIELD_NAME(TEXT("ImportTable")));
 
-					for (int32 i = 0; i < Linker->ImportMap.Num(); i++)
-					{
-						FObjectImport& Import = Linker->ImportMap[i];
-						if (Import.XObject)
+						for (int32 i = 0; i < Linker->ImportMap.Num(); i++)
 						{
-							// Set the package index.
-							if (Import.XObject->GetOuter())
+							FObjectImport& Import = Linker->ImportMap[i];
+							if (Import.XObject)
 							{
-								if (Import.XObject->GetOuter()->IsIn(InOuter))
+								// Set the package index.
+								if (Import.XObject->GetOuter())
 								{
-									if (!Import.XObject->HasAllFlags(RF_Transient) || !Import.XObject->IsNative())
+									if (Import.XObject->GetOuter()->IsIn(InOuter))
 									{
-										UE_LOG(LogSavePackage, Warning, TEXT("Bad Object=%s"), *Import.XObject->GetFullName());
+										if (!Import.XObject->HasAllFlags(RF_Transient) || !Import.XObject->IsNative())
+										{
+											UE_LOG(LogSavePackage, Warning, TEXT("Bad Object=%s"), *Import.XObject->GetFullName());
+										}
+										else
+										{
+											// if an object is marked RF_Transient and native, it is either an intrinsic class or
+											// a property of an intrinsic class.  Only properties of intrinsic classes will have
+											// an Outer that passes the check for "GetOuter()->IsIn(InOuter)" (thus ending up in this
+											// block of code).  Just verify that the Outer for this property is also marked RF_Transient and Native
+											check(Import.XObject->GetOuter()->HasAllFlags(RF_Transient) && Import.XObject->GetOuter()->IsNative());
+										}
 									}
-									else
+									check(!Import.XObject->GetOuter()->IsIn(InOuter) || Import.XObject->HasAllFlags(RF_Transient) || Import.XObject->IsNative());
+#if WITH_EDITOR
+									UObject** ReplacedOuter = ReplacedImportOuters.Find(Import.XObject);
+									if (ReplacedOuter && *ReplacedOuter)
 									{
-										// if an object is marked RF_Transient and native, it is either an intrinsic class or
-										// a property of an intrinsic class.  Only properties of intrinsic classes will have
-										// an Outer that passes the check for "GetOuter()->IsIn(InOuter)" (thus ending up in this
-										// block of code).  Just verify that the Outer for this property is also marked RF_Transient and Native
-										check(Import.XObject->GetOuter()->HasAllFlags(RF_Transient) && Import.XObject->GetOuter()->IsNative());
+										Import.OuterIndex = Linker->MapObject(*ReplacedOuter);
+										ensure(Import.OuterIndex != FPackageIndex());
 									}
-								}
-								check(!Import.XObject->GetOuter()->IsIn(InOuter) || Import.XObject->HasAllFlags(RF_Transient) || Import.XObject->IsNative());
-#if WITH_EDITOR
-								UObject** ReplacedOuter = ReplacedImportOuters.Find(Import.XObject);
-								if (ReplacedOuter && *ReplacedOuter)
-								{
-									Import.OuterIndex = Linker->MapObject(*ReplacedOuter);
-									ensure(Import.OuterIndex != FPackageIndex());
-								}
-								else
+									else
 #endif
-								{
-									Import.OuterIndex = Linker->MapObject(Import.XObject->GetOuter());
-								}
+									{
+										Import.OuterIndex = Linker->MapObject(Import.XObject->GetOuter());
+									}
 
-								if (Linker->IsCooking() && IsEventDrivenLoaderEnabledInCookedBuilds())
-								{
-									// Only package imports are allowed to have no outer
-									ensureMsgf(Import.OuterIndex != FPackageIndex() || Import.ClassName == NAME_Package, TEXT("Import %s has no valid outer when cooking!"), *Import.XObject->GetPathName());
+									if (Linker->IsCooking() && IsEventDrivenLoaderEnabledInCookedBuilds())
+									{
+										// Only package imports are allowed to have no outer
+										ensureMsgf(Import.OuterIndex != FPackageIndex() || Import.ClassName == NAME_Package, TEXT("Import %s has no valid outer when cooking!"), *Import.XObject->GetPathName());
+									}
 								}
 							}
-						}
-						else
-						{
-							checkf(Conform != nullptr, TEXT("NULL XObject for import %i - Object: %s Class: %s"), i, *Import.ObjectName.ToString(), *Import.ClassName.ToString());
-						}
+							else
+							{
+								checkf(Conform != nullptr, TEXT("NULL XObject for import %i - Object: %s Class: %s"), i, *Import.ObjectName.ToString(), *Import.ClassName.ToString());
+							}
 
-						// Save it.
-						ImportTableStream.EnterElement() << Import;
+							// Save it.
+							ImportTableStream.EnterElement() << Import;
+						}
 					}
 				}
 
 				// Save the export map.
 				if (!bTextFormat)
 				{
-				check( Linker->Tell() == OffsetAfterImportMap );
+					check( Linker->Tell() == OffsetAfterImportMap );
 					Linker->Seek(Linker->Summary.ExportOffset);
-				}
 
-				int32 NumExports = Linker->ExportMap.Num();
-				FStructuredArchive::FStream ExportTableStream = StructuredArchiveRoot.EnterStream(SA_FIELD_NAME(TEXT("ExportTable")));
-				{
+					int32 NumExports = Linker->ExportMap.Num();
+					FStructuredArchive::FStream ExportTableStream = StructuredArchiveRoot.EnterStream(SA_FIELD_NAME(TEXT("ExportTable")));
+					{
 #if WITH_EDITOR
-					FArchive::FScopeSetDebugSerializationFlags S(*Linker, DSF_IgnoreDiff, true);
-					FArchiveStackTraceIgnoreScope IgnoreSummaryDiffsScope(DiffSettings.bIgnoreHeaderDiffs);
+						FArchive::FScopeSetDebugSerializationFlags S(*Linker, DSF_IgnoreDiff, true);
+						FArchiveStackTraceIgnoreScope IgnoreSummaryDiffsScope(DiffSettings.bIgnoreHeaderDiffs);
 #endif
-					for (int32 i = 0; i < Linker->ExportMap.Num(); i++)
-					{
-						FObjectExport& Export = Linker->ExportMap[i];
-						ExportTableStream.EnterElement() << Export;
+						for (int32 i = 0; i < Linker->ExportMap.Num(); i++)
+						{
+							FObjectExport& Export = Linker->ExportMap[i];
+							ExportTableStream.EnterElement() << Export;
+						}
 					}
-				}
 
-				if (!bTextFormat)
-				{
 					check( Linker->Tell() == OffsetAfterExportMap );
 				}
 
diff --git a/Engine/Source/Runtime/CoreUObject/Public/UObject/LinkerLoad.h b/Engine/Source/Runtime/CoreUObject/Public/UObject/LinkerLoad.h
--- a/Engine/Source/Runtime/CoreUObject/Public/UObject/LinkerLoad.h
+++ b/Engine/Source/Runtime/CoreUObject/Public/UObject/LinkerLoad.h
@@ -929,10 +929,12 @@ private:
 	 */
 	ELinkerStatus SerializeExportMap();
 
+#if WITH_TEXT_ARCHIVE_SUPPORT
 	/**
 	 * Create an import and export table when loading a text asset.
 	 */
 	ELinkerStatus ReconstructImportAndExportMap();
+#endif
 
 	ELinkerStatus SerializeDependsMap();
 
@@ -1196,6 +1198,13 @@ private:
 	// FLinkerLoad creation helpers END
 	//
 
+private:
+
+#if WITH_TEXT_ARCHIVE_SUPPORT
+	// Cache of the export names in a text asset. Allows us to enter those export slots by index rather than needing to reconstruct the name
+	TArray<FName> OriginalExportNames;
+#endif
+
 public:
 
 	//~ FArchive interface
