Fixed truncation of bulkdata size to a 32-bit integer even though we try to handle 64-bit sizes elsewhere

diff --git a/Engine/Source/Runtime/CoreUObject/Private/UObject/SavePackage.cpp b/Engine/Source/Runtime/CoreUObject/Private/UObject/SavePackage.cpp
--- a/Engine/Source/Runtime/CoreUObject/Private/UObject/SavePackage.cpp
+++ b/Engine/Source/Runtime/CoreUObject/Private/UObject/SavePackage.cpp
@@ -5734,7 +5734,7 @@ FSavePackageResultStruct UPackage::Save(UPackage* InOuter, UObject* Base, EObjec
 						int64 BulkEndOffset = TargetArchive->Tell();
 						int64 LinkerEndOffset = Linker->Tell();
 
-						int64 SizeOnDisk = (int32)(BulkEndOffset - BulkStartOffset);
+						int64 SizeOnDisk = BulkEndOffset - BulkStartOffset;
 				
 						Linker->Seek(BulkDataStorageInfo.BulkDataFlagsPos);
 						*Linker << ModifiedBulkDataFlags;
