Restructured how bulk data saving works to prepare for package store related changes. Partially to reduce code complexity but also to reduce memory footprint

(ushell-p4-cherrypick of 8640322 by [REDACTED_PERSON])

diff --git a/Engine/Source/Runtime/CoreUObject/Private/UObject/SavePackage.cpp b/Engine/Source/Runtime/CoreUObject/Private/UObject/SavePackage.cpp
--- a/Engine/Source/Runtime/CoreUObject/Private/UObject/SavePackage.cpp
+++ b/Engine/Source/Runtime/CoreUObject/Private/UObject/SavePackage.cpp
@@ -488,7 +488,7 @@ void AsyncWriteFile(FLargeMemoryPtr Data, const int64 DataSize, const TCHAR* Fil
 			check(DataSize);
 			FString TempFilename; 
 
-			if ( bUseTempFilename )
+			if (bUseTempFilename)
 			{
 				TempFilename = FPaths::GetBaseFilename(Filename, false);
 				TempFilename += TEXT(".t");
@@ -510,17 +510,21 @@ void AsyncWriteFile(FLargeMemoryPtr Data, const int64 DataSize, const TCHAR* Fil
 
 				if (IFileManager::Get().FileSize(*TempFilename) == DataSize)
 				{
-					if ( bUseTempFilename )
+					if (bUseTempFilename)
 					{
-						if( !IFileManager::Get().Move(*Filename, *TempFilename, true, true, false, false))
+						if (!IFileManager::Get().Move(*Filename, *TempFilename, true, true, false, false))
 						{
 							UE_LOG(LogSavePackage, Fatal, TEXT("Could not move to %s."),*Filename);
 						}
 
-						// if everything worked, this is not necessary, but we will make every effort to avoid leaving junk in the cache
+						// If everything worked as it should this is not necessary, but we want to ensure
+						// no temporary files are left lying around if something does go wrong.
+
 						if (FPaths::FileExists(TempFilename))
 						{
 							IFileManager::Get().Delete(*TempFilename);
+
+							UE_LOG(LogSavePackage, Fatal, TEXT("Deleted temporary file '%s'"), *TempFilename);
 						}
 					}
 
@@ -531,6 +535,9 @@ void AsyncWriteFile(FLargeMemoryPtr Data, const int64 DataSize, const TCHAR* Fil
 				}
 				else
 				{
+					// Don't leave incomplete files around
+					IFileManager::Get().Delete(*TempFilename);
+
 					UE_LOG(LogSavePackage, Fatal, TEXT("Could not save to %s!"),*TempFilename);
 				}
 			}
@@ -3505,14 +3512,6 @@ void VerifyEDLCookInfo()
 	FEDLCookChecker::Verify();
 }
 
-static void AddArchiveToHash(FBufferArchive& Ar, FMD5& OutHash)
-{
-	if (int64 Size = Ar.TotalSize())
-	{
-		OutHash.Update(Ar.GetData(), Size);
-	}
-}
-
 void AddFileToHash(FString const &Filename, FMD5 &Hash)
 {
 	TArray<uint8> LocalScratch;
@@ -5601,7 +5600,7 @@ FSavePackageResultStruct UPackage::Save(UPackage* InOuter, UObject* Base, EObjec
 
 				// now we write all the bulkdata that is supposed to be at the end of the package
 				// and fix up the offset
-				int64 StartOfBulkDataArea = Linker->Tell();
+				const int64 StartOfBulkDataArea = Linker->Tell();
 				Linker->Summary.BulkDataStartOffset = StartOfBulkDataArea;
 
 				check(!bTextFormat || Linker->BulkDataToAppend.Num() == 0);
@@ -5612,57 +5611,37 @@ FSavePackageResultStruct UPackage::Save(UPackage* InOuter, UObject* Base, EObjec
 
 					FScopedSlowTask BulkDataFeedback(Linker->BulkDataToAppend.Num());
 
-					FArchive* BulkArchive = nullptr;
-					FArchive* OptionalBulkArchive = nullptr;
-					FArchive* MappedBulkArchive = nullptr;
+					TUniquePtr<FLargeMemoryWriter> BulkArchive;
+					TUniquePtr<FLargeMemoryWriter> OptionalBulkArchive;
+					TUniquePtr<FLargeMemoryWriter> MappedBulkArchive;
+
 					uint32 ExtraBulkDataFlags = 0;
 
-					static struct FUseSeperateBulkDataFiles
+					static const struct FUseSeparateBulkDataFiles
 					{
-						bool bEnable;
-						FUseSeperateBulkDataFiles()
+						bool bEnable = false;
+
+						FUseSeparateBulkDataFiles()
 						{
-							if (!GConfig->GetBool(TEXT("Core.System"), TEXT("UseSeperateBulkDataFiles"), bEnable, GEngineIni))
-							{
-								bEnable = false;
-							}
+							GConfig->GetBool(TEXT("Core.System"), TEXT("UseSeperateBulkDataFiles"), /* out */ bEnable, GEngineIni);
+
 							if (IsEventDrivenLoaderEnabledInCookedBuilds())
 							{
 								// Always split bulk data when splitting cooked files
 								bEnable = true;
 							}
 						}
-					} ShouldUseSeperateBulkDataFiles;
-
-					const bool bShouldUseSeparateBulkFile = ShouldUseSeperateBulkDataFiles.bEnable && Linker->IsCooking();
+					} ShouldUseSeparateBulkDataFiles;
 
-					const FString BulkFilename = FPaths::ChangeExtension(Filename, TEXT(".ubulk"));
+					const bool bShouldUseSeparateBulkFile = ShouldUseSeparateBulkDataFiles.bEnable && Linker->IsCooking();
 
-					const static TCHAR* OptionalBulkFileExtension = TEXT("uptnl");
-
-					const FString OptionalBulkFilename = FPaths::ChangeExtension(Filename, OptionalBulkFileExtension);
-
-					const static TCHAR* MappedBulkFileExtension = TEXT(".m.ubulk");
-
-					const FString MappedBulkFilename = FPaths::ChangeExtension(Filename, MappedBulkFileExtension);
-
-					bool bBufferedBulkArchives = false;
 					if (bShouldUseSeparateBulkFile)
 					{
-						ExtraBulkDataFlags = BULKDATA_PayloadInSeperateFile;
-						if (bSaveAsync || bDiffing)
-						{
-							bBufferedBulkArchives = true;
-							BulkArchive = new FBufferArchive(true);
-							OptionalBulkArchive = new FBufferArchive(true);
-							MappedBulkArchive = new FBufferArchive(true);
-						}
-						else
-						{
-							BulkArchive = IFileManager::Get().CreateFileWriter(*BulkFilename);
-							OptionalBulkArchive = IFileManager::Get().CreateFileWriter(*OptionalBulkFilename);
-							MappedBulkArchive = IFileManager::Get().CreateFileWriter(*MappedBulkFilename);
-						}
+						ExtraBulkDataFlags	= BULKDATA_PayloadInSeperateFile;
+
+						BulkArchive.Reset(new FLargeMemoryWriter(0, /* IsPersistent */ true));
+						OptionalBulkArchive.Reset(new FLargeMemoryWriter(0, /* IsPersistent */ true));
+						MappedBulkArchive.Reset(new FLargeMemoryWriter(0, /* IsPersistent */ true));
 					}
 
 					bool bAlignBulkData = false;
@@ -5674,16 +5653,14 @@ FSavePackageResultStruct UPackage::Save(UPackage* InOuter, UObject* Base, EObjec
 						BulkDataAlignment = TargetPlatform->GetMemoryMappingAlignment();
 					}
 
-					for (int32 i=0; i < Linker->BulkDataToAppend.Num(); ++i)
+					for (FLinkerSave::FBulkDataStorageInfo& BulkDataStorageInfo : Linker->BulkDataToAppend)
 					{
 						BulkDataFeedback.EnterProgressFrame();
 
-						FLinkerSave::FBulkDataStorageInfo& BulkDataStorageInfo = Linker->BulkDataToAppend[i];
-
 						// Set bulk data flags to what they were during initial serialization (they might have changed after that)
 						const uint32 OldBulkDataFlags = BulkDataStorageInfo.BulkData->GetBulkDataFlags();
 						uint32 ModifiedBulkDataFlags = BulkDataStorageInfo.BulkDataFlags | ExtraBulkDataFlags;
-						bool bBulkItemIsOptional = (ModifiedBulkDataFlags & BULKDATA_OptionalPayload) != 0;
+						const bool bBulkItemIsOptional = (ModifiedBulkDataFlags & BULKDATA_OptionalPayload) != 0;
 						bool bBulkItemIsMapped = bAlignBulkData && ((ModifiedBulkDataFlags & BULKDATA_MemoryMappedPayload) != 0);
 
 						if (bBulkItemIsMapped && bBulkItemIsOptional)
@@ -5697,42 +5674,64 @@ FSavePackageResultStruct UPackage::Save(UPackage* InOuter, UObject* Base, EObjec
 						BulkDataStorageInfo.BulkData->SetBulkDataFlags(ModifiedBulkDataFlags);
 
 						FArchive* TargetArchive = Linker.Get();
-						if ( bShouldUseSeparateBulkFile )
+
+						if (bShouldUseSeparateBulkFile)
 						{
-							check( OptionalBulkArchive && BulkArchive );
-							TargetArchive = bBulkItemIsOptional ? OptionalBulkArchive : (bBulkItemIsMapped ? MappedBulkArchive : BulkArchive);
+							check(MappedBulkArchive && OptionalBulkArchive && BulkArchive);
+
+							if (bBulkItemIsOptional)
+							{
+								TargetArchive = OptionalBulkArchive.Get();
+							}
+							else if (bBulkItemIsMapped)
+							{
+								TargetArchive = MappedBulkArchive.Get();
+							}
+							else
+							{
+								TargetArchive = BulkArchive.Get();
+							}
 						}
 
-						int64 BulkStartOffset = TargetArchive->Tell();
+						// Pad archive for proper alignment for memory mapping
 
-						if (bBulkItemIsMapped && BulkDataAlignment > 0 && !IsAligned(BulkStartOffset, BulkDataAlignment))
+						if (bBulkItemIsMapped && BulkDataAlignment > 0)
 						{
-							int64 AlignedOffset = Align(BulkStartOffset, BulkDataAlignment);
-							int64 Padding = AlignedOffset - BulkStartOffset;
-							check(Padding > 0);
-							uint64 Zero64 = 0;
-							while (Padding >= 8)
-							{
-								*TargetArchive << Zero64;
-								Padding -= 8;
-							}
-							uint8 Zero8 = 0;
-							while (Padding > 0)
+							const int64 BulkStartOffset = TargetArchive->Tell();
+
+							if (!IsAligned(BulkStartOffset, BulkDataAlignment))
 							{
-								*TargetArchive << Zero8;
-								Padding--;
+								const int64 AlignedOffset = Align(BulkStartOffset, BulkDataAlignment);
+
+								int64 Padding = AlignedOffset - BulkStartOffset;
+								check(Padding > 0);
+
+								uint64 Zero64 = 0;
+								while(Padding >= 8)
+								{
+									*TargetArchive << Zero64;
+									Padding -= 8;
+								}
+
+								uint8 Zero8 = 0;
+								while(Padding > 0)
+								{
+									*TargetArchive << Zero8;
+									Padding--;
+								}
+
+								check(TargetArchive->Tell() == AlignedOffset);
 							}
-							BulkStartOffset = TargetArchive->Tell();
-							check(BulkStartOffset == AlignedOffset);
 						}
 
+						const int64 BulkStartOffset = TargetArchive->Tell();
 
 						int64 StoredBulkStartOffset = BulkStartOffset - StartOfBulkDataArea;
 
 						BulkDataStorageInfo.BulkData->SerializeBulkData(*TargetArchive, BulkDataStorageInfo.BulkData->Lock(LOCK_READ_ONLY));
 
 						int64 BulkEndOffset = TargetArchive->Tell();
-						int64 LinkerEndOffset = Linker->Tell();
+						const int64 LinkerEndOffset = Linker->Tell();
 
 						int64 SizeOnDisk = BulkEndOffset - BulkStartOffset;
 				
@@ -5764,51 +5763,36 @@ FSavePackageResultStruct UPackage::Save(UPackage* InOuter, UObject* Base, EObjec
 
 					if (BulkArchive)
 					{
-						check( OptionalBulkArchive);
+						check(OptionalBulkArchive);
+						check(MappedBulkArchive);
+
+						const bool bWriteBulkToDisk = !bDiffing; 
 
-						const bool bWriteBulkToDisk = bSaveAsync && !bDiffing; 
-						auto FinalizeBulkDataFile = [bWriteBulkToDisk](FArchive* Archive, const FString& ArchiveFilename)
+						auto WriteBulkData = [&](FLargeMemoryWriter* Archive, const TCHAR* BulkFileExtension)
 							{
-								Archive->Close();
-								if (bWriteBulkToDisk)
+								if (const int64 DataSize = Archive->TotalSize())
 								{
-									if (int64 DataSize = Archive->TotalSize())
+									if (bComputeHash)
 									{
-										// TODO - update the BulkBuffer code to write into a FLargeMemoryWriter so we can 
-										// take ownership of the data
-										uint8* Data = new uint8[DataSize];
-										FMemory::Memcpy(Data, ((FBufferArchive*)(Archive))->GetData(), DataSize);
-										FLargeMemoryPtr DataPtr = FLargeMemoryPtr(Data);
-		
-										AsyncWriteFile(MoveTemp(DataPtr), DataSize, *ArchiveFilename, FDateTime::MinValue(), false);
+										CookedPackageHash.Update(Archive->GetData(), DataSize);
 									}
-								}
-								delete Archive;
-							};
 
-						if (bComputeHash && bBufferedBulkArchives)
-						{
-							AddArchiveToHash(*(FBufferArchive*)BulkArchive, CookedPackageHash);
-							AddArchiveToHash(*(FBufferArchive*)OptionalBulkArchive, CookedPackageHash);
-							AddArchiveToHash(*(FBufferArchive*)MappedBulkArchive, CookedPackageHash);
-						}
-						
-						TotalPackageSizeUncompressed += BulkArchive->TotalSize();
-						TotalPackageSizeUncompressed += OptionalBulkArchive->TotalSize();		
-						TotalPackageSizeUncompressed += MappedBulkArchive->TotalSize();
+									TotalPackageSizeUncompressed += DataSize;
 
-						FinalizeBulkDataFile(BulkArchive, BulkFilename);
-						FinalizeBulkDataFile(OptionalBulkArchive, OptionalBulkFilename);
-						FinalizeBulkDataFile(MappedBulkArchive, MappedBulkFilename);
+									if (bWriteBulkToDisk)
+									{
+										FLargeMemoryPtr DataPtr(Archive->ReleaseOwnership());
 
-						if (bComputeHash && !bBufferedBulkArchives)
-						{
-							check(!bWriteBulkToDisk);
+										const FString ArchiveFilename = FPaths::ChangeExtension(Filename, BulkFileExtension);
 
-							AddFileToHash(BulkFilename, CookedPackageHash);
-							AddFileToHash(OptionalBulkFilename, CookedPackageHash);
-							AddFileToHash(MappedBulkFilename, CookedPackageHash);
-						}
+										AsyncWriteFile(MoveTemp(DataPtr), DataSize, *ArchiveFilename, FDateTime::MinValue(), /* bUseTempFileName */ false);
+									}
+								}
+							};
+
+						WriteBulkData(BulkArchive.Get(),			TEXT(".ubulk"));		// Regular separate bulk data file
+						WriteBulkData(OptionalBulkArchive.Get(),	TEXT(".uptnl"));		// Optional bulk data
+						WriteBulkData(MappedBulkArchive.Get(),		TEXT(".m.ubulk"));		// Memory-mapped bulk data
 					}
 				}
 
