Consolidated writing of cooked package data into a single function instead of scattering the same logic in multiple places, to simplify package store refactoring (and simplify maintenance)

(ushell-p4-cherrypick of 8645569 by [REDACTED_PERSON])

diff --git a/Engine/Source/Runtime/CoreUObject/Private/UObject/SavePackage.cpp b/Engine/Source/Runtime/CoreUObject/Private/UObject/SavePackage.cpp
--- a/Engine/Source/Runtime/CoreUObject/Private/UObject/SavePackage.cpp
+++ b/Engine/Source/Runtime/CoreUObject/Private/UObject/SavePackage.cpp
@@ -442,6 +442,28 @@ void UPackage::WaitForAsyncFileWrites()
 	}
 }
 
+static void WriteToFile(const FString& Filename, uint8* InDataPtr, int64 InDataSize)
+{
+	IFileManager& FileManager = IFileManager::Get();
+
+	if (FArchive* Ar = FileManager.CreateFileWriter(*Filename))
+	{
+		Ar->Serialize(InDataPtr, InDataSize);
+		delete Ar;
+
+		if (FileManager.FileSize(*Filename) != InDataSize)
+		{
+			FileManager.Delete(*Filename);
+
+			UE_LOG(LogSavePackage, Fatal, TEXT("Could not save to %s!"), *Filename);
+		}
+	}
+	else
+	{
+		UE_LOG(LogSavePackage, Fatal, TEXT("Could not write to %s!"), *Filename);
+	}
+}
+
 struct FLargeMemoryDelete
 {
 	void operator()(uint8* Ptr) const
@@ -455,96 +477,27 @@ struct FLargeMemoryDelete
 
 typedef TUniquePtr<uint8, FLargeMemoryDelete> FLargeMemoryPtr;
 
-void AsyncWriteFile(FLargeMemoryPtr Data, const int64 DataSize, const TCHAR* Filename, const FDateTime& TimeStamp, bool bUseTempFilename = true)
+static void AsyncWriteFile(FLargeMemoryPtr Data, const int64 DataSize, const TCHAR* Filename)
 {
 	class FAsyncWriteWorker : public FNonAbandonableTask
 	{
 	public:
-		/** Filename To write to**/
-		FString Filename;
-		/** Should we write to a temp file then move it */
-		bool bUseTempFilename;
-		/** Data for the file. Will be freed after write **/
+		FString			Filename;
 		FLargeMemoryPtr Data;
-		/** Size of data */
-		const int64 DataSize;
-		/** Timestamp to give the file. MinValue if shouldn't be modified */
-		FDateTime FinalTimeStamp;
-
-		/** Constructor
-		*/
-		FAsyncWriteWorker(const TCHAR* InFilename, FLargeMemoryPtr InData, const int64 InDataSize, const FDateTime& InTimeStamp, bool inbUseTempFilename)
+		const int64		DataSize;
+
+		FAsyncWriteWorker(const TCHAR* InFilename, FLargeMemoryPtr InData, const int64 InDataSize)
 			: Filename(InFilename)
-			, bUseTempFilename(inbUseTempFilename)
 			, Data(MoveTemp(InData))
 			, DataSize(InDataSize)
-			, FinalTimeStamp(InTimeStamp)
 		{
+			check(InDataSize);
 		}
 		
 		/** Write the file  */
 		void DoWork()
 		{
-			check(DataSize);
-			FString TempFilename; 
-
-			if (bUseTempFilename)
-			{
-				TempFilename = FPaths::GetBaseFilename(Filename, false);
-				TempFilename += TEXT(".t");
-			}
-			else
-			{
-				TempFilename = Filename;
-			}
-			
-			// Open a file writer for saving
-			FArchive* Ar = IFileManager::Get().CreateFileWriter(*TempFilename);
-			if (Ar)
-			{
-				Ar->Serialize(Data.Get(), DataSize);
-				delete Ar;
-				
-				// Clean-up the memory as soon as we save the file to reduce the memory footprint.
-				Data.Reset();
-
-				if (IFileManager::Get().FileSize(*TempFilename) == DataSize)
-				{
-					if (bUseTempFilename)
-					{
-						if (!IFileManager::Get().Move(*Filename, *TempFilename, true, true, false, false))
-						{
-							UE_LOG(LogSavePackage, Fatal, TEXT("Could not move to %s."),*Filename);
-						}
-
-						// If everything worked as it should this is not necessary, but we want to ensure
-						// no temporary files are left lying around if something does go wrong.
-
-						if (FPaths::FileExists(TempFilename))
-						{
-							IFileManager::Get().Delete(*TempFilename);
-
-							UE_LOG(LogSavePackage, Fatal, TEXT("Deleted temporary file '%s'"), *TempFilename);
-						}
-					}
-
-					if (FinalTimeStamp != FDateTime::MinValue())
-					{
-						IFileManager::Get().SetTimeStamp(*Filename, FinalTimeStamp);
-					}
-				}
-				else
-				{
-					// Don't leave incomplete files around
-					IFileManager::Get().Delete(*TempFilename);
-
-					UE_LOG(LogSavePackage, Fatal, TEXT("Could not save to %s!"),*TempFilename);
-				}
-			}
-			else
-			{
-				UE_LOG(LogSavePackage, Fatal, TEXT("Could not write to %s!"),*TempFilename);
-			}
+			WriteToFile(Filename, Data.Get(), DataSize);
 			
 			OutstandingAsyncWrites.Decrement();
 		}
@@ -553,150 +506,40 @@ void AsyncWriteFile(FLargeMemoryPtr Data, const int64 DataSize, const TCHAR* Fil
 		{
 			RETURN_QUICK_DECLARE_CYCLE_STAT(FAsyncWriteWorker, STATGROUP_ThreadPoolAsyncTasks);
 		}
-
 	};
 
 	OutstandingAsyncWrites.Increment();
-	(new FAutoDeleteAsyncTask<FAsyncWriteWorker>(Filename, MoveTemp(Data), DataSize, TimeStamp, bUseTempFilename))->StartBackgroundTask();
+	(new FAutoDeleteAsyncTask<FAsyncWriteWorker>(Filename, MoveTemp(Data), DataSize))->StartBackgroundTask();
 }
 
-void AsyncWriteFileWithSplitExports(FLargeMemoryPtr Data, const int64 DataSize, const int64 HeaderSize, const TCHAR* Filename, const FDateTime& TimeStamp, bool bUseTempFilename = true)
+static void AsyncWriteFileWithSplitExports(FLargeMemoryPtr Data, const int64 DataSize, const int64 HeaderSize, const TCHAR* Filename)
 {
 	class FAsyncWriteWorkerWithSplitExports : public FNonAbandonableTask
 	{
 	public:
-		/** Filename To write to**/
-		FString Filename;
-		/** Should we write to a temp file then move it */
-		bool bUseTempFilename;
-		/** Data for the file. Will be freed after write **/
+		FString			Filename;
 		FLargeMemoryPtr Data;
-		/** Size of data */
-		const int64 DataSize;
-		/** Size of data */
-		const int64 HeaderSize;
-		/** Timestamp to give the file. MinValue if shouldn't be modified */
-		FDateTime FinalTimeStamp;
-
-		/** Constructor
-		*/
-		FAsyncWriteWorkerWithSplitExports(const TCHAR* InFilename, FLargeMemoryPtr InData, const int64 InDataSize, const int64 InHeaderSize, const FDateTime& InTimeStamp, bool inbUseTempFilename)
+		const int64		DataSize;
+		const int64		HeaderSize;
+
+		FAsyncWriteWorkerWithSplitExports(const TCHAR* InFilename, FLargeMemoryPtr InData, const int64 InDataSize, const int64 InHeaderSize)
 			: Filename(InFilename)
-			, bUseTempFilename(inbUseTempFilename)
 			, Data(MoveTemp(InData))
 			, DataSize(InDataSize)
 			, HeaderSize(InHeaderSize)
-			, FinalTimeStamp(InTimeStamp)
 		{
 			check(InDataSize);
 		}
 
-		/** Write the file  */
 		void DoWork()
 		{
-			check(DataSize);
-			FString BaseFilename;
-			FString TempFilename;
-			FString FilenameExports;
-			FString TempFilenameExports;
+			// Write .uasset file
+			WriteToFile(Filename, Data.Get(), HeaderSize);
 
-			BaseFilename = FPaths::GetBaseFilename(Filename, false);
-			FilenameExports = BaseFilename;
-			FilenameExports += TEXT(".uexp");
-			if (bUseTempFilename)
-			{
-				TempFilename = BaseFilename;
-				TempFilename += TEXT(".t");
-				TempFilenameExports = BaseFilename;
-				TempFilenameExports += TEXT(".e");
-			}
-			else
-			{
-				TempFilename = Filename;
-				TempFilenameExports = FilenameExports;
-			}
-
-			{
-				FArchive* Ar = IFileManager::Get().CreateFileWriter(*TempFilename);
-				if (Ar)
-				{
-					Ar->Serialize(Data.Get(), HeaderSize);
-					delete Ar;
-					Ar = nullptr;
-
-					if (IFileManager::Get().FileSize(*TempFilename) == HeaderSize)
-					{
-						if (bUseTempFilename)
-						{
-							if (!IFileManager::Get().Move(*Filename, *TempFilename, true, true, false, false))
-							{
-								UE_LOG(LogSavePackage, Fatal, TEXT("Could not move to %s."), *Filename);
-							}
-
-							// if everything worked, this is not necessary, but we will make every effort to avoid leaving junk in the cache
-							if (FPaths::FileExists(TempFilename))
-							{
-								IFileManager::Get().Delete(*TempFilename);
-							}
-						}
-
-						if (FinalTimeStamp != FDateTime::MinValue())
-						{
-							IFileManager::Get().SetTimeStamp(*Filename, FinalTimeStamp);
-						}
-					}
-					else
-					{
-						UE_LOG(LogSavePackage, Fatal, TEXT("Could not save to %s!"), *TempFilename);
-					}
-				}
-				else
-				{
-					UE_LOG(LogSavePackage, Fatal, TEXT("Could not write to %s!"), *TempFilename);
-				}
-			}
-			{
-				FArchive* Ar = IFileManager::Get().CreateFileWriter(*TempFilenameExports);
-				if (Ar)
-				{
-					Ar->Serialize(Data.Get() + HeaderSize, DataSize - HeaderSize);
-					delete Ar;
-					Ar = nullptr;
-
-					// Clean-up the memory as soon as we save the file to reduce the memory footprint.
-					Data.Reset();
-
-					if (IFileManager::Get().FileSize(*TempFilenameExports) == DataSize - HeaderSize)
-					{
-						if (bUseTempFilename)
-						{
-							if (!IFileManager::Get().Move(*FilenameExports, *TempFilenameExports, true, true, false, false))
-							{
-								UE_LOG(LogSavePackage, Fatal, TEXT("Could not move to %s."), *Filename);
-							}
+			// Write .uexp file
+			const FString FilenameExports = FPaths::ChangeExtension(Filename, TEXT(".uexp"));
+			WriteToFile(FilenameExports, Data.Get() + HeaderSize, DataSize - HeaderSize);
 
-							// if everything worked, this is not necessary, but we will make every effort to avoid leaving junk in the cache
-							if (FPaths::FileExists(TempFilenameExports))
-							{
-								IFileManager::Get().Delete(*TempFilenameExports);
-							}
-						}
-
-						if (FinalTimeStamp != FDateTime::MinValue())
-						{
-							IFileManager::Get().SetTimeStamp(*FilenameExports, FinalTimeStamp);
-						}
-					}
-					else
-					{
-						UE_LOG(LogSavePackage, Fatal, TEXT("Could not save to %s!"), *TempFilenameExports);
-					}
-				}
-				else
-				{
-					UE_LOG(LogSavePackage, Fatal, TEXT("Could not write to %s!"), *TempFilenameExports);
-				}
-			}
 			OutstandingAsyncWrites.Decrement();
 		}
 
@@ -704,11 +547,10 @@ void AsyncWriteFileWithSplitExports(FLargeMemoryPtr Data, const int64 DataSize,
 		{
 			RETURN_QUICK_DECLARE_CYCLE_STAT(FAsyncWriteWorkerWithSplitExports, STATGROUP_ThreadPoolAsyncTasks);
 		}
-
 	};
 
 	OutstandingAsyncWrites.Increment();
-	(new FAutoDeleteAsyncTask<FAsyncWriteWorkerWithSplitExports>(Filename, MoveTemp(Data), DataSize, HeaderSize, TimeStamp, bUseTempFilename))->StartBackgroundTask();
+	(new FAutoDeleteAsyncTask<FAsyncWriteWorkerWithSplitExports>(Filename, MoveTemp(Data), DataSize, HeaderSize))->StartBackgroundTask();
 }
 
 /** 
@@ -5785,7 +5627,7 @@ FSavePackageResultStruct UPackage::Save(UPackage* InOuter, UObject* Base, EObjec
 
 										const FString ArchiveFilename = FPaths::ChangeExtension(Filename, BulkFileExtension);
 
-										AsyncWriteFile(MoveTemp(DataPtr), DataSize, *ArchiveFilename, FDateTime::MinValue(), /* bUseTempFileName */ false);
+										AsyncWriteFile(MoveTemp(DataPtr), DataSize, *ArchiveFilename);
 									}
 								}
 							};
@@ -5812,9 +5654,9 @@ FSavePackageResultStruct UPackage::Save(UPackage* InOuter, UObject* Base, EObjec
 						}
 						if (bWriteFileToDisk)
 						{
-							FLargeMemoryPtr DataPtr(Writer.GetData());
-							Writer.ReleaseOwnership();
-							AsyncWriteFile(MoveTemp(DataPtr), Size, *Writer.GetArchiveName(), FDateTime::MinValue(), false);
+							FLargeMemoryPtr DataPtr(Writer.ReleaseOwnership());
+
+							AsyncWriteFile(MoveTemp(DataPtr), Size, *Writer.GetArchiveName());
 						}
 					}
 					AdditionalFilesFromExports.Empty();
@@ -6020,7 +5862,6 @@ FSavePackageResultStruct UPackage::Save(UPackage* InOuter, UObject* Base, EObjec
 
 				if( Success == true )
 				{
-					// Compress the temporarily file to destination.
 					if (bSaveAsync)
 					{						
 						FString NewPathToSave = NewPath;
@@ -6069,27 +5910,25 @@ FSavePackageResultStruct UPackage::Save(UPackage* InOuter, UObject* Base, EObjec
 						{
 							UE_LOG(LogSavePackage, Verbose, TEXT("Async saving from memory to '%s'"), *NewPathToSave);
 
-							// Detach archive used for memory saving.
 							FLargeMemoryWriter* Writer = (FLargeMemoryWriter*)(Linker->Saver);
-							int64 DataSize = Writer->TotalSize();
+							const int64 DataSize = Writer->TotalSize();
 
-							COOK_STAT(FScopedDurationTimer SaveTimer(SavePackageStats::AsyncWriteTimeSec));
 							TotalPackageSizeUncompressed += DataSize;
 
 							if (bComputeHash)
 							{
-								CookedPackageHash.Update( Writer->GetData(), Writer->TotalSize() );
+								CookedPackageHash.Update(Writer->GetData(), DataSize);
 							}
 
-							FLargeMemoryPtr DataPtr(Writer->GetData());
-							Writer->ReleaseOwnership();
+							FLargeMemoryPtr DataPtr(Writer->ReleaseOwnership());
+
 							if (IsEventDrivenLoaderEnabledInCookedBuilds() && Linker->IsCooking())
 							{
-								AsyncWriteFileWithSplitExports(MoveTemp(DataPtr), DataSize, Linker->Summary.TotalHeaderSize, *NewPathToSave, FinalTimeStamp);
+								AsyncWriteFileWithSplitExports(MoveTemp(DataPtr), DataSize, Linker->Summary.TotalHeaderSize, *NewPathToSave);
 							}
 							else
 							{
-								AsyncWriteFile(MoveTemp(DataPtr), DataSize, *NewPathToSave, FinalTimeStamp);
+								AsyncWriteFile(MoveTemp(DataPtr), DataSize, *NewPathToSave);
 							}
 						}
 						Linker->CloseAndDestroySaver();
