MakeUnique fixes.

diff --git a/Engine/Source/Runtime/CoreUObject/Private/Serialization/AsyncLoading.cpp b/Engine/Source/Runtime/CoreUObject/Private/Serialization/AsyncLoading.cpp
--- a/Engine/Source/Runtime/CoreUObject/Private/Serialization/AsyncLoading.cpp
+++ b/Engine/Source/Runtime/CoreUObject/Private/Serialization/AsyncLoading.cpp
@@ -6313,7 +6313,7 @@ void FAsyncPackage::AddImportDependency(const FName& PendingImport, FFlushTree*
 		!PackageToStream->bLoadHasFailed)
 	{
 		const bool bInternalCallback = true;
-		TUniquePtr<FLoadPackageAsyncDelegate> InternalDelegate(new FLoadPackageAsyncDelegate(FLoadPackageAsyncDelegate::CreateRaw(this, &FAsyncPackage::ImportFullyLoadedCallback)));
+		TUniquePtr<FLoadPackageAsyncDelegate> InternalDelegate = MakeUnique<FLoadPackageAsyncDelegate>(FLoadPackageAsyncDelegate::CreateRaw(this, &FAsyncPackage::ImportFullyLoadedCallback));
 		PackageToStream->AddCompletionCallback(MoveTemp(InternalDelegate), bInternalCallback);
 		PackageToStream->DependencyRefCount.Increment();
 		PendingImportedPackages.Add(PackageToStream);
@@ -7378,7 +7378,7 @@ int32 LoadPackageAsync(const FString& InName, const FGuid* InGuid /*= nullptr*/,
 		TUniquePtr<FLoadPackageAsyncDelegate> CompletionDelegatePtr;
 		if (InCompletionDelegate.IsBound())
 		{
-			CompletionDelegatePtr.Reset(new FLoadPackageAsyncDelegate(InCompletionDelegate));
+			CompletionDelegatePtr = MakeUnique<FLoadPackageAsyncDelegate>(InCompletionDelegate);
 		}
 
 		// Add new package request
