Removed some unnecessary static functions from public Package.h header

diff --git a/Engine/Source/Runtime/CoreUObject/Private/UObject/SavePackage.cpp b/Engine/Source/Runtime/CoreUObject/Private/UObject/SavePackage.cpp
--- a/Engine/Source/Runtime/CoreUObject/Private/UObject/SavePackage.cpp
+++ b/Engine/Source/Runtime/CoreUObject/Private/UObject/SavePackage.cpp
@@ -67,6 +67,11 @@ static FCriticalSection InitializeCoreClassesCritSec;
 #define VALIDATE_INITIALIZECORECLASSES 0
 #define EXPORT_SORTING_DETAILED_LOGGING 0
 
+static void SaveThumbnails(UPackage* InOuter, FLinkerSave* Linker, FStructuredArchive::FSlot Slot);
+static void SaveAssetRegistryData(UPackage* InOuter, FLinkerSave* Linker, FStructuredArchive::FSlot Slot);
+static void SaveWorldLevelInfo(UPackage* InOuter, FLinkerSave* Linker, FStructuredArchive::FRecord Record);
+static EObjectMark GetExcludedObjectMarksForTargetPlatform(const class ITargetPlatform* TargetPlatform, const bool bIsCooking);
+
 #if ENABLE_COOK_STATS
 #include "ProfilingDebugging/ScopedTimers.h"
 
@@ -919,7 +924,7 @@ FArchive& FArchiveSaveTagExports::operator<<(UObject*& Obj)
 	}
 
 	// Check outer chain for any exlcuded object marks
-	const EObjectMark ExcludedObjectMarks = UPackage::GetExcludedObjectMarksForTargetPlatform(CookingTarget(), IsCooking());
+	const EObjectMark ExcludedObjectMarks = GetExcludedObjectMarksForTargetPlatform(CookingTarget(), IsCooking());
 	ConditionallyExcludeObjectForTarget(Obj, ExcludedObjectMarks, CookingTarget(), IsCooking());
 
 	if (!Obj->HasAnyMarks((EObjectMark)(ExcludedObjectMarks)))
@@ -1126,7 +1131,7 @@ FArchive& FArchiveSaveTagImports::operator<<( UObject*& Obj )
 	// Check transient and pending kill flags for outers
 	CheckObjectPriorToSave(*this, Obj, nullptr);
 
-	const EObjectMark ExcludedObjectMarks = UPackage::GetExcludedObjectMarksForTargetPlatform( CookingTarget(), IsCooking() );
+	const EObjectMark ExcludedObjectMarks = GetExcludedObjectMarksForTargetPlatform( CookingTarget(), IsCooking() );
 	ConditionallyExcludeObjectForTarget(Obj, ExcludedObjectMarks, CookingTarget(), IsCooking());
 	bool bExcludePackageFromCook = Obj && FCoreUObjectDelegates::ShouldCookPackageForPlatform.IsBound() ? !FCoreUObjectDelegates::ShouldCookPackageForPlatform.Execute(Obj->GetOutermost(), CookingTarget()) : false;
 
@@ -2951,7 +2956,15 @@ static bool ValidateConformCompatibility(UPackage* NewPackage, FLinkerLoad* OldL
 	return !bHadCompatibilityErrors;
 }
 
-EObjectMark UPackage::GetExcludedObjectMarksForTargetPlatform( const class ITargetPlatform* TargetPlatform, const bool bIsCooking )
+/**
+ * Determines the set of object marks that should be excluded for the target platform
+ *
+ * @param TargetPlatform	The platform being saved for
+ * @param bIsCooking		Whether we are cooking or not
+ *
+ * @return Excluded object marks specific for the particular target platform, objects with any of these marks will be rejected from the cook
+ */
+EObjectMark GetExcludedObjectMarksForTargetPlatform( const class ITargetPlatform* TargetPlatform, const bool bIsCooking )
 {
 	EObjectMark ObjectMarks = OBJECTMARK_NOMARKS;
 
@@ -3883,7 +3896,7 @@ FSavePackageResultStruct UPackage::Save(UPackage* InOuter, UObject* Base, EObjec
 					TArray<UObject*> TagExpObjects;
 					GetObjectsWithAnyMarks(TagExpObjects, OBJECTMARK_TagExp);
 
-					const EObjectMark ExcludedObjectMarks = UPackage::GetExcludedObjectMarksForTargetPlatform(TargetPlatform, Linker->IsCooking());
+					const EObjectMark ExcludedObjectMarks = GetExcludedObjectMarksForTargetPlatform(TargetPlatform, Linker->IsCooking());
 					if (Linker->IsCooking() && ExcludedObjectMarks != OBJECTMARK_NOMARKS)
 					{
 						// Make sure that nothing is in the export table that should have been filtered out
@@ -4489,7 +4502,7 @@ FSavePackageResultStruct UPackage::Save(UPackage* InOuter, UObject* Base, EObjec
 				{
 					TArray<UObject*> TagImpObjects;
 
-					const EObjectMark ExcludedObjectMarks = UPackage::GetExcludedObjectMarksForTargetPlatform(TargetPlatform, Linker->IsCooking());
+					const EObjectMark ExcludedObjectMarks = GetExcludedObjectMarksForTargetPlatform(TargetPlatform, Linker->IsCooking());
 					GetObjectsWithAnyMarks(TagImpObjects, OBJECTMARK_TagImp);
 
 					if (Linker->IsCooking() && ExcludedObjectMarks != OBJECTMARK_NOMARKS)
@@ -4940,16 +4953,16 @@ FSavePackageResultStruct UPackage::Save(UPackage* InOuter, UObject* Base, EObjec
 #endif // WITH_EDITOR
 
 					// Save thumbnails
-					UPackage::SaveThumbnails(InOuter, Linker.Get(), StructuredArchiveRoot.EnterField(SA_FIELD_NAME(TEXT("Thumbnails"))));
+					SaveThumbnails(InOuter, Linker.Get(), StructuredArchiveRoot.EnterField(SA_FIELD_NAME(TEXT("Thumbnails"))));
 
 					if (!bTextFormat)
 					{	
 						// Save asset registry data so the editor can search for information about assets in this package
-						UPackage::SaveAssetRegistryData(InOuter, Linker.Get(), StructuredArchiveRoot.EnterField(SA_FIELD_NAME(TEXT("AssetRegistry"))));
+						SaveAssetRegistryData(InOuter, Linker.Get(), StructuredArchiveRoot.EnterField(SA_FIELD_NAME(TEXT("AssetRegistry"))));
 					}
 
 					// Save level information used by World browser
-					UPackage::SaveWorldLevelInfo(InOuter, Linker.Get(), StructuredArchiveRoot);
+					SaveWorldLevelInfo(InOuter, Linker.Get(), StructuredArchiveRoot);
 				}
 
 
@@ -5043,7 +5056,7 @@ FSavePackageResultStruct UPackage::Save(UPackage* InOuter, UObject* Base, EObjec
 					FArchiveStackTraceIgnoreScope IgnoreSummaryDiffsScope(DiffSettings.bIgnoreHeaderDiffs);
 #endif // WITH_EDITOR
 
-					const EObjectMark ExcludedObjectMarks = UPackage::GetExcludedObjectMarksForTargetPlatform(Linker->CookingTarget(), Linker->IsCooking());
+					const EObjectMark ExcludedObjectMarks = GetExcludedObjectMarksForTargetPlatform(Linker->CookingTarget(), Linker->IsCooking());
 					Linker->Summary.PreloadDependencyCount = 0;
 
 					auto IncludeObjectAsDependency = [&Linker, ExcludedObjectMarks](int32 CallSite, TSet<FPackageIndex>& AddTo, UObject* ToTest, UObject* ForObj, bool bMandatory, bool bOnlyIfInLinkerTable)
@@ -6118,7 +6131,7 @@ bool UPackage::SavePackage(UPackage* InOuter, UObject* Base, EObjectFlags TopLev
  * @param	Linker							linker we're currently saving with
  * @param	Slot							structed archive slot we are saving too (temporary)
  */
-void UPackage::SaveThumbnails(UPackage* InOuter, FLinkerSave* Linker, FStructuredArchive::FSlot Slot)
+static void SaveThumbnails(UPackage* InOuter, FLinkerSave* Linker, FStructuredArchive::FSlot Slot)
 {
 	FStructuredArchive::FRecord Record = Slot.EnterRecord();
 
@@ -6233,7 +6246,7 @@ void UPackage::SaveThumbnails(UPackage* InOuter, FLinkerSave* Linker, FStructure
 #endif
 }
 
-void UPackage::SaveAssetRegistryData(UPackage* InOuter, FLinkerSave* Linker, FStructuredArchive::FSlot Slot)
+static void SaveAssetRegistryData(UPackage* InOuter, FLinkerSave* Linker, FStructuredArchive::FSlot Slot)
 {
 	// Make a copy of the tag map
 	TArray<UObject*> AssetObjects;
@@ -6267,13 +6280,13 @@ void UPackage::SaveAssetRegistryData(UPackage* InOuter, FLinkerSave* Linker, FSt
 		FString ObjectPath = Object->GetPathName(Object->GetOutermost());
 		FString ObjectClassName = Object->GetClass()->GetName();
 		
-		TArray<FAssetRegistryTag> SourceTags;
+		TArray<UObject::FAssetRegistryTag> SourceTags;
 		Object->GetAssetRegistryTags(SourceTags);
 
-		TArray<FAssetRegistryTag> Tags;
-		for (FAssetRegistryTag& SourceTag : SourceTags)
+		TArray<UObject::FAssetRegistryTag> Tags;
+		for (UObject::FAssetRegistryTag& SourceTag : SourceTags)
 		{
-			FAssetRegistryTag* Existing = Tags.FindByPredicate([SourceTag](const FAssetRegistryTag& InTag) { return InTag.Name == SourceTag.Name; });
+			UObject::FAssetRegistryTag* Existing = Tags.FindByPredicate([SourceTag](const UObject::FAssetRegistryTag& InTag) { return InTag.Name == SourceTag.Name; });
 			if (Existing)
 			{
 				Existing->Value = SourceTag.Value;
@@ -6291,7 +6304,7 @@ void UPackage::SaveAssetRegistryData(UPackage* InOuter, FLinkerSave* Linker, FSt
 
 		FStructuredArchive::FMap TagMap = AssetRecord.EnterField(SA_FIELD_NAME(TEXT("Tags"))).EnterMap(TagCount);
 				
-		for (TArray<FAssetRegistryTag>::TConstIterator TagIter(Tags); TagIter; ++TagIter)
+		for (TArray<UObject::FAssetRegistryTag>::TConstIterator TagIter(Tags); TagIter; ++TagIter)
 		{
 			FString Key = TagIter->Name.ToString();
 			FString Value = TagIter->Value;
@@ -6301,7 +6314,7 @@ void UPackage::SaveAssetRegistryData(UPackage* InOuter, FLinkerSave* Linker, FSt
 	}
 }
 
-void UPackage::SaveWorldLevelInfo(UPackage* InOuter, FLinkerSave* Linker, FStructuredArchive::FRecord Record)
+static void SaveWorldLevelInfo(UPackage* InOuter, FLinkerSave* Linker, FStructuredArchive::FRecord Record)
 {
 	Linker->Summary.WorldTileInfoDataOffset = 0;
 	
diff --git a/Engine/Source/Runtime/CoreUObject/Public/UObject/Package.h b/Engine/Source/Runtime/CoreUObject/Public/UObject/Package.h
--- a/Engine/Source/Runtime/CoreUObject/Public/UObject/Package.h
+++ b/Engine/Source/Runtime/CoreUObject/Public/UObject/Package.h
@@ -575,31 +575,6 @@ public:
 	/** Wait for any SAVE_Async file writes to complete **/
 	static void WaitForAsyncFileWrites();
 
-	/**
-	* Static: Saves thumbnail data for the specified package outer and linker
-	*
-	* @param	InOuter							the outer to use for the new package
-	* @param	Linker							linker we're currently saving with
-	*/
-	static void SaveThumbnails(UPackage* InOuter, FLinkerSave* Linker, FStructuredArchive::FSlot Slot);
-
-	/**
-	* Static: Saves asset registry data for the specified package outer and linker
-	*
-	* @param	InOuter							the outer to use for the new package
-	* @param	Linker							linker we're currently saving with
-	*/
-	static void SaveAssetRegistryData(UPackage* InOuter, FLinkerSave* Linker, FStructuredArchive::FSlot Slot);
-
-	/**
-	* Static: Saves the level information used by the World browser
-	*
-	* @param	InOuter							the outer to use for the new package
-	* @param	Linker							linker we're currently saving with
-	* @param	Record							the structured archive record we should add the world level info to
-	*/
-	static void SaveWorldLevelInfo(UPackage* InOuter, FLinkerSave* Linker, FStructuredArchive::FRecord Record);
-
 	/**
 	* Determines if a package contains no more assets.
 	*
@@ -608,16 +583,5 @@ public:
 	* @return true if Package contains no more assets.
 	*/
 	static bool IsEmptyPackage(UPackage* Package, const UObject* LastReferencer = NULL);
-
-	/**
-	* Determines the set of object marks that should be excluded for the target platform
-	*
-	* @param TargetPlatform	The platform being saved for
-	* @param bIsCooking		Whether we are cooking or not
-	*
-	* @return Excluded object marks specific for the particular target platform, objects with any of these marks will be rejected from the cook
-	*/
-	static EObjectMark GetExcludedObjectMarksForTargetPlatform( const class ITargetPlatform* TargetPlatform, const bool bIsCooking );
-
 };
 PRAGMA_ENABLE_DEPRECATION_WARNINGS
