Strip out code related to now-defunct cook stats

diff --git a/Engine/Source/Runtime/CoreUObject/Private/UObject/SavePackage.cpp b/Engine/Source/Runtime/CoreUObject/Private/UObject/SavePackage.cpp
--- a/Engine/Source/Runtime/CoreUObject/Private/UObject/SavePackage.cpp
+++ b/Engine/Source/Runtime/CoreUObject/Private/UObject/SavePackage.cpp
@@ -92,9 +92,6 @@ namespace SavePackageStats
 	static double MBWritten = 0.0;
 	TMap<FName, FArchiveDiffStats> PackageDiffStats;
 	static int32 NumberOfDifferentPackages = 0;
-	static TMap<FName, FCookStatsManager::TKeyValuePair<double, uint32>> PackageClassSerializeTimes;
-	static TMap<FName, FCookStatsManager::TKeyValuePair<double, uint32>> TagExportSerializeTimes;
-	static TMap<FName, FCookStatsManager::TKeyValuePair<double, uint32>> ClassPreSaveTimes;
 	static FCookStatsManager::FAutoRegisterCallback RegisterCookStats([](FCookStatsManager::AddStatFuncRef AddStat)
 	{
 		// Don't use FCookStatsManager::CreateKeyValueArray because there's just too many arguments. Don't need to overburden the compiler here.
@@ -178,127 +175,6 @@ namespace SavePackageStats
 		#undef ADD_COOK_STAT		
 		
 		const FString TotalString = TEXT("Total");
-
-		if (PackageClassSerializeTimes.Num() > 0)
-		{
-			// Sort the class serialize times in reverse order.
-			typedef FCookStatsManager::TKeyValuePair<FName, FCookStatsManager::TKeyValuePair<double, uint32>> FClassSerializeTimeData;
-			TArray<FClassSerializeTimeData> SerializeTimesArray;
-			SerializeTimesArray.Empty(PackageClassSerializeTimes.Num());
-			for (const TPair<FName, FCookStatsManager::TKeyValuePair<double, uint32>>& KV : PackageClassSerializeTimes)
-			{
-				SerializeTimesArray.Emplace(FCookStatsManager::MakePair(KV.Key, FCookStatsManager::MakePair(KV.Value.Key, KV.Value.Value)));
-			}
-			SerializeTimesArray.Sort([](const FClassSerializeTimeData& LHS, const FClassSerializeTimeData& RHS)
-			{
-				return LHS.Value.Key > RHS.Value.Key;
-			});
-
-			// always print at least the top n, but not anything < 0.1% of total save time.
-			int ClassesLogged = 0;
-			for (const FClassSerializeTimeData& KV : SerializeTimesArray)
-			{
-				// since we're sorted on size already. Just find the first one below the threshold and stop there
-				if (ClassesLogged >= 10 && KV.Value.Key < 0.001 * SavePackageTimeSec)
-				{
-					break;
-				}
-				const FString ClassName = KV.Key.ToString();
-				AddStat(TEXT("Package.Serialize"), FCookStatsManager::CreateKeyValueArray(
-					TEXT("Class"), ClassName, 
-					TEXT("TimeSec"), KV.Value.Key, 
-					TEXT("Calls"), KV.Value.Value));
-				ClassesLogged++;
-			}
-		}
-
-		if (TagExportSerializeTimes.Num() > 0)
-		{
-			// Sort the class serialize times in reverse order.
-			typedef FCookStatsManager::TKeyValuePair<FName, FCookStatsManager::TKeyValuePair<double, uint32>> FClassSerializeTimeData;
-			TArray<FClassSerializeTimeData> SerializeTimesArray;
-			SerializeTimesArray.Empty(TagExportSerializeTimes.Num());
-			for (const TPair<FName, FCookStatsManager::TKeyValuePair<double, uint32>>& KV : TagExportSerializeTimes)
-			{
-				SerializeTimesArray.Emplace(FCookStatsManager::MakePair(KV.Key, FCookStatsManager::MakePair(KV.Value.Key, KV.Value.Value)));
-			}
-			SerializeTimesArray.Sort([](const FClassSerializeTimeData& LHS, const FClassSerializeTimeData& RHS)
-			{
-				return LHS.Value.Key > RHS.Value.Key;
-			});
-
-			double TotalSerializeTime = 0.0;
-			int TotalSerializeCalls = 0;
-
-			// always print at least the top n, but not anything < 0.1% of total save time.
-			int ClassesLogged = 0;
-			for (const FClassSerializeTimeData& KV : SerializeTimesArray)
-			{
-				TotalSerializeTime += KV.Value.Key;
-				TotalSerializeCalls += KV.Value.Value;
-
-				// since we're sorted on size already. Just find the first one below the threshold and stop there
-				if (ClassesLogged >= 10 && KV.Value.Key < 0.001 * SavePackageTimeSec)
-				{
-					break;
-				}
-				const FString ClassName = KV.Key.ToString();
-				AddStat(TEXT("Package.TagExportSerialize"), FCookStatsManager::CreateKeyValueArray(
-					TEXT("Class"), ClassName,
-					TEXT("TimeSec"), KV.Value.Key,
-					TEXT("Calls"), KV.Value.Value));
-				ClassesLogged++;
-			}
-
-			AddStat(TEXT("Package.TagExportSerialize"), FCookStatsManager::CreateKeyValueArray(
-				TEXT("Class"), TotalString,
-				TEXT("TimeSec"), TotalSerializeTime,
-				TEXT("Calls"), TotalSerializeCalls));
-		}
-
-		if (ClassPreSaveTimes.Num() > 0)
-		{
-			// Sort the class serialize times in reverse order.
-			typedef FCookStatsManager::TKeyValuePair<FName, FCookStatsManager::TKeyValuePair<double, uint32>> FClassSerializeTimeData;
-			TArray<FClassSerializeTimeData> SerializeTimesArray;
-			SerializeTimesArray.Empty(ClassPreSaveTimes.Num());
-			for (const TPair<FName, FCookStatsManager::TKeyValuePair<double, uint32>>& KV : ClassPreSaveTimes)
-			{
-				SerializeTimesArray.Emplace(FCookStatsManager::MakePair(KV.Key, FCookStatsManager::MakePair(KV.Value.Key, KV.Value.Value)));
-			}
-			SerializeTimesArray.Sort([](const FClassSerializeTimeData& LHS, const FClassSerializeTimeData& RHS)
-			{
-				return LHS.Value.Key > RHS.Value.Key;
-			});
-
-			// always print at least the top n, but not anything < 0.1% of total save time.
-			// even if we don't print them then add up the time spent inside presave so we can at least account for it
-			double TotalPreSaveTime = 0.0;
-			int TotalPreSaveCalls = 0;
-			int ClassesLogged = 0;
-			for (const FClassSerializeTimeData& KV : SerializeTimesArray)
-			{
-				TotalPreSaveTime += KV.Value.Key;
-				TotalPreSaveCalls += KV.Value.Value;
-
-				// since we're sorted on size already. Just find the first one below the threshold and stop there
-				if (ClassesLogged >= 10 && KV.Value.Key < 0.001 * SavePackageTimeSec)
-				{
-					break;
-				}
-				const FString ClassName = KV.Key.ToString();
-				AddStat(TEXT("Package.PreSave"), FCookStatsManager::CreateKeyValueArray(
-					TEXT("Class"), ClassName,
-					TEXT("TimeSec"), KV.Value.Key,
-					TEXT("Calls"), KV.Value.Value));
-				ClassesLogged++;
-			}
-			
-			AddStat(TEXT("Package.PreSave"), FCookStatsManager::CreateKeyValueArray(
-				TEXT("Class"), TotalString,
-				TEXT("TimeSec"), TotalPreSaveTime,
-				TEXT("Calls"), TotalPreSaveCalls));
-		}
 	});
 }
 #endif
