TBA: Enter exports slot when accessing record for default objects

diff --git a/Engine/Source/Runtime/CoreUObject/Private/UObject/LinkerLoad.cpp b/Engine/Source/Runtime/CoreUObject/Private/UObject/LinkerLoad.cpp
--- a/Engine/Source/Runtime/CoreUObject/Private/UObject/LinkerLoad.cpp
+++ b/Engine/Source/Runtime/CoreUObject/Private/UObject/LinkerLoad.cpp
@@ -3690,7 +3690,7 @@ void FLinkerLoad::Preload( UObject* Object )
 #if WITH_EDITOR && WITH_TEXT_ARCHIVE_SUPPORT
 						if (IsTextFormat())
 						{
-							FStructuredArchive::FSlot ExportSlot = StructuredArchiveRootRecord->EnterField(SA_FIELD_NAME(*OriginalExportNames[Export.ThisIndex.ToExport()].ToString()));
+							FStructuredArchive::FSlot ExportSlot = StructuredArchiveRootRecord->EnterRecord(SA_FIELD_NAME(TEXT("Exports"))).EnterField(SA_FIELD_NAME(*OriginalExportNames[Export.ThisIndex.ToExport()].ToString()));
 
 							if (bClassSupportsTextFormat)
 							{
