TBA: When writing out soft object paths to text assets, write the full path rather than just the asset path

diff --git a/Engine/Source/Runtime/CoreUObject/Private/Serialization/Formatters/JsonArchiveOutputFormatter.cpp b/Engine/Source/Runtime/CoreUObject/Private/Serialization/Formatters/JsonArchiveOutputFormatter.cpp
--- a/Engine/Source/Runtime/CoreUObject/Private/Serialization/Formatters/JsonArchiveOutputFormatter.cpp
+++ b/Engine/Source/Runtime/CoreUObject/Private/Serialization/Formatters/JsonArchiveOutputFormatter.cpp
@@ -371,7 +371,7 @@ void FJsonArchiveOutputFormatter::Serialize(FSoftObjectPath& Value)
 {
 	if (Value.IsValid())
 	{
-		SerializeStringInternal(FString::Printf(TEXT("Object:%s"), *Value.GetAssetPathName().ToString()));
+		SerializeStringInternal(FString::Printf(TEXT("Object:%s"), *Value.ToString()));
 	}
 	else
 	{
