Eliminate some unused package flags PKG_Need and PKG_DisallowLazyLoading

diff --git a/Engine/Source/Runtime/CoreUObject/Private/UObject/LinkerLoad.cpp b/Engine/Source/Runtime/CoreUObject/Private/UObject/LinkerLoad.cpp
--- a/Engine/Source/Runtime/CoreUObject/Private/UObject/LinkerLoad.cpp
+++ b/Engine/Source/Runtime/CoreUObject/Private/UObject/LinkerLoad.cpp
@@ -1368,7 +1368,7 @@ FLinkerLoad::ELinkerStatus FLinkerLoad::SerializePackageFileSummary()
 		}
 		
 		// Propagate fact that package cannot use lazy loading to archive (aka this).
-		if( (Summary.PackageFlags & PKG_DisallowLazyLoading) || IsTextFormat() )
+		if (IsTextFormat())
 		{
 			ArAllowLazyLoading = false;
 		}
diff --git a/Engine/Source/Runtime/CoreUObject/Public/UObject/ObjectMacros.h b/Engine/Source/Runtime/CoreUObject/Public/UObject/ObjectMacros.h
--- a/Engine/Source/Runtime/CoreUObject/Public/UObject/ObjectMacros.h
+++ b/Engine/Source/Runtime/CoreUObject/Public/UObject/ObjectMacros.h
@@ -116,11 +116,11 @@ enum EPackageFlags
 //	PKG_Unused						= 0x00001000,
 //	PKG_Unused						= 0x00002000,
 	PKG_ContainsMapData				= 0x00004000,   ///< Contains map data (UObjects only referenced by a single ULevel) but is stored in a different package
-	PKG_Need						= 0x00008000,	///< Client needs to download this package.
+//	PKG_Unused						= 0x00008000,
 	PKG_Compiling					= 0x00010000,	///< package is currently being compiled
 	PKG_ContainsMap					= 0x00020000,	///< Set if the package contains a ULevel/ UWorld object
 	PKG_RequiresLocalizationGather	= 0x00040000,	///< Set if the package contains any data to be gathered by localization
-	PKG_DisallowLazyLoading			= 0x00080000,	///< Set if the archive serializing this package cannot use lazy loading
+//	PKG_Unused						= 0x00080000,
 	PKG_PlayInEditor				= 0x00100000,	///< Set if the package was created for the purpose of PIE
 	PKG_ContainsScript				= 0x00200000,	///< Package is allowed to contain UClass objects
 	PKG_DisallowExport				= 0x00400000,	///< Editor should not export asset in this package
