MoveTemp some async delegates
#rnx

(ushell-p4-cherrypick of 8637603 by marc.audy)

diff --git a/Engine/Source/Runtime/CoreUObject/Private/Serialization/AsyncLoading.cpp b/Engine/Source/Runtime/CoreUObject/Private/Serialization/AsyncLoading.cpp
--- a/Engine/Source/Runtime/CoreUObject/Private/Serialization/AsyncLoading.cpp
+++ b/Engine/Source/Runtime/CoreUObject/Private/Serialization/AsyncLoading.cpp
@@ -7378,7 +7378,7 @@ int32 LoadPackageAsync(const FString& InName, const FGuid* InGuid /*= nullptr*/,
 		TUniquePtr<FLoadPackageAsyncDelegate> CompletionDelegatePtr;
 		if (InCompletionDelegate.IsBound())
 		{
-			CompletionDelegatePtr = MakeUnique<FLoadPackageAsyncDelegate>(InCompletionDelegate);
+			CompletionDelegatePtr = MakeUnique<FLoadPackageAsyncDelegate>(MoveTemp(InCompletionDelegate));
 		}
 
 		// Add new package request
@@ -7397,7 +7397,7 @@ int32 LoadPackageAsync(const FString& PackageName, FLoadPackageAsyncDelegate Com
 {
 	const FGuid* Guid = nullptr;
 	const TCHAR* PackageToLoadFrom = nullptr;
-	return LoadPackageAsync(PackageName, Guid, PackageToLoadFrom, CompletionDelegate, InPackageFlags, InPIEInstanceID, InPackagePriority );
+	return LoadPackageAsync(PackageName, Guid, PackageToLoadFrom, MoveTemp(CompletionDelegate), InPackageFlags, InPIEInstanceID, InPackagePriority );
 }
 
 void CancelAsyncLoading()
