[AsyncLoading] Robomerge fix

diff --git a/Engine/Source/Runtime/CoreUObject/Private/Serialization/AsyncLoading.cpp b/Engine/Source/Runtime/CoreUObject/Private/Serialization/AsyncLoading.cpp
--- a/Engine/Source/Runtime/CoreUObject/Private/Serialization/AsyncLoading.cpp
+++ b/Engine/Source/Runtime/CoreUObject/Private/Serialization/AsyncLoading.cpp
@@ -5342,12 +5342,10 @@ void FAsyncLoadingThread::CancelAsyncLoadingInternal()
 
 	if (GEventDrivenLoaderEnabled)
 	{
-		FPrecacheCallbackHandler* PrecacheHandler = FCoreUObjectModule::GetGlobalPrecacheHandler();
-
-		while (PrecacheHandler->AnyIOOutstanding())
+		while (GPrecacheCallbackHandler->AnyIOOutstanding())
 		{
-			PrecacheHandler->WaitForIO(10.0f);
-			PrecacheHandler->ProcessIncoming();
+			GPrecacheCallbackHandler->WaitForIO(10.0f);
+			GPrecacheCallbackHandler->ProcessIncoming();
 		}
 	}
 
