TBA: Increase precision of floats in text and don't use exponential forms (because we don't parse them back in properly)

diff --git a/Engine/Source/Runtime/CoreUObject/Private/Serialization/Formatters/JsonArchiveOutputFormatter.cpp b/Engine/Source/Runtime/CoreUObject/Private/Serialization/Formatters/JsonArchiveOutputFormatter.cpp
--- a/Engine/Source/Runtime/CoreUObject/Private/Serialization/Formatters/JsonArchiveOutputFormatter.cpp
+++ b/Engine/Source/Runtime/CoreUObject/Private/Serialization/Formatters/JsonArchiveOutputFormatter.cpp
@@ -279,7 +279,7 @@ void FJsonArchiveOutputFormatter::Serialize(float& Value)
 	}
 	else
 	{
-		FString String = FString::Printf(TEXT("%.9g"), Value);
+		FString String = FString::Printf(TEXT("%.17f"), Value);
 #if DO_GUARD_SLOW
 		float RoundTripped;
 		LexFromString(RoundTripped, *String);
@@ -297,13 +297,13 @@ void FJsonArchiveOutputFormatter::Serialize(double& Value)
 	}
 	else
 	{
-		FString String = FString::Printf(TEXT("%.17g"), Value);
+		FString String = FString::Printf(TEXT("%.17f"), Value);
 #if DO_GUARD_SLOW
 		double RoundTripped;
 		LexFromString(RoundTripped, *String);
 		check(RoundTripped == Value);
 #endif
-		WriteValue(FString::Printf(TEXT("%.17g"), Value));
+		WriteValue(FString::Printf(TEXT("%.17f"), Value));
 	}
 }
 
