TBA: Use TryEnterField to "structure" uobject guid serialization properly
 - Encapsulates the bool/value pair that already existed, while allowing text archives to use their inherent format to decide whether the value exists or not

diff --git a/Engine/Source/Runtime/CoreUObject/Private/UObject/LazyObjectPtr.cpp b/Engine/Source/Runtime/CoreUObject/Private/UObject/LazyObjectPtr.cpp
--- a/Engine/Source/Runtime/CoreUObject/Private/UObject/LazyObjectPtr.cpp
+++ b/Engine/Source/Runtime/CoreUObject/Private/UObject/LazyObjectPtr.cpp
@@ -102,9 +102,8 @@ void FLazyObjectPtr::PossiblySerializeObjectGuid(UObject *Object, FStructuredArc
 	if (UnderlyingArchive.IsSaving() || UnderlyingArchive.IsCountingMemory())
 	{
 		FUniqueObjectGuid Guid = GuidAnnotation.GetAnnotation(Object);
-		bool HasGuid = Guid.IsValid();
-		Record << SA_VALUE(TEXT("HasGuid"), HasGuid);
-		if (HasGuid)
+		TOptional<FStructuredArchiveSlot> GuidSlot = Record.TryEnterField(SA_FIELD_NAME(TEXT("Guid")), Guid.IsValid());
+		if (GuidSlot.IsSet())
 		{
 			if (UnderlyingArchive.GetPortFlags() & PPF_DuplicateForPIE)
 			{
@@ -120,17 +119,16 @@ void FLazyObjectPtr::PossiblySerializeObjectGuid(UObject *Object, FStructuredArc
 				}
 			}
 
-			Record << SA_VALUE(TEXT("Guid"), Guid);
+			GuidSlot.GetValue() << Guid;
 		}
 	}
 	else if (UnderlyingArchive.IsLoading())
 	{
-		bool HasGuid = false;
-		Record << SA_VALUE(TEXT("HasGuid"), HasGuid);
-		if (HasGuid)
+		TOptional<FStructuredArchiveSlot> GuidSlot = Record.TryEnterField(SA_FIELD_NAME(TEXT("Guid")), false);
+		if (GuidSlot.IsSet())
 		{
 			FUniqueObjectGuid Guid;
-			Record << SA_VALUE(TEXT("Guid"), Guid);
+			GuidSlot.GetValue() << Guid;
 
 			// Don't try and resolve GUIDs when loading a package for diffing
 			const UPackage* Package = Object->GetOutermost();
