Optimize cooker ConditionallyExcludeForTarget with a new OBJECTMARK_KeepForTargetPlatform

diff --git a/Engine/Source/Runtime/CoreUObject/Private/UObject/SavePackage.cpp b/Engine/Source/Runtime/CoreUObject/Private/UObject/SavePackage.cpp
--- a/Engine/Source/Runtime/CoreUObject/Private/UObject/SavePackage.cpp
+++ b/Engine/Source/Runtime/CoreUObject/Private/UObject/SavePackage.cpp
@@ -588,10 +588,14 @@ static void ConditionallyExcludeObjectForTarget(UObject* Obj, EObjectMark Exclud
 		MarksToModify = (EObjectMark)(MarksToModify | (ObjToCheckMarks & MarkMask));
 	};
 
+	// MarksToProcess is a superset of marks retrieved from UPackage::GetExcludedObjectMarksForTargetPlatform
+	const uint32 MarksToProcess = OBJECTMARK_EditorOnly | OBJECTMARK_NotForClient | OBJECTMARK_NotForServer | OBJECTMARK_KeepForTargetPlatform;
+	check((ExcludedObjectMarks & ~MarksToProcess) == 0);
+
 	EObjectMark CurrentMarks = OBJECTMARK_NOMARKS;
-	InheritMarks(CurrentMarks, Obj, OBJECTMARK_EditorOnly | OBJECTMARK_NotForClient | OBJECTMARK_NotForServer);
+	InheritMarks(CurrentMarks, Obj, MarksToProcess);
 
-	if ((CurrentMarks & ExcludedObjectMarks) != 0)
+	if ((CurrentMarks & MarksToProcess) != 0)
 	{
 		// Already marked
 		return;
@@ -677,6 +681,12 @@ static void ConditionallyExcludeObjectForTarget(UObject* Obj, EObjectMark Exclud
 		NewMarks = (EObjectMark)(NewMarks | OBJECTMARK_EditorOnly);
 	}
 
+	// If not excluded after a full set of tests, it is implicitly a keep
+	if (NewMarks == 0)
+	{
+		NewMarks = OBJECTMARK_KeepForTargetPlatform;
+	}
+
 	// If our marks are different than original, set them on the object
 	if (CurrentMarks != NewMarks)
 	{
@@ -3609,8 +3619,8 @@ FSavePackageResultStruct UPackage::Save(UPackage* InOuter, UObject* Base, EObjec
 			
 				{
 					COOK_STAT(FScopedDurationTimer SaveTimer(SavePackageStats::TagPackageExportsTimeSec));
-					// Clear OBJECTMARK_TagExp again as we need to redo tagging below.
-					UnMarkAllObjects((EObjectMark)(OBJECTMARK_TagExp|OBJECTMARK_EditorOnly));
+					// Clear all marks (OBJECTMARK_TagExp and exclusion marks) again as we need to redo tagging below.
+					UnMarkAllObjects();
 			
 					// We need to serialize objects yet again to tag objects that were created by PreSave as OBJECTMARK_TagExp.
 					PackageExportTagger.TagPackageExports( ExportTaggerArchive, false );
diff --git a/Engine/Source/Runtime/CoreUObject/Public/UObject/UObjectMarks.h b/Engine/Source/Runtime/CoreUObject/Public/UObject/UObjectMarks.h
--- a/Engine/Source/Runtime/CoreUObject/Public/UObject/UObjectMarks.h
+++ b/Engine/Source/Runtime/CoreUObject/Public/UObject/UObjectMarks.h
@@ -50,6 +50,7 @@ enum EObjectMark
 	OBJECTMARK_NotForServer					= 0x00000040,		// Temporary save tag for server load flag.
 	OBJECTMARK_NotAlwaysLoadedForEditorGame	= 0x00000080,		// Temporary save tag for editorgame load flag.
 	OBJECTMARK_EditorOnly					= 0x00000100,		// Temporary editor only flag
+	OBJECTMARK_KeepForTargetPlatform		= 0x00000200,		// Temporary save tag to keep an object, derived from the load exclusion flags above
 	OBJECTMARK_ALLMARKS						= 0xffffffff,		// -1, all possible marks
 };
 
