TBA: Various text asset saving/loading fixes
 - Require users of farchive adapter classes to .Close() the adapter after writing is finished. It was previously done via scoped destruction but due to the inheritance patterns of the plain archive with the uobject overrides, this didn't work properly.
 - Create a structured archive sub-reader for ever export in the package. This allows us to jump backwards and forward between exports without breaking state.
 - When saving a text asset, use the export taggers mapping of allowed objects to make sure we don't serialize references to objects that binary packages would have ignored (transients etc.)

diff --git a/Engine/Source/Runtime/Core/Private/Serialization/ArchiveFromStructuredArchive.cpp b/Engine/Source/Runtime/Core/Private/Serialization/ArchiveFromStructuredArchive.cpp
--- a/Engine/Source/Runtime/Core/Private/Serialization/ArchiveFromStructuredArchive.cpp
+++ b/Engine/Source/Runtime/Core/Private/Serialization/ArchiveFromStructuredArchive.cpp
@@ -16,7 +16,7 @@ struct FArchiveFromStructuredArchiveImpl::FImpl
 
 	static constexpr int32 MaxBufferSize = 128;
 
-	bool bPendingSerialize = true;
+	bool bPendingSerialize = false;
 	bool bWasOpened        = false;
 
 	TArray<uint8> Buffer;
@@ -45,7 +45,7 @@ FArchiveFromStructuredArchiveImpl::FArchiveFromStructuredArchiveImpl(FStructured
 
 FArchiveFromStructuredArchiveImpl::~FArchiveFromStructuredArchiveImpl()
 {
-	Commit();
+	checkf(!Pimpl->bPendingSerialize, TEXT("Archive adapters must be closed before destruction"));
 }
 
 void FArchiveFromStructuredArchiveImpl::Flush()
@@ -153,29 +153,14 @@ FArchive& FArchiveFromStructuredArchiveImpl::operator<<(class UObject*& Value)
 			{
 				FStructuredArchive::FStream Stream = Pimpl->Root->EnterStream(SA_FIELD_NAME(TEXT("Objects")));
 
-				// We know exactly which stream index we want to load here, but because of the API we need to read through them
-				// in order, consuming the string name until we reach the entry we want and then load it as a uobject reference.
-				// If we are loading from a text archive, we could easily specify here which index we want, and the internal formatter
-				// can just push that single value by itself onto the value stack, but that same API couldn't be implemented for a
-				// binary archive as we can't skip over entries because we don't know how big they are. Maybe we could specify a stride
-				// or something, but at this point the API is complex and pretty formatter specific. Thought required!
-				// For now, just consume all the string names of the objects up until the one we need, then load that as an object
-				// pointer.
-
-				FString Dummy;
-				for (int32 Index = 0; Index < Pimpl->Objects.Num(); ++Index)
+				// Skip earlier elements
+				FString Str;
+				for (int32 Index = 0; Index < ObjectIdx; ++Index)
 				{
-					if (Index == ObjectIdx)
-					{
-						Stream.EnterElement() << Value;
-						Pimpl->Objects[ObjectIdx] = Value;
-					}
-					else
-					{
-						Stream.EnterElement() << Dummy;
-					}
+					Stream.EnterElement() << Str;
 				}
 
+				Stream.EnterElement() << Value;
 				Pimpl->Objects[ObjectIdx] = Value;
 				Pimpl->ObjectsValid[ObjectIdx] = true;
 			}
@@ -248,15 +233,17 @@ void FArchiveFromStructuredArchiveImpl::Commit()
 {
 	if (Pimpl->bWasOpened && InnerArchive.IsTextFormat())
 	{
-		SerializeInternal(Pimpl->Root.GetValue());
+		Finalize(Pimpl->Root.GetValue());
 	}
 }
 
-void FArchiveFromStructuredArchiveImpl::SerializeInternal(FStructuredArchive::FRecord Record)
+bool FArchiveFromStructuredArchiveImpl::Finalize(FStructuredArchive::FRecord Record)
 {
 	check(Pimpl->bWasOpened);
+	bool bShouldSerialize = Pimpl->bPendingSerialize;
+	Pimpl->bPendingSerialize = false;
 
-	if (Pimpl->bPendingSerialize)
+	if (bShouldSerialize)
 	{
 		FStructuredArchive::FSlot DataSlot = Record.EnterField(SA_FIELD_NAME(TEXT("Data")));
 		DataSlot.Serialize(Pimpl->Buffer);
@@ -291,9 +278,9 @@ void FArchiveFromStructuredArchiveImpl::SerializeInternal(FStructuredArchive::FR
 		{
 			NamesSlot.GetValue() << Pimpl->Names;
 		}
-
-		Pimpl->bPendingSerialize = false;
 	}
+
+	return bShouldSerialize;
 }
 
 void FArchiveFromStructuredArchiveImpl::OpenArchive()
@@ -304,11 +291,12 @@ void FArchiveFromStructuredArchiveImpl::OpenArchive()
 
 		if (InnerArchive.IsTextFormat())
 		{
+			Pimpl->bPendingSerialize = true;
 			Pimpl->Root = Pimpl->RootSlot.EnterRecord();
 
 			if (IsLoading())
 			{
-				SerializeInternal(Pimpl->Root.GetValue());
+				Finalize(Pimpl->Root.GetValue());
 			}
 		}
 		else
diff --git a/Engine/Source/Runtime/Core/Public/Serialization/StructuredArchive.h b/Engine/Source/Runtime/Core/Public/Serialization/StructuredArchive.h
--- a/Engine/Source/Runtime/Core/Public/Serialization/StructuredArchive.h
+++ b/Engine/Source/Runtime/Core/Public/Serialization/StructuredArchive.h
@@ -668,7 +668,7 @@ public:
 	bool ContainsData() const;
 
 protected:
-	virtual void SerializeInternal(FStructuredArchiveRecord Record);
+	virtual bool Finalize(FStructuredArchiveRecord Record);
 	void OpenArchive();
 
 private:
@@ -689,6 +689,8 @@ public:
 	      FArchive& GetArchive()       { return Impl; }
 	const FArchive& GetArchive() const { return Impl; }
 
+	void Close() { Impl.Close(); }
+
 private:
 	FArchiveFromStructuredArchiveImpl Impl;
 };
@@ -706,6 +708,8 @@ public:
 	      FArchive& GetArchive()       { return Ar; }
 	const FArchive& GetArchive() const { return Ar; }
 
+	void Close() {}
+
 private:
 	FArchive& Ar;
 };
diff --git a/Engine/Source/Runtime/CoreUObject/Private/Serialization/ArchiveUObjectFromStructuredArchive.cpp b/Engine/Source/Runtime/CoreUObject/Private/Serialization/ArchiveUObjectFromStructuredArchive.cpp
--- a/Engine/Source/Runtime/CoreUObject/Private/Serialization/ArchiveUObjectFromStructuredArchive.cpp
+++ b/Engine/Source/Runtime/CoreUObject/Private/Serialization/ArchiveUObjectFromStructuredArchive.cpp
@@ -6,7 +6,6 @@
 
 FArchiveUObjectFromStructuredArchiveImpl::FArchiveUObjectFromStructuredArchiveImpl(FStructuredArchive::FSlot Slot)
 	: Super(Slot)
-	, bPendingSerialize(true)
 {
 }
 
@@ -147,11 +146,10 @@ FArchive& FArchiveUObjectFromStructuredArchiveImpl::operator<<(FWeakObjectPtr& V
 	return *this;
 }
 
-void FArchiveUObjectFromStructuredArchiveImpl::SerializeInternal(FStructuredArchive::FRecord Record)
+bool FArchiveUObjectFromStructuredArchiveImpl::Finalize(FStructuredArchive::FRecord Record)
 {
-	Super::SerializeInternal(Record);
-
-	if (bPendingSerialize)
+	bool bShouldSerialize = Super::Finalize(Record);
+	if (bShouldSerialize)
 	{
 		TOptional<FStructuredArchive::FSlot> LazyObjectPtrsSlot = Record.TryEnterField(SA_FIELD_NAME(TEXT("LazyObjectPtrs")), LazyObjectPtrs.Num() > 0);
 		if (LazyObjectPtrsSlot.IsSet())
@@ -176,9 +174,8 @@ void FArchiveUObjectFromStructuredArchiveImpl::SerializeInternal(FStructuredArch
 		{
 			WeakObjectPtrsSlot.GetValue() << WeakObjectPtrs;
 		}
-
-		bPendingSerialize = true;
 	}
+	return bShouldSerialize;
 }
 
 #endif
\ No newline at end of file
diff --git a/Engine/Source/Runtime/CoreUObject/Private/Serialization/Formatters/JsonArchiveOutputFormatter.cpp b/Engine/Source/Runtime/CoreUObject/Private/Serialization/Formatters/JsonArchiveOutputFormatter.cpp
--- a/Engine/Source/Runtime/CoreUObject/Private/Serialization/Formatters/JsonArchiveOutputFormatter.cpp
+++ b/Engine/Source/Runtime/CoreUObject/Private/Serialization/Formatters/JsonArchiveOutputFormatter.cpp
@@ -332,13 +332,13 @@ void FJsonArchiveOutputFormatter::Serialize(FName& Value)
 
 void FJsonArchiveOutputFormatter::Serialize(UObject*& Value)
 {
-	if (Value == nullptr)
+	if (Value != nullptr && IsObjectAllowed(Value))
 	{
-		WriteValue(TEXT("null"));
+		SerializeStringInternal(FString::Printf(TEXT("Object:%s"), *Value->GetFullName()));
 	}
 	else
 	{
-		SerializeStringInternal(FString::Printf(TEXT("Object:%s"), *Value->GetFullName()));
+		WriteValue(TEXT("null"));
 	}
 }
 
@@ -351,7 +351,7 @@ void FJsonArchiveOutputFormatter::Serialize(FText& Value)
 
 void FJsonArchiveOutputFormatter::Serialize(FWeakObjectPtr& Value)
 {
-	if (Value.IsValid())
+	if (Value.IsValid() && IsObjectAllowed(Value.Get()))
 	{
 		SerializeStringInternal(FString::Printf(TEXT("Object:%s"), *Value.Get()->GetFullName()));
 	}
@@ -381,7 +381,7 @@ void FJsonArchiveOutputFormatter::Serialize(FSoftObjectPath& Value)
 
 void FJsonArchiveOutputFormatter::Serialize(FLazyObjectPtr& Value)
 {
-	if (Value.IsValid())
+	if (Value.IsValid() && IsObjectAllowed(Value.Get()))
 	{
 		SerializeStringInternal(FString::Printf(TEXT("Lazy:%s"), *Value.GetUniqueID().ToString()));
 	}
@@ -588,4 +588,9 @@ void FJsonArchiveOutputFormatter::SerializeStringInternal(const FString& String)
 	WriteValue(Result);
 }
 
+bool FJsonArchiveOutputFormatter::IsObjectAllowed(UObject* InObject) const
+{
+	return ObjectIndicesMap == nullptr || ObjectIndicesMap->Contains(InObject);
+}
+
 #endif
diff --git a/Engine/Source/Runtime/CoreUObject/Private/UObject/Class.cpp b/Engine/Source/Runtime/CoreUObject/Private/UObject/Class.cpp
--- a/Engine/Source/Runtime/CoreUObject/Private/UObject/Class.cpp
+++ b/Engine/Source/Runtime/CoreUObject/Private/UObject/Class.cpp
@@ -2215,6 +2215,7 @@ void UScriptStruct::SerializeItem(FStructuredArchive::FSlot Slot, void* Value, v
 				// The struct said that serialization succeeded but it didn't actually write anything.
 				Slot.EnterRecord();
 			}
+			Adapter.Close();
 #else
 			bItemSerialized = TheCppStructOps->Serialize(Slot.GetUnderlyingArchive(), Value);
 #endif
diff --git a/Engine/Source/Runtime/CoreUObject/Private/UObject/LinkerLoad.cpp b/Engine/Source/Runtime/CoreUObject/Private/UObject/LinkerLoad.cpp
--- a/Engine/Source/Runtime/CoreUObject/Private/UObject/LinkerLoad.cpp
+++ b/Engine/Source/Runtime/CoreUObject/Private/UObject/LinkerLoad.cpp
@@ -999,11 +999,17 @@ FLinkerLoad::ELinkerStatus FLinkerLoad::CreateLoader(
 			if (Filename.EndsWith(FPackageName::GetTextAssetPackageExtension()) || Filename.EndsWith(FPackageName::GetTextMapPackageExtension()))
 			{
 				INC_DWORD_STAT(STAT_TextAssetLinkerCount);
+				DECLARE_SCOPE_CYCLE_COUNTER(TEXT("FLinkerLoad::CreateTextArchiveFormatter"), STAT_LinkerLoad_CreateTextArchiveFormatter, STATGROUP_LinkerLoad);
+				TRACE_CPUPROFILER_EVENT_SCOPE_TEXT(TEXT("FLinkerLoad::CreateTextArchiveFormatter"));
 				Loader = IFileManager::Get().CreateFileReader(*Filename);
 				StructuredArchiveFormatter = new FJsonArchiveInputFormatter(*this, [this](const FString& InFullPath)
 				{
 					FPackageIndex Index = FindOrCreateImportOrExport(InFullPath);
-					if (Index.IsImport())
+					if (Index.IsNull())
+					{
+						return (UObject*)nullptr;
+					}
+					else if (Index.IsImport())
 					{
 						return CreateImport(Index.ToImport());
 					}
@@ -1734,6 +1740,13 @@ FLinkerLoad::ELinkerStatus FLinkerLoad::SerializeExportMap()
 
 #if WITH_TEXT_ARCHIVE_SUPPORT
 
+FStructuredArchiveSlot FLinkerLoad::GetExportSlot(FPackageIndex InExportIndex)
+{
+	check(InExportIndex.IsExport());
+	int32 Index = InExportIndex.ToExport();
+	return ExportReaders[Index]->GetRoot();
+}
+
 FPackageIndex FLinkerLoad::FindOrCreateImport(const FName InObjectName, const FName InClassName, const FName InClassPackageName)
 {
 	for (int32 ImportIndex = 0; ImportIndex < ImportMap.Num(); ++ImportIndex)
@@ -1828,12 +1841,13 @@ FLinkerLoad::ELinkerStatus FLinkerLoad::ReconstructImportAndExportMap()
 	if (!bHasReconstructedImportAndExportMap && IsTextFormat())
 	{
 		int32 NumExports = 0;
-		FStructuredArchive::FMap PackageExports = StructuredArchiveRootRecord->EnterMap(SA_FIELD_NAME(TEXT("Exports")), NumExports);
+		FStructuredArchiveMap PackageExports = StructuredArchiveRootRecord->EnterMap(SA_FIELD_NAME(TEXT("Exports")), NumExports);
 
 		TArray<FObjectTextExport> ExportRecords;
 		ExportRecords.Reserve(NumExports);
 		ExportMap.SetNum(NumExports);
 		OriginalExportNames.SetNum(NumExports);
+		ExportReaders.AddDefaulted(NumExports);
 
 		Summary.ExportCount = ExportMap.Num();
 		Summary.ImportCount = 0;
@@ -1842,9 +1856,26 @@ FLinkerLoad::ELinkerStatus FLinkerLoad::ReconstructImportAndExportMap()
 		{
 			FObjectTextExport& TextExport = ExportRecords.Emplace_GetRef(ExportMap[ExportIndex], nullptr);
 			FString ExportName;
-			PackageExports.EnterElement(ExportName) << TextExport;
+			ExportReaders[ExportIndex] = new FStructuredArchiveChildReader(PackageExports.EnterElement(ExportName));
+			ExportReaders[ExportIndex]->GetRoot() << TextExport;
 			OriginalExportNames[ExportIndex] = *ExportName;
-			ExportMap[ExportIndex].ObjectName = *ExtractObjectName(OriginalExportNames[ExportIndex].ToString());
+			
+			if (TextExport.OuterName.Len())
+			{
+				FString Class, Package, Object, SubObject;
+				FPackageName::SplitFullObjectPath(TextExport.OuterName, Class, Package, Object, SubObject);
+				FString OuterPath = Object;
+				if (SubObject.Len())
+				{
+					OuterPath += TEXT(":") + SubObject;
+				}
+				if (ExportName.StartsWith(OuterPath))
+				{
+					ExportName = ExportName.Right(ExportName.Len() - OuterPath.Len() - 1);
+				}
+			}
+
+			ExportMap[ExportIndex].ObjectName = *ExportName;
 		}
 		
 		// Now pass over all the exports and rebuild the export/import records
@@ -3690,7 +3721,7 @@ void FLinkerLoad::Preload( UObject* Object )
 #if WITH_EDITOR && WITH_TEXT_ARCHIVE_SUPPORT
 						if (IsTextFormat())
 						{
-							FStructuredArchive::FSlot ExportSlot = StructuredArchiveRootRecord->EnterRecord(SA_FIELD_NAME(TEXT("Exports"))).EnterField(SA_FIELD_NAME(*OriginalExportNames[Export.ThisIndex.ToExport()].ToString()));
+							FStructuredArchiveSlot ExportSlot = GetExportSlot(Export.ThisIndex);
 
 							if (bClassSupportsTextFormat)
 							{
@@ -3730,7 +3761,7 @@ void FLinkerLoad::Preload( UObject* Object )
 #if WITH_EDITOR && WITH_TEXT_ARCHIVE_SUPPORT
 						if (IsTextFormat())
 						{
-							FStructuredArchive::FSlot ExportSlot = StructuredArchiveRootRecord->EnterRecord(SA_FIELD_NAME(TEXT("Exports"))).EnterField(SA_FIELD_NAME(*OriginalExportNames[Export.ThisIndex.ToExport()].ToString()));
+							FStructuredArchive::FSlot ExportSlot = GetExportSlot(Export.ThisIndex);
 
 							if (bClassSupportsTextFormat)
 							{
@@ -4862,6 +4893,11 @@ void FLinkerLoad::Detach()
 
 	delete StructuredArchive;
 	StructuredArchive = nullptr;
+	for (FStructuredArchiveChildReader* Reader : ExportReaders)
+	{
+		delete Reader;
+	}
+	ExportReaders.Empty();
 	delete StructuredArchiveFormatter;
 	StructuredArchiveFormatter = nullptr;
 
diff --git a/Engine/Source/Runtime/CoreUObject/Private/UObject/SavePackage.cpp b/Engine/Source/Runtime/CoreUObject/Private/UObject/SavePackage.cpp
--- a/Engine/Source/Runtime/CoreUObject/Private/UObject/SavePackage.cpp
+++ b/Engine/Source/Runtime/CoreUObject/Private/UObject/SavePackage.cpp
@@ -3726,7 +3726,9 @@ FSavePackageResultStruct UPackage::Save(UPackage* InOuter, UObject* Base, EObjec
 				if (bTextFormat)
 				{
 					TextFormatArchive = IFileManager::Get().CreateFileWriter(*TextFormatTempFilename);
-					Formatter = new FJsonArchiveOutputFormatter(*TextFormatArchive);
+					FJsonArchiveOutputFormatter* OutputFormatter = new FJsonArchiveOutputFormatter(*TextFormatArchive);
+					OutputFormatter->SetObjectIndicesMap(&Linker->ObjectIndicesMap);
+					Formatter = OutputFormatter;
 				}
 				else
 #endif
@@ -5283,6 +5285,7 @@ FSavePackageResultStruct UPackage::Save(UPackage* InOuter, UObject* Base, EObjec
 								{
 									FArchiveUObjectFromStructuredArchive Adapter(ExportSlot);
 									Export.Object->GetClass()->SerializeDefaultObject(Export.Object, Adapter.GetArchive());
+									Adapter.Close();
 								}
 							}
 							else
@@ -5298,6 +5301,7 @@ FSavePackageResultStruct UPackage::Save(UPackage* InOuter, UObject* Base, EObjec
 								{
 									FArchiveUObjectFromStructuredArchive Adapter(ExportSlot);
 									Export.Object->Serialize(Adapter.GetArchive());
+									Adapter.Close();
 								}
 
 #if WITH_EDITOR
diff --git a/Engine/Source/Runtime/CoreUObject/Public/Serialization/ArchiveUObjectFromStructuredArchive.h b/Engine/Source/Runtime/CoreUObject/Public/Serialization/ArchiveUObjectFromStructuredArchive.h
--- a/Engine/Source/Runtime/CoreUObject/Public/Serialization/ArchiveUObjectFromStructuredArchive.h
+++ b/Engine/Source/Runtime/CoreUObject/Public/Serialization/ArchiveUObjectFromStructuredArchive.h
@@ -32,8 +32,6 @@ public:
 
 private:
 
-	bool bPendingSerialize;
-
 	TArray<FLazyObjectPtr> LazyObjectPtrs;
 	TArray<FWeakObjectPtr> WeakObjectPtrs;
 	TArray<FSoftObjectPtr> SoftObjectPtrs;
@@ -44,7 +42,7 @@ private:
 	TMap<FSoftObjectPtr, int32> SoftObjectPtrToIndex;
 	TMap<FSoftObjectPath, int32> SoftObjectPathToIndex;
 
-	virtual void SerializeInternal(FStructuredArchive::FRecord Record) override;
+	virtual bool Finalize(FStructuredArchive::FRecord Record) override;
 };
 
 class FArchiveUObjectFromStructuredArchive
@@ -58,6 +56,8 @@ public:
 	      FArchive& GetArchive()       { return Impl; }
 	const FArchive& GetArchive() const { return Impl; }
 
+	void Close() { Impl.Close(); }
+
 private:
 	FArchiveUObjectFromStructuredArchiveImpl Impl;
 };
diff --git a/Engine/Source/Runtime/CoreUObject/Public/Serialization/Formatters/JsonArchiveOutputFormatter.h b/Engine/Source/Runtime/CoreUObject/Public/Serialization/Formatters/JsonArchiveOutputFormatter.h
--- a/Engine/Source/Runtime/CoreUObject/Public/Serialization/Formatters/JsonArchiveOutputFormatter.h
+++ b/Engine/Source/Runtime/CoreUObject/Public/Serialization/Formatters/JsonArchiveOutputFormatter.h
@@ -4,6 +4,7 @@
 
 #include "CoreMinimal.h"
 #include "Serialization/StructuredArchiveFormatter.h"
+#include "UObject/ObjectResource.h"
 
 #if WITH_TEXT_ARCHIVE_SUPPORT
 
@@ -75,9 +76,16 @@ public:
 	virtual void Serialize(TArray<uint8>& Value) override;
 	virtual void Serialize(void* Data, uint64 DataSize) override;
 
+	void SetObjectIndicesMap(const TMap<UObject*, FPackageIndex>* InObjectIndicesMap)
+	{
+		ObjectIndicesMap = InObjectIndicesMap;
+	}
+
 private:
 	FArchive& Inner;
 
+	const TMap<UObject*, FPackageIndex>* ObjectIndicesMap = nullptr;
+
 	TArray<ANSICHAR> Newline;
 	bool bNeedsComma   = false;
 	bool bNeedsNewline = false;
@@ -99,6 +107,8 @@ private:
 	void WriteOptionalAttributedBlockClosing();
 
 	void SerializeStringInternal(const FString& String);
+
+	bool IsObjectAllowed(UObject* InObject) const;
 };
 
 #endif
diff --git a/Engine/Source/Runtime/CoreUObject/Public/UObject/LinkerLoad.h b/Engine/Source/Runtime/CoreUObject/Public/UObject/LinkerLoad.h
--- a/Engine/Source/Runtime/CoreUObject/Public/UObject/LinkerLoad.h
+++ b/Engine/Source/Runtime/CoreUObject/Public/UObject/LinkerLoad.h
@@ -142,6 +142,7 @@ private:
 	FStructuredArchive* StructuredArchive;
 	FArchiveFormatterType* StructuredArchiveFormatter;
 	TOptional<FStructuredArchive::FRecord> StructuredArchiveRootRecord;
+	TArray<FStructuredArchiveChildReader*> ExportReaders;
 
 	/** The archive that actually reads the raw data from disk.																*/
 	FArchive*				Loader;
@@ -1203,6 +1204,9 @@ private:
 #if WITH_TEXT_ARCHIVE_SUPPORT
 	// Cache of the export names in a text asset. Allows us to enter those export slots by index rather than needing to reconstruct the name
 	TArray<FName> OriginalExportNames;
+
+	// Function to get a slot for a given export
+	FStructuredArchiveSlot GetExportSlot(FPackageIndex InExportIndex);
 #endif
 
 public:
diff --git a/Engine/Source/Runtime/CoreUObject/Public/UObject/ObjectMacros.h b/Engine/Source/Runtime/CoreUObject/Public/UObject/ObjectMacros.h
--- a/Engine/Source/Runtime/CoreUObject/Public/UObject/ObjectMacros.h
+++ b/Engine/Source/Runtime/CoreUObject/Public/UObject/ObjectMacros.h
@@ -1416,7 +1416,7 @@ namespace UM
 	}
 
 #define IMPLEMENT_FARCHIVE_SERIALIZER( TClass ) void TClass::Serialize(FArchive& Ar) { TClass::Serialize(FStructuredArchiveFromArchive(Ar).GetSlot().EnterRecord()); }
-#define IMPLEMENT_FSTRUCTUREDARCHIVE_SERIALIZER( TClass ) void TClass::Serialize(FStructuredArchive::FRecord Record) { FArchiveUObjectFromStructuredArchive Ar(Record.EnterField(SA_FIELD_NAME(TEXT("BaseClassAutoGen")))); TClass::Serialize(Ar.GetArchive()); }
+#define IMPLEMENT_FSTRUCTUREDARCHIVE_SERIALIZER( TClass ) void TClass::Serialize(FStructuredArchive::FRecord Record) { FArchiveUObjectFromStructuredArchive Ar(Record.EnterField(SA_FIELD_NAME(TEXT("BaseClassAutoGen")))); TClass::Serialize(Ar.GetArchive()); Ar.Close(); }
 #define DECLARE_FARCHIVE_SERIALIZER( TClass, API ) virtual API void Serialize(FArchive& Ar) override;
 #define DECLARE_FSTRUCTUREDARCHIVE_SERIALIZER( TClass, API ) virtual API void Serialize(FStructuredArchive::FRecord Record) override;
 
