FMaxPackageSummarySize will now be a standalone thing not encapsulated inside of FAsyncLoadingThread

diff --git a/Engine/Source/Runtime/CoreUObject/Private/Serialization/AsyncLoading.cpp b/Engine/Source/Runtime/CoreUObject/Private/Serialization/AsyncLoading.cpp
--- a/Engine/Source/Runtime/CoreUObject/Private/Serialization/AsyncLoading.cpp
+++ b/Engine/Source/Runtime/CoreUObject/Private/Serialization/AsyncLoading.cpp
@@ -5008,7 +5008,9 @@ EAsyncPackageState::Type FAsyncLoadingThread::TickAsyncLoading(bool bUseTimeLimi
 }
 
 
-FMaxPackageSummarySize::FMaxPackageSummarySize()
+int32 FMaxPackageSummarySize::Value = 8192;
+
+void FMaxPackageSummarySize::Init()
 {
 	// this is used for the initial precache and should be large enough to find the actual Sum.TotalHeaderSize
 	// the editor packages may not have the AdditionalPackagesToCook array stripped so we need to allocate more memory
@@ -5025,8 +5027,6 @@ FMaxPackageSummarySize::FMaxPackageSummarySize()
 			Value = MinimumPackageSummarySize;
 		}
 	}
-#else
-	Value = 8192;
 #endif
 }
 
@@ -7433,7 +7433,7 @@ void ResetAsyncLoadingStats()
 void InitAsyncThread()
 {
 	FCoreDelegates::OnSyncLoadPackage.AddStatic([](const FString&) { GSyncLoadCount++; });
-
+	FMaxPackageSummarySize::Init();
 	FAsyncLoadingThread::Get().InitializeAsyncThread();
 }
 
@@ -7824,7 +7824,7 @@ void FAsyncArchive::ReadCallback(bool bWasCancelled, IAsyncReadRequest* Request)
 			}
 			else
 			{
-				int64 Size = FMath::Min<int64>(FAsyncLoadingThread::Get().MaxPackageSummarySize.Value, FileSize);
+				int64 Size = FMath::Min<int64>(FMaxPackageSummarySize::Value, FileSize);
 				LogItem(TEXT("Starting Summary"), 0, Size);
 				SummaryRequestPtr = Handle->ReadRequest(0, Size, AIOP_Normal, &ReadCallbackFunction);
 				// I need a precache request here to keep the memory alive until I submit the header request
@@ -7851,7 +7851,7 @@ void FAsyncArchive::ReadCallback(bool bWasCancelled, IAsyncReadRequest* Request)
 		}
 		else
 		{
-			FBufferReader Ar(Mem, FMath::Min<int64>(FAsyncLoadingThread::Get().MaxPackageSummarySize.Value, FileSize),/*bInFreeOnClose=*/ false, /*bIsPersistent=*/ true);
+			FBufferReader Ar(Mem, FMath::Min<int64>(FMaxPackageSummarySize::Value, FileSize),/*bInFreeOnClose=*/ false, /*bIsPersistent=*/ true);
 			FPackageFileSummary Sum;
 			Ar << Sum;
 			if (Ar.IsError() || Sum.TotalHeaderSize > FileSize || Sum.GetFileVersionUE4() < VER_UE4_OLDEST_LOADABLE_PACKAGE)
@@ -7863,9 +7863,9 @@ void FAsyncArchive::ReadCallback(bool bWasCancelled, IAsyncReadRequest* Request)
 				FScopeLock Lock(&SummaryRacePreventer);
 				//@todoio change header format to put the TotalHeaderSize at the start of the file
 				// we need to be sure that we can at least get the size from the initial request. This is an early warning that custom versions are starting to get too big, relocate the total size to be at offset 4!
-				checkf(Ar.Tell() < FAsyncLoadingThread::Get().MaxPackageSummarySize.Value / 2, 
+				checkf(Ar.Tell() < FMaxPackageSummarySize::Value / 2,
 					TEXT("The initial read request was too small (%d) compared to package %s header size (%lld). Try increasing s.MaxPackageSummarySize value in DefaultEngine.ini."),
-					FAsyncLoadingThread::Get().MaxPackageSummarySize.Value, *FileName, Ar.Tell());
+					FMaxPackageSummarySize::Value, *FileName, Ar.Tell());
 				
 				// Support for cooked EDL packages in the editor
 				bCookedForEDLInEditor = !FPlatformProperties::RequiresCookedData() && (Sum.PackageFlags & PKG_FilterEditorOnly) && Sum.PreloadDependencyCount > 0 && Sum.PreloadDependencyOffset > 0;
@@ -7877,7 +7877,7 @@ void FAsyncArchive::ReadCallback(bool bWasCancelled, IAsyncReadRequest* Request)
 				LoadPhase = ELoadPhase::WaitingForHeader;
 			}
 			FMemory::Free(Mem);
-			DEC_MEMORY_STAT_BY(STAT_AsyncFileMemory, FMath::Min<int64>(FAsyncLoadingThread::Get().MaxPackageSummarySize.Value, FileSize));
+			DEC_MEMORY_STAT_BY(STAT_AsyncFileMemory, FMath::Min<int64>(FMaxPackageSummarySize::Value, FileSize));
 		}
 	}
 	else
diff --git a/Engine/Source/Runtime/CoreUObject/Private/Serialization/AsyncLoadingThread.h b/Engine/Source/Runtime/CoreUObject/Private/Serialization/AsyncLoadingThread.h
--- a/Engine/Source/Runtime/CoreUObject/Private/Serialization/AsyncLoadingThread.h
+++ b/Engine/Source/Runtime/CoreUObject/Private/Serialization/AsyncLoadingThread.h
@@ -134,10 +134,13 @@ struct FFlushTree
 	}
 };
 
+/** Holds the maximum package summary size that can be set via ini files
+  * This is used for the initial precache and should be large enough to hold the actual Sum.TotalHeaderSize 
+  */
 struct FMaxPackageSummarySize
 {
-	int32 Value;
-	FMaxPackageSummarySize();
+	static int32 Value;
+	static void Init();
 };
 
 /**
@@ -207,8 +210,6 @@ class FAsyncLoadingThread : public FRunnable
 public:
 	/** [EDL] Async Packages that are ready for tick */
 	TArray<FAsyncPackage*> AsyncPackagesReadyForTick;
-	/* This is used for the initial precache and should be large enough to find the actual Sum.TotalHeaderSize */
-	const FMaxPackageSummarySize MaxPackageSummarySize;
 private:
 
 #if THREADSAFE_UOBJECTS
