TBA: If an array or object JSON element is entered but no data is written, don't add a newline before closing that element. Gives us nice [] and {} output for empty objects/array

diff --git a/Engine/Source/Runtime/CoreUObject/Private/Serialization/Formatters/JsonArchiveOutputFormatter.cpp b/Engine/Source/Runtime/CoreUObject/Private/Serialization/Formatters/JsonArchiveOutputFormatter.cpp
--- a/Engine/Source/Runtime/CoreUObject/Private/Serialization/Formatters/JsonArchiveOutputFormatter.cpp
+++ b/Engine/Source/Runtime/CoreUObject/Private/Serialization/Formatters/JsonArchiveOutputFormatter.cpp
@@ -40,6 +40,7 @@ void FJsonArchiveOutputFormatter::EnterRecord()
 	Write("{");
 	Newline.Add('\t');
 	bNeedsNewline = true;
+	TextStartPosStack.Push(Inner.Tell());
 }
 
 void FJsonArchiveOutputFormatter::EnterRecord_TextOnly(TArray<FString>& OutFieldNames)
@@ -51,6 +52,10 @@ void FJsonArchiveOutputFormatter::EnterRecord_TextOnly(TArray<FString>& OutField
 void FJsonArchiveOutputFormatter::LeaveRecord()
 {
 	Newline.Pop(false);
+	if (TextStartPosStack.Pop() == Inner.Tell())
+	{
+		bNeedsNewline = false;
+	}
 	WriteOptionalNewline();
 	Write("}");
 	bNeedsComma = true;
@@ -118,6 +123,7 @@ void FJsonArchiveOutputFormatter::EnterStream()
 	Write("[");
 	Newline.Add('\t');
 	bNeedsNewline = true;
+	TextStartPosStack.Push(Inner.Tell());
 }
 
 void FJsonArchiveOutputFormatter::EnterStream_TextOnly(int32& OutNumElements)
@@ -129,6 +135,10 @@ void FJsonArchiveOutputFormatter::EnterStream_TextOnly(int32& OutNumElements)
 void FJsonArchiveOutputFormatter::LeaveStream()
 {
 	Newline.Pop(false);
+	if (TextStartPosStack.Pop() == Inner.Tell())
+	{
+		bNeedsNewline = false;
+	}
 	WriteOptionalNewline();
 	Write("]");
 	bNeedsComma = true;
diff --git a/Engine/Source/Runtime/CoreUObject/Public/Serialization/Formatters/JsonArchiveOutputFormatter.h b/Engine/Source/Runtime/CoreUObject/Public/Serialization/Formatters/JsonArchiveOutputFormatter.h
--- a/Engine/Source/Runtime/CoreUObject/Public/Serialization/Formatters/JsonArchiveOutputFormatter.h
+++ b/Engine/Source/Runtime/CoreUObject/Public/Serialization/Formatters/JsonArchiveOutputFormatter.h
@@ -91,6 +91,7 @@ private:
 	bool bNeedsNewline = false;
 
 	TArray<int32> NumAttributesStack;
+	TArray<int64> TextStartPosStack;
 
 	void Write(ANSICHAR Character);
 
