Fixed deprecation warning regarding GIsRequestingExit
#rnx

diff --git a/Engine/Source/Runtime/CoreUObject/Private/Serialization/AsyncLoading.cpp b/Engine/Source/Runtime/CoreUObject/Private/Serialization/AsyncLoading.cpp
--- a/Engine/Source/Runtime/CoreUObject/Private/Serialization/AsyncLoading.cpp
+++ b/Engine/Source/Runtime/CoreUObject/Private/Serialization/AsyncLoading.cpp
@@ -7397,7 +7397,7 @@ void CancelAsyncLoading()
 
 	FAsyncLoadingThread::Get().CancelAsyncLoading();
 
-	if (!GIsRequestingExit)
+	if (!IsEngineExitRequested())
 	{
 		CollectGarbage(GARBAGE_COLLECTION_KEEPFLAGS, true);
 	}
