Make StaticFindObject() handle searching for any package with a package-specifying string

The change in behavior was introduced by an optimization in UE 4.23

diff --git a/Engine/Source/Runtime/CoreUObject/Private/UObject/UObjectHash.cpp b/Engine/Source/Runtime/CoreUObject/Private/UObject/UObjectHash.cpp
--- a/Engine/Source/Runtime/CoreUObject/Private/UObject/UObjectHash.cpp
+++ b/Engine/Source/Runtime/CoreUObject/Private/UObject/UObjectHash.cpp
@@ -500,40 +500,31 @@ static FName SplitInnerAndOuter(const TCHAR* OuterAndInner, const TCHAR* Delimit
 
 static FName ExtractInnerAndOuterFromPath(FName ObjectPath, FName& OutOuter)
 {
-	TCHAR Path[NAME_SIZE];
-	ObjectPath.GetPlainNameString(Path);
+	TCHAR PathBuffer[NAME_SIZE];
+	ObjectPath.GetPlainNameString(PathBuffer);
 
 	// Find package separator . or subobject separator :
 	constexpr FAsciiSet DotColon(".:");
-	const TCHAR* DelimiterOrEnd = FAsciiSet::FindFirstOrEnd(Path, DotColon);
-
-	if (*DelimiterOrEnd == '\0')
-	{
-		return ObjectPath;
-	}
-	else if (*DelimiterOrEnd == ':')
+	const TCHAR* Path = PathBuffer;
+	while (true)
 	{
-		return SplitInnerAndOuter(Path, DelimiterOrEnd, OutOuter, ObjectPath.GetNumber());
-	}
+		const TCHAR* DelimiterOrEnd = FAsciiSet::FindFirstOrEnd(Path, DotColon);
 
-	// We have a package prefix, drop it
-	check(*DelimiterOrEnd == '.');
-	const TCHAR* PathWithoutPackage = DelimiterOrEnd + 1;
+		if (*DelimiterOrEnd == '\0')
+		{
+			return ObjectPath;
+		}
+		else if (*DelimiterOrEnd == ':')
+		{
+			return SplitInnerAndOuter(Path, DelimiterOrEnd, OutOuter, ObjectPath.GetNumber());
+		}
 
-	// Check if there is also a subobject delimiter
-	constexpr FAsciiSet Colon(":");
-	const TCHAR* ColonOrEnd = FAsciiSet::FindFirstOrEnd(PathWithoutPackage, Colon);
-	if (*ColonOrEnd == '\0')
-	{
-		return FName(ColonOrEnd - PathWithoutPackage, PathWithoutPackage, ObjectPath.GetNumber());
-	}
-	else		
-	{
-		return SplitInnerAndOuter(PathWithoutPackage, ColonOrEnd, OutOuter, ObjectPath.GetNumber());
+		// We have a package prefix, drop it
+		check(*DelimiterOrEnd == '.');
+		Path = DelimiterOrEnd + 1;
 	}
 }
 
-
 UObject* StaticFindObjectFastInternalThreadSafe(FUObjectHashTables& ThreadHash, const UClass* ObjectClass, const UObject* ObjectPackage, FName ObjectName, bool bExactClass, bool bAnyPackage, EObjectFlags ExcludeFlags, EInternalObjectFlags ExclusiveInternalFlags)
 {
 	ExclusiveInternalFlags |= EInternalObjectFlags::Unreachable;
