TBA: Number parsing improvements
- Add verifies around json string reading of numbers
- Write numbers to json using %g format specs (again) to give engineering forms and more compact representations of more complex numbers

diff --git a/Engine/Source/Runtime/CoreUObject/Private/Serialization/Formatters/JsonArchiveInputFormatter.cpp b/Engine/Source/Runtime/CoreUObject/Private/Serialization/Formatters/JsonArchiveInputFormatter.cpp
--- a/Engine/Source/Runtime/CoreUObject/Private/Serialization/Formatters/JsonArchiveInputFormatter.cpp
+++ b/Engine/Source/Runtime/CoreUObject/Private/Serialization/Formatters/JsonArchiveInputFormatter.cpp
@@ -304,52 +304,52 @@ bool FJsonArchiveInputFormatter::TryEnterAttribute(FArchiveFieldName AttributeNa
 
 void FJsonArchiveInputFormatter::Serialize(uint8& Value)
 {
-	ValueStack.Top()->TryGetNumber(Value);
+	verify(ValueStack.Top()->TryGetNumber(Value));
 }
 
 void FJsonArchiveInputFormatter::Serialize(uint16& Value)
 {
-	ValueStack.Top()->TryGetNumber(Value);
+	verify(ValueStack.Top()->TryGetNumber(Value));
 }
 
 void FJsonArchiveInputFormatter::Serialize(uint32& Value)
 {
-	ValueStack.Top()->TryGetNumber(Value);
+	verify(ValueStack.Top()->TryGetNumber(Value));
 }
 
 void FJsonArchiveInputFormatter::Serialize(uint64& Value)
 {
-	ValueStack.Top()->TryGetNumber(Value);
+	verify(ValueStack.Top()->TryGetNumber(Value));
 }
 
 void FJsonArchiveInputFormatter::Serialize(int8& Value)
 {
-	ValueStack.Top()->TryGetNumber(Value);
+	verify(ValueStack.Top()->TryGetNumber(Value));
 }
 
 void FJsonArchiveInputFormatter::Serialize(int16& Value)
 {
-	ValueStack.Top()->TryGetNumber(Value);;
+	verify(ValueStack.Top()->TryGetNumber(Value));
 }
 
 void FJsonArchiveInputFormatter::Serialize(int32& Value)
 {
-	ValueStack.Top()->TryGetNumber(Value);
+	verify(ValueStack.Top()->TryGetNumber(Value));
 }
 
 void FJsonArchiveInputFormatter::Serialize(int64& Value)
 {
-	ValueStack.Top()->TryGetNumber(Value);
+	verify(ValueStack.Top()->TryGetNumber(Value));
 }
 
 void FJsonArchiveInputFormatter::Serialize(float& Value)
 {
-	ValueStack.Top()->TryGetNumber(Value);
+	verify(ValueStack.Top()->TryGetNumber(Value));
 }
 
 void FJsonArchiveInputFormatter::Serialize(double& Value)
 {
-	ValueStack.Top()->TryGetNumber(Value);
+	verify(ValueStack.Top()->TryGetNumber(Value));
 }
 
 void FJsonArchiveInputFormatter::Serialize(bool& Value)
diff --git a/Engine/Source/Runtime/CoreUObject/Private/Serialization/Formatters/JsonArchiveOutputFormatter.cpp b/Engine/Source/Runtime/CoreUObject/Private/Serialization/Formatters/JsonArchiveOutputFormatter.cpp
--- a/Engine/Source/Runtime/CoreUObject/Private/Serialization/Formatters/JsonArchiveOutputFormatter.cpp
+++ b/Engine/Source/Runtime/CoreUObject/Private/Serialization/Formatters/JsonArchiveOutputFormatter.cpp
@@ -289,7 +289,7 @@ void FJsonArchiveOutputFormatter::Serialize(float& Value)
 	}
 	else
 	{
-		FString String = FString::Printf(TEXT("%.17f"), Value);
+		FString String = FString::Printf(TEXT("%.17g"), Value);
 #if DO_GUARD_SLOW
 		float RoundTripped;
 		LexFromString(RoundTripped, *String);
@@ -307,13 +307,13 @@ void FJsonArchiveOutputFormatter::Serialize(double& Value)
 	}
 	else
 	{
-		FString String = FString::Printf(TEXT("%.17f"), Value);
+		FString String = FString::Printf(TEXT("%.17g"), Value);
 #if DO_GUARD_SLOW
 		double RoundTripped;
 		LexFromString(RoundTripped, *String);
 		check(RoundTripped == Value);
 #endif
-		WriteValue(FString::Printf(TEXT("%.17f"), Value));
+		WriteValue(String);
 	}
 }
 
