Fix struct properties being skipped during recursion
The original fix from CL# 4056552 has been reverted, as this change fixes both the original problem UE-58753 and a newly identified issue with arrays of structs only finding the first element.  Which also now finds the ArrayOfStructs elements in the test code, which needs its property totals updated.

https://udn.unrealengine.com/questions/524953/tpropertyvalueiterator-only-finds-the-first-value.html

diff --git a/Engine/Source/Runtime/CoreUObject/Private/UObject/UnrealType.cpp b/Engine/Source/Runtime/CoreUObject/Private/UObject/UnrealType.cpp
--- a/Engine/Source/Runtime/CoreUObject/Private/UObject/UnrealType.cpp
+++ b/Engine/Source/Runtime/CoreUObject/Private/UObject/UnrealType.cpp
@@ -5,7 +5,7 @@
 
 DEFINE_LOG_CATEGORY(LogType);
 
-bool FPropertyValueIterator::NextValue(EPropertyValueIteratorFlags InRecursionFlags, bool bReturningFromStruct)
+bool FPropertyValueIterator::NextValue(EPropertyValueIteratorFlags InRecursionFlags)
 {
 	if (PropertyIteratorStack.Num() == 0)
 	{
@@ -15,12 +15,6 @@ bool FPropertyValueIterator::NextValue(EPropertyValueIteratorFlags InRecursionFl
 
 	FPropertyValueStackEntry& Entry = PropertyIteratorStack.Last();
 
-	// If we are returning from struct, increment value index as we delayed incrementing it when entering recursion
-	if (bReturningFromStruct)
-	{
-		Entry.ValueIndex++;
-	}
-
 	// If we have pending values, deal with them
 	if (Entry.ValueIndex < Entry.ValueArray.Num())
 	{
@@ -82,21 +76,19 @@ bool FPropertyValueIterator::NextValue(EPropertyValueIteratorFlags InRecursionFl
 		{
 			if (InRecursionFlags == EPropertyValueIteratorFlags::FullRecursion)
 			{
-				// And child to stack and rerun. This invalidates Entry. We leave ValueIndex the same so it can be used for recursive property lookup
-				PropertyIteratorStack.Emplace(StructProperty->Struct, PropertyValue, DeprecatedPropertyFlags);
-
-				return NextValue(InRecursionFlags, false);
+				// We don't need recursion on these - this will happen naturally as we process these values
+				for (FPropertyValueIterator Iter(UProperty::StaticClass(), StructProperty->Struct, PropertyValue, EPropertyValueIteratorFlags::NoRecursion, DeprecatedPropertyFlags); Iter; ++Iter)
+				{
+					Entry.ValueArray.EmplaceAt(InsertIndex++, Iter->Key, Iter->Value);
+				}
 			}
 		}
 
 		// Else this is a normal property and has nothing to expand
 		// We don't expand enum properties because EnumProperty handles value wrapping for us
 
-		// We didn't recurse into a struct, so increment next value to check, unless we we just popped a struct off the stack
-		if (!bReturningFromStruct)
-		{
-			Entry.ValueIndex++;
-		}
+		// Increment next value to check
+		Entry.ValueIndex++;
 	}
 
 	// Out of pending values, try to add more
@@ -109,7 +101,10 @@ bool FPropertyValueIterator::NextValue(EPropertyValueIteratorFlags InRecursionFl
 
 			if (PropertyIteratorStack.Num() > 0)
 			{
-				return NextValue(InRecursionFlags, true);
+				// Increment value index as we delayed incrementing it when entering recursion
+				PropertyIteratorStack.Last().ValueIndex++;
+
+				return NextValue(InRecursionFlags);
 			}
 
 			return false;
@@ -146,7 +141,7 @@ void FPropertyValueIterator::IterateToNext()
 		bSkipRecursionOnce = false;
 	}
 
-	while (NextValue(LocalRecursionFlags, false))
+	while (NextValue(LocalRecursionFlags))
 	{
 		// If this property is valid type, stop iteration
 		FPropertyValueStackEntry& Entry = PropertyIteratorStack.Last();
diff --git a/Engine/Source/Runtime/CoreUObject/Public/UObject/UnrealType.h b/Engine/Source/Runtime/CoreUObject/Public/UObject/UnrealType.h
--- a/Engine/Source/Runtime/CoreUObject/Public/UObject/UnrealType.h
+++ b/Engine/Source/Runtime/CoreUObject/Public/UObject/UnrealType.h
@@ -5276,7 +5276,7 @@ private:
 	bool bSkipRecursionOnce;
 
 	/** Goes to the next Property/value pair. Returns true if next value is valid */
-	bool NextValue(EPropertyValueIteratorFlags RecursionFlags, bool bReturningFromStruct);
+	bool NextValue(EPropertyValueIteratorFlags RecursionFlags);
 
 	/** Iterates to next property being checked for or until reaching the end of the structure */
 	COREUOBJECT_API void IterateToNext();
