TBA: In object path strings, include the class package as well as the name so that it can be resolved properly when loading

diff --git a/Engine/Source/Runtime/CoreUObject/Private/Misc/PackageName.cpp b/Engine/Source/Runtime/CoreUObject/Private/Misc/PackageName.cpp
--- a/Engine/Source/Runtime/CoreUObject/Private/Misc/PackageName.cpp
+++ b/Engine/Source/Runtime/CoreUObject/Private/Misc/PackageName.cpp
@@ -503,51 +503,28 @@ bool FPackageName::SplitLongPackageName(const FString& InLongPackageName, FStrin
 void FPackageName::SplitFullObjectPath(const FString& InFullObjectPath, FString& OutClassName, FString& OutPackageName, FString& OutObjectName, FString& OutSubObjectName)
 {
 	FString Sanitized = InFullObjectPath.TrimStartAndEnd();
+	const TCHAR* Cur = *Sanitized;
 
-	int32 StartOfClass = 0;
-	int32 StartOfPackage = 0;
-	int32 StartOfObject = 0;
-	int32 StartOfSubObject = 0;
-	int32 EndOfClass = 0;
-	int32 EndOfPackage = 0;
-	int32 EndOfObject = 0;
-	int32 EndOfSubObject = 0;
-
-	int32 TempIndex = INDEX_NONE;
-	if (Sanitized.FindChar(' ', TempIndex))
-	{
-		EndOfClass = TempIndex;
-		StartOfPackage = EndOfPackage = EndOfClass + 1;
-	}
-
-	if (Sanitized.FindChar('.', TempIndex))
+	auto ExtractBeforeDelim = [&Cur](TCHAR Delim, FString& OutString)
 	{
-		EndOfPackage = TempIndex;
-		StartOfObject = EndOfObject = EndOfPackage + 1;
-
-		if (Sanitized.FindChar(':', TempIndex))
+		const TCHAR* Start = Cur;
+		while (*Cur != '\0' && *Cur != Delim)
 		{
-			EndOfObject = TempIndex;
-			if (TempIndex < InFullObjectPath.Len())
-			{
-				StartOfSubObject = TempIndex + 1;
-				EndOfSubObject = InFullObjectPath.Len();
-			}
+			++Cur;
 		}
-		else
+
+		OutString = FString(Cur - Start, Start);
+
+		if (*Cur == Delim)
 		{
-			EndOfObject = Sanitized.Len();
+			++Cur;
 		}
-	}
-	else
-	{
-		EndOfPackage = Sanitized.Len();
-	}
+	};
 
-	OutClassName = Sanitized.Mid(StartOfClass, EndOfClass - StartOfClass);
-	OutPackageName = Sanitized.Mid(StartOfPackage, EndOfPackage - StartOfPackage);
-	OutObjectName = Sanitized.Mid(StartOfObject, EndOfObject - StartOfObject);
-	OutSubObjectName = Sanitized.Mid(StartOfSubObject, EndOfSubObject - StartOfSubObject);
+	ExtractBeforeDelim(' ', OutClassName);
+	ExtractBeforeDelim('.', OutPackageName);
+	ExtractBeforeDelim(':', OutObjectName);
+	ExtractBeforeDelim('\0', OutSubObjectName);
 }
 
 FString FPackageName::GetLongPackageAssetName(const FString& InLongPackageName)
diff --git a/Engine/Source/Runtime/CoreUObject/Private/Serialization/Formatters/JsonArchiveOutputFormatter.cpp b/Engine/Source/Runtime/CoreUObject/Private/Serialization/Formatters/JsonArchiveOutputFormatter.cpp
--- a/Engine/Source/Runtime/CoreUObject/Private/Serialization/Formatters/JsonArchiveOutputFormatter.cpp
+++ b/Engine/Source/Runtime/CoreUObject/Private/Serialization/Formatters/JsonArchiveOutputFormatter.cpp
@@ -325,7 +325,7 @@ void FJsonArchiveOutputFormatter::Serialize(bool& Value)
 void FJsonArchiveOutputFormatter::Serialize(FString& Value)
 {
 	// Insert a "String:" prefix to prevent incorrect interpretation as another explicit type
-	if (Value.StartsWith(TEXT("Name:")) || Value.StartsWith(TEXT("Object:")) || Value.StartsWith(TEXT("String:")) || Value.StartsWith(TEXT("Base64:")))
+	if (Value.StartsWith(TEXT("Object:")) || Value.StartsWith(TEXT("String:")) || Value.StartsWith(TEXT("Base64:")))
 	{
 		SerializeStringInternal(FString::Printf(TEXT("String:%s"), *Value));
 	}
@@ -337,14 +337,15 @@ void FJsonArchiveOutputFormatter::Serialize(FString& Value)
 
 void FJsonArchiveOutputFormatter::Serialize(FName& Value)
 {
-	SerializeStringInternal(FString::Printf(TEXT("Name:%s"), *Value.ToString()));
+	SerializeStringInternal(*Value.ToString());
 }
 
 void FJsonArchiveOutputFormatter::Serialize(UObject*& Value)
 {
 	if (Value != nullptr && IsObjectAllowed(Value))
 	{
-		SerializeStringInternal(FString::Printf(TEXT("Object:%s"), *Value->GetFullName()));
+		FString FullObjectName = Value->GetFullName(nullptr, EObjectFullNameFlags::IncludeClassPackage);
+		SerializeStringInternal(FString::Printf(TEXT("Object:%s"), *FullObjectName));
 	}
 	else
 	{
@@ -363,7 +364,7 @@ void FJsonArchiveOutputFormatter::Serialize(FWeakObjectPtr& Value)
 {
 	if (Value.IsValid() && IsObjectAllowed(Value.Get()))
 	{
-		SerializeStringInternal(FString::Printf(TEXT("Object:%s"), *Value.Get()->GetFullName()));
+		SerializeStringInternal(FString::Printf(TEXT("Object:%s"), *Value.Get()->GetFullName(nullptr, EObjectFullNameFlags::IncludeClassPackage)));
 	}
 	else
 	{
diff --git a/Engine/Source/Runtime/CoreUObject/Private/UObject/LinkerLoad.cpp b/Engine/Source/Runtime/CoreUObject/Private/UObject/LinkerLoad.cpp
--- a/Engine/Source/Runtime/CoreUObject/Private/UObject/LinkerLoad.cpp
+++ b/Engine/Source/Runtime/CoreUObject/Private/UObject/LinkerLoad.cpp
@@ -1813,12 +1813,28 @@ FPackageIndex FLinkerLoad::FindOrCreateImportOrExport(const FString& InFullPath)
 	else
 	{
 		FName ClassName = *Class;
+		FName ClassPackageName;
 		FName PackageName = *Package;
 
-		if (UClass* ObjectClass = FindObjectFast<UClass>(nullptr, ClassName, false, true))
+		if (Class.StartsWith(TEXT("/")))
 		{
-			FName ClassPackageName = *ObjectClass->GetOuterUPackage()->GetPathName();
+			int32 EndOfPackage = INDEX_NONE;
+			if (Class.FindChar('.', EndOfPackage))
+			{
+				ClassPackageName = *FString(EndOfPackage, *Class);
+				ClassName = *FString(Class.Len() - EndOfPackage, *Class + EndOfPackage + 1);
+			}
+		}
+		else
+		{
+			if (UClass * ObjectClass = FindObjectFast<UClass>(nullptr, ClassName, false, true))
+			{
+				ClassPackageName = *ObjectClass->GetOuterUPackage()->GetPathName();
+			}
+		}
 
+		if (ClassPackageName != NAME_None)
+		{
 			FPackageIndex ImportOuterIndex = FindOrCreateImport(PackageName, NAME_Package, FName(TEXT("/Script/CoreUObject")));
 			FPackageIndex ImportIndex = FindOrCreateImport(ObjectName, ClassName, ClassPackageName);
 			ImportMap[ImportIndex.ToImport()].OuterIndex = ImportOuterIndex;
diff --git a/Engine/Source/Runtime/CoreUObject/Private/UObject/UObjectBaseUtility.cpp b/Engine/Source/Runtime/CoreUObject/Private/UObject/UObjectBaseUtility.cpp
--- a/Engine/Source/Runtime/CoreUObject/Private/UObject/UObjectBaseUtility.cpp
+++ b/Engine/Source/Runtime/CoreUObject/Private/UObject/UObjectBaseUtility.cpp
@@ -72,22 +72,29 @@ void UObjectBaseUtility::GetPathName(const UObject* StopOuter, FString& ResultSt
  *
  * @note	safe to call on NULL object pointers!
  */
-FString UObjectBaseUtility::GetFullName(const UObject* StopOuter) const
+FString UObjectBaseUtility::GetFullName(const UObject* StopOuter, EObjectFullNameFlags Flags) const
 {
 	FString Result;
 	Result.Empty(128);
-	GetFullName(StopOuter, Result);
+	GetFullName(StopOuter, Result, Flags);
 	return Result;
 }
 
  /**
   * Version of GetFullName() that eliminates unnecessary copies and appends an existing string.
   */
-void UObjectBaseUtility::GetFullName(const UObject* StopOuter, FString& ResultString) const
+void UObjectBaseUtility::GetFullName(const UObject* StopOuter, FString& ResultString, EObjectFullNameFlags Flags) const
 {
 	if (this != nullptr)
 	{
-		GetClass()->AppendName(ResultString);
+		if (EnumHasAllFlags(Flags, EObjectFullNameFlags::IncludeClassPackage))
+		{
+			ResultString += GetClass()->GetPathName();
+		}
+		else
+		{
+			GetClass()->AppendName(ResultString);
+		}
 		ResultString += TEXT(' ');
 		GetPathName(StopOuter, ResultString);
 	}
diff --git a/Engine/Source/Runtime/CoreUObject/Public/UObject/UObjectBaseUtility.h b/Engine/Source/Runtime/CoreUObject/Public/UObject/UObjectBaseUtility.h
--- a/Engine/Source/Runtime/CoreUObject/Public/UObject/UObjectBaseUtility.h
+++ b/Engine/Source/Runtime/CoreUObject/Public/UObject/UObjectBaseUtility.h
@@ -19,6 +19,20 @@
 	#endif
 #endif
 
+/**
+* Enum which specifies the mode in which full object names are constructed
+*/
+enum class EObjectFullNameFlags
+{
+	// Standard object full name (i.e. "Type PackageName.ObjectName:SubobjectName")
+	None = 0,
+
+	// Adds package to the type portion (i.e. "TypePackage.TypeName PackageName.ObjectName:SubobjectName")
+	IncludeClassPackage = 1,
+};
+
+ENUM_CLASS_FLAGS(EObjectFullNameFlags);
+
 /**
  * Provides utility functions for UObject, this class should not be used directly
  */
@@ -287,15 +301,16 @@ public:
 	 *
 	 * @param	StopOuter	if specified, indicates that the output string should be relative to this object.  if StopOuter
 	 *						does not exist in this object's Outer chain, the result would be the same as passing NULL.
+	 * @param	Flags		flags that control the behavior of full name generation
 	 *
 	 * @note	safe to call on NULL object pointers!
 	 */
-	FString GetFullName( const UObject* StopOuter=NULL ) const;
+	FString GetFullName( const UObject* StopOuter=NULL, EObjectFullNameFlags Flags = EObjectFullNameFlags::None ) const;
 
 	/**
 	 * Version of GetFullName() that eliminates unnecessary copies.
 	 */
-	void GetFullName(const UObject* StopOuter, FString& ResultString) const;
+	void GetFullName( const UObject* StopOuter, FString& ResultString, EObjectFullNameFlags Flags = EObjectFullNameFlags::None ) const;
 
 	/**
 	 * Returns the fully qualified pathname for this object, in the format:
