TBA: Don't look for Name: at the start of string fields. Explicit marking up of name fields has been removed

diff --git a/Engine/Source/Runtime/CoreUObject/Private/Serialization/Formatters/JsonArchiveInputFormatter.cpp b/Engine/Source/Runtime/CoreUObject/Private/Serialization/Formatters/JsonArchiveInputFormatter.cpp
--- a/Engine/Source/Runtime/CoreUObject/Private/Serialization/Formatters/JsonArchiveInputFormatter.cpp
+++ b/Engine/Source/Runtime/CoreUObject/Private/Serialization/Formatters/JsonArchiveInputFormatter.cpp
@@ -371,16 +371,7 @@ void FJsonArchiveInputFormatter::Serialize(FString& Value)
 void FJsonArchiveInputFormatter::Serialize(FName& Value)
 {
 	FString StringValue = ValueStack.Top()->AsString();
-
-	const TCHAR Prefix[] = TEXT("Name:");
-	if (ensure(StringValue.StartsWith(Prefix)))
-	{
-		Value = FName(*StringValue + UE_ARRAY_COUNT(Prefix) - 1);
-	}
-	else
-	{
-		Value = NAME_None;
-	}
+	Value = FName(*StringValue);
 }
 
 void FJsonArchiveInputFormatter::Serialize(UObject*& Value)
@@ -544,11 +535,7 @@ EArchiveValueType FJsonArchiveInputFormatter::GetValueType(const FJsonValue& Val
 	switch (Value.Type)
 	{
 	case EJson::String:
-		if (Value.AsString().StartsWith(TEXT("Name:")))
-		{
-			return EArchiveValueType::Name;
-		}
-		else if (Value.AsString().StartsWith(TEXT("Object:")))
+		if (Value.AsString().StartsWith(TEXT("Object:")))
 		{
 			return EArchiveValueType::Object;
 		}
