Compile fixes

diff --git a/Engine/Source/Runtime/CoreUObject/Private/Serialization/AsyncLoading2.cpp b/Engine/Source/Runtime/CoreUObject/Private/Serialization/AsyncLoading2.cpp
--- a/Engine/Source/Runtime/CoreUObject/Private/Serialization/AsyncLoading2.cpp
+++ b/Engine/Source/Runtime/CoreUObject/Private/Serialization/AsyncLoading2.cpp
@@ -1587,17 +1587,17 @@ struct FAsyncPackageScope2
 };
 
 /** Makes sure the specified object reference is added to the package reference list by the time we exit a function (early or not) */
-struct FScopedAddObjectreference
+struct FScopedAddObjectreference2
 {
 	FAsyncPackage2& Package;
 	UObject*& Reference;
 
-	FScopedAddObjectreference(FAsyncPackage2& InPackage, UObject*& InReference)
+	FScopedAddObjectreference2(FAsyncPackage2& InPackage, UObject*& InReference)
 		: Package(InPackage)
 		, Reference(InReference)
 	{
 	}
-	~FScopedAddObjectreference()
+	~FScopedAddObjectreference2()
 	{
 		if (Reference)
 		{
@@ -2520,7 +2520,7 @@ void FAsyncPackage2::LinkSlimport(int32 LocalImportIndex, int32 GlobalImportInde
 	}
 
 	UObject*& Object = GlobalImportObjects[GlobalImportIndex];
-	FScopedAddObjectreference OnExitAddReference(*this, Object);
+	FScopedAddObjectreference2 OnExitAddReference(*this, Object);
 
 	if (Object)
 	{
@@ -2628,7 +2628,7 @@ void FAsyncPackage2::EventDrivenCreateExport(int32 LocalExportIndex)
 	if (!Export.Object && !Export.bExportLoadFailed)
 	{
 		FUObjectSerializeContext* LoadContext = GetSerializeContext();
-		FScopedAddObjectreference OnExitAddReference(*this, Export.Object);
+		FScopedAddObjectreference2 OnExitAddReference(*this, Export.Object);
 
 		if (!Linker->FilterExport(Export)) // for some acceptable position, it was not "not for" 
 		{
diff --git a/Engine/Source/Runtime/CoreUObject/Public/Serialization/AsyncPackageLoader.h b/Engine/Source/Runtime/CoreUObject/Public/Serialization/AsyncPackageLoader.h
--- a/Engine/Source/Runtime/CoreUObject/Public/Serialization/AsyncPackageLoader.h
+++ b/Engine/Source/Runtime/CoreUObject/Public/Serialization/AsyncPackageLoader.h
@@ -27,7 +27,7 @@ bool CanPostLoadOnAsyncLoadingThread(UObject* Object);
 template <typename T>
 void ClearFlagsAndDissolveClustersFromLoadedObjects(T& LoadedObjects)
 {
-	const EObjectFlags ObjectLoadFlags = EObjectFlags(RF_NeedLoad | RF_NeedPostLoad | RF_NeedPostLoadSubobjects | RF_WasLoaded);
+	const EObjectFlags ObjectLoadFlags = RF_NeedLoad | RF_NeedPostLoad | RF_NeedPostLoadSubobjects | RF_WasLoaded;
 	for (UObject* ObjRef : LoadedObjects)
 	{
 		if (ObjRef)
