Fix package loader shutdown crash if exiting before the loader has been created.

diff --git a/Engine/Source/Runtime/CoreUObject/Private/Serialization/AsyncPackageLoader.cpp b/Engine/Source/Runtime/CoreUObject/Private/Serialization/AsyncPackageLoader.cpp
--- a/Engine/Source/Runtime/CoreUObject/Private/Serialization/AsyncPackageLoader.cpp
+++ b/Engine/Source/Runtime/CoreUObject/Private/Serialization/AsyncPackageLoader.cpp
@@ -562,7 +562,10 @@ void FlushAsyncLoading(int32 PackageID /* = INDEX_NONE */)
 	CheckImageIntegrityAtRuntime();
 #endif
 	checkf(IsInGameThread(), TEXT("Unable to FlushAsyncLoading from any thread other than the game thread."));
-	GetAsyncPackageLoader().FlushLoading(PackageID);
+	if (GPackageLoader)
+	{
+		return GPackageLoader->FlushLoading(PackageID);
+	}
 }
 
 EAsyncPackageState::Type ProcessAsyncLoadingUntilComplete(TFunctionRef<bool()> CompletionPredicate, float TimeLimit)
