Fix void function returning a value

diff --git a/Engine/Source/Runtime/CoreUObject/Private/Serialization/AsyncPackageLoader.cpp b/Engine/Source/Runtime/CoreUObject/Private/Serialization/AsyncPackageLoader.cpp
--- a/Engine/Source/Runtime/CoreUObject/Private/Serialization/AsyncPackageLoader.cpp
+++ b/Engine/Source/Runtime/CoreUObject/Private/Serialization/AsyncPackageLoader.cpp
@@ -564,7 +564,7 @@ void FlushAsyncLoading(int32 PackageID /* = INDEX_NONE */)
 	checkf(IsInGameThread(), TEXT("Unable to FlushAsyncLoading from any thread other than the game thread."));
 	if (GPackageLoader)
 	{
-		return GPackageLoader->FlushLoading(PackageID);
+		GPackageLoader->FlushLoading(PackageID);
 	}
 }
 
