Cleanup LinkerLoad after copy-up.
Tested editor and development client with old and new loader.

diff --git a/Engine/Source/Runtime/CoreUObject/Private/Serialization/AsyncLoading2.cpp b/Engine/Source/Runtime/CoreUObject/Private/Serialization/AsyncLoading2.cpp
--- a/Engine/Source/Runtime/CoreUObject/Private/Serialization/AsyncLoading2.cpp
+++ b/Engine/Source/Runtime/CoreUObject/Private/Serialization/AsyncLoading2.cpp
@@ -2580,8 +2580,7 @@ UObject* FAsyncPackage2::EventDrivenIndexToObject(FPackageIndex Index, bool bChe
 	else if (Index.IsImport())
 	{
 		int32 GlobalImportIndex = LocalImportIndices[Index.ToImport()];
-		Result = GlobalImportObjects[GlobalImportIndex];
-		// Result = FindExistingSlimport(GlobalImportIndex);
+		Result = FindExistingSlimport(GlobalImportIndex);
 		check(Result);
 	}
 	if (bCheckSerialized && !IsFullyLoadedObj(Result))
@@ -4405,12 +4404,16 @@ EAsyncPackageState::Type FAsyncPackage2::CreateLinker()
 
 	check(!FLinkerLoad::FindExistingLinkerForPackage(Package));
 
-	uint32 LinkerFlags = LOAD_None | LOAD_Async | LOAD_NoVerify;
-	Linker = FLinkerLoad::CreateLinkerAsync2(/*LoadContext*/ nullptr, Package, *AsyncLoadingThread.GetPackageFileName(GlobalPackageId), LinkerFlags);
-	check(Linker);
-
-	check(Linker->AsyncRoot == nullptr);
-	Linker->AsyncRoot = this;
+	{
+		uint32 LinkerFlags = LOAD_None | LOAD_Async | LOAD_NoVerify;
+		Linker = new FLinkerLoad(Package, *AsyncLoadingThread.GetPackageFileName(GlobalPackageId), LinkerFlags);
+		Linker->bIsAsyncLoader = false;
+		Linker->bLockoutLegacyOperations = true;
+		Linker->SetIsLoading(true);
+		Linker->SetIsPersistent(true);
+		Linker->AsyncRoot = this;
+	}
+	Package->LinkerLoad = Linker;
 
 	AsyncLoadingThread.EnqueueIoRequest(this, CreateChunkId(PackageChunkId, 0, EChunkType::PackageSummary));
 
@@ -4449,8 +4452,7 @@ EAsyncPackageState::Type FAsyncPackage2::FinishLinker()
 
 			Summary.SetFileVersions(GPackageFileUE4Version, GPackageFileLicenseeUE4Version, /*unversioned*/true);
 
-			FLinkerLoad::UpdateLinkerAndLoaderFromSummary(Summary, *Linker);
-			FLinkerLoad::UpdateUPackageFromSummary(Summary, /*bCustomVersionIsLatest*/true, Linker->LinkerRoot);
+			Linker->UpdateFromPackageFileSummary();
 		}
 
 		// TODO: FLinker should not be a FArchive - NameMap is only required for operator<<(FName& Name)
diff --git a/Engine/Source/Runtime/CoreUObject/Private/UObject/LinkerLoad.cpp b/Engine/Source/Runtime/CoreUObject/Private/UObject/LinkerLoad.cpp
--- a/Engine/Source/Runtime/CoreUObject/Private/UObject/LinkerLoad.cpp
+++ b/Engine/Source/Runtime/CoreUObject/Private/UObject/LinkerLoad.cpp
@@ -646,105 +646,6 @@ FLinkerLoad* FLinkerLoad::CreateLinkerAsync(FUObjectSerializeContext* LoadContex
 	return Linker;
 }
 
-void FLinkerLoad::UpdateLinkerAndLoaderFromSummary(const FPackageFileSummary& Summary, FLinkerLoad& Linker, FArchive* Loader)
-{
-	Linker.SetUE4Ver(Summary.GetFileVersionUE4());
-	Linker.SetLicenseeUE4Ver(Summary.GetFileVersionLicenseeUE4());
-	Linker.SetEngineVer(Summary.SavedByEngineVersion);
-
-	const FCustomVersionContainer& SummaryVersions = Summary.GetCustomVersionContainer();
-	Linker.SetSharedCustomVersionContainerForOptimizedLoading(&SummaryVersions);
-
-	// TODO: Also done by void operator<<(FStructuredArchive::FSlot Slot, FPackageFileSummary& Sum)
-	if (Summary.PackageFlags & PKG_FilterEditorOnly)
-	{
-		Linker.SetFilterEditorOnly(true);
-	}
-
-	// Propagate fact that package cannot use lazy loading to archive (aka this).
-	if (Linker.IsTextFormat())
-	{
-		Linker.ArAllowLazyLoading = false;
-	}
-	else
-	{
-		Linker.ArAllowLazyLoading = true;
-	}
-
-	if (Loader)
-	{
-		Loader->SetUE4Ver(Summary.GetFileVersionUE4());
-		Loader->SetLicenseeUE4Ver(Summary.GetFileVersionLicenseeUE4());
-		Loader->SetEngineVer(Summary.SavedByEngineVersion);
-		Loader->SetSharedCustomVersionContainerForOptimizedLoading(&SummaryVersions);
-	}
-}
-
-void FLinkerLoad::UpdateUPackageFromSummary(const FPackageFileSummary& Summary, bool bCustomVersionIsLatest, UPackage* LinkerRootPackage)
-{
-	if (!LinkerRootPackage)
-	{
-		return;
-	}
-
-	// Preserve PIE package flag
-	uint32 NewPackageFlags = Summary.PackageFlags;
-	if (LinkerRootPackage->HasAnyPackageFlags(PKG_PlayInEditor))
-	{
-		NewPackageFlags |= PKG_PlayInEditor;
-	}
-
-	// Propagate package flags
-	LinkerRootPackage->SetPackageFlagsTo(NewPackageFlags);
-
-#if WITH_EDITORONLY_DATA
-	// Propagate package folder name
-	LinkerRootPackage->SetFolderName(*Summary.FolderName);
-#endif
-
-	// Propagate streaming install ChunkID
-	LinkerRootPackage->SetChunkIDs(Summary.ChunkIDs);
-
-	// Propagate package file size
-	// LinkerRootPackage->FileSize = TotalSize(); // TODO!!!!
-
-	// Propagate package Guid
-	LinkerRootPackage->SetGuid( Summary.Guid );
-
-#if WITH_EDITORONLY_DATA
-	LinkerRootPackage->SetPersistentGuid( Summary.PersistentGuid );
-	LinkerRootPackage->SetOwnerPersistentGuid( Summary.OwnerPersistentGuid );
-#endif
-
-	// Remember the linker versions
-	LinkerRootPackage->LinkerPackageVersion = Summary.GetFileVersionUE4();
-	LinkerRootPackage->LinkerLicenseeVersion = Summary.GetFileVersionLicenseeUE4();
-
-	// Only set the custom version if it is not already latest.
-	// If it is latest, we will compare against latest in GetLinkerCustomVersion
-	if (!bCustomVersionIsLatest)
-	{
-		LinkerRootPackage->LinkerCustomVersion = Summary.GetCustomVersionContainer();
-	}
-
-#if WITH_EDITORONLY_DATA
-	LinkerRootPackage->bIsCookedForEditor = !!(Summary.PackageFlags & PKG_FilterEditorOnly);
-#endif
-}
-
-FLinkerLoad* FLinkerLoad::CreateLinkerAsync2(FUObjectSerializeContext* LoadContext, UPackage* Parent, const TCHAR* FileName, uint32 LoadFlags)
-{
-	LoadFlags |= LOAD_Async;
-	FLinkerLoad* Linker = new FLinkerLoad(Parent, FileName, LoadFlags);
-	Linker->bIsAsyncLoader = false;
-	Linker->bLockoutLegacyOperations = true;
-	Linker->SetIsLoading(true);
-	Linker->SetIsPersistent(true);
-	Parent->LinkerLoad = Linker;
-
-	return Linker;
-}
-
 void FLinkerLoad::SetSerializeContext(FUObjectSerializeContext* InLoadContext)
 {
 }
@@ -1368,45 +1269,12 @@ FLinkerLoad::ELinkerStatus FLinkerLoad::SerializePackageFileSummary()
 		}
 #endif // PLATFORM_WINDOWS
 
-		// When unversioned, pretend we are the latest version
-		bool bCustomVersionIsLatest = Summary.bUnversioned;
-		if (!bCustomVersionIsLatest)
+		ELinkerStatus UpdateStatus = UpdateFromPackageFileSummary();
+		if (UpdateStatus != LINKER_Loaded)
 		{
-			// Check custom versions.
-			const FCustomVersionContainer& LatestCustomVersions  = FCustomVersionContainer::GetRegistered();
-			bool bAllSavedVersionsMatch = true;
-			const FCustomVersionArray&  PackageCustomVersions = Summary.GetCustomVersionContainer().GetAllVersions();
-			for (auto It = PackageCustomVersions.CreateConstIterator(); It; ++It)
-			{
-				const FCustomVersion& SerializedCustomVersion = *It;
-
-				const FCustomVersion* LatestVersion = LatestCustomVersions.GetVersion(SerializedCustomVersion.Key);
-				if (!LatestVersion)
-				{
-					// Loading a package with custom integration that we don't know about!
-					// Temporarily just warn and continue. @todo: this needs to be fixed properly
-					UE_LOG(LogLinker, Warning, TEXT("Package %s was saved with a custom integration that is not present. Tag %s  Version %d"), *Filename, *SerializedCustomVersion.Key.ToString(), SerializedCustomVersion.Version);
-					bAllSavedVersionsMatch = false;
-				}
-				else if (SerializedCustomVersion.Version > LatestVersion->Version)
-				{
-					// Loading a package with a newer custom version than the current one.
-					UE_LOG(LogLinker, Error, TEXT("Package %s was saved with a newer custom version than the current. Tag %s Name '%s' PackageVersion %d  MaxExpected %d"), *Filename, *SerializedCustomVersion.Key.ToString(), *LatestVersion->GetFriendlyName().ToString(), SerializedCustomVersion.Version, LatestVersion->Version);
-					return LINKER_Failed;
-				}
-				else if (SerializedCustomVersion.Version != LatestVersion->Version)
-				{
-					bAllSavedVersionsMatch = false;
-				}
-			}
-
-			const bool bSameNumberOfVersions = (PackageCustomVersions.Num() == LatestCustomVersions.GetAllVersions().Num());
-			bCustomVersionIsLatest = bSameNumberOfVersions && bAllSavedVersionsMatch;
+			return UpdateStatus;
 		}
 
-		UpdateLinkerAndLoaderFromSummary(Summary, *this, Loader);
-		UpdateUPackageFromSummary(Summary, bCustomVersionIsLatest, LinkerRoot);
-
 		// Slack everything according to summary.
 		ImportMap					.Empty( Summary.ImportCount				);
 		ExportMap					.Empty( Summary.ExportCount				);
@@ -1421,6 +1289,125 @@ FLinkerLoad::ELinkerStatus FLinkerLoad::SerializePackageFileSummary()
 	return !IsTimeLimitExceeded( TEXT("serializing package file summary") ) ? LINKER_Loaded : LINKER_TimedOut;
 }
 
+FLinkerLoad::ELinkerStatus FLinkerLoad::UpdateFromPackageFileSummary()
+{
+	// When unversioned, pretend we are the latest version
+	bool bCustomVersionIsLatest = Summary.bUnversioned;
+	if (!bCustomVersionIsLatest)
+	{
+		// Check custom versions.
+		const FCustomVersionContainer& LatestCustomVersions  = FCustomVersionContainer::GetRegistered();
+		bool bAllSavedVersionsMatch = true;
+		const FCustomVersionArray&  PackageCustomVersions = Summary.GetCustomVersionContainer().GetAllVersions();
+		for (auto It = PackageCustomVersions.CreateConstIterator(); It; ++It)
+		{
+			const FCustomVersion& SerializedCustomVersion = *It;
+
+			const FCustomVersion* LatestVersion = LatestCustomVersions.GetVersion(SerializedCustomVersion.Key);
+			if (!LatestVersion)
+			{
+				// Loading a package with custom integration that we don't know about!
+				// Temporarily just warn and continue. @todo: this needs to be fixed properly
+				UE_LOG(LogLinker, Warning, TEXT("Package %s was saved with a custom integration that is not present. Tag %s  Version %d"), *Filename, *SerializedCustomVersion.Key.ToString(), SerializedCustomVersion.Version);
+				bAllSavedVersionsMatch = false;
+			}
+			else if (SerializedCustomVersion.Version > LatestVersion->Version)
+			{
+				// Loading a package with a newer custom version than the current one.
+				UE_LOG(LogLinker, Error, TEXT("Package %s was saved with a newer custom version than the current. Tag %s Name '%s' PackageVersion %d  MaxExpected %d"), *Filename, *SerializedCustomVersion.Key.ToString(), *LatestVersion->GetFriendlyName().ToString(), SerializedCustomVersion.Version, LatestVersion->Version);
+				return LINKER_Failed;
+			}
+			else if (SerializedCustomVersion.Version != LatestVersion->Version)
+			{
+				bAllSavedVersionsMatch = false;
+			}
+		}
+
+		const bool bSameNumberOfVersions = (PackageCustomVersions.Num() == LatestCustomVersions.GetAllVersions().Num());
+		bCustomVersionIsLatest = bSameNumberOfVersions && bAllSavedVersionsMatch;
+	}
+
+	const FCustomVersionContainer& SummaryVersions = Summary.GetCustomVersionContainer();
+
+	SetUE4Ver(Summary.GetFileVersionUE4());
+	SetLicenseeUE4Ver(Summary.GetFileVersionLicenseeUE4());
+	SetEngineVer(Summary.SavedByEngineVersion);
+	SetSharedCustomVersionContainerForOptimizedLoading(&SummaryVersions);
+
+	if (Summary.PackageFlags & PKG_FilterEditorOnly)
+	{
+		SetFilterEditorOnly(true);
+	}
+
+	// Propagate fact that package cannot use lazy loading to archive (aka this).
+	if (IsTextFormat())
+	{
+		ArAllowLazyLoading = false;
+	}
+	else
+	{
+		ArAllowLazyLoading = true;
+	}
+
+	// Loader needs to be the same version.
+	if (Loader)
+	{
+		Loader->SetUE4Ver(Summary.GetFileVersionUE4());
+		Loader->SetLicenseeUE4Ver(Summary.GetFileVersionLicenseeUE4());
+		Loader->SetEngineVer(Summary.SavedByEngineVersion);
+		Loader->SetSharedCustomVersionContainerForOptimizedLoading(&SummaryVersions);
+	}
+
+	if (UPackage* LinkerRootPackage = LinkerRoot)
+	{
+		// Preserve PIE package flag
+		uint32 NewPackageFlags = Summary.PackageFlags;
+		if (LinkerRootPackage->HasAnyPackageFlags(PKG_PlayInEditor))
+		{
+			NewPackageFlags |= PKG_PlayInEditor;
+		}
+
+		// Propagate package flags
+		LinkerRootPackage->SetPackageFlagsTo(NewPackageFlags);
+
+#if WITH_EDITORONLY_DATA
+		// Propagate package folder name
+		LinkerRootPackage->SetFolderName(*Summary.FolderName);
+#endif
+
+		// Propagate streaming install ChunkID
+		LinkerRootPackage->SetChunkIDs(Summary.ChunkIDs);
+
+		// Propagate package file size
+		LinkerRootPackage->FileSize = Loader ? Loader->TotalSize() : 0;
+
+		// Propagate package Guids
+		LinkerRootPackage->SetGuid( Summary.Guid );
+
+#if WITH_EDITORONLY_DATA
+		LinkerRootPackage->SetPersistentGuid( Summary.PersistentGuid );
+		LinkerRootPackage->SetOwnerPersistentGuid( Summary.OwnerPersistentGuid );
+#endif
+
+		// Remember the linker versions
+		LinkerRootPackage->LinkerPackageVersion = Summary.GetFileVersionUE4();
+		LinkerRootPackage->LinkerLicenseeVersion = Summary.GetFileVersionLicenseeUE4();
+
+		// Only set the custom version if it is not already latest.
+		// If it is latest, we will compare against latest in GetLinkerCustomVersion
+		if (!bCustomVersionIsLatest)
+		{
+			LinkerRootPackage->LinkerCustomVersion = SummaryVersions;
+		}
+
+#if WITH_EDITORONLY_DATA
+		LinkerRootPackage->bIsCookedForEditor = !!(Summary.PackageFlags & PKG_FilterEditorOnly);
+#endif
+	}
+
+	return LINKER_Loaded;
+}
+
 /**
  * Serializes the name table.
  */
diff --git a/Engine/Source/Runtime/CoreUObject/Public/UObject/LinkerLoad.h b/Engine/Source/Runtime/CoreUObject/Public/UObject/LinkerLoad.h
--- a/Engine/Source/Runtime/CoreUObject/Public/UObject/LinkerLoad.h
+++ b/Engine/Source/Runtime/CoreUObject/Public/UObject/LinkerLoad.h
@@ -135,11 +135,6 @@ public:
 		return bIsAsyncLoader ? (FAsyncArchive*)Loader : nullptr;
 	}
 
-public:
-	int32* LocalImportIndices = 0;
-	UObject** GlobalImportObjects = nullptr;
-	const TArray<FNameEntryId>* ActiveNameMap = &NameMap;
-
 private:
 
 	/** Structured archive interface. Wraps underlying loader to provide contextual metadata to the values being written
@@ -153,6 +148,10 @@ private:
 	/** The archive that actually reads the raw data from disk.																*/
 	FArchive*				Loader;
 
+	int32* LocalImportIndices = nullptr;
+	UObject** GlobalImportObjects = nullptr;
+	const TArray<FNameEntryId>* ActiveNameMap = &NameMap;
+
 protected:
 
 	void SetLoader(FArchive* InLoader);
@@ -872,9 +871,6 @@ private:
 	COREUOBJECT_API static FLinkerLoad* CreateLinkerAsync(FUObjectSerializeContext* LoadContext, UPackage* Parent, const TCHAR* Filename, uint32 LoadFlags
 		, TFunction<void()>&& InSummaryReadyCallback
 	);
-	COREUOBJECT_API static FLinkerLoad* CreateLinkerAsync2(FUObjectSerializeContext* LoadContext, UPackage* Parent, const TCHAR* Filename, uint32 LoadFlags);
-	static void UpdateLinkerAndLoaderFromSummary(const FPackageFileSummary& Summary, FLinkerLoad& Linker, FArchive* Loader = nullptr);
-	static void UpdateUPackageFromSummary(const FPackageFileSummary& Summary, bool bCustomVersionIsLatest, UPackage* LinkerRootPackage);
 
 protected: // Daniel L: Made this protected so I can override the constructor and create a custom loader to load the header of the linker in the DiffFilesCommandlet
 	/**
@@ -920,6 +916,11 @@ private:
 	 */
 	ELinkerStatus SerializePackageFileSummary();
 
+	/**
+	 * Updates the linker, loader and root package with data from the package file summary.
+	 * */
+	ELinkerStatus UpdateFromPackageFileSummary();
+
 	/**
 	 * Serializes the name map.
 	 */
