AsyncLoading2: Replace referenced object sets with an internal async flag to prevent GC of objects during loading

diff --git a/Engine/Source/Runtime/CoreUObject/Private/Serialization/AsyncLoading2.cpp b/Engine/Source/Runtime/CoreUObject/Private/Serialization/AsyncLoading2.cpp
--- a/Engine/Source/Runtime/CoreUObject/Private/Serialization/AsyncLoading2.cpp
+++ b/Engine/Source/Runtime/CoreUObject/Private/Serialization/AsyncLoading2.cpp
@@ -20,12 +20,10 @@
 #include "Misc/PackageName.h"
 #include "UObject/PackageFileSummary.h"
 #include "UObject/Linker.h"
-// #include "UObject/CoreRedirects.h"
 #include "UObject/SoftObjectPath.h"
 #include "UObject/LinkerLoad.h"
 #include "Serialization/DeferredMessageLog.h"
 #include "UObject/UObjectThreadContext.h"
-// #include "UObject/LinkerManager.h"
 #include "Misc/Paths.h"
 #include "Misc/ExclusiveLoadPackageTimeTracker.h"
 #include "ProfilingDebugging/LoadTimeTracker.h"
@@ -50,8 +48,9 @@
 #include "UObject/GCObject.h"
 #include "UObject/ObjectRedirector.h"
 
-#if UE_BUILD_DEVELOPMENT
+#if UE_BUILD_DEVELOPMENT || UE_BUILD_DEBUG
 //PRAGMA_DISABLE_OPTIMIZATION
+#define ALT2_VERIFY_ASYNC_FLAGS
 #endif
 
 struct FAsyncPackage2;
@@ -89,11 +88,17 @@ struct FGlobalImport
 
 struct FGlobalImportRuntime
 {
+	/** Persistent data */
 	int32 Count = 0;
 	FName* Names;
 	FPackageIndex* Outers;
 	FPackageIndex* Packages;
 	UObject** Objects;
+	/** Reference tracking for GC management */
+	TAtomic<int32>* RefCounts;
+	TArray<UObject*> KeepAliveObjects;
+	void OnPreGarbageCollect();
+	void OnPostGarbageCollect();
 };
 	
 struct FPackageStoreEntrySerialized
@@ -691,6 +696,25 @@ struct FAsyncPackage2
 			ImportedAsyncPackage->ReleaseRef();
 		}
 		ImportedAsyncPackages.Empty();
+		ReleaseGlobalImportObjectReferences();
+	}
+
+	void AddGlobalImportObjectReferences()
+	{
+		for (int32 LocalImportIndex = 0; LocalImportIndex < LocalImportCount; ++LocalImportIndex)
+		{
+			const int32 GlobalImportIndex = LocalImportIndices[LocalImportIndex];
+			++GlobalImportObjectRefCounts[GlobalImportIndex];
+		}
+	}
+
+	void ReleaseGlobalImportObjectReferences()
+	{
+		for (int32 LocalImportIndex = 0; LocalImportIndex < LocalImportCount; ++LocalImportIndex)
+		{
+			const int32 GlobalImportIndex = LocalImportIndices[LocalImportIndex];
+			--GlobalImportObjectRefCounts[GlobalImportIndex];
+		}
 	}
 
 	/** Marks a specific request as complete */
@@ -786,17 +810,18 @@ struct FAsyncPackage2
 	}
 
 	/** FGCObject Interface */
-	virtual void AddReferencedObjects(FReferenceCollector& Collector) override;
+	virtual void AddReferencedObjects(FReferenceCollector& Collector) override {}
 	virtual FString GetReferencerName() const override
 	{
 		return FString::Printf(TEXT("FAsyncPackage %s"), *GetPackageName().ToString());
 	}
 
-	/** Adds a new object referenced by this package */
-	void AddObjectReference(UObject* InObject);
+	void AddOwnedObject(UObject* Object)
+	{
+		OwnedObjects.Add(Object);
+	}
 
-	/** Removes all objects from the list and clears async loading flags */
-	void EmptyReferencedObjects();
+	void ClearOwnedObjects();
 
 	/** Returns the UPackage wrapped by this, if it is valid */
 	UPackage* GetLoadedPackage();
@@ -898,10 +923,8 @@ private:
 	/** Number of times we recursed to load this package. */
 	int32 ReentryCount;
 	TArray<FAsyncPackage2*> ImportedAsyncPackages;
-	/** List of objects referenced by this package */
-	TSet<UObject*> ReferencedObjects;
-	/** Critical section for referenced objects list */
-	FCriticalSection ReferencedObjectsCritical;
+	/** List of OwnedObjects = Exports + UPackage + ObjectsCreatedFromExports */
+	TArray<UObject*> OwnedObjects;
 	/** Cached async loading thread object this package was created by */
 	FAsyncLoadingThread2Impl& AsyncLoadingThread;
 	IEDLBootNotificationManager& EDLBootNotificationManager;
@@ -924,11 +947,12 @@ private:
 
 	int32 GlobalImportCount = 0;
 	int32 LocalImportCount = 0;
-	int32* LocalImportIndices = 0;
+	int32* LocalImportIndices = nullptr;
 	FName* GlobalImportNames = nullptr;
 	FPackageIndex* GlobalImportOuters = nullptr;
 	FPackageIndex* GlobalImportPackages = nullptr;
 	UObject** GlobalImportObjects = nullptr;
+	TAtomic<int32>* GlobalImportObjectRefCounts = nullptr;
 public:
 
 	FAsyncLoadingThread2Impl& GetAsyncLoadingThread()
@@ -1271,6 +1295,11 @@ private:
 		return GlobalImportRuntime.Packages;
 	}
 
+	inline TAtomic<int32>* GetGlobalImportObjectRefCounts()
+	{
+		return GlobalImportRuntime.RefCounts;
+	}
+
 	inline int32* GetPackageSlimports(FGlobalPackageId GlobalPackageId, int32& OutCount) const 
 	{
 		OutCount = StoreEntriesSerialized[GlobalPackageId.Id].SlimportCount;
@@ -1303,6 +1332,7 @@ public:
 	/** Start the async loading thread */
 	void StartThread();
 
+	/** [GC] Management of global import objects */
 	void OnPreGarbageCollect();
 	void OnPostGarbageCollect();
 
@@ -1586,26 +1616,6 @@ struct FAsyncPackageScope2
 	}
 };
 
-/** Makes sure the specified object reference is added to the package reference list by the time we exit a function (early or not) */
-struct FScopedAddObjectreference2
-{
-	FAsyncPackage2& Package;
-	UObject*& Reference;
-
-	FScopedAddObjectreference2(FAsyncPackage2& InPackage, UObject*& InReference)
-		: Package(InPackage)
-		, Reference(InReference)
-	{
-	}
-	~FScopedAddObjectreference2()
-	{
-		if (Reference)
-		{
-			Package.AddObjectReference(Reference);
-		}
-	}
-};
-
 /** Just like TGuardValue for FAsyncLoadingThread::AsyncLoadingTickCounter but only works for the game thread */
 struct FAsyncLoadingTickScope2
 {
@@ -1712,8 +1722,8 @@ void FAsyncLoadingThread2Impl::InitializeLoading()
 		GlobalImportRuntime.Outers = new FPackageIndex[GlobalImportRuntime.Count];
 		GlobalImportRuntime.Packages = new FPackageIndex[GlobalImportRuntime.Count];
 		GlobalImportRuntime.Objects = new UObject*[GlobalImportRuntime.Count];
+		GlobalImportRuntime.RefCounts = new TAtomic<int32>[GlobalImportRuntime.Count];
 
-		// FName fixup
 		for (int I = 0; I < GlobalImportRuntime.Count; ++I)
 		{
 			FGlobalImport& Import = Imports[I];
@@ -1722,6 +1732,7 @@ void FAsyncLoadingThread2Impl::InitializeLoading()
 			GlobalImportRuntime.Outers[I] = Import.OuterIndex;
 			GlobalImportRuntime.Packages[I] = Import.OutermostIndex;
 			GlobalImportRuntime.Objects[I] = nullptr;
+			GlobalImportRuntime.RefCounts[I] = 0;
 		}
 	}
 
@@ -2266,7 +2277,6 @@ void FAsyncPackage2::ImportPackagesRecursive()
 
 		if (ImportedObject)
 		{
-			AddObjectReference(ImportedObject);
 			if (IsFullyLoadedObj(ImportedObject))
 			{
 				continue;
@@ -2520,7 +2530,6 @@ void FAsyncPackage2::LinkSlimport(int32 LocalImportIndex, int32 GlobalImportInde
 	}
 
 	UObject*& Object = GlobalImportObjects[GlobalImportIndex];
-	FScopedAddObjectreference2 OnExitAddReference(*this, Object);
 
 	if (Object)
 	{
@@ -2602,7 +2611,6 @@ UObject* FAsyncPackage2::EventDrivenIndexToObject(FPackageIndex Index, bool bChe
 	if (Result)
 	{
 		UE_CLOG(Result->HasAnyInternalFlags(EInternalObjectFlags::Unreachable), LogStreaming, Fatal, TEXT("Returning an object  (%s) from EventDrivenIndexToObject that is unreachable."), *Result->GetFullName());
-		checkSlow(ReferencedObjects.Contains(Result));
 	}
 	return Result;
 }
@@ -2627,7 +2635,6 @@ void FAsyncPackage2::EventDrivenCreateExport(int32 LocalExportIndex)
 	if (!Export.Object && !Export.bExportLoadFailed)
 	{
 		FUObjectSerializeContext* LoadContext = GetSerializeContext();
-		FScopedAddObjectreference2 OnExitAddReference(*this, Export.Object);
 
 		if (!Linker->FilterExport(Export)) // for some acceptable position, it was not "not for" 
 		{
@@ -2649,10 +2656,12 @@ void FAsyncPackage2::EventDrivenCreateExport(int32 LocalExportIndex)
 						DCD->AddToRoot();
 					}
 				}
-				if (DCD)
-				{
-					AddObjectReference(DCD);
-				}
+				// TODO: Replace with DCD->SetInternalFlags(EObjectInternalFlag::Async)?
+				// But then it also need to be cleared in ClearOwnedObjects()
+				// if (DCD)
+				// {
+				// 	AddObjectReference(DCD);
+				// }
 				UE_LOG(LogStreaming, Verbose, TEXT("EventDrivenCreateExport: Created dynamic class %s"), *Export.Object->GetFullName());
 				if (Export.Object)
 				{
@@ -2732,37 +2741,39 @@ void FAsyncPackage2::EventDrivenCreateExport(int32 LocalExportIndex)
 				// Object is found in memory.
 				if (Export.Object)
 				{
-					// Mark that we need to dissociate forced exports later on if we are a forced export.
-					if (Export.bForcedExport)
-					{
-						// see FLinkerLoad::CreateExport, there may be some more we can do here
-						check(!Export.bForcedExport); // this is leftover from seekfree loading I think
-					}
-					// Associate linker with object to avoid detachment mismatches.
-					else
-					{
-						Export.Object->SetLinker(Linker, LocalExportIndex);
+					check(!Export.bForcedExport);
+					Export.Object->SetLinker(Linker, LocalExportIndex);
 
-						// If this object was allocated but never loaded (components created by a constructor, CDOs, etc) make sure it gets loaded
-						// Do this for all subobjects created in the native constructor.
-						if (!Export.Object->HasAnyFlags(RF_LoadCompleted))
+					// If this object was allocated but never loaded (components created by a constructor, CDOs, etc) make sure it gets loaded
+					// Do this for all subobjects created in the native constructor.
+					if (!Export.Object->HasAnyFlags(RF_LoadCompleted))
+					{
+						UE_LOG(LogStreaming, VeryVerbose, TEXT("Note2: %s was constructed during load and is an export and so needs loading."), *Export.Object->GetFullName());
+						UE_CLOG(!Export.Object->HasAllFlags(RF_WillBeLoaded), LogStreaming, Fatal, TEXT("%s was found in memory and is an export but does not have all load flags."), *Export.Object->GetFullName());
+						if(Export.Object->HasAnyFlags(RF_ClassDefaultObject))
+						{
+							// never call PostLoadSubobjects on class default objects, this matches the behavior of the old linker where
+							// StaticAllocateObject prevents setting of RF_NeedPostLoad and RF_NeedPostLoadSubobjects, but FLinkerLoad::Preload
+							// assigns RF_NeedPostLoad for blueprint CDOs:
+							Export.Object->SetFlags(RF_NeedLoad | RF_NeedPostLoad | RF_WasLoaded);
+						}
+						else
 						{
-								UE_LOG(LogStreaming, VeryVerbose, TEXT("Note2: %s was constructed during load and is an export and so needs loading."), *Export.Object->GetFullName());
-								UE_CLOG(!Export.Object->HasAllFlags(RF_WillBeLoaded), LogStreaming, Fatal, TEXT("%s was found in memory and is an export but does not have all load flags."), *Export.Object->GetFullName());
-							if(Export.Object->HasAnyFlags(RF_ClassDefaultObject))
-							{
-								// never call PostLoadSubobjects on class default objects, this matches the behavior of the old linker where
-								// StaticAllocateObject prevents setting of RF_NeedPostLoad and RF_NeedPostLoadSubobjects, but FLinkerLoad::Preload
-										// assigns RF_NeedPostLoad for blueprint CDOs:
-								Export.Object->SetFlags(RF_NeedLoad | RF_NeedPostLoad | RF_WasLoaded);
-							}
-							else
-							{
 							Export.Object->SetFlags(RF_NeedLoad | RF_NeedPostLoad | RF_NeedPostLoadSubobjects | RF_WasLoaded);
-							}
-								Export.Object->ClearFlags(RF_WillBeLoaded);
 						}
+						Export.Object->ClearFlags(RF_WillBeLoaded);
+						if (IsInGameThread())
+						{
+							Export.Object->SetInternalFlags(EInternalObjectFlags::Async);
+							AddOwnedObject(Export.Object);
+						}
+					}
+					else
+					{
+						Export.Object->SetInternalFlags(EInternalObjectFlags::Async);
+						AddOwnedObject(Export.Object);
 					}
+					check(Export.Object->HasAnyInternalFlags(EInternalObjectFlags::Async));
 				}
 				else
 				{
@@ -2876,6 +2887,12 @@ void FAsyncPackage2::EventDrivenCreateExport(int32 LocalExportIndex)
 					{
 						Export.Object->AddToRoot();
 					}
+					if (IsInGameThread())
+					{
+						Export.Object->SetInternalFlags(EInternalObjectFlags::Async);
+						AddOwnedObject(Export.Object);
+					}
+					check(Export.Object->HasAnyInternalFlags(EInternalObjectFlags::Async));
 					Export.Object->SetLinker(Linker, LocalExportIndex);
 					check(Export.Object->GetClass() == LoadClass);
 					check(NewName == Export.ObjectName);
@@ -3062,7 +3079,6 @@ EAsyncPackageState::Type FAsyncPackage2::Event_StartPostload(FAsyncPackage2* Pac
 		{
 			FObjectExport& Export = Package->Linker->ExportMap[LocalExportIndex];
 			UObject* Object = Export.Object;
-			checkSlow(!(Object && !Package->ReferencedObjects.Contains(Object)));
 			if (Object && (Object->HasAnyFlags(RF_NeedPostLoad) || Package->Linker->bDynamicClassLinker || Object->HasAnyInternalFlags(EInternalObjectFlags::AsyncLoading)))
 			{
 				check(Object->IsValidLowLevelFast());
@@ -3403,32 +3419,17 @@ EAsyncPackageState::Type FAsyncLoadingThread2Impl::ProcessLoadedPackagesFromGame
 			{
 				FScopeLock LockAsyncPackages(&AsyncPackagesCritical);
 				AsyncPackageNameLookup.Remove(Package->GetPackageName());
+				Package->ClearOwnedObjects();
 			}
 
 			// Remove the package from the list before we trigger the callbacks, 
 			// this is to ensure we can re-enter FlushAsyncLoading from any of the callbacks
 			LoadedPackagesToProcess.RemoveAt(PackageIndex--);
 
-			if (FPlatformProperties::RequiresCookedData())
-			{
-				// Emulates ResetLoaders on the package linker's linkerroot.
-				if (!Package->IsBeingProcessedRecursively())
-				{
-					Package->ResetLoader();
-				}
-			}
-			else
+			// Emulates ResetLoaders on the package linker's linkerroot.
+			if (!Package->IsBeingProcessedRecursively())
 			{
-				if (GIsEditor)
-				{
-					// Flush linker cache for all objects loaded with this package.
-					// This may be slow, hence we only do it in the editor
-					Package->FlushObjectLinkerCache();
-				}
-				// Detach linker in mutex scope to make sure that if something requests this package
-				// before it's been deleted does not try to associate the new async package with the old linker
-				// while this async package is still bound to it.
-				Package->DetachLinker();
+				Package->ResetLoader();
 			}
 
 			// Close linkers opened by synchronous loads during async loading
@@ -3987,53 +3988,97 @@ float FAsyncLoadingThread2Impl::GetAsyncLoadPercentage(const FName& PackageName)
 	return LoadPercentage;
 }
 
+void FGlobalImportRuntime::OnPreGarbageCollect()
+{
+	int32 NumCleared = 0;
+
+	for (int32 GlobalImportIndex = 0; GlobalImportIndex < Count; ++GlobalImportIndex)
+	{
+		UObject*& Object = Objects[GlobalImportIndex];
+		if (!Object)
+		{
+			continue;
+		}
+
+		// Import objects in packages currently being loaded already have the Async flag set.
+		// They will never be destroyed during GC, and the object pointers are safe to keep.
+		const bool bHasAsyncFlag = Object->HasAnyInternalFlags(EInternalObjectFlags::Async);
+		if (bHasAsyncFlag)
+		{
+			continue;
+		}
+
+		const int32 RefCount = RefCounts[GlobalImportIndex];
+		check(RefCount >= 0);
+
+		if (RefCount > 0)
+		{
+			// Import objects in native packages will never be garbage collected and does not need marking.
+			FPackageIndex PackageIndex = Packages[GlobalImportIndex];
+			UPackage* Package = CastChecked<UPackage>(Objects[PackageIndex.ToImport()]);
+			if (Package->HasAnyPackageFlags(PKG_CompiledIn))
+			{
+				continue;
+			}
+			// Mark object to be kept alive during GC
+			Object->SetInternalFlags(EInternalObjectFlags::Async);
+			KeepAliveObjects.Add(Object);
+		}
+		else
+		{
+			// Clear object pointer since object may get destroyed during GC
+			Object = nullptr;
+			++NumCleared;
+		}
+	}
+
+	UE_LOG(LogStreaming, Log, TEXT("FGlobalImportRuntime::OnPreGarbageCollect - Marked %d objects, cleared %d object pointers"),
+		KeepAliveObjects.Num(), NumCleared);
+}
+
+void FGlobalImportRuntime::OnPostGarbageCollect()
+{
+	if (KeepAliveObjects.Num() == 0)
+	{
+		return;
+	}
+
+	for (UObject* Object : KeepAliveObjects)
+	{
+		Object->ClearInternalFlags(EInternalObjectFlags::Async);
+	}
+
+	const int32 UnmarkedCount = KeepAliveObjects.Num();
+	KeepAliveObjects.Reset();
+	UE_LOG(LogStreaming, Log, TEXT("FGlobalImportRuntime::UpdateGlobalImportsPostGC - Unmarked %d objects"),
+		UnmarkedCount);
+}
+
 void FAsyncLoadingThread2Impl::OnPreGarbageCollect()
 {
 	TRACE_CPUPROFILER_EVENT_SCOPE(AltPreGC);
-	int32 GlobalImportCount = 0;
-	if (UObject** GlobalImportObjects = GetGlobalImportObjects(GlobalImportCount))
-	{
-		UE_LOG(LogStreaming, Display, TEXT("AltPreGC - clearing objects"));
-		int32 NumObjects = 0;
-		FMemory::Memset(GlobalImportObjects, 0, GlobalImportCount * sizeof(UObject*));
-		// FName* GlobalImportNames = IPackageResourceManager::Get().GetGlobalImportNames(GlobalImportCount);
-		// for (int32 GlobalImportIndex = 0; GlobalImportIndex < GlobalImportCount; ++GlobalImportIndex)
-		// {
-		// 	if (GlobalImportObjects[GlobalImportIndex])
-		// 	{
-		// 		UE_LOG(LogStreaming, Display, TEXT("AltPreGC - %d:%s cleared"),
-		// 			GlobalImportIndex,
-		// 			*GlobalImportNames[GlobalImportIndex].ToString());
-		// 		GlobalImportObjects[GlobalImportIndex] = nullptr;
-		// 		++NumObjects;
-		// 	}
-		// }
-		UE_LOG(LogStreaming, Display, TEXT("AltPreGC - %d objects cleared"), NumObjects);
+	if (!IsAsyncLoadingPackages())
+	{
+#ifdef ALT2_VERIFY_ASYNC_FLAGS
+		const EInternalObjectFlags AsyncFlags = EInternalObjectFlags::Async | EInternalObjectFlags::AsyncLoading;
+		for (int32 ObjectIndex = 0; ObjectIndex < GUObjectArray.GetObjectArrayNum(); ++ObjectIndex)
+		{
+			FUObjectItem* ObjectItem = &GUObjectArray.GetObjectItemArrayUnsafe()[ObjectIndex];
+			if (UObject* Obj = static_cast<UObject*>(ObjectItem->Object))
+			{
+				ensure(!Obj->HasAnyInternalFlags(AsyncFlags));
+			}
+		}
+#endif
+		return;
 	}
+	GlobalImportRuntime.OnPreGarbageCollect();
 }
 
 void FAsyncLoadingThread2Impl::OnPostGarbageCollect()
 {
 	TRACE_CPUPROFILER_EVENT_SCOPE(AltPostGC);
-	int32 GlobalImportCount = 0;
-	if (UObject** GlobalImportObjects = GetGlobalImportObjects(GlobalImportCount))
-	{
-		UE_LOG(LogStreaming, Display, TEXT("AltPostGC - finding objects"));
-		int32 NumObjects = 0;
-		 FPackageIndex* GlobalImportOuters = GetGlobalImportOuters(GlobalImportCount);
-		 FName* GlobalImportNames = GetGlobalImportNames(GlobalImportCount);
-		 for (int32 GlobalImportIndex = 0; GlobalImportIndex < GlobalImportCount; ++GlobalImportIndex)
-		 {
-		 	if (GFindExistingSlimport(GlobalImportIndex, GlobalImportObjects, GlobalImportOuters, GlobalImportNames))
-		 	{
-		 		// UE_LOG(LogStreaming, Display, TEXT("AltPreGC - %d:%s found"),
-		 		// 	GlobalImportIndex,
-		 		// 	*GlobalImportNames[GlobalImportIndex].ToString());
-		 		++NumObjects;
-		 	}
-		 }
-		UE_LOG(LogStreaming, Display, TEXT("AltPostGC - %d objects found"), NumObjects);
-	}
+	GlobalImportRuntime.OnPostGarbageCollect();
 }
 
 /**
@@ -4053,7 +4098,7 @@ void FAsyncLoadingThread2Impl::NotifyConstructedDuringAsyncLoading(UObject* Obje
 	FUObjectThreadContext& ThreadContext = FUObjectThreadContext::Get();
 	check(ThreadContext.AsyncPackage); // Otherwise something is wrong and we're creating objects outside of async loading code
 	FAsyncPackage2* AsyncPackage2 = (FAsyncPackage2*)ThreadContext.AsyncPackage;
-	AsyncPackage2->AddObjectReference(Object);
+	AsyncPackage2->AddOwnedObject(Object);
 	
 	// if this is in the package and is an export, then let mark it as needing load now
 	if (Object->GetOutermost() == AsyncPackage2->GetLinkerRoot() && 
@@ -4103,6 +4148,7 @@ FAsyncPackage2::FAsyncPackage2(const FAsyncPackageDesc& InDesc, int32 InSerialNu
 , GraphAllocator(InGraphAllocator)
 , GlobalPackageId(InGlobalPackageId)
 // Begin EDL specific properties
+, GlobalImportObjectRefCounts(AsyncLoadingThread.GetGlobalImportObjectRefCounts())
 , AsyncPackageLoadingState(EAsyncPackageLoadingState2::NewPackage)
 , SerialNumber(InSerialNumber)
 , bAllExportsSerialized(false)
@@ -4121,6 +4167,7 @@ FAsyncPackage2::FAsyncPackage2(const FAsyncPackageDesc& InDesc, int32 InSerialNu
 		GlobalImportOuters = AsyncLoadingThread.GetGlobalImportOuters(GlobalImportCount);
 		GlobalImportPackages = AsyncLoadingThread.GetGlobalImportPackages(GlobalImportCount);
 		GlobalImportObjects = AsyncLoadingThread.GetGlobalImportObjects(GlobalImportCount);
+		AddGlobalImportObjectReferences();
 	}
 
 	{
@@ -4208,38 +4255,32 @@ FAsyncPackage2::~FAsyncPackage2()
 	DetachLinker();
 	SerialNumber = 0; // the weak pointer will always fail now
 	
-	EmptyReferencedObjects();
+	ensure(OwnedObjects.Num() == 0);
 }
 
 
-void FAsyncPackage2::AddReferencedObjects(FReferenceCollector& Collector)
+void FAsyncPackage2::ClearOwnedObjects()
 {
-	Collector.AddReferencedObjects(ReferencedObjects);
-	Collector.AddReferencedObjects(DeferredFinalizeObjects);
-	Collector.AddReferencedObjects(PackageObjLoaded);
-}
-
-void FAsyncPackage2::AddObjectReference(UObject* InObject)
-{
-	if (InObject)
+	for (UObject* Object : OwnedObjects)
 	{
-		UE_CLOG(!IsInGameThread() && !IsGarbageCollectionLocked(), LogStreaming, Fatal, TEXT("Trying to add an object %s to FAsyncPackage referenced objects list outside of a FGCScopeGuard."), *InObject->GetFullName());
+		const EObjectFlags Flags = Object->GetFlags();
+		const EInternalObjectFlags InternalFlags = Object->GetInternalFlags();
+		EInternalObjectFlags InternalFlagsToClear = EInternalObjectFlags::None;
+
+		ensure(!(Flags & (RF_NeedPostLoad | RF_NeedPostLoadSubobjects)));
+		if (!!(InternalFlags & EInternalObjectFlags::AsyncLoading))
 		{
-			FScopeLock ReferencedObjectsLock(&ReferencedObjectsCritical);
-			if (!ReferencedObjects.Contains(InObject))
-			{
-				ReferencedObjects.Add(InObject);
-				InObject->AtomicallyClearInternalFlags(EInternalObjectFlags::Async);
-			}
+			ensure(!(Flags & RF_WasLoaded));
+			InternalFlagsToClear |= EInternalObjectFlags::AsyncLoading;
 		}
-		UE_CLOG(InObject->HasAnyInternalFlags(EInternalObjectFlags::Unreachable), LogStreaming, Fatal, TEXT("Trying to add an unreachable object %s to FAsyncPackage %s referenced objects list."), *InObject->GetFullName(), *GetPackageName().ToString());
-	}
-}
 
-void FAsyncPackage2::EmptyReferencedObjects()
-{
-	FScopeLock ReferencedObjectsLock(&ReferencedObjectsCritical);
-	ReferencedObjects.Reset();
+		if (!!(InternalFlags & EInternalObjectFlags::Async))
+		{
+			InternalFlagsToClear |= EInternalObjectFlags::Async;
+		}
+		Object->ClearInternalFlags(InternalFlagsToClear);
+	}
+	OwnedObjects.Empty();
 }
 
 void FAsyncPackage2::AddRequestID(int32 Id)
@@ -4387,19 +4428,19 @@ EAsyncPackageState::Type FAsyncPackage2::CreateLinker()
 		Package = FindObjectFast<UPackage>(nullptr, Desc.Name);
 		if (!Package)
 		{
-			// TODO: Are the redirect lookups done by the big and slow CreatePackage() ever needed in a cooked game???
 			Package = NewObject<UPackage>(/*Outer*/nullptr, Desc.Name, RF_Public);
+			Package->SetPackageFlags(Desc.PackageFlags);
+			Package->FileName = Desc.NameToLoad;
 			bCreatedLinkerRoot = true;
 		}
 	}
 	check(!IsNativeCodePackage(Package));
-	if (bCreatedLinkerRoot)
+	if (IsInGameThread() || !bCreatedLinkerRoot)
 	{
-		Package->SetPackageFlags(Desc.PackageFlags);
-		Package->FileName = Desc.NameToLoad;
+		Package->SetInternalFlags(EInternalObjectFlags::Async);
+		AddOwnedObject(Package);
 	}
-	
-	AddObjectReference(Package);
+	check(Package->HasAnyInternalFlags(EInternalObjectFlags::Async));
 	LinkerRoot = Package;
 
 	check(!FLinkerLoad::FindExistingLinkerForPackage(Package));
@@ -4497,6 +4538,7 @@ EAsyncPackageState::Type FAsyncPackage2::FinishLinker()
 				ObjectNameWithOuterToExport.Add(TPair<FName, FPackageIndex>(Export.ObjectName, Export.OuterIndex), Index);
 			}
 
+			OwnedObjects.Reserve(PackageSummary->ExportCount + 1); // ExportCount + UPackage
 			ExportIoBuffers.SetNum(PackageSummary->ExportCount);
 		}
 
@@ -4600,8 +4642,6 @@ EAsyncPackageState::Type FAsyncPackage2::PostLoadObjects()
 			// All object must be finalized on the game thread
 			DeferredFinalizeObjects.Add(Object);
 			check(Object->IsValidLowLevelFast());
-			// Make sure all objects in DeferredFinalizeObjects are referenced too
-			AddObjectReference(Object);
 		}
 	}
 
@@ -4863,14 +4903,11 @@ EAsyncPackageState::Type FAsyncPackage2::FinishObjects()
 	const bool bInternalCallbacks = true;
 	CallCompletionCallbacks(bInternalCallbacks, LoadingResult);
 
+	for (UObject* Object : OwnedObjects)
 	{
-		FScopeLock ReferencedObjectsLock(&ReferencedObjectsCritical);
-		for (UObject* Obj : ReferencedObjects)
+		if (!Object->HasAnyFlags(RF_NeedPostLoad | RF_NeedPostLoadSubobjects))
 		{
-			if (Obj && !Obj->HasAnyFlags(RF_NeedPostLoad | RF_NeedPostLoadSubobjects))
-			{
-				Obj->AtomicallyClearInternalFlags(EInternalObjectFlags::AsyncLoading);
-			}
+			Object->ClearInternalFlags(EInternalObjectFlags::AsyncLoading);
 		}
 	}
 
@@ -4941,13 +4978,10 @@ void FAsyncPackage2::Cancel()
 
 	{
 		// Clear load flags from any referenced objects
-		FScopeLock RereferncedObjectsLock(&ReferencedObjectsCritical);
-		ClearFlagsAndDissolveClustersFromLoadedObjects(ReferencedObjects);
 		ClearFlagsAndDissolveClustersFromLoadedObjects(PackageObjLoaded);
 		ClearFlagsAndDissolveClustersFromLoadedObjects(DeferredFinalizeObjects);
 	
 		// Release references
-		EmptyReferencedObjects();
 		PackageObjLoaded.Empty();
 		DeferredFinalizeObjects.Empty();
 	}
