Removed include to package resource

diff --git a/Engine/Source/Runtime/CoreUObject/Private/Serialization/AsyncLoading2.cpp b/Engine/Source/Runtime/CoreUObject/Private/Serialization/AsyncLoading2.cpp
--- a/Engine/Source/Runtime/CoreUObject/Private/Serialization/AsyncLoading2.cpp
+++ b/Engine/Source/Runtime/CoreUObject/Private/Serialization/AsyncLoading2.cpp
@@ -44,7 +44,6 @@
 #include "Serialization/Zenaphore.h"
 #include "IO/IoDispatcher.h"
 #include "Containers/CircularQueue.h"
-#include "UObject/PackageResource.h"
 #include "UObject/GCObject.h"
 #include "UObject/ObjectRedirector.h"
 
@@ -53,6 +52,59 @@
 #define ALT2_VERIFY_ASYNC_FLAGS
 #endif
 
+class FSimpleArchive
+	: public FArchive
+{
+public:
+	FSimpleArchive(const uint8* BufferPtr, uint64 BufferSize)
+	{
+#if DEVIRTUALIZE_FLinkerLoad_Serialize
+		ActiveFPLB->OriginalFastPathLoadBuffer = BufferPtr;
+		ActiveFPLB->StartFastPathLoadBuffer = BufferPtr;
+		ActiveFPLB->EndFastPathLoadBuffer = BufferPtr + BufferSize;
+#endif
+	}
+
+	int64 TotalSize() override
+	{
+#if DEVIRTUALIZE_FLinkerLoad_Serialize
+		return ActiveFPLB->EndFastPathLoadBuffer - ActiveFPLB->OriginalFastPathLoadBuffer;
+#else
+		return 0;
+#endif
+	}
+
+	int64 Tell() override
+	{
+#if DEVIRTUALIZE_FLinkerLoad_Serialize
+		return ActiveFPLB->StartFastPathLoadBuffer - ActiveFPLB->OriginalFastPathLoadBuffer;
+#else
+		return 0;
+#endif
+	}
+
+	void Seek(int64 Position) override
+	{
+#if DEVIRTUALIZE_FLinkerLoad_Serialize
+		ActiveFPLB->StartFastPathLoadBuffer = ActiveFPLB->OriginalFastPathLoadBuffer + Position;
+		check(ActiveFPLB->StartFastPathLoadBuffer <= ActiveFPLB->EndFastPathLoadBuffer);
+#endif
+	}
+
+	void Serialize(void* Data, int64 Length) override
+	{
+#if DEVIRTUALIZE_FLinkerLoad_Serialize
+		if (!Length || ArIsError)
+		{
+			return;
+		}
+		check(ActiveFPLB->StartFastPathLoadBuffer + Length <= ActiveFPLB->EndFastPathLoadBuffer);
+		FMemory::Memcpy(Data, ActiveFPLB->StartFastPathLoadBuffer, Length);
+		ActiveFPLB->StartFastPathLoadBuffer += Length;
+#endif
+	}
+};
+
 struct FAsyncPackage2;
 class FAsyncLoadingThread2Impl;
 
diff --git a/Engine/Source/Runtime/CoreUObject/Public/CoreUObject.h b/Engine/Source/Runtime/CoreUObject/Public/CoreUObject.h
--- a/Engine/Source/Runtime/CoreUObject/Public/CoreUObject.h
+++ b/Engine/Source/Runtime/CoreUObject/Public/CoreUObject.h
@@ -84,4 +84,3 @@ MONOLITHIC_HEADER_BOILERPLATE()
 #include "Misc/RedirectCollector.h"
 #include "UObject/ScriptStackTracker.h"
 #include "Misc/WorldCompositionUtility.h"
-#include "UObject/PackageResource.h"
diff --git a/Engine/Source/Runtime/CoreUObject/Public/UObject/PackageResource.h b/Engine/Source/Runtime/CoreUObject/Public/UObject/PackageResource.h
deleted file mode 100644
