AsyncLoading2: Double the IoRequestQueue size to 131072

diff --git a/Engine/Source/Runtime/CoreUObject/Private/Serialization/AsyncLoading2.cpp b/Engine/Source/Runtime/CoreUObject/Private/Serialization/AsyncLoading2.cpp
--- a/Engine/Source/Runtime/CoreUObject/Private/Serialization/AsyncLoading2.cpp
+++ b/Engine/Source/Runtime/CoreUObject/Private/Serialization/AsyncLoading2.cpp
@@ -1722,7 +1722,7 @@ void FAsyncLoadingThread2Impl::InitializeLoading()
 		UE_CLOG(!ReaderStatus.IsOk(), LogStreaming, Error, TEXT("Failed to initialize I/O dispatcher: '%s'"), *ReaderStatus.ToString());
 
 		IoDispatcher.Mount(IoStoreReader.Get());
-		IoRequestQueue = MakeUnique<FIoRequestQueue>(IoDispatcher, AltZenaphore, 65536);
+		IoRequestQueue = MakeUnique<FIoRequestQueue>(IoDispatcher, AltZenaphore, 131072);
 	}
 
 	{
