Pre-merging selected files from Main

diff --git a/Engine/Config/BaseEditor.ini b/Engine/Config/BaseEditor.ini
--- a/Engine/Config/BaseEditor.ini
+++ b/Engine/Config/BaseEditor.ini
@@ -507,7 +507,6 @@ bNativizeAnimBPOnlyWhenNonReducibleFuncitons=false
 +NoExportTypesWithDirectNativeFieldAccess=/Script/Engine.FormatArgumentData
 
 [/Script/UnrealEd.BlueprintEditorProjectSettings]
-bUseCompilationManager=true
 bValidateUnloadedSoftActorReferences=true
 
 [/Script/UnrealEd.AssetViewerSettings]
diff --git a/Engine/Source/Runtime/CoreUObject/Private/UObject/PropertyMulticastDelegate.cpp b/Engine/Source/Runtime/CoreUObject/Private/UObject/PropertyMulticastDelegate.cpp
--- a/Engine/Source/Runtime/CoreUObject/Private/UObject/PropertyMulticastDelegate.cpp
+++ b/Engine/Source/Runtime/CoreUObject/Private/UObject/PropertyMulticastDelegate.cpp
@@ -284,7 +284,7 @@ const TCHAR* UMulticastDelegateProperty::ImportText_Remove( const TCHAR* Buffer,
 		return NULL;
 	}
 
-	// Remove this delegate to our multicast delegate's invocation list
+	// Remove this delegate from our multicast delegate's invocation list
 	RemoveDelegate(ImportedDelegate, Parent, PropertyValue);
 
 	SkipWhitespace(Buffer);
@@ -391,7 +391,7 @@ void UMulticastInlineDelegateProperty::AddDelegate(FScriptDelegate ScriptDelegat
 
 	FMulticastScriptDelegate& MulticastDelegate = (*(FMulticastScriptDelegate*)PropertyValue);
 
-	// Remove this delegate to our multicast delegate's invocation list
+	// Add this delegate to our multicast delegate's invocation list
 	MulticastDelegate.AddUnique(MoveTemp(ScriptDelegate));
 }
 
@@ -401,7 +401,7 @@ void UMulticastInlineDelegateProperty::RemoveDelegate(const FScriptDelegate& Scr
 
 	FMulticastScriptDelegate& MulticastDelegate = (*(FMulticastScriptDelegate*)PropertyValue);
 
-	// Remove this delegate to our multicast delegate's invocation list
+	// Remove this delegate from our multicast delegate's invocation list
 	MulticastDelegate.Remove(ScriptDelegate);
 }
 
