TBA: Improvements to tagged property serialization in text files
 - FPropertyTag writes out its data as attributes
 - Properties are serialized as fields of a record rather than an array of tag/value pairs. Also wrapped static arrays up
 - Array properties write out as structured arrays, with only the inner type and struct name being written as attributes, rather than the full property tag

diff --git a/Engine/Source/Runtime/CoreUObject/Private/UObject/Class.cpp b/Engine/Source/Runtime/CoreUObject/Private/UObject/Class.cpp
--- a/Engine/Source/Runtime/CoreUObject/Private/UObject/Class.cpp
+++ b/Engine/Source/Runtime/CoreUObject/Private/UObject/Class.cpp
@@ -970,123 +970,90 @@ void UStruct::SerializeBinEx( FStructuredArchive::FSlot Slot, void* Data, void c
 	}
 }
 
-void UStruct::SerializeTaggedProperties(FStructuredArchive::FSlot Slot, uint8* Data, UStruct* DefaultsStruct, uint8* Defaults, const UObject* BreakRecursionIfFullyLoad) const
+void UStruct::LoadTaggedPropertiesFromText(FStructuredArchive::FSlot Slot, uint8* Data, UStruct* DefaultsStruct, uint8* Defaults, const UObject* BreakRecursionIfFullyLoad) const
 {
 	FArchive& UnderlyingArchive = Slot.GetUnderlyingArchive();
-	//SCOPED_LOADTIMER(SerializeTaggedPropertiesTime);
-
-	// Determine if this struct supports optional property guid's (UBlueprintGeneratedClasses Only)
-	const bool bArePropertyGuidsAvailable = (UnderlyingArchive.UE4Ver() >= VER_UE4_PROPERTY_GUID_IN_PROPERTY_TAG) && !FPlatformProperties::RequiresCookedData() && ArePropertyGuidsAvailable();
 	const bool bUseRedirects = !FPlatformProperties::RequiresCookedData() || UnderlyingArchive.IsSaveGame();
+	int32 NumProperties = 0;
+	FStructuredArchiveMap PropertiesMap = Slot.EnterMap(NumProperties);
 
-	if( UnderlyingArchive.IsLoading() )
+	for (int32 PropertyIndex = 0; PropertyIndex < NumProperties; ++PropertyIndex)
 	{
-		// Load tagged properties.
-		FStructuredArchive::FStream PropertiesStream = Slot.EnterStream();
+		FString PropertyNameString;
+		FStructuredArchiveSlot PropertySlot = PropertiesMap.EnterElement(PropertyNameString);
+		FName PropertyName = *PropertyNameString;
 
-		// This code assumes that properties are loaded in the same order they are saved in. This removes a n^2 search 
-		// and makes it an O(n) when properties are saved in the same order as they are loaded (default case). In the 
-		// case that a property was reordered the code falls back to a slower search.
-		UProperty*	Property			= PropertyLink;
-		bool		bAdvanceProperty	= false;
-		int32		RemainingArrayDim	= Property ? Property->ArrayDim : 0;
-
-		// Load all stored properties, potentially skipping unknown ones.
-		while (true)
+		// If this property has a guid attached then we need to resolve it to the right name before we start loading
+		TOptional<FStructuredArchiveSlot> PropertyGuidSlot = PropertySlot.TryEnterAttribute(SA_FIELD_NAME(TEXT("PropertyGuid")), false);
+		if (PropertyGuidSlot.IsSet())
 		{
-			FStructuredArchive::FRecord PropertyRecord = PropertiesStream.EnterElement().EnterRecord();
-
-			FPropertyTag Tag;
-			PropertyRecord << SA_VALUE(TEXT("Tag"), Tag);
-
-			if (Tag.Name.IsNone())
+			FGuid PropertyGuid;
+			PropertyGuidSlot.GetValue() << PropertyGuid;
+			if (PropertyGuid.IsValid())
 			{
-				break;
-			}
-
-			// Move to the next property to be serialized
-			if( bAdvanceProperty && --RemainingArrayDim <= 0 )
-			{
-				Property = Property->PropertyLinkNext;
-				// Skip over properties that don't need to be serialized.
-				while( Property && !Property->ShouldSerializeValue( UnderlyingArchive ) )
+				FName NewName = FindPropertyNameFromGuid(PropertyGuid);
+				if (NewName != NAME_None)
 				{
-					Property = Property->PropertyLinkNext;
+					PropertyName = NewName;
 				}
-				RemainingArrayDim = Property ? Property->ArrayDim : 0;
 			}
-			bAdvanceProperty = false;
-			
-			// Optionally resolve properties using Guid Property tags in non cooked builds that support it.
-			if (bArePropertyGuidsAvailable && Tag.HasPropertyGuid)
+		}
+
+		// Resolve any redirects if necessary
+		if (bUseRedirects && !UnderlyingArchive.HasAnyPortFlags(PPF_DuplicateForPIE | PPF_Duplicate))
+		{
+			for (UStruct* CheckStruct = GetOwnerStruct(); CheckStruct; CheckStruct = CheckStruct->GetSuperStruct())
 			{
-				// Use property guids from blueprint generated classes to redirect serialised data.
-				FName Result = FindPropertyNameFromGuid(Tag.PropertyGuid);
-				if (Result != NAME_None && Tag.Name != Result)
+				FName NewTagName = UProperty::FindRedirectedPropertyName(CheckStruct, PropertyName);
+				if (!NewTagName.IsNone())
 				{
-					Tag.Name = Result;
+					PropertyName = NewTagName;
+					break;
 				}
 			}
-			// If this property is not the one we expect (e.g. skipped as it matches the default value), do the brute force search.
-			if( Property == nullptr || Property->GetFName() != Tag.Name )
-			{
-				// No need to check redirects on platforms where everything is cooked. Always check for save games
-				if (bUseRedirects && !UnderlyingArchive.HasAnyPortFlags(PPF_DuplicateForPIE|PPF_Duplicate))
-				{
-					for (UStruct* CheckStruct = GetOwnerStruct(); CheckStruct; CheckStruct = CheckStruct->GetSuperStruct())
-					{
-						FName NewTagName = UProperty::FindRedirectedPropertyName(CheckStruct, Tag.Name);
-						if (!NewTagName.IsNone())
-						{
-							Tag.Name = NewTagName;
-							break;
-						}
-					}
-				}
-
-				UProperty* CurrentProperty = Property;
-				// Search forward...
-				for ( ; Property; Property=Property->PropertyLinkNext )
-				{
-					if( Property->GetFName() == Tag.Name )
-					{
-						break;
-					}
-				}
-				// ... and then search from the beginning till we reach the current property if it's not found.
-				if( Property == nullptr )
-				{
-					for( Property = PropertyLink; Property && Property != CurrentProperty; Property = Property->PropertyLinkNext )
-					{
-						if( Property->GetFName() == Tag.Name )
-						{
-							break;
-						}
-					}
+		}
 
-					if( Property == CurrentProperty )
-					{
-						// Property wasn't found.
-						Property = nullptr;
-					}
-				}
+		// Now we know what the property name is, we can try and load it
+		UProperty* Property = FindPropertyByName(PropertyName);
 
-				RemainingArrayDim = Property ? Property->ArrayDim : 0;
-			}
+		if (Property == nullptr)
+		{
+			Property = CustomFindProperty(PropertyName);
+		}
 
-			const int64 StartOfProperty = UnderlyingArchive.Tell();
+		if (Property && Property->ShouldSerializeValue(UnderlyingArchive))
+		{
+			FName PropID = Property->GetID();
 
-			if (!Property)
+			// Static arrays of tagged properties are special cases where the slot is always an array with no tag data attached. We currently have no TryEnterArray we can't 
+			// react based on what is in the file (yet) so we'll just have to assume that nobody converts a property from an array to a single value and go with whatever 
+			// the code property tells us.
+			TOptional<FStructuredArchiveArray> SlotArray;
+			int32 NumItems = Property->ArrayDim;
+			if (Property->ArrayDim > 1)
 			{
-				Property = CustomFindProperty(Tag.Name);
+				int32 NumAvailableItems = 0;
+				SlotArray.Emplace(PropertySlot.EnterArray(NumAvailableItems));
+				NumItems = FMath::Min(Property->ArrayDim, NumAvailableItems);
 			}
 
-			if (Property)
+			for (int32 ItemIndex = 0; ItemIndex < NumItems; ++ItemIndex)
 			{
-				FName PropID = Property->GetID();
+				TOptional<FStructuredArchiveSlot> ItemSlot;
+				if (SlotArray.IsSet())
+				{
+					ItemSlot.Emplace(SlotArray->EnterElement());
+				}
+				else
+				{
+					ItemSlot.Emplace(PropertySlot);
+				}
+
+				FPropertyTag Tag;
+				ItemSlot.GetValue() << Tag;
+				Tag.ArrayIndex = ItemIndex;
+				Tag.Name = PropertyName;
 
-				// Check if this is a struct property and we have a redirector
-				// No need to check redirects on platforms where everything is cooked. Always check for save games
 				if (bUseRedirects)
 				{
 					if (Tag.Type == NAME_StructProperty && PropID == NAME_StructProperty)
@@ -1106,34 +1073,214 @@ void UStruct::SerializeTaggedProperties(FStructuredArchive::FSlot Slot, uint8* D
 							Tag.EnumName = NewName;
 						}
 					}
+
+					if (!(BreakRecursionIfFullyLoad && BreakRecursionIfFullyLoad->HasAllFlags(RF_LoadCompleted)))
+					{
+						switch (Property->ConvertFromType(Tag, ItemSlot.GetValue(), Data, DefaultsStruct))
+						{
+						case EConvertFromTypeResult::Converted:
+							break;
+
+						case EConvertFromTypeResult::UseSerializeItem:
+							if (Tag.Type != PropID)
+							{
+								UE_LOG(LogClass, Warning, TEXT("Type mismatch in %s of %s - Previous (%s) Current(%s) for package:  %s"), *Tag.Name.ToString(), *GetName(), *Tag.Type.ToString(), *PropID.ToString(), *UnderlyingArchive.GetArchiveName());
+							}
+							else
+							{
+								uint8* DestAddress = Property->ContainerPtrToValuePtr<uint8>(Data, Tag.ArrayIndex);
+								uint8* DefaultsFromParent = Property->ContainerPtrToValuePtrForDefaults<uint8>(DefaultsStruct, Defaults, Tag.ArrayIndex);
+
+								// This property is ok.
+								Tag.SerializeTaggedProperty(ItemSlot.GetValue(), Property, DestAddress, DefaultsFromParent);
+							}
+							break;
+
+						case EConvertFromTypeResult::CannotConvert:
+							break;
+
+						default:
+							check(false);
+						}
+					}
 				}
+			}
+		}
+	}
+}
 
-#if WITH_EDITOR
-				if (BreakRecursionIfFullyLoad && BreakRecursionIfFullyLoad->HasAllFlags(RF_LoadCompleted))
+void UStruct::SerializeTaggedProperties(FStructuredArchive::FSlot Slot, uint8* Data, UStruct* DefaultsStruct, uint8* Defaults, const UObject* BreakRecursionIfFullyLoad) const
+{
+	FArchive& UnderlyingArchive = Slot.GetUnderlyingArchive();
+	//SCOPED_LOADTIMER(SerializeTaggedPropertiesTime);
+
+	// Determine if this struct supports optional property guid's (UBlueprintGeneratedClasses Only)
+	const bool bArePropertyGuidsAvailable = (UnderlyingArchive.UE4Ver() >= VER_UE4_PROPERTY_GUID_IN_PROPERTY_TAG) && !FPlatformProperties::RequiresCookedData() && ArePropertyGuidsAvailable();
+	const bool bUseRedirects = !FPlatformProperties::RequiresCookedData() || UnderlyingArchive.IsSaveGame();
+
+	if( UnderlyingArchive.IsLoading() )
+	{
+#if WITH_TEXT_ARCHIVE_SUPPORT
+		if (UnderlyingArchive.IsTextFormat())
+		{
+			LoadTaggedPropertiesFromText(Slot, Data, DefaultsStruct, Defaults, BreakRecursionIfFullyLoad);
+		}
+		else
+#endif // WITH_TEXT_ARCHIVE_SUPPORT
+		{
+			// Load tagged properties.
+			FStructuredArchive::FStream PropertiesStream = Slot.EnterStream();
+
+			// This code assumes that properties are loaded in the same order they are saved in. This removes a n^2 search 
+			// and makes it an O(n) when properties are saved in the same order as they are loaded (default case). In the 
+			// case that a property was reordered the code falls back to a slower search.
+			UProperty* Property = PropertyLink;
+			bool		bAdvanceProperty = false;
+			int32		RemainingArrayDim = Property ? Property->ArrayDim : 0;
+
+			// Load all stored properties, potentially skipping unknown ones.
+			while (true)
+			{
+				FStructuredArchive::FRecord PropertyRecord = PropertiesStream.EnterElement().EnterRecord();
+
+				FPropertyTag Tag;
+				PropertyRecord << SA_VALUE(TEXT("Tag"), Tag);
+
+				if (Tag.Name.IsNone())
 				{
+					break;
 				}
-#endif // WITH_EDITOR
-				// editoronly properties should be skipped if we are NOT the editor, or we are 
-				// the editor but are cooking for console (editoronly implies notforconsole)
-				else if ((Property->PropertyFlags & CPF_EditorOnly) && !FPlatformProperties::HasEditorOnlyData() && !GForceLoadEditorOnly)
+
+				// Move to the next property to be serialized
+				if (bAdvanceProperty && --RemainingArrayDim <= 0)
 				{
+					Property = Property->PropertyLinkNext;
+					// Skip over properties that don't need to be serialized.
+					while (Property && !Property->ShouldSerializeValue(UnderlyingArchive))
+					{
+						Property = Property->PropertyLinkNext;
+					}
+					RemainingArrayDim = Property ? Property->ArrayDim : 0;
 				}
-				// check for valid array index
-				else if (Tag.ArrayIndex >= Property->ArrayDim || Tag.ArrayIndex < 0)
+				bAdvanceProperty = false;
+
+				// Optionally resolve properties using Guid Property tags in non cooked builds that support it.
+				if (bArePropertyGuidsAvailable && Tag.HasPropertyGuid)
 				{
-					UE_LOG(LogClass, Warning, TEXT("Array bound exceeded (var %s=%d, exceeds %s [0-%d] in package:  %s"),
-						*Tag.Name.ToString(), Tag.ArrayIndex, *GetName(), Property->ArrayDim - 1, *UnderlyingArchive.GetArchiveName());
+					// Use property guids from blueprint generated classes to redirect serialised data.
+					FName Result = FindPropertyNameFromGuid(Tag.PropertyGuid);
+					if (Result != NAME_None && Tag.Name != Result)
+					{
+						Tag.Name = Result;
+					}
 				}
-				else if (!Property->ShouldSerializeValue(UnderlyingArchive))
+				// If this property is not the one we expect (e.g. skipped as it matches the default value), do the brute force search.
+				if (Property == nullptr || Property->GetFName() != Tag.Name)
 				{
-					UE_CLOG((UnderlyingArchive.IsPersistent() && FPlatformProperties::RequiresCookedData()), LogClass, Warning, TEXT("Skipping saved property %s of %s since it is no longer serializable for asset:  %s. (Maybe resave asset?)"), *Tag.Name.ToString(), *GetName(), *UnderlyingArchive.GetArchiveName());
+					// No need to check redirects on platforms where everything is cooked. Always check for save games
+					if (bUseRedirects && !UnderlyingArchive.HasAnyPortFlags(PPF_DuplicateForPIE | PPF_Duplicate))
+					{
+						for (UStruct* CheckStruct = GetOwnerStruct(); CheckStruct; CheckStruct = CheckStruct->GetSuperStruct())
+						{
+							FName NewTagName = UProperty::FindRedirectedPropertyName(CheckStruct, Tag.Name);
+							if (!NewTagName.IsNone())
+							{
+								Tag.Name = NewTagName;
+								break;
+							}
+						}
+					}
+
+					UProperty* CurrentProperty = Property;
+					// Search forward...
+					for (; Property; Property = Property->PropertyLinkNext)
+					{
+						if (Property->GetFName() == Tag.Name)
+						{
+							break;
+						}
+					}
+					// ... and then search from the beginning till we reach the current property if it's not found.
+					if (Property == nullptr)
+					{
+						for (Property = PropertyLink; Property && Property != CurrentProperty; Property = Property->PropertyLinkNext)
+						{
+							if (Property->GetFName() == Tag.Name)
+							{
+								break;
+							}
+						}
+
+						if (Property == CurrentProperty)
+						{
+							// Property wasn't found.
+							Property = nullptr;
+						}
+					}
+
+					RemainingArrayDim = Property ? Property->ArrayDim : 0;
 				}
-				else
+
+				const int64 StartOfProperty = UnderlyingArchive.Tell();
+
+				if (!Property)
+				{
+					Property = CustomFindProperty(Tag.Name);
+				}
+
+				if (Property)
 				{
-					FStructuredArchive::FSlot ValueSlot = PropertyRecord.EnterField(SA_FIELD_NAME(TEXT("Value")));
+					FName PropID = Property->GetID();
 
-					switch (Property->ConvertFromType(Tag, ValueSlot, Data, DefaultsStruct))
+					// Check if this is a struct property and we have a redirector
+					// No need to check redirects on platforms where everything is cooked. Always check for save games
+					if (bUseRedirects)
+					{
+						if (Tag.Type == NAME_StructProperty && PropID == NAME_StructProperty)
+						{
+							const FName NewName = FLinkerLoad::FindNewNameForStruct(Tag.StructName);
+							const FName StructName = CastChecked<UStructProperty>(Property)->Struct->GetFName();
+							if (NewName == StructName)
+							{
+								Tag.StructName = NewName;
+							}
+						}
+						else if ((PropID == NAME_EnumProperty) && ((Tag.Type == NAME_EnumProperty) || (Tag.Type == NAME_ByteProperty)))
+						{
+							const FName NewName = FLinkerLoad::FindNewNameForEnum(Tag.EnumName);
+							if (!NewName.IsNone())
+							{
+								Tag.EnumName = NewName;
+							}
+						}
+					}
+
+#if WITH_EDITOR
+					if (BreakRecursionIfFullyLoad && BreakRecursionIfFullyLoad->HasAllFlags(RF_LoadCompleted))
+					{
+					}
+#endif // WITH_EDITOR
+					// editoronly properties should be skipped if we are NOT the editor, or we are 
+					// the editor but are cooking for console (editoronly implies notforconsole)
+					else if ((Property->PropertyFlags & CPF_EditorOnly) && !FPlatformProperties::HasEditorOnlyData() && !GForceLoadEditorOnly)
+					{
+					}
+					// check for valid array index
+					else if (Tag.ArrayIndex >= Property->ArrayDim || Tag.ArrayIndex < 0)
+					{
+						UE_LOG(LogClass, Warning, TEXT("Array bound exceeded (var %s=%d, exceeds %s [0-%d] in package:  %s"),
+							*Tag.Name.ToString(), Tag.ArrayIndex, *GetName(), Property->ArrayDim - 1, *UnderlyingArchive.GetArchiveName());
+					}
+					else if (!Property->ShouldSerializeValue(UnderlyingArchive))
+					{
+						UE_CLOG((UnderlyingArchive.IsPersistent() && FPlatformProperties::RequiresCookedData()), LogClass, Warning, TEXT("Skipping saved property %s of %s since it is no longer serializable for asset:  %s. (Maybe resave asset?)"), *Tag.Name.ToString(), *GetName(), *UnderlyingArchive.GetArchiveName());
+					}
+					else
 					{
+						FStructuredArchive::FSlot ValueSlot = PropertyRecord.EnterField(SA_FIELD_NAME(TEXT("Value")));
+
+						switch (Property->ConvertFromType(Tag, ValueSlot, Data, DefaultsStruct))
+						{
 						case EConvertFromTypeResult::Converted:
 							bAdvanceProperty = true;
 							break;
@@ -1159,19 +1306,26 @@ void UStruct::SerializeTaggedProperties(FStructuredArchive::FSlot Slot, uint8* D
 
 						default:
 							check(false);
+						}
 					}
 				}
-			}
 
-			if (!bAdvanceProperty)
-			{
-				UnderlyingArchive.Seek(StartOfProperty + Tag.Size);
+				int64 Loaded = UnderlyingArchive.Tell() - StartOfProperty;
+
+				if (!bAdvanceProperty)
+				{
+					UnderlyingArchive.Seek(StartOfProperty + Tag.Size);
+				}
+				else
+				{
+					check(Tag.Size == Loaded);
+				}
 			}
 		}
 	}
 	else
 	{
-		FStructuredArchive::FStream PropertiesStream = Slot.EnterStream();
+		FStructuredArchive::FRecord PropertiesRecord = Slot.EnterRecord();
 
 		check(UnderlyingArchive.IsSaving() || UnderlyingArchive.IsCountingMemory());
 
@@ -1196,11 +1350,19 @@ void UStruct::SerializeTaggedProperties(FStructuredArchive::FSlot Slot, uint8* D
 			{
 				const int32 LoopMin = CustomPropertyNode ? CustomPropertyNode->ArrayIndex : 0;
 				const int32 LoopMax = CustomPropertyNode ? LoopMin + 1 : Property->ArrayDim;
+
+				TOptional<FStructuredArchive::FArray> StaticArrayContainer;
+				if (((LoopMax - 1) > LoopMin) && UnderlyingArchive.IsTextFormat())
+				{
+					int32 NumItems = LoopMax - LoopMin;
+					StaticArrayContainer.Emplace(PropertiesRecord.EnterArray(SA_FIELD_NAME((*Property->GetName())), NumItems));
+				}
+
 				for( int32 Idx = LoopMin; Idx < LoopMax; Idx++ )
 				{
 					uint8* DataPtr      = Property->ContainerPtrToValuePtr           <uint8>(Data, Idx);
 					uint8* DefaultValue = Property->ContainerPtrToValuePtrForDefaults<uint8>(DefaultsStruct, Defaults, Idx);
-					if( CustomPropertyNode || !UnderlyingArchive.DoDelta() || UnderlyingArchive.IsTransacting() || (!Defaults && !dynamic_cast<const UClass*>(this)) || !Property->Identical( DataPtr, DefaultValue, UnderlyingArchive.GetPortFlags()) )
+					if( StaticArrayContainer.IsSet() || CustomPropertyNode || !UnderlyingArchive.DoDelta() || UnderlyingArchive.IsTransacting() || (!Defaults && !dynamic_cast<const UClass*>(this)) || !Property->Identical( DataPtr, DefaultValue, UnderlyingArchive.GetPortFlags()) )
 					{
 						if (bUseAtomicSerialization)
 						{
@@ -1219,9 +1381,9 @@ void UStruct::SerializeTaggedProperties(FStructuredArchive::FSlot Slot, uint8* D
 							Tag.SetPropertyGuid(PropertyGuid);
 						}
 
-						FStructuredArchive::FRecord PropertyRecord = PropertiesStream.EnterElement().EnterRecord();
+						FStructuredArchive::FSlot PropertySlot = StaticArrayContainer.IsSet() ? StaticArrayContainer->EnterElement() : PropertiesRecord.EnterField(SA_FIELD_NAME(*Tag.Name.ToString()));
 
-						PropertyRecord << SA_VALUE(TEXT("Tag"), Tag);
+						PropertySlot << Tag;
 
 						// need to know how much data this call to SerializeTaggedProperty consumes, so mark where we are
 						int64 DataOffset = UnderlyingArchive.Tell();
@@ -1234,12 +1396,7 @@ void UStruct::SerializeTaggedProperties(FStructuredArchive::FSlot Slot, uint8* D
 							UnderlyingArchive.ArCustomPropertyList = CustomPropertyNode->SubPropertyList;
 						}
 
-						FStructuredArchive::FSlot PropertyField = PropertyRecord.EnterField(SA_FIELD_NAME(TEXT("Value")));
-						Tag.SerializeTaggedProperty(PropertyField, Property, DataPtr, DefaultValue);
-						if (!PropertyField.IsFilled())
-						{
-							Tag.SerializeTaggedProperty(PropertyField, Property, DataPtr, DefaultValue);
-						}
+						Tag.SerializeTaggedProperty(PropertySlot, Property, DataPtr, DefaultValue);
 
 						// restore the original custom property list after serializing
 						if (SavedCustomPropertyList)
@@ -1267,8 +1424,11 @@ void UStruct::SerializeTaggedProperties(FStructuredArchive::FSlot Slot, uint8* D
 			}
 		}
 
-		static FName Temp(NAME_None);
-		PropertiesStream.EnterElement().EnterRecord().EnterField(SA_FIELD_NAME(TEXT("Tag"))).EnterRecord() << SA_VALUE(TEXT("Name"), Temp);
+		if (!UnderlyingArchive.IsTextFormat())
+		{
+			static FName Temp(NAME_None);
+			UnderlyingArchive << Temp;
+		}
 	}
 }
 void UStruct::FinishDestroy()
diff --git a/Engine/Source/Runtime/CoreUObject/Private/UObject/PropertyArray.cpp b/Engine/Source/Runtime/CoreUObject/Private/UObject/PropertyArray.cpp
--- a/Engine/Source/Runtime/CoreUObject/Private/UObject/PropertyArray.cpp
+++ b/Engine/Source/Runtime/CoreUObject/Private/UObject/PropertyArray.cpp
@@ -62,14 +62,22 @@ void UArrayProperty::SerializeItem(FStructuredArchive::FSlot Slot, void* Value,
 {
 	checkSlow(Inner);
 	FArchive& UnderlyingArchive = Slot.GetUnderlyingArchive();
-	FStructuredArchive::FRecord Record = Slot.EnterRecord();
-
+	FPropertyTag InnerTag(UnderlyingArchive, Inner, 0, (uint8*)Value, (uint8*)Defaults);
+	
+	if (UnderlyingArchive.IsTextFormat() && InnerTag.Type == NAME_StructProperty)
+	{
+		Slot << SA_ATTRIBUTE(TEXT("InnerStructName"), InnerTag.StructName);
+		Slot << SA_DEFAULTED_ATTRIBUTE(TEXT("InnerStructGuid"), InnerTag.StructGuid, FGuid());
+	}
+	
 	// Ensure that the Inner itself has been loaded before calling SerializeItem() on it
 	UnderlyingArchive.Preload(Inner);
 
 	FScriptArrayHelper ArrayHelper(this, Value);
 	int32		n		= ArrayHelper.Num();
-	Record << SA_VALUE(TEXT("Count"), n);
+	
+	FStructuredArchiveArray Array = Slot.EnterArray(n);
+
 	if( UnderlyingArchive.IsLoading() )
 	{
 		// If using a custom property list, don't empty the array on load. Not all indices may have been serialized, so we need to preserve existing values at those slots.
@@ -93,17 +101,15 @@ void UArrayProperty::SerializeItem(FStructuredArchive::FSlot Slot, void* Value,
 	ArrayHelper.CountBytes( UnderlyingArchive );
 
 	// Serialize a PropertyTag for the inner property of this array, allows us to validate the inner struct to see if it has changed
-	FPropertyTag InnerTag(UnderlyingArchive, Inner, 0, (uint8*)Value, (uint8*)Defaults);
 	if (UnderlyingArchive.UE4Ver() >= VER_UE4_INNER_ARRAY_TAG_INFO && InnerTag.Type == NAME_StructProperty)
 	{
-		if (UnderlyingArchive.IsSaving())
+		if (!UnderlyingArchive.IsTextFormat())
 		{
-			Record << SA_VALUE(TEXT("InnerTag"), InnerTag);
+			UnderlyingArchive << InnerTag;
 		}
-		else if (UnderlyingArchive.IsLoading())
+		
+		if (UnderlyingArchive.IsLoading())
 		{
-			Record << SA_VALUE(TEXT("InnerTag"), InnerTag);
-
 			auto CanSerializeFromStructWithDifferentName = [](const FPropertyTag& PropertyTag, const UStructProperty* StructProperty)
 			{
 				return PropertyTag.StructGuid.IsValid()
@@ -155,8 +161,6 @@ void UArrayProperty::SerializeItem(FStructuredArchive::FSlot Slot, void* Value,
 	// need to know how much data this call to SerializeItem consumes, so mark where we are
 	int64 DataOffset = UnderlyingArchive.Tell();
 
-	FStructuredArchive::FStream ValueStream = Record.EnterField(SA_FIELD_NAME(TEXT("Values"))).EnterStream();
-
 	// If we're using a custom property list, first serialize any explicit indices
 	int32 i = 0;
 	bool bSerializeRemainingItems = true;
@@ -191,7 +195,7 @@ void UArrayProperty::SerializeItem(FStructuredArchive::FSlot Slot, void* Value,
 
 				// Serialize the item at this array index
 				i = PropertyNode->ArrayIndex;
-				Inner->SerializeItem(ValueStream.EnterElement(), ArrayHelper.GetRawPtr(i));
+				Inner->SerializeItem(Array.EnterElement(), ArrayHelper.GetRawPtr(i));
 				PropertyNode = PropertyNode->PropertyListNext;
 
 				// Restore the current property list
@@ -215,7 +219,7 @@ void UArrayProperty::SerializeItem(FStructuredArchive::FSlot Slot, void* Value,
 			NAME_UArraySerializeCount.SetNumber(i);
 			FArchive::FScopeAddDebugData P(UnderlyingArchive, NAME_UArraySerializeCount);
 #endif
-			Inner->SerializeItem(ValueStream.EnterElement(), ArrayHelper.GetRawPtr(i++));
+			Inner->SerializeItem(Array.EnterElement(), ArrayHelper.GetRawPtr(i++));
 		}
 
 		// Restore use of the custom property list (if it was previously enabled)
@@ -600,10 +604,16 @@ EConvertFromTypeResult UArrayProperty::ConvertFromType(const FPropertyTag& Tag,
 	{
 		void* ArrayPropertyData = ContainerPtrToValuePtr<void>(Data);
 
-		FStructuredArchive::FRecord ArrayRecord = Slot.EnterRecord();
-
 		int32 ElementCount = 0;
-		ArrayRecord << SA_VALUE(TEXT("Count"), ElementCount);
+
+		if (Slot.GetUnderlyingArchive().IsTextFormat())
+		{
+			Slot.EnterArray(ElementCount);
+		}
+		else
+		{
+			Slot.GetUnderlyingArchive() << ElementCount;
+		}
 
 		FScriptArrayHelper ScriptArrayHelper(this, ArrayPropertyData);
 		ScriptArrayHelper.EmptyAndAddValues(ElementCount);
@@ -615,7 +625,7 @@ EConvertFromTypeResult UArrayProperty::ConvertFromType(const FPropertyTag& Tag,
 			InnerPropertyTag.Type = Tag.InnerType;
 			InnerPropertyTag.ArrayIndex = 0;
 
-			FStructuredArchive::FStream ValueStream = ArrayRecord.EnterField(SA_FIELD_NAME(TEXT("Value"))).EnterStream();
+			FStructuredArchive::FStream ValueStream = Slot.EnterStream();
 
 			if (Inner->ConvertFromType(InnerPropertyTag, ValueStream.EnterElement(), ScriptArrayHelper.GetRawPtr(0), DefaultsStruct) == EConvertFromTypeResult::Converted)
 			{
diff --git a/Engine/Source/Runtime/CoreUObject/Private/UObject/PropertyTag.cpp b/Engine/Source/Runtime/CoreUObject/Private/UObject/PropertyTag.cpp
--- a/Engine/Source/Runtime/CoreUObject/Private/UObject/PropertyTag.cpp
+++ b/Engine/Source/Runtime/CoreUObject/Private/UObject/PropertyTag.cpp
@@ -80,28 +80,36 @@ FArchive& operator<<(FArchive& Ar, FPropertyTag& Tag)
 void operator<<(FStructuredArchive::FSlot Slot, FPropertyTag& Tag)
 {
 	FArchive& UnderlyingArchive = Slot.GetUnderlyingArchive();
-	FStructuredArchive::FRecord Record = Slot.EnterRecord();
+	bool bIsTextFormat = UnderlyingArchive.IsTextFormat();
+
 	int32 Version = UnderlyingArchive.UE4Ver();
 
 	checkf(!UnderlyingArchive.IsSaving() || Tag.Prop, TEXT("FPropertyTag must be constructed with a valid property when used for saving data!"));
 
-	// Name.
-	Record << SA_VALUE(TEXT("Name"), Tag.Name);
-	if (Tag.Name.IsNone())
+	if (!bIsTextFormat)
 	{
-		return;
+		// Name.
+		Slot << SA_ATTRIBUTE(TEXT("Name"), Tag.Name);
+		if (Tag.Name.IsNone())
+		{
+			return;
+		}
 	}
 
-	Record << SA_VALUE(TEXT("Type"), Tag.Type);
+	Slot << SA_ATTRIBUTE(TEXT("Type"), Tag.Type);
+
 	if (UnderlyingArchive.IsSaving())
 	{
 		// remember the offset of the Size variable - UStruct::SerializeTaggedProperties will update it after the
 		// property has been serialized.
 		Tag.SizeOffset = UnderlyingArchive.Tell();
 	}
+
+	if (!bIsTextFormat)
 	{
 		FArchive::FScopeSetDebugSerializationFlags S(UnderlyingArchive, DSF_IgnoreDiff);
-		Record << SA_VALUE(TEXT("Size"), Tag.Size) << SA_VALUE(TEXT("ArrayIndex"), Tag.ArrayIndex);
+		Slot << SA_ATTRIBUTE(TEXT("Size"), Tag.Size);
+		Slot << SA_ATTRIBUTE(TEXT("ArrayIndex"), Tag.ArrayIndex);
 	}
 
 	if (Tag.Type.GetNumber() == 0)
@@ -111,10 +119,17 @@ void operator<<(FStructuredArchive::FSlot Slot, FPropertyTag& Tag)
 		// only need to serialize this for structs
 		if (TagType == NAME_StructProperty)
 		{
-			Record << SA_VALUE(TEXT("StructName"), Tag.StructName);
+			Slot << SA_ATTRIBUTE(TEXT("StructName"), Tag.StructName);
 			if (Version >= VER_UE4_STRUCT_GUID_IN_PROPERTY_TAG)
 			{
-				Record << SA_VALUE(TEXT("StructGuid"), Tag.StructGuid);
+				if (bIsTextFormat)
+				{
+					Slot << SA_DEFAULTED_ATTRIBUTE(TEXT("StructGuid"), Tag.StructGuid, FGuid());
+				}
+				else
+				{
+					Slot << SA_ATTRIBUTE(TEXT("StructGuid"), Tag.StructGuid);
+				}
 			}
 		}
 		// only need to serialize this for bools
@@ -123,36 +138,47 @@ void operator<<(FStructuredArchive::FSlot Slot, FPropertyTag& Tag)
 			if (UnderlyingArchive.IsSaving())
 			{
 				FSerializedPropertyScope SerializedProperty(UnderlyingArchive, Tag.Prop);
-				Record << SA_VALUE(TEXT("BoolVal"), Tag.BoolVal);
+				Slot << SA_ATTRIBUTE(TEXT("BoolVal"), Tag.BoolVal);
 			}
 			else
 			{
-				Record << SA_VALUE(TEXT("BoolVal"), Tag.BoolVal);
+				Slot << SA_ATTRIBUTE(TEXT("BoolVal"), Tag.BoolVal);
 			}
 		}
 		// only need to serialize this for bytes/enums
-		else if (TagType == NAME_ByteProperty || TagType == NAME_EnumProperty)
+		else if (TagType == NAME_ByteProperty)
+		{
+			if (UnderlyingArchive.IsTextFormat())
+			{
+				Slot << SA_DEFAULTED_ATTRIBUTE(TEXT("EnumName"), Tag.EnumName, NAME_None);
+			}
+			else
+			{
+				Slot << SA_ATTRIBUTE(TEXT("EnumName"), Tag.EnumName);
+			}
+		}
+		else if (TagType == NAME_EnumProperty)
 		{
-			Record << SA_VALUE(TEXT("EnumName"), Tag.EnumName);
+			Slot << SA_ATTRIBUTE(TEXT("EnumName"), Tag.EnumName);
 		}
 		// only need to serialize this for arrays
 		else if (TagType == NAME_ArrayProperty)
 		{
 			if (Version >= VAR_UE4_ARRAY_PROPERTY_INNER_TAGS)
 			{
-				Record << SA_VALUE(TEXT("InnerType"), Tag.InnerType);
+				Slot << SA_ATTRIBUTE(TEXT("InnerType"), Tag.InnerType);
 			}
 		}
 		else if (Version >= VER_UE4_PROPERTY_TAG_SET_MAP_SUPPORT)
 		{
 			if (TagType == NAME_SetProperty)
 			{
-				Record << SA_VALUE(TEXT("InnerType"), Tag.InnerType);
+				Slot << SA_ATTRIBUTE(TEXT("InnerType"), Tag.InnerType);
 			}
 			else if (TagType == NAME_MapProperty)
 			{
-				Record << SA_VALUE(TEXT("InnerType"), Tag.InnerType);
-				Record << SA_VALUE(TEXT("ValueType"), Tag.ValueType);
+				Slot << SA_ATTRIBUTE(TEXT("InnerType"), Tag.InnerType);
+				Slot << SA_ATTRIBUTE(TEXT("ValueType"), Tag.ValueType);
 			}
 		}
 	}
@@ -160,10 +186,18 @@ void operator<<(FStructuredArchive::FSlot Slot, FPropertyTag& Tag)
 	// Property tags to handle renamed blueprint properties effectively.
 	if (Version >= VER_UE4_PROPERTY_GUID_IN_PROPERTY_TAG)
 	{
-		Record << SA_VALUE(TEXT("HasPropertyGuid"), Tag.HasPropertyGuid);
-		if (Tag.HasPropertyGuid)
+		if (bIsTextFormat)
 		{
-			Record << SA_VALUE(TEXT("PropertyGuid"), Tag.PropertyGuid);
+			Slot << SA_DEFAULTED_ATTRIBUTE(TEXT("PropertyGuid"), Tag.PropertyGuid, FGuid());
+			Tag.HasPropertyGuid = Tag.PropertyGuid.IsValid();
+		}
+		else
+		{
+			Slot << SA_ATTRIBUTE(TEXT("HasPropertyGuid"), Tag.HasPropertyGuid);
+			if (Tag.HasPropertyGuid)
+			{
+				Slot << SA_ATTRIBUTE(TEXT("PropertyGuid"), Tag.PropertyGuid);
+			}
 		}
 	}
 }
diff --git a/Engine/Source/Runtime/CoreUObject/Public/UObject/Class.h b/Engine/Source/Runtime/CoreUObject/Public/UObject/Class.h
--- a/Engine/Source/Runtime/CoreUObject/Public/UObject/Class.h
+++ b/Engine/Source/Runtime/CoreUObject/Public/UObject/Class.h
@@ -489,6 +489,9 @@ protected:
 	/** Returns if we have access to property guids */
 	virtual bool ArePropertyGuidsAvailable() const { return false; }
 
+	/** Serializes list of properties to a te, using property tags to handle mismatches */
+	void LoadTaggedPropertiesFromText(FStructuredArchive::FSlot Slot, uint8* Data, UStruct* DefaultsStruct, uint8* Defaults, const UObject* BreakRecursionIfFullyLoad) const;
+
 private:
 #if USTRUCT_FAST_ISCHILDOF_IMPL == USTRUCT_ISCHILDOF_STRUCTARRAY
 	// For UObjectBaseUtility
