ZenCreator and AsyncLoading2 cleanup of redundant data and files.

diff --git a/Engine/Source/Runtime/CoreUObject/Private/Serialization/AsyncLoading2.cpp b/Engine/Source/Runtime/CoreUObject/Private/Serialization/AsyncLoading2.cpp
--- a/Engine/Source/Runtime/CoreUObject/Private/Serialization/AsyncLoading2.cpp
+++ b/Engine/Source/Runtime/CoreUObject/Private/Serialization/AsyncLoading2.cpp
@@ -155,20 +155,21 @@ struct FGlobalImportRuntime
 	
 struct FPackageStoreEntrySerialized
 {
-	FGuid Guid;
 	FName Name;
 	FName FileName;
-	uint32 PackageFlags;
 	int32 ImportCount;
-	int32 ImportOffset;
 	int32 SlimportCount;
 	int32 SlimportOffset;
 	int32 ExportCount;
+};
+
+struct FPackageSummary
+{
+	uint32 PackageFlags;
 	int32 ExportOffset;
-	int32 PreloadDependencyCount;
-	int32 PreloadDependencyOffset;
-	int32 Pad;
-	int64 BulkDataStartOffset;
+	int32 GraphDataOffset;
+	int32 GraphDataSize;
+	int32 BulkDataStartOffset;
 };
 
 class FGlobalNameMap
@@ -304,19 +305,6 @@ uint16 GetChunkIndex(const FIoChunkId& ChunkId)
 
 }
 
-struct FPackageSummary
-{
-	FGuid Guid;
-	uint32 PackageFlags;
-	int32 ImportCount;
-	int32 ExportCount;
-	int32 PreloadDependencyCount;
-	int32 ExportOffset;
-	int32 GraphDataOffset;
-	int32 GraphDataSize;
-	int32 BulkDataStartOffset;
-};
-
 class FIoRequestQueue final : private FRunnable
 {
 public:
@@ -4224,7 +4212,6 @@ FAsyncPackage2::FAsyncPackage2(const FAsyncPackageDesc& InDesc, int32 InSerialNu
 
 	{
 		TRACE_CPUPROFILER_EVENT_SCOPE(CreateNodes);
-		FPackageFileSummary Summary;
 		uint32 ImportCount = AsyncLoadingThread.GetPackageImportCount(GlobalPackageId);
 		ImportNodeCount = ImportCount * EEventLoadNode2::Import_NumPhases;
 		uint32 ExportCount = AsyncLoadingThread.GetPackageExportCount(GlobalPackageId);
@@ -4528,23 +4515,15 @@ EAsyncPackageState::Type FAsyncPackage2::FinishLinker()
 		SCOPED_LOADTIMER(LinkerLoad_FinalizeCreation);
 
 		const FPackageSummary* PackageSummary = reinterpret_cast<const FPackageSummary*>(PackageSummaryBuffer.Get());
+		const int32 ExportCount = AsyncLoadingThread.GetPackageExportCount(GlobalPackageId);
 
 		{
 			SCOPED_LOADTIMER(LinkerLoad_SerializePackageFileSummary);
 
 			FPackageFileSummary& Summary	= Linker->Summary;
-
-			Summary.Tag						= PACKAGE_FILE_TAG;
-			Summary.Guid					= PackageSummary->Guid;
 			Summary.PackageFlags			= PackageSummary->PackageFlags;
-			Summary.ExportCount				= PackageSummary->ExportCount;
-			Summary.ImportCount				= PackageSummary->ImportCount;
-			Summary.PreloadDependencyCount	= PackageSummary->PreloadDependencyCount;
-			Summary.PreloadDependencyOffset	= 1; // HACK: circumvent check in FLinkerLoad::SerializePreloadDependencies()
 			Summary.BulkDataStartOffset		= PackageSummary->BulkDataStartOffset;
-
 			Summary.SetFileVersions(GPackageFileUE4Version, GPackageFileLicenseeUE4Version, /*unversioned*/true);
-
 			Linker->UpdateFromPackageFileSummary();
 		}
 
@@ -4563,35 +4542,33 @@ EAsyncPackageState::Type FAsyncPackage2::FinishLinker()
 			Linker->GlobalImportObjects = AsyncLoadingThread.GetGlobalImportObjects(TmpGlobalImportCount);
 		}
 
-		if (PackageSummary->ExportCount)
+		if (ExportCount)
 		{
 			SCOPED_LOADTIMER(LinkerLoad_SerializeExportMap);
 			const FObjectExport* Exports = reinterpret_cast<const FObjectExport*>(PackageSummaryBuffer.Get() + PackageSummary->ExportOffset);
 
-			Linker->ExportMap.AddUninitialized(PackageSummary->ExportCount);
-			FMemory::Memcpy(Linker->ExportMap.GetData(), Exports, PackageSummary->ExportCount * sizeof(FObjectExport));
+			Linker->ExportMap.AddUninitialized(ExportCount);
+			Linker->ExportMapIndex = ExportCount;
+			FMemory::Memcpy(Linker->ExportMap.GetData(), Exports, ExportCount * sizeof(FObjectExport));
 
 			for (FObjectExport& Export : Linker->ExportMap)
 			{
 				Export.ObjectName = AsyncLoadingThread.GlobalNameMap.FromSerializedName(Export.ObjectName);
 			}
 
-			// Foreach Export.bWasFiltered = FilterExport(Export); // TODO: Should never return true in packaged game, if so move to cooker!!!
-			Linker->ExportMapIndex = PackageSummary->ExportCount;
-
 			// ObjectNameWithOuterToExport has two use cases:
 			// - used for SetupImports during initial loading
 			// - MarkNewObjectForLoadIfItIsAnExport from NotifyConstructedDuringAsyncLoading
-			ObjectNameWithOuterToExport.Reserve(PackageSummary->ExportCount);
-			for (int32 LocalExportIndex = 0; LocalExportIndex < PackageSummary->ExportCount; ++LocalExportIndex)
+			ObjectNameWithOuterToExport.Reserve(ExportCount);
+			for (int32 LocalExportIndex = 0; LocalExportIndex < ExportCount; ++LocalExportIndex)
 			{
 				FPackageIndex Index = FPackageIndex::FromExport(LocalExportIndex);
 				const FObjectExport& Export = Linker->Exp(Index);
 				ObjectNameWithOuterToExport.Add(TPair<FName, FPackageIndex>(Export.ObjectName, Export.OuterIndex), Index);
 			}
 
-			OwnedObjects.Reserve(PackageSummary->ExportCount + 1); // ExportCount + UPackage
-			ExportIoBuffers.SetNum(PackageSummary->ExportCount);
+			OwnedObjects.Reserve(ExportCount + 1); // ExportCount + UPackage
+			ExportIoBuffers.SetNum(ExportCount);
 		}
 
 		// Add this linker to the object manager's linker array.
