Re-add stats instrumentation of local blueprint function calls

diff --git a/Engine/Source/Runtime/CoreUObject/Private/UObject/ScriptCore.cpp b/Engine/Source/Runtime/CoreUObject/Private/UObject/ScriptCore.cpp
--- a/Engine/Source/Runtime/CoreUObject/Private/UObject/ScriptCore.cpp
+++ b/Engine/Source/Runtime/CoreUObject/Private/UObject/ScriptCore.cpp
@@ -1051,6 +1051,10 @@ void ProcessLocalFunction(UObject* Context, UFunction* Fn, FFrame& Stack, RESULT
 	}
 	else
 	{
+#if PER_FUNCTION_SCRIPT_STATS
+		const bool bShouldTrackFunction = Stats::IsThreadCollectingData();
+		FScopeCycleCounterUObject FunctionScope(bShouldTrackFunction ? Fn : nullptr);
+#endif // PER_FUNCTION_SCRIPT_STATS
 		ProcessScriptFunction(Context, Fn, Stack, RESULT_PARAM, ProcessLocalScriptFunction);
 	}
 }
