Fix compile error in FUObjectAnnotationChunked

diff --git a/Engine/Source/Runtime/CoreUObject/Public/UObject/UObjectAnnotation.h b/Engine/Source/Runtime/CoreUObject/Public/UObject/UObjectAnnotation.h
--- a/Engine/Source/Runtime/CoreUObject/Public/UObject/UObjectAnnotation.h
+++ b/Engine/Source/Runtime/CoreUObject/Public/UObject/UObjectAnnotation.h
@@ -671,7 +671,7 @@ public:
 	TAnnotation& AddOrGetAnnotation(const UObjectBase *Object, TFunctionRef<TAnnotation()> NewAnnotationFn)
 	{
 		check(Object);
-		AddOrGetAnnotation(GUObjectArray.ObjectToIndex(Object), NewAnnotationFn);
+		return AddOrGetAnnotation(GUObjectArray.ObjectToIndex(Object), NewAnnotationFn);
 	}
 	/**
 	 * Add an annotation to the annotation list. If the Annotation is the default, then the annotation is removed.
