AsyncLoading2: Some GC bug fixes until revised later

diff --git a/Engine/Source/Runtime/CoreUObject/Private/Serialization/AsyncLoading2.cpp b/Engine/Source/Runtime/CoreUObject/Private/Serialization/AsyncLoading2.cpp
--- a/Engine/Source/Runtime/CoreUObject/Private/Serialization/AsyncLoading2.cpp
+++ b/Engine/Source/Runtime/CoreUObject/Private/Serialization/AsyncLoading2.cpp
@@ -142,14 +142,15 @@ struct FGlobalImportRuntime
 {
 	/** Persistent data */
 	int32 Count = 0;
-	FName* Names;
-	FPackageIndex* Outers;
-	FPackageIndex* Packages;
-	UObject** Objects;
+	FName* Names = nullptr;
+	FPackageIndex* Outers = nullptr;
+	FPackageIndex* Packages = nullptr;
+	UObject** Objects = nullptr;
 	/** Reference tracking for GC management */
-	TAtomic<int32>* RefCounts;
+	TAtomic<int32>* RefCounts = nullptr;
 	TArray<UObject*> KeepAliveObjects;
-	void OnPreGarbageCollect();
+	bool bIsLoadingPackages = false;
+	void OnPreGarbageCollect(bool bInIsLoadingPackages);
 	void OnPostGarbageCollect();
 };
 	
@@ -856,9 +857,41 @@ struct FAsyncPackage2
 		return FString::Printf(TEXT("FAsyncPackage %s"), *GetPackageName().ToString());
 	}
 
-	void AddOwnedObject(UObject* Object)
+	void AddOwnedObjectFromCallback(UObject* Object, bool bSubObject)
+	{
+		if (bSubObject)
+		{
+			if (!OwnedObjects.Contains(Object))
+			{
+				OwnedObjects.Add(Object);
+			}
+		}
+		else
+		{
+			ensure(!OwnedObjects.Contains(Object));
+			OwnedObjects.Add(Object);
+		}
+	}
+
+	void AddOwnedObject(UObject* Object, bool bForceAdd)
+	{
+		if (bForceAdd || !IsInAsyncLoadingThread())
+		{
+			ensure(!OwnedObjects.Contains(Object));
+			OwnedObjects.Add(Object);
+		}
+		ensure(OwnedObjects.Contains(Object));
+	}
+
+	void AddOwnedObjectWithAsyncFlag(UObject* Object, bool bForceAdd)
 	{
-		OwnedObjects.Add(Object);
+		AddOwnedObject(Object, bForceAdd);
+		if (bForceAdd || IsInGameThread())
+		{
+			ensure(!Object->HasAnyInternalFlags(EInternalObjectFlags::Async));
+			Object->SetInternalFlags(EInternalObjectFlags::Async);
+		}
+		ensure(Object->HasAnyInternalFlags(EInternalObjectFlags::Async));
 	}
 
 	void ClearOwnedObjects();
@@ -2802,18 +2835,12 @@ void FAsyncPackage2::EventDrivenCreateExport(int32 LocalExportIndex)
 							Export.Object->SetFlags(RF_NeedLoad | RF_NeedPostLoad | RF_NeedPostLoadSubobjects | RF_WasLoaded);
 						}
 						Export.Object->ClearFlags(RF_WillBeLoaded);
-						if (IsInGameThread())
-						{
-							Export.Object->SetInternalFlags(EInternalObjectFlags::Async);
-							AddOwnedObject(Export.Object);
-						}
+						AddOwnedObjectWithAsyncFlag(Export.Object, /*bForceAdd*/ false);
 					}
 					else
 					{
-						Export.Object->SetInternalFlags(EInternalObjectFlags::Async);
-						AddOwnedObject(Export.Object);
+						AddOwnedObjectWithAsyncFlag(Export.Object, /*bForceAdd*/ true);
 					}
-					check(Export.Object->HasAnyInternalFlags(EInternalObjectFlags::Async));
 				}
 				else
 				{
@@ -2927,12 +2954,8 @@ void FAsyncPackage2::EventDrivenCreateExport(int32 LocalExportIndex)
 					{
 						Export.Object->AddToRoot();
 					}
-					if (IsInGameThread())
-					{
-						Export.Object->SetInternalFlags(EInternalObjectFlags::Async);
-						AddOwnedObject(Export.Object);
-					}
-					check(Export.Object->HasAnyInternalFlags(EInternalObjectFlags::Async));
+
+					AddOwnedObjectWithAsyncFlag(Export.Object, /*bForceAdd*/ false);
 					Export.Object->SetLinker(Linker, LocalExportIndex);
 					check(Export.Object->GetClass() == LoadClass);
 					check(NewName == Export.ObjectName);
@@ -4028,10 +4051,9 @@ float FAsyncLoadingThread2Impl::GetAsyncLoadPercentage(const FName& PackageName)
 	return LoadPercentage;
 }
 
-void FGlobalImportRuntime::OnPreGarbageCollect()
+void FGlobalImportRuntime::OnPreGarbageCollect(bool bInIsLoadingPackages)
 {
 	int32 NumCleared = 0;
-
 	for (int32 GlobalImportIndex = 0; GlobalImportIndex < Count; ++GlobalImportIndex)
 	{
 		UObject*& Object = Objects[GlobalImportIndex];
@@ -4072,6 +4094,29 @@ void FGlobalImportRuntime::OnPreGarbageCollect()
 		}
 	}
 
+#ifdef ALT2_VERIFY_ASYNC_FLAGS
+	if (bIsLoadingPackages && !bInIsLoadingPackages)
+	{
+		for (int32 GlobalImportIndex = 0; GlobalImportIndex < Count; ++GlobalImportIndex)
+		{
+			check(Objects[GlobalImportIndex] == nullptr);
+			check(RefCounts[GlobalImportIndex] == 0);
+		}
+
+		const EInternalObjectFlags AsyncFlags = EInternalObjectFlags::Async | EInternalObjectFlags::AsyncLoading;
+		for (int32 ObjectIndex = 0; ObjectIndex < GUObjectArray.GetObjectArrayNum(); ++ObjectIndex)
+		{
+			FUObjectItem* ObjectItem = &GUObjectArray.GetObjectItemArrayUnsafe()[ObjectIndex];
+			if (UObject* Obj = static_cast<UObject*>(ObjectItem->Object))
+			{
+				ensure(!Obj->HasAnyInternalFlags(AsyncFlags));
+			}
+		}
+	}
+#endif
+
+	bIsLoadingPackages = bInIsLoadingPackages;
+
 	UE_LOG(LogStreaming, Log, TEXT("FGlobalImportRuntime::OnPreGarbageCollect - Marked %d objects, cleared %d object pointers"),
 		KeepAliveObjects.Num(), NumCleared);
 }
@@ -4082,6 +4127,7 @@ void FGlobalImportRuntime::OnPostGarbageCollect()
 	{
 		return;
 	}
+	check(bIsLoadingPackages);
 
 	for (UObject* Object : KeepAliveObjects)
 	{
@@ -4097,22 +4143,8 @@ void FGlobalImportRuntime::OnPostGarbageCollect()
 void FAsyncLoadingThread2Impl::OnPreGarbageCollect()
 {
 	TRACE_CPUPROFILER_EVENT_SCOPE(AltPreGC);
-	if (!IsAsyncLoadingPackages())
-	{
-#ifdef ALT2_VERIFY_ASYNC_FLAGS
-		const EInternalObjectFlags AsyncFlags = EInternalObjectFlags::Async | EInternalObjectFlags::AsyncLoading;
-		for (int32 ObjectIndex = 0; ObjectIndex < GUObjectArray.GetObjectArrayNum(); ++ObjectIndex)
-		{
-			FUObjectItem* ObjectItem = &GUObjectArray.GetObjectItemArrayUnsafe()[ObjectIndex];
-			if (UObject* Obj = static_cast<UObject*>(ObjectItem->Object))
-			{
-				ensure(!Obj->HasAnyInternalFlags(AsyncFlags));
-			}
-		}
-#endif
-		return;
-	}
-	GlobalImportRuntime.OnPreGarbageCollect();
+	const bool bIsAsyncLoadingPackages = IsAsyncLoadingPackages();
+	GlobalImportRuntime.OnPreGarbageCollect(bIsAsyncLoadingPackages);
 }
 
 void FAsyncLoadingThread2Impl::OnPostGarbageCollect()
@@ -4138,7 +4170,7 @@ void FAsyncLoadingThread2Impl::NotifyConstructedDuringAsyncLoading(UObject* Obje
 	FUObjectThreadContext& ThreadContext = FUObjectThreadContext::Get();
 	check(ThreadContext.AsyncPackage); // Otherwise something is wrong and we're creating objects outside of async loading code
 	FAsyncPackage2* AsyncPackage2 = (FAsyncPackage2*)ThreadContext.AsyncPackage;
-	AsyncPackage2->AddOwnedObject(Object);
+	AsyncPackage2->AddOwnedObjectFromCallback(Object, /*bForceAdd*/ bSubObject);
 	
 	// if this is in the package and is an export, then let mark it as needing load now
 	if (Object->GetOutermost() == AsyncPackage2->GetLinkerRoot() && 
@@ -4461,6 +4493,9 @@ EAsyncPackageState::Type FAsyncPackage2::CreateLinker()
 {
 	check(Linker == nullptr);
 
+	const int32 ExportCount = AsyncLoadingThread.GetPackageExportCount(GlobalPackageId);
+	OwnedObjects.Reserve(ExportCount + 1); // ExportCount + UPackage
+
 	// Try to find existing package or create it if not already present.
 	UPackage* Package = nullptr;
 	{
@@ -4474,11 +4509,7 @@ EAsyncPackageState::Type FAsyncPackage2::CreateLinker()
 		}
 	}
 	check(!IsNativeCodePackage(Package));
-	if (IsInGameThread() || !bCreatedLinkerRoot)
-	{
-		Package->SetInternalFlags(EInternalObjectFlags::Async);
-		AddOwnedObject(Package);
-	}
+	AddOwnedObjectWithAsyncFlag(Package, /*bForceAdd*/ !bCreatedLinkerRoot);
 	check(Package->HasAnyInternalFlags(EInternalObjectFlags::Async));
 	LinkerRoot = Package;
 
@@ -4567,7 +4598,6 @@ EAsyncPackageState::Type FAsyncPackage2::FinishLinker()
 				ObjectNameWithOuterToExport.Add(TPair<FName, FPackageIndex>(Export.ObjectName, Export.OuterIndex), Index);
 			}
 
-			OwnedObjects.Reserve(ExportCount + 1); // ExportCount + UPackage
 			ExportIoBuffers.SetNum(ExportCount);
 		}
 
