AsyncLoading2: Serialize slim exportmap with global import index (to be continued...)

diff --git a/Engine/Source/Runtime/CoreUObject/Private/Serialization/AsyncLoading2.cpp b/Engine/Source/Runtime/CoreUObject/Private/Serialization/AsyncLoading2.cpp
--- a/Engine/Source/Runtime/CoreUObject/Private/Serialization/AsyncLoading2.cpp
+++ b/Engine/Source/Runtime/CoreUObject/Private/Serialization/AsyncLoading2.cpp
@@ -167,7 +167,8 @@ struct FPackageStoreEntrySerialized
 struct FPackageSummary
 {
 	uint32 PackageFlags;
-	int32 ExportOffset;
+	int32 ExportMapOffset;
+	int32 ExportMapSize;
 	int32 GraphDataOffset;
 	int32 GraphDataSize;
 	int32 BulkDataStartOffset;
@@ -1222,27 +1223,6 @@ private:
 	int32 ThreadId = 0;
 };
 
-static FPackageIndex FindExportFromSlimport(FLinkerLoad* ImportLinker,
-	int32 GlobalImportIndex,
-	FPackageIndex* GlobalImportOuters,
-	FName* GlobalImportNames)
-{
-	check(ImportLinker && ImportLinker->AsyncRoot && ((FAsyncPackage2*)ImportLinker->AsyncRoot)->ObjectNameWithOuterToExport.Num());
-	FPackageIndex Result;
-	FPackageIndex ImportOuterIndex = GlobalImportOuters[GlobalImportIndex];
-	if (ImportOuterIndex.IsImport())
-	{
-		FName ObjectName = GlobalImportNames[GlobalImportIndex];
-		FPackageIndex ExportOuterIndex = FindExportFromSlimport(ImportLinker, ImportOuterIndex.ToImport(), GlobalImportOuters, GlobalImportNames);
-		FPackageIndex* PotentialExport = ((FAsyncPackage2*)ImportLinker->AsyncRoot)->ObjectNameWithOuterToExport.Find(TPair<FName, FPackageIndex>(ObjectName, ExportOuterIndex));
-		if (PotentialExport)
-		{
-			Result = *PotentialExport;
-		}
-	}
-	return Result;
-}
-
 static FPackageIndex FindExportFromObject2(FLinkerLoad* Linker, UObject *Object)
 {
 	check(Linker && Linker->AsyncRoot && ((FAsyncPackage2*)Linker->AsyncRoot)->ObjectNameWithOuterToExport.Num());
@@ -2242,7 +2222,7 @@ EAsyncPackageState::Type FAsyncPackage2::Event_StartImportPackages(FAsyncPackage
 
 	{
 		TRACE_CPUPROFILER_EVENT_SCOPE(SetupSerializedArcs);
-		const TArray<FNameEntryId>* GlobalNameMap = &Package->AsyncLoadingThread.GlobalNameMap.GetNameEntries();
+		const FGlobalNameMap& GlobalNameMap = Package->AsyncLoadingThread.GlobalNameMap;
 
 		const FPackageSummary* PackageSummary = reinterpret_cast<const FPackageSummary*>(Package->PackageSummaryBuffer.Get());
 		uint64 ZenHeaderDataSize = PackageSummary->GraphDataSize;
@@ -2265,8 +2245,7 @@ EAsyncPackageState::Type FAsyncPackage2::Event_StartImportPackages(FAsyncPackage
 			int32 ImportedPackageNameIndex;
 			int32 ImportedPackageNameNumber;
 			ZenHeaderArchive << ImportedPackageNameIndex << ImportedPackageNameNumber;
-			FNameEntryId MappedName = (*GlobalNameMap)[ImportedPackageNameIndex];
-			FName ImportedPackageName = FName::CreateFromDisplayId(MappedName, ImportedPackageNameNumber);
+			FName ImportedPackageName = GlobalNameMap.GetName(ImportedPackageNameIndex, ImportedPackageNameNumber);
 			FAsyncPackage2* ImportedPackage = Package->AsyncLoadingThread.FindAsyncPackage(ImportedPackageName);
 			int32 ExternalArcCount;
 			ZenHeaderArchive << ExternalArcCount;
@@ -3083,7 +3062,7 @@ void FAsyncPackage2::EventDrivenSerializeExport(int32 LocalExportIndex)
 		LoadContext->SerializedObject = PrevSerializedObject;
 		Linker->bForceSimpleIndexToObject = false;
 
-		UE_CLOG(Ar.Tell() != Export.SerialSize, LogStreaming, Warning, TEXT("Serialize mismatch, ObjectName='%s'"), *Object->GetFullName());
+		//UE_CLOG(Ar.Tell() != Export.SerialSize, LogStreaming, Warning, TEXT("Serialize mismatch, ObjectName='%s'"), *Object->GetFullName());
 
 		Linker->Loader = OldLoader;
 
@@ -4545,8 +4524,10 @@ EAsyncPackageState::Type FAsyncPackage2::FinishLinker()
 	{
 		SCOPED_LOADTIMER(LinkerLoad_FinalizeCreation);
 
+		const FGlobalNameMap& GlobalNameMap = AsyncLoadingThread.GlobalNameMap;
 		const FPackageSummary* PackageSummary = reinterpret_cast<const FPackageSummary*>(PackageSummaryBuffer.Get());
 		const int32 ExportCount = AsyncLoadingThread.GetPackageExportCount(GlobalPackageId);
+		check(ExportCount);
 
 		{
 			SCOPED_LOADTIMER(LinkerLoad_SerializePackageFileSummary);
@@ -4573,28 +4554,35 @@ EAsyncPackageState::Type FAsyncPackage2::FinishLinker()
 			Linker->GlobalImportObjects = AsyncLoadingThread.GetGlobalImportObjects(TmpGlobalImportCount);
 		}
 
-		if (ExportCount)
 		{
 			SCOPED_LOADTIMER(LinkerLoad_SerializeExportMap);
-			const FObjectExport* Exports = reinterpret_cast<const FObjectExport*>(PackageSummaryBuffer.Get() + PackageSummary->ExportOffset);
-
-			Linker->ExportMap.AddUninitialized(ExportCount);
-			Linker->ExportMapIndex = ExportCount;
-			FMemory::Memcpy(Linker->ExportMap.GetData(), Exports, ExportCount * sizeof(FObjectExport));
+			const uint8* ZenExportData = PackageSummaryBuffer.Get() + PackageSummary->ExportMapOffset;
+			FSimpleArchive ZenExportMapArchive(ZenExportData, PackageSummary->ExportMapSize);
 
-			for (FObjectExport& Export : Linker->ExportMap)
-			{
-				Export.ObjectName = AsyncLoadingThread.GlobalNameMap.FromSerializedName(Export.ObjectName);
-			}
-
-			// ObjectNameWithOuterToExport has two use cases:
-			// - used for SetupImports during initial loading
-			// - MarkNewObjectForLoadIfItIsAnExport from NotifyConstructedDuringAsyncLoading
+			// TODO: Remove ObjectNameWithOuterToExport
+			// It is only used from MarkNewObjectForLoadIfItIsAnExport from NotifyConstructedDuringAsyncLoading
 			ObjectNameWithOuterToExport.Reserve(ExportCount);
+			TArray<FObjectExport>& ExportMap = Linker->ExportMap;
+			ExportMap.AddZeroed(ExportCount);
 			for (int32 LocalExportIndex = 0; LocalExportIndex < ExportCount; ++LocalExportIndex)
 			{
 				FPackageIndex Index = FPackageIndex::FromExport(LocalExportIndex);
-				const FObjectExport& Export = Linker->Exp(Index);
+				FObjectExport& Export = ExportMap[LocalExportIndex];
+				int32 ObjectNameIndex;
+				int32 ObjectNameNumber;
+				int32 GlobalImportIndex;
+
+				ZenExportMapArchive << ObjectNameIndex << ObjectNameNumber;
+				ZenExportMapArchive << Export.OuterIndex;
+				ZenExportMapArchive << Export.ClassIndex;
+				ZenExportMapArchive << Export.SuperIndex;
+				ZenExportMapArchive << Export.TemplateIndex;
+				ZenExportMapArchive << (uint32&)Export.ObjectFlags;
+				ZenExportMapArchive << GlobalImportIndex;
+
+				Export.ObjectName = GlobalNameMap.GetName(ObjectNameIndex, ObjectNameNumber);
+				Export.ThisIndex = Index;
+
 				ObjectNameWithOuterToExport.Add(TPair<FName, FPackageIndex>(Export.ObjectName, Export.OuterIndex), Index);
 			}
 
