AsyncLoading2: Simplify NotifyConstructedDuringAsyncLoading and stop using RF_WillBeLoaded.
Remove MarkNewObjectForLoadIfItIsAnExport, FindExportFromObject2 and ObjectNameWithOuterToExport.

diff --git a/Engine/Source/Runtime/CoreUObject/Private/Serialization/AsyncLoading2.cpp b/Engine/Source/Runtime/CoreUObject/Private/Serialization/AsyncLoading2.cpp
--- a/Engine/Source/Runtime/CoreUObject/Private/Serialization/AsyncLoading2.cpp
+++ b/Engine/Source/Runtime/CoreUObject/Private/Serialization/AsyncLoading2.cpp
@@ -149,7 +149,7 @@ struct FGlobalImportRuntime
 	/** Reference tracking for GC management */
 	TAtomic<int32>* RefCounts = nullptr;
 	TArray<UObject*> KeepAliveObjects;
-	bool bIsLoadingPackages = false;
+	bool bNeedToRunGarbageCollect = false;
 	void OnPreGarbageCollect(bool bInIsLoadingPackages);
 	void OnPostGarbageCollect();
 };
@@ -1044,8 +1044,6 @@ public:
 	EAsyncPackageLoadingState2 AsyncPackageLoadingState;
 	int32 SerialNumber;
 
-	TMap<TPair<FName, FPackageIndex>, FPackageIndex> ObjectNameWithOuterToExport;
-
 	bool bHasImportedPackagesRecursive = false;
 
 	bool bAllExportsSerialized;
@@ -1223,23 +1221,6 @@ private:
 	int32 ThreadId = 0;
 };
 
-static FPackageIndex FindExportFromObject2(FLinkerLoad* Linker, UObject *Object)
-{
-	check(Linker && Linker->AsyncRoot && ((FAsyncPackage2*)Linker->AsyncRoot)->ObjectNameWithOuterToExport.Num());
-	FPackageIndex Result;
-	UObject* Outer = Object->GetOuter();
-	if (Outer)
-	{
-		FPackageIndex OuterIndex = FindExportFromObject2(Linker, Outer);
-		FPackageIndex* PotentialExport = ((FAsyncPackage2*)Linker->AsyncRoot)->ObjectNameWithOuterToExport.Find(TPair<FName, FPackageIndex>(Object->GetFName(), OuterIndex));
-		if (PotentialExport)
-		{
-			Result = *PotentialExport;
-		}
-	}
-	return Result;
-}
-
 class FAsyncLoadingThread2Impl
 	: public FRunnable
 {
@@ -2798,11 +2779,12 @@ void FAsyncPackage2::EventDrivenCreateExport(int32 LocalExportIndex)
 
 					// If this object was allocated but never loaded (components created by a constructor, CDOs, etc) make sure it gets loaded
 					// Do this for all subobjects created in the native constructor.
-					if (!Export.Object->HasAnyFlags(RF_LoadCompleted))
+					const EObjectFlags ObjectFlags = Export.Object->GetFlags();
+					if (!(ObjectFlags & RF_LoadCompleted))
 					{
 						UE_LOG(LogStreaming, VeryVerbose, TEXT("Note2: %s was constructed during load and is an export and so needs loading."), *Export.Object->GetFullName());
-						UE_CLOG(!Export.Object->HasAllFlags(RF_WillBeLoaded), LogStreaming, Fatal, TEXT("%s was found in memory and is an export but does not have all load flags."), *Export.Object->GetFullName());
-						if(Export.Object->HasAnyFlags(RF_ClassDefaultObject))
+						ensure(!(ObjectFlags & (RF_NeedLoad | RF_WasLoaded))); // If export exist but is not completed, it is expected to have been created from a native constructor and not from EventDrivenCreateExport, but who knows...?
+						if (ObjectFlags & RF_ClassDefaultObject)
 						{
 							// never call PostLoadSubobjects on class default objects, this matches the behavior of the old linker where
 							// StaticAllocateObject prevents setting of RF_NeedPostLoad and RF_NeedPostLoadSubobjects, but FLinkerLoad::Preload
@@ -2813,7 +2795,6 @@ void FAsyncPackage2::EventDrivenCreateExport(int32 LocalExportIndex)
 						{
 							Export.Object->SetFlags(RF_NeedLoad | RF_NeedPostLoad | RF_NeedPostLoadSubobjects | RF_WasLoaded);
 						}
-						Export.Object->ClearFlags(RF_WillBeLoaded);
 						AddOwnedObjectWithAsyncFlag(Export.Object, /*bForceAdd*/ false);
 					}
 					else
@@ -2949,20 +2930,6 @@ void FAsyncPackage2::EventDrivenCreateExport(int32 LocalExportIndex)
 	check(Export.Object || Export.bExportLoadFailed);
 }
 
-
-void FAsyncPackage2::MarkNewObjectForLoadIfItIsAnExport(UObject *Object)
-{
-	if (!Object->HasAnyFlags(RF_WillBeLoaded | RF_LoadCompleted | RF_NeedLoad))
-	{
-		FPackageIndex MaybeExportIndex = FindExportFromObject2(Linker, Object);
-		if (MaybeExportIndex.IsExport())
-		{
-			UE_LOG(LogStreaming, VeryVerbose, TEXT("Note: %s was constructed during load and is an export and so needs loading."), *Object->GetFullName());
-			Object->SetFlags(RF_WillBeLoaded);
-		}
-	}
-}
-
 void FAsyncPackage2::EventDrivenSerializeExport(int32 LocalExportIndex)
 {
 	//SCOPED_LOADTIMER(Package_PreLoadObjects);
@@ -2985,7 +2952,7 @@ void FAsyncPackage2::EventDrivenSerializeExport(int32 LocalExportIndex)
 			UObject* LocObj = ConstructDynamicType(*Linker->GetExportPathName(LocalExportIndex), EConstructDynamicType::CallZConstructor);
 			check(UD == LocObj);
 		}
-		Object->ClearFlags(RF_NeedLoad | RF_WillBeLoaded);
+		Object->ClearFlags(RF_NeedLoad);
 	}
 	else if (Object && Object->HasAnyFlags(RF_NeedLoad))
 	{
@@ -3871,6 +3838,7 @@ uint32 FAsyncLoadingThread2Impl::Run()
 		}
 		else
 		{
+			GlobalImportRuntime.bNeedToRunGarbageCollect |= IsAsyncLoadingPackages();
 			bool bDidSomething = false;
 			{
 				FGCScopeGuard GCGuard;
@@ -4030,8 +3998,45 @@ float FAsyncLoadingThread2Impl::GetAsyncLoadPercentage(const FName& PackageName)
 	return LoadPercentage;
 }
 
+static void VerifyLoadFlagsWhenFinishedLoading()
+{
+	const EInternalObjectFlags AsyncFlags =
+		EInternalObjectFlags::Async | EInternalObjectFlags::AsyncLoading;
+
+	const EObjectFlags LoadIntermediateFlags = 
+		EObjectFlags::RF_NeedLoad | EObjectFlags::RF_WillBeLoaded |
+		EObjectFlags::RF_NeedPostLoad | RF_NeedPostLoadSubobjects;
+
+	for (int32 ObjectIndex = 0; ObjectIndex < GUObjectArray.GetObjectArrayNum(); ++ObjectIndex)
+	{
+		FUObjectItem* ObjectItem = &GUObjectArray.GetObjectItemArrayUnsafe()[ObjectIndex];
+		if (UObject* Obj = static_cast<UObject*>(ObjectItem->Object))
+		{
+			const EInternalObjectFlags InternalFlags = Obj->GetInternalFlags();
+			const EObjectFlags Flags = Obj->GetFlags();
+			const bool bHasAnyAsyncFlags = !!(InternalFlags & AsyncFlags);
+			const bool bHasAnyLoadIntermediateFlags = !!(Flags & LoadIntermediateFlags);
+			const bool bWasLoaded = !!(Flags & RF_WasLoaded);
+			const bool bLoadCompleted = !!(Flags & RF_LoadCompleted);
+			ensure(!bHasAnyAsyncFlags);
+			ensure(!bHasAnyLoadIntermediateFlags);
+			if (bWasLoaded)
+			{
+				const bool bIsPackage = Obj->IsA(UPackage::StaticClass());
+				ensure(bIsPackage || bLoadCompleted);
+			}
+		}
+	}
+	UE_LOG(LogStreaming, Log, TEXT("Verified load flags when finished loading"));
+}
 void FGlobalImportRuntime::OnPreGarbageCollect(bool bInIsLoadingPackages)
 {
+	if (!bNeedToRunGarbageCollect && !bInIsLoadingPackages)
+	{
+		return;
+	}
+	bNeedToRunGarbageCollect = bInIsLoadingPackages;
+
 	int32 NumCleared = 0;
 	for (int32 GlobalImportIndex = 0; GlobalImportIndex < Count; ++GlobalImportIndex)
 	{
@@ -4073,29 +4078,23 @@ void FGlobalImportRuntime::OnPreGarbageCollect(bool bInIsLoadingPackages)
 		}
 	}
 
+	if (!bInIsLoadingPackages)
+	{
+		check(KeepAliveObjects.Num() == 0);
+	}
+
 #ifdef ALT2_VERIFY_ASYNC_FLAGS
-	if (bIsLoadingPackages && !bInIsLoadingPackages)
+	if (!bInIsLoadingPackages)
 	{
 		for (int32 GlobalImportIndex = 0; GlobalImportIndex < Count; ++GlobalImportIndex)
 		{
 			check(Objects[GlobalImportIndex] == nullptr);
 			check(RefCounts[GlobalImportIndex] == 0);
 		}
-
-		const EInternalObjectFlags AsyncFlags = EInternalObjectFlags::Async | EInternalObjectFlags::AsyncLoading;
-		for (int32 ObjectIndex = 0; ObjectIndex < GUObjectArray.GetObjectArrayNum(); ++ObjectIndex)
-		{
-			FUObjectItem* ObjectItem = &GUObjectArray.GetObjectItemArrayUnsafe()[ObjectIndex];
-			if (UObject* Obj = static_cast<UObject*>(ObjectItem->Object))
-			{
-				ensure(!Obj->HasAnyInternalFlags(AsyncFlags));
-			}
-		}
+		VerifyLoadFlagsWhenFinishedLoading();
 	}
 #endif
 
-	bIsLoadingPackages = bInIsLoadingPackages;
-
 	UE_LOG(LogStreaming, Log, TEXT("FGlobalImportRuntime::OnPreGarbageCollect - Marked %d objects, cleared %d object pointers"),
 		KeepAliveObjects.Num(), NumCleared);
 }
@@ -4106,7 +4105,7 @@ void FGlobalImportRuntime::OnPostGarbageCollect()
 	{
 		return;
 	}
-	check(bIsLoadingPackages);
+	check(bNeedToRunGarbageCollect);
 
 	for (UObject* Object : KeepAliveObjects)
 	{
@@ -4150,14 +4149,6 @@ void FAsyncLoadingThread2Impl::NotifyConstructedDuringAsyncLoading(UObject* Obje
 	check(ThreadContext.AsyncPackage); // Otherwise something is wrong and we're creating objects outside of async loading code
 	FAsyncPackage2* AsyncPackage2 = (FAsyncPackage2*)ThreadContext.AsyncPackage;
 	AsyncPackage2->AddOwnedObjectFromCallback(Object, /*bForceAdd*/ bSubObject);
-	
-	// if this is in the package and is an export, then let mark it as needing load now
-	if (Object->GetOutermost() == AsyncPackage2->GetLinkerRoot() && 
-		AsyncPackage2->AsyncPackageLoadingState <= EAsyncPackageLoadingState2::ProcessNewImportsAndExports &&
-		AsyncPackage2->AsyncPackageLoadingState > EAsyncPackageLoadingState2::WaitingForSummary)
-	{
-		AsyncPackage2->MarkNewObjectForLoadIfItIsAnExport(Object);
-	}
 }
 
 void FAsyncLoadingThread2Impl::FireCompletedCompiledInImport(FGCObject* AsyncPackage, FPackageIndex Import)
@@ -4559,9 +4550,6 @@ EAsyncPackageState::Type FAsyncPackage2::FinishLinker()
 			const uint8* ZenExportData = PackageSummaryBuffer.Get() + PackageSummary->ExportMapOffset;
 			FSimpleArchive ZenExportMapArchive(ZenExportData, PackageSummary->ExportMapSize);
 
-			// TODO: Remove ObjectNameWithOuterToExport
-			// It is only used from MarkNewObjectForLoadIfItIsAnExport from NotifyConstructedDuringAsyncLoading
-			ObjectNameWithOuterToExport.Reserve(ExportCount);
 			TArray<FObjectExport>& ExportMap = Linker->ExportMap;
 			ExportMap.AddZeroed(ExportCount);
 			for (int32 LocalExportIndex = 0; LocalExportIndex < ExportCount; ++LocalExportIndex)
@@ -4582,8 +4570,6 @@ EAsyncPackageState::Type FAsyncPackage2::FinishLinker()
 
 				Export.ObjectName = GlobalNameMap.GetName(ObjectNameIndex, ObjectNameNumber);
 				Export.ThisIndex = Index;
-
-				ObjectNameWithOuterToExport.Add(TPair<FName, FPackageIndex>(Export.ObjectName, Export.OuterIndex), Index);
 			}
 
 			ExportIoBuffers.SetNum(ExportCount);
