AsyncLoading2: Replace FLinkerLoad with new FSimpleExportArchive when serializing exports.
Add and use ExportMap in AsyncPackage2.

diff --git a/Engine/Source/Runtime/CoreUObject/Private/Serialization/AsyncLoading2.cpp b/Engine/Source/Runtime/CoreUObject/Private/Serialization/AsyncLoading2.cpp
--- a/Engine/Source/Runtime/CoreUObject/Private/Serialization/AsyncLoading2.cpp
+++ b/Engine/Source/Runtime/CoreUObject/Private/Serialization/AsyncLoading2.cpp
@@ -52,6 +52,9 @@
 #define ALT2_VERIFY_ASYNC_FLAGS
 #endif
 
+struct FAsyncPackage2;
+class FAsyncLoadingThread2Impl;
+
 class FSimpleArchive
 	: public FArchive
 {
@@ -105,8 +108,109 @@ public:
 	}
 };
 
-struct FAsyncPackage2;
-class FAsyncLoadingThread2Impl;
+class FSimpleExportArchive
+	: public FSimpleArchive
+{
+public:
+	FSimpleExportArchive(const uint8* BufferPtr, uint64 BufferSize) : FSimpleArchive(BufferPtr, BufferSize) {}
+
+	using FArchive::operator<<; // For visibility of the overloads we don't override
+
+	//~ Begin FArchive::FArchiveUObject Interface
+	virtual FArchive& operator<<(FSoftObjectPath& Value) override { return FArchiveUObject::SerializeSoftObjectPath(*this, Value); }
+	virtual FArchive& operator<<(FWeakObjectPtr& Value) override { return FArchiveUObject::SerializeWeakObjectPtr(*this, Value); }
+	//~ End FArchive::FArchiveUObject Interface
+
+	//~ Begin FArchive::FLinkerLoad Interface
+	UObject* GetArchetypeFromLoader(const UObject* Obj) { return TemplateForGetArchetypeFromLoader; }
+
+	virtual bool AttachExternalReadDependency(FExternalReadCallback& ReadCallback) override
+	{
+		ExternalReadDependencies->Add(ReadCallback);
+		return true;
+	}
+
+	virtual FArchive& operator<<( UObject*& Object ) override
+	{
+		FPackageIndex Index;
+		FArchive& Ar = *this;
+		Ar << Index;
+
+		if (Index.IsNull())
+		{
+			Object = nullptr;
+		}
+		else if (Index.IsExport())
+		{
+			Object = (*ExportMap)[Index.ToExport()].Object;
+		}
+		else
+		{
+			Object = GlobalImportObjects[LocalImportIndices[Index.ToImport()]];
+		}
+		return *this;
+	}
+
+	inline virtual FArchive& operator<<(FLazyObjectPtr& LazyObjectPtr) override
+	{
+		FArchive& Ar = *this;
+		FUniqueObjectGuid ID;
+		Ar << ID;
+		LazyObjectPtr = ID;
+		return Ar;
+	}
+
+	inline virtual FArchive& operator<<(FSoftObjectPtr& Value) override
+	{
+		FArchive& Ar = *this;
+		FSoftObjectPath ID;
+		ID.Serialize(Ar);
+		Value = ID;
+		return Ar;
+	}
+
+	void BadNameIndexError(int32 NameIndex)
+	{
+		UE_LOG(LogLinker, Error, TEXT("Bad name index %i/%i"), NameIndex, NameMap->Num());
+	}
+
+	inline virtual FArchive& operator<<(FName& Name) override
+	{
+		FArchive& Ar = *this;
+		int32 NameIndex;
+		Ar << NameIndex;
+		int32 Number = 0;
+		Ar << Number;
+
+		if (NameMap->IsValidIndex(NameIndex))
+		{
+			// if the name wasn't loaded (because it wasn't valid in this context)
+			FNameEntryId MappedName = (*NameMap)[NameIndex];
+
+			// simply create the name from the NameMap's name and the serialized instance number
+			Name = FName::CreateFromDisplayId(MappedName, Number);
+		}
+		else
+		{
+			Name = FName();
+			BadNameIndexError(NameIndex);
+			ArIsError = true;
+			ArIsCriticalError = true;
+		}
+		return *this;
+	}
+	//~ End FArchive::FLinkerLoad Interface
+
+private:
+	friend FAsyncPackage2;
+
+	UObject* TemplateForGetArchetypeFromLoader = nullptr;
+	int32* LocalImportIndices = nullptr;
+	UObject** GlobalImportObjects = nullptr;
+	const TArray<FNameEntryId>* NameMap = nullptr;
+	const TArray<FObjectExport>* ExportMap = nullptr;
+	TArray<FExternalReadCallback>* ExternalReadDependencies;
+};
 
 struct FPackageStoreEntryRuntime
 {
@@ -1019,6 +1123,8 @@ private:
 	TUniquePtr<uint8[]> PackageSummaryBuffer;
 	TArray<FIoBuffer> ExportIoBuffers;
 
+	// FZenLinkerLoad
+	TArray<FObjectExport> ExportMap;
 	int32 GlobalImportCount = 0;
 	int32 LocalImportCount = 0;
 	int32* LocalImportIndices = nullptr;
@@ -2503,7 +2609,7 @@ EAsyncPackageState::Type FAsyncPackage2::Event_SetupExports(FAsyncPackage2* Pack
 {
 	TRACE_CPUPROFILER_EVENT_SCOPE(Event_SetupExports);
 
-	Package->ExportIndex = Package->Linker->ExportMap.Num();
+	Package->ExportIndex = Package->ExportMap.Num();
 	check(Package->AsyncPackageLoadingState == EAsyncPackageLoadingState2::SetupExports);
 	Package->AsyncPackageLoadingState = EAsyncPackageLoadingState2::ProcessNewImportsAndExports;
 	return EAsyncPackageState::Complete;
@@ -2548,7 +2654,7 @@ EAsyncPackageState::Type FAsyncPackage2::Event_SerializeExport(FAsyncPackage2* P
 	FScopedAsyncPackageEvent2 Scope(Package);
 
 	Package->EventDrivenSerializeExport(LocalExportIndex);
-	FObjectExport& Export = Package->Linker->ExportMap[LocalExportIndex];
+	FObjectExport& Export = Package->ExportMap[LocalExportIndex];
 	UObject* Object = Export.Object;
 	check(!Object || !Object->HasAnyFlags(RF_NeedLoad));
 	return EAsyncPackageState::Complete;
@@ -2617,7 +2723,7 @@ UObject* FAsyncPackage2::EventDrivenIndexToObject(FPackageIndex Index, bool bChe
 	}
 	if (Index.IsExport())
 	{
-		Result = Linker->Exp(Index).Object;
+		Result = ExportMap[Index.ToExport()].Object;
 	}
 	else if (Index.IsImport())
 	{
@@ -2652,7 +2758,7 @@ UObject* FAsyncPackage2::EventDrivenIndexToObject(FPackageIndex Index, bool bChe
 void FAsyncPackage2::EventDrivenCreateExport(int32 LocalExportIndex)
 {
 	//SCOPED_LOADTIMER(Package_CreateExports);
-	FObjectExport& Export = Linker->ExportMap[LocalExportIndex];
+	FObjectExport& Export = ExportMap[LocalExportIndex];
 
 	TRACE_LOADTIME_CREATE_EXPORT_SCOPE(Linker, &Export.Object, Export.SerialOffset, Export.SerialSize, Export.bIsAsset);
 
@@ -2775,7 +2881,7 @@ void FAsyncPackage2::EventDrivenCreateExport(int32 LocalExportIndex)
 				if (Export.Object)
 				{
 					check(!Export.bForcedExport);
-					Export.Object->SetLinker(Linker, LocalExportIndex);
+					//Export.Object->SetLinker(Linker, LocalExportIndex);
 
 					// If this object was allocated but never loaded (components created by a constructor, CDOs, etc) make sure it gets loaded
 					// Do this for all subobjects created in the native constructor.
@@ -2916,7 +3022,7 @@ void FAsyncPackage2::EventDrivenCreateExport(int32 LocalExportIndex)
 					}
 
 					AddOwnedObjectWithAsyncFlag(Export.Object, /*bForceAdd*/ false);
-					Export.Object->SetLinker(Linker, LocalExportIndex);
+					//Export.Object->SetLinker(Linker, LocalExportIndex);
 					check(Export.Object->GetClass() == LoadClass);
 					check(NewName == Export.ObjectName);
 				}
@@ -2934,7 +3040,7 @@ void FAsyncPackage2::EventDrivenSerializeExport(int32 LocalExportIndex)
 {
 	//SCOPED_LOADTIMER(Package_PreLoadObjects);
 
-	FObjectExport& Export = Linker->ExportMap[LocalExportIndex];
+	FObjectExport& Export = ExportMap[LocalExportIndex];
 
 	LLM_SCOPE(ELLMTag::UObject);
 	LLM_SCOPED_TAG_WITH_OBJECT_IN_SET(GetLinkerRoot(), ELLMTagSet::Assets);
@@ -2956,8 +3062,8 @@ void FAsyncPackage2::EventDrivenSerializeExport(int32 LocalExportIndex)
 	}
 	else if (Object && Object->HasAnyFlags(RF_NeedLoad))
 	{
-		check(Object->GetLinker() == Linker);
-		check(Object->GetLinkerIndex() == LocalExportIndex);
+		//check(Object->GetLinker() == Linker);
+		//check(Object->GetLinkerIndex() == LocalExportIndex);
 		UClass* Cls = nullptr;
 
 		// If this is a struct, make sure that its parent struct is completely loaded
@@ -2988,14 +3094,25 @@ void FAsyncPackage2::EventDrivenSerializeExport(int32 LocalExportIndex)
 		const FPackageFileSummary& Summary = Linker->Summary;
 		const FCustomVersionContainer& SummaryVersions = Summary.GetCustomVersionContainer();
 
-		FSimpleArchive Ar(ExportIoBuffers[LocalExportIndex].Data(), ExportIoBuffers[LocalExportIndex].DataSize());
-		Ar.SetUE4Ver(Linker->Summary.GetFileVersionUE4());
-		Ar.SetLicenseeUE4Ver(Linker->Summary.GetFileVersionLicenseeUE4());
-		Ar.SetEngineVer(Linker->Summary.SavedByEngineVersion);
+		FSimpleExportArchive Ar(ExportIoBuffers[LocalExportIndex].Data(), ExportIoBuffers[LocalExportIndex].DataSize());
+		Ar.SetUE4Ver(Summary.GetFileVersionUE4());
+		Ar.SetLicenseeUE4Ver(Summary.GetFileVersionLicenseeUE4());
+		Ar.SetEngineVer(Summary.SavedByEngineVersion);
 		Ar.SetCustomVersions(SummaryVersions);
+		Ar.SetIsLoading(true);
+		Ar.SetIsPersistent(true);
+		if (Summary.PackageFlags & PKG_FilterEditorOnly)
+		{
+			Ar.SetFilterEditorOnly(true);
+		}
+		Ar.ArAllowLazyLoading = true;
 
-		FArchive* OldLoader = Linker->Loader;
-		Linker->Loader = &Ar;
+		// FSimpleExportArchive special fields
+		Ar.NameMap = &AsyncLoadingThread.GlobalNameMap.GetNameEntries();
+		Ar.LocalImportIndices = LocalImportIndices;
+		Ar.GlobalImportObjects = GlobalImportObjects;
+		Ar.ExportMap = &ExportMap;
+		Ar.ExternalReadDependencies = &Linker->ExternalReadDependencies;
 
 		Object->ClearFlags(RF_NeedLoad);
 
@@ -3011,28 +3128,22 @@ void FAsyncPackage2::EventDrivenSerializeExport(int32 LocalExportIndex)
 		UObject* Template = EventDrivenIndexToObject(Export.TemplateIndex, true);
 		check(Template);
 
-		check(!Linker->TemplateForGetArchetypeFromLoader);
-		Linker->TemplateForGetArchetypeFromLoader = Template;
+		Ar.TemplateForGetArchetypeFromLoader = Template;
 
 		if (Object->HasAnyFlags(RF_ClassDefaultObject))
 		{
-			Object->GetClass()->SerializeDefaultObject(Object, *Linker);
+			Object->GetClass()->SerializeDefaultObject(Object, Ar);
 		}
 		else
 		{
-			Object->Serialize(*Linker);
+			Object->Serialize(Ar);
 		}
-		check(Linker->TemplateForGetArchetypeFromLoader == Template);
-		Linker->TemplateForGetArchetypeFromLoader = nullptr;
 
 		Object->SetFlags(RF_LoadCompleted);
 		LoadContext->SerializedObject = PrevSerializedObject;
-		Linker->bForceSimpleIndexToObject = false;
 
 		//UE_CLOG(Ar.Tell() != Export.SerialSize, LogStreaming, Warning, TEXT("Serialize mismatch, ObjectName='%s'"), *Object->GetFullName());
 
-		Linker->Loader = OldLoader;
-
 #if DO_CHECK
 		if (Object->HasAnyFlags(RF_ClassDefaultObject) && Object->GetClass()->HasAnyClassFlags(CLASS_CompiledFromBlueprint))
 		{
@@ -3051,7 +3162,7 @@ EAsyncPackageState::Type FAsyncPackage2::Event_StartIO(FAsyncPackage2* Package,
 	TRACE_CPUPROFILER_EVENT_SCOPE(Event_StartIO);
 
 	check(Package->AsyncPackageLoadingState == EAsyncPackageLoadingState2::ProcessNewImportsAndExports);
-	FObjectExport& Export = Package->Linker->ExportMap[LocalExportIndex];
+	FObjectExport& Export = Package->ExportMap[LocalExportIndex];
 	if (Package->Linker->bDynamicClassLinker || //native blueprint, there is no IO for these
 		(Export.Object && !Export.Object->HasAnyFlags(RF_NeedLoad)))
 	{
@@ -3083,10 +3194,10 @@ EAsyncPackageState::Type FAsyncPackage2::Event_StartPostload(FAsyncPackage2* Pac
 	check(Package->bAllExportsSerialized);
 	{
 		check(Package->PackageObjLoaded.Num() == 0);
-		Package->PackageObjLoaded.Reserve(Package->Linker->ExportMap.Num());
-		for (int32 LocalExportIndex = 0; LocalExportIndex < Package->Linker->ExportMap.Num(); LocalExportIndex++)
+		Package->PackageObjLoaded.Reserve(Package->ExportMap.Num());
+		for (int32 LocalExportIndex = 0; LocalExportIndex < Package->ExportMap.Num(); LocalExportIndex++)
 		{
-			FObjectExport& Export = Package->Linker->ExportMap[LocalExportIndex];
+			FObjectExport& Export = Package->ExportMap[LocalExportIndex];
 			UObject* Object = Export.Object;
 			if (Object && (Object->HasAnyFlags(RF_NeedPostLoad) || Package->Linker->bDynamicClassLinker || Object->HasAnyInternalFlags(EInternalObjectFlags::AsyncLoading)))
 			{
@@ -3355,7 +3466,7 @@ bool FAsyncPackage2::AreAllDependenciesFullyLoadedInternal(FAsyncPackage2* Packa
 					OutError = FString::Printf(TEXT("%s Doesn't have all objects processed by DeferredPostLoad"), *Package->GetPackageName().ToString());
 					return false;
 				}
-				for (FObjectExport& Export : ImportPackageLinker->ExportMap)
+				for (FObjectExport& Export : AsyncRoot->ExportMap)
 				{
 					if (Export.Object && Export.Object->HasAnyFlags(RF_NeedPostLoad|RF_NeedLoad))
 					{
@@ -4511,74 +4622,55 @@ EAsyncPackageState::Type FAsyncPackage2::FinishLinker()
 {
 	LLM_SCOPE(ELLMTag::AsyncLoading);
 
-	//SCOPED_LOADTIMER(FinishLinkerTime);
-	{
-		SCOPED_LOADTIMER(LinkerLoad_FinalizeCreation);
+	SCOPED_LOADTIMER(LinkerLoad_FinalizeCreation);
 
-		const FGlobalNameMap& GlobalNameMap = AsyncLoadingThread.GlobalNameMap;
-		const FPackageSummary* PackageSummary = reinterpret_cast<const FPackageSummary*>(PackageSummaryBuffer.Get());
-		const int32 ExportCount = AsyncLoadingThread.GetPackageExportCount(GlobalPackageId);
-		check(ExportCount);
-
-		{
-			SCOPED_LOADTIMER(LinkerLoad_SerializePackageFileSummary);
+	const FGlobalNameMap& GlobalNameMap = AsyncLoadingThread.GlobalNameMap;
+	const FPackageSummary* PackageSummary = reinterpret_cast<const FPackageSummary*>(PackageSummaryBuffer.Get());
+	const int32 ExportCount = AsyncLoadingThread.GetPackageExportCount(GlobalPackageId);
+	check(ExportCount);
 
-			FPackageFileSummary& Summary	= Linker->Summary;
-			Summary.PackageFlags			= PackageSummary->PackageFlags;
-			Summary.BulkDataStartOffset		= PackageSummary->BulkDataStartOffset;
-			Summary.SetFileVersions(GPackageFileUE4Version, GPackageFileLicenseeUE4Version, /*unversioned*/true);
-			Linker->UpdateFromPackageFileSummary();
-		}
+	{
+		SCOPED_LOADTIMER(LinkerLoad_SerializePackageFileSummary);
 
-		// TODO: FLinker should not be a FArchive - NameMap is only required for operator<<(FName& Name)
-		{
-			SCOPED_LOADTIMER(LinkerLoad_SerializeNameMap);
-			Linker->ActiveNameMap = &AsyncLoadingThread.GlobalNameMap.GetNameEntries();
-		}
+		FPackageFileSummary& Summary	= Linker->Summary;
+		Summary.PackageFlags			= PackageSummary->PackageFlags;
+		Summary.BulkDataStartOffset		= PackageSummary->BulkDataStartOffset;
+		Summary.SetFileVersions(GPackageFileUE4Version, GPackageFileLicenseeUE4Version, /*unversioned*/true);
+		Linker->UpdateFromPackageFileSummary();
+	}
 
-		// TODO: FLinker should not be a FArchive - Slimports are only required for operator<<( UObject*& Object )
-		{
-			SCOPED_LOADTIMER(LinkerLoad_SerializeImportMap);
-			int32 TmpImportCount = 0;
-			int32 TmpGlobalImportCount = 0;
-			Linker->LocalImportIndices = AsyncLoadingThread.GetPackageSlimports(GlobalPackageId, TmpImportCount);
-			Linker->GlobalImportObjects = AsyncLoadingThread.GetGlobalImportObjects(TmpGlobalImportCount);
-		}
+	{
+		SCOPED_LOADTIMER(LinkerLoad_SerializeExportMap);
+		const uint8* ZenExportData = PackageSummaryBuffer.Get() + PackageSummary->ExportMapOffset;
+		FSimpleArchive ZenExportMapArchive(ZenExportData, PackageSummary->ExportMapSize);
 
+		ExportMap.AddZeroed(ExportCount);
+		for (int32 LocalExportIndex = 0; LocalExportIndex < ExportCount; ++LocalExportIndex)
 		{
-			SCOPED_LOADTIMER(LinkerLoad_SerializeExportMap);
-			const uint8* ZenExportData = PackageSummaryBuffer.Get() + PackageSummary->ExportMapOffset;
-			FSimpleArchive ZenExportMapArchive(ZenExportData, PackageSummary->ExportMapSize);
+			FPackageIndex Index = FPackageIndex::FromExport(LocalExportIndex);
+			FObjectExport& Export = ExportMap[LocalExportIndex];
+			int32 ObjectNameIndex;
+			int32 ObjectNameNumber;
+			int32 GlobalImportIndex;
 
-			TArray<FObjectExport>& ExportMap = Linker->ExportMap;
-			ExportMap.AddZeroed(ExportCount);
-			for (int32 LocalExportIndex = 0; LocalExportIndex < ExportCount; ++LocalExportIndex)
-			{
-				FPackageIndex Index = FPackageIndex::FromExport(LocalExportIndex);
-				FObjectExport& Export = ExportMap[LocalExportIndex];
-				int32 ObjectNameIndex;
-				int32 ObjectNameNumber;
-				int32 GlobalImportIndex;
-
-				ZenExportMapArchive << ObjectNameIndex << ObjectNameNumber;
-				ZenExportMapArchive << Export.OuterIndex;
-				ZenExportMapArchive << Export.ClassIndex;
-				ZenExportMapArchive << Export.SuperIndex;
-				ZenExportMapArchive << Export.TemplateIndex;
-				ZenExportMapArchive << (uint32&)Export.ObjectFlags;
-				ZenExportMapArchive << GlobalImportIndex;
-
-				Export.ObjectName = GlobalNameMap.GetName(ObjectNameIndex, ObjectNameNumber);
-				Export.ThisIndex = Index;
-			}
+			ZenExportMapArchive << ObjectNameIndex << ObjectNameNumber;
+			ZenExportMapArchive << Export.OuterIndex;
+			ZenExportMapArchive << Export.ClassIndex;
+			ZenExportMapArchive << Export.SuperIndex;
+			ZenExportMapArchive << Export.TemplateIndex;
+			ZenExportMapArchive << (uint32&)Export.ObjectFlags;
+			ZenExportMapArchive << GlobalImportIndex;
 
-			ExportIoBuffers.SetNum(ExportCount);
+			Export.ObjectName = GlobalNameMap.GetName(ObjectNameIndex, ObjectNameNumber);
+			Export.ThisIndex = Index;
 		}
 
-		// Add this linker to the object manager's linker array.
-		// FLinkerManager::Get().AddLoader(Linker);
+		ExportIoBuffers.SetNum(ExportCount);
 	}
 
+	// Add this linker to the object manager's linker array.
+	// FLinkerManager::Get().AddLoader(Linker);
+
 	return EAsyncPackageState::Complete;
 }
 
@@ -4593,7 +4685,9 @@ EAsyncPackageState::Type FAsyncPackage2::FinishExternalReadDependencies()
 	while (FinishExternalReadDependenciesIndex < PackageObjLoaded.Num())
 	{
 		UObject* Obj = PackageObjLoaded[FinishExternalReadDependenciesIndex];
-		FLinkerLoad* LinkerLoad = Obj ? Obj->GetLinker() : nullptr;
+		//FLinkerLoad* LinkerLoad = Obj ? Obj->GetLinker() : nullptr;
+		UPackage* Package = Obj ? Obj->GetOutermost() : nullptr;
+		FLinkerLoad* LinkerLoad = Package ? FLinkerLoad::FindExistingLinkerForPackage(Package) : nullptr;
 		if (LinkerLoad && LinkerLoad != VisitedLinkerLoad)
 		{
 			if (!LinkerLoad->FinishExternalReadDependencies(0.0) || FAsyncLoadingThreadState2::Get().IsTimeLimitExceeded())
diff --git a/Engine/Source/Runtime/CoreUObject/Private/Serialization/BulkData.cpp b/Engine/Source/Runtime/CoreUObject/Private/Serialization/BulkData.cpp
--- a/Engine/Source/Runtime/CoreUObject/Private/Serialization/BulkData.cpp
+++ b/Engine/Source/Runtime/CoreUObject/Private/Serialization/BulkData.cpp
@@ -981,11 +981,27 @@ void FUntypedBulkData::Serialize( FArchive& Ar, UObject* Owner, int32 Idx, bool
 			// determine whether the payload is stored inline or at the end of the file
 			const bool bPayloadInline = !(BulkDataFlags&BULKDATA_PayloadAtEndOfFile);
 
+			// GetLinker
+#if WITH_EDITOR
+			if (Owner != nullptr)
+			{
+				Linker = Owner->GetLinker();
+			}
+#else
+			FLinker* Linker = nullptr;
+			if (Owner != nullptr)
+			{
+				Package = Owner->GetOutermost();
+				check(Package.IsValid());
+				Linker = FLinkerLoad::FindExistingLinkerForPackage(Package.Get());
+			}
+#endif
+
 			// fix up the file offset, but only if not stored inline
 			int64 OffsetInFileFixup = 0;
-			if (Owner != NULL && Owner->GetLinker() && !bPayloadInline)
+			if (Owner != NULL && Linker && !bPayloadInline)
 			{
-				OffsetInFileFixup = Owner->GetLinker()->Summary.BulkDataStartOffset;
+				OffsetInFileFixup = Linker->Summary.BulkDataStartOffset;
 			}
 			BulkDataOffsetInFile += OffsetInFileFixup;
 
@@ -994,20 +1010,16 @@ void FUntypedBulkData::Serialize( FArchive& Ar, UObject* Owner, int32 Idx, bool
 			if (Ar.IsAllowingLazyLoading() && Owner != NULL && CacheableArchive)
 			{				
 #if WITH_EDITOR
-				Linker = Owner->GetLinker();
 				check(Linker);
 				CacheableArchive->AttachBulkData(Owner, this);
 				check(!CacheableArchive->IsTextFormat());
 				AttachedAr = CacheableArchive;
 				Filename = Linker->Filename;
 #else
-				Package = Owner->GetOutermost();
-				check(Package.IsValid());
-				auto Linker = FLinkerLoad::FindExistingLinkerForPackage(Package.Get());
 				check(Linker);
 				Filename = Linker->Filename;
 
-				check(Owner->GetLinker() == Linker);
+				check(!Owner->GetLinker() || Owner->GetLinker() == Linker);
 #endif // WITH_EDITOR
 
 				if (bPayloadInline)
diff --git a/Engine/Source/Runtime/CoreUObject/Private/UObject/LinkerLoad.cpp b/Engine/Source/Runtime/CoreUObject/Private/UObject/LinkerLoad.cpp
--- a/Engine/Source/Runtime/CoreUObject/Private/UObject/LinkerLoad.cpp
+++ b/Engine/Source/Runtime/CoreUObject/Private/UObject/LinkerLoad.cpp
@@ -5021,28 +5021,15 @@ FArchive& FLinkerLoad::operator<<( UObject*& Object )
 		{
 			Object = nullptr;
 		}
+		else if (Index.IsExport())
+		{
+			Object = Exp(Index).Object;
+		}
 		else
 		{
-			if (Index.IsExport())
-			{
-				Object = Exp(Index).Object;
-			}
-			else // Index.IsImport()
-			{
-				if (LocalImportIndices)
-				{
-					check(GlobalImportObjects);
-					Object = GlobalImportObjects[LocalImportIndices[Index.ToImport()]];
-				}
-				else
-				{
-					Object = Imp(Index).XObject;
-				}
-			}
+			Object = Imp(Index).XObject;
 		}
 
-		//Object = ((FAsyncPackage*)AsyncRoot)->EventDrivenIndexToObject(Index, false);
-
 		return *this;
 	}
 
diff --git a/Engine/Source/Runtime/CoreUObject/Public/UObject/LinkerLoad.h b/Engine/Source/Runtime/CoreUObject/Public/UObject/LinkerLoad.h
--- a/Engine/Source/Runtime/CoreUObject/Public/UObject/LinkerLoad.h
+++ b/Engine/Source/Runtime/CoreUObject/Public/UObject/LinkerLoad.h
@@ -148,10 +148,6 @@ private:
 	/** The archive that actually reads the raw data from disk.																*/
 	FArchive*				Loader;
 
-	int32* LocalImportIndices = nullptr;
-	UObject** GlobalImportObjects = nullptr;
-	const TArray<FNameEntryId>* ActiveNameMap = &NameMap;
-
 protected:
 
 	void SetLoader(FArchive* InLoader);
@@ -823,10 +819,10 @@ private:
 		int32 Number = 0;
 		Ar << Number;
 
-		if (ActiveNameMap->IsValidIndex(NameIndex))
+		if (NameMap.IsValidIndex(NameIndex))
 		{
 			// if the name wasn't loaded (because it wasn't valid in this context)
-			FNameEntryId MappedName = (*ActiveNameMap)[NameIndex];
+			FNameEntryId MappedName = NameMap[NameIndex];
 
 			// simply create the name from the NameMap's name and the serialized instance number
 			Name = FName::CreateFromDisplayId(MappedName, Number);
diff --git a/Engine/Source/Runtime/Engine/Private/Materials/Material.cpp b/Engine/Source/Runtime/Engine/Private/Materials/Material.cpp
--- a/Engine/Source/Runtime/Engine/Private/Materials/Material.cpp
+++ b/Engine/Source/Runtime/Engine/Private/Materials/Material.cpp
@@ -680,7 +680,8 @@ void SerializeInlineShaderMaps(
 
 		if (OutOffsetToFirstResource)
 		{
-			int64 Tmp = Ar.Tell() - Ar.GetLinker()->Summary.TotalHeaderSize;
+			const FLinker* Linker = Ar.GetLinker();
+			int64 Tmp = Ar.Tell() - (Linker ? Linker->Summary.TotalHeaderSize : 0);
 			check(Tmp >= 0 && Tmp <= 0xffffffffLL);
 			*OutOffsetToFirstResource = uint32(Tmp);
 		}
