Factoring out the BulkData saving code from UPackage::SavePackage to it's own function (::SaveBulkData) so that the logic is easier to follow.

diff --git a/Engine/Source/Runtime/CoreUObject/Private/UObject/SavePackage.cpp b/Engine/Source/Runtime/CoreUObject/Private/UObject/SavePackage.cpp
--- a/Engine/Source/Runtime/CoreUObject/Private/UObject/SavePackage.cpp
+++ b/Engine/Source/Runtime/CoreUObject/Private/UObject/SavePackage.cpp
@@ -70,6 +70,8 @@ static FCriticalSection InitializeCoreClassesCritSec;
 
 static void SaveThumbnails(UPackage* InOuter, FLinkerSave* Linker, FStructuredArchive::FSlot Slot);
 static void SaveAssetRegistryData(UPackage* InOuter, FLinkerSave* Linker, FStructuredArchive::FSlot Slot);
+static void SaveBulkData(FLinkerSave* Linker, const UPackage* InOuter, const TCHAR* Filename, const ITargetPlatform* TargetPlatform,
+						 FSavePackageContext* SavePackageContext, const bool bTextFormat, const bool bDiffing, FMD5* CookedPackageHash, int64 TotalPackageSizeUncompressed);
 static void SaveWorldLevelInfo(UPackage* InOuter, FLinkerSave* Linker, FStructuredArchive::FRecord Record);
 static EObjectMark GetExcludedObjectMarksForTargetPlatform(const class ITargetPlatform* TargetPlatform, const bool bIsCooking);
 
@@ -5430,231 +5432,11 @@ FSavePackageResultStruct UPackage::Save(UPackage* InOuter, UObject* Base, EObjec
 				{ 
 					return ESavePackageResult::Canceled;
 				}
-				SlowTask.EnterProgressFrame(1, NSLOCTEXT("Core", "SerializingBulkData", "Serializing bulk data"));
-
-				// now we write all the bulkdata that is supposed to be at the end of the package
-				// and fix up the offset
-				const int64 StartOfBulkDataArea = Linker->Tell();
-				Linker->Summary.BulkDataStartOffset = StartOfBulkDataArea;
-
-				check(!bTextFormat || Linker->BulkDataToAppend.Num() == 0);
-
-				if (!bTextFormat && Linker->BulkDataToAppend.Num() > 0)
-				{
-					COOK_STAT(FScopedDurationTimer SaveTimer(SavePackageStats::SerializeBulkDataTimeSec));
-
-					FScopedSlowTask BulkDataFeedback(Linker->BulkDataToAppend.Num());
-
-					TUniquePtr<FLargeMemoryWriter> BulkArchive;
-					TUniquePtr<FLargeMemoryWriter> OptionalBulkArchive;
-					TUniquePtr<FLargeMemoryWriter> MappedBulkArchive;
-
-					uint32 ExtraBulkDataFlags = 0;
-
-					static const struct FUseSeparateBulkDataFiles
-					{
-						bool bEnable = false;
-
-						FUseSeparateBulkDataFiles()
-						{
-							GConfig->GetBool(TEXT("Core.System"), TEXT("UseSeperateBulkDataFiles"), /* out */ bEnable, GEngineIni);
-
-							if (IsEventDrivenLoaderEnabledInCookedBuilds())
-							{
-								// Always split bulk data when splitting cooked files
-								bEnable = true;
-							}
-						}
-					} ShouldUseSeparateBulkDataFiles;
-
-					const bool bShouldUseSeparateBulkFile = ShouldUseSeparateBulkDataFiles.bEnable && Linker->IsCooking();
-
-					if (bShouldUseSeparateBulkFile)
-					{
-						ExtraBulkDataFlags	= BULKDATA_PayloadInSeperateFile;
-
-						BulkArchive.Reset(new FLargeMemoryWriter(0, /* IsPersistent */ true));
-						OptionalBulkArchive.Reset(new FLargeMemoryWriter(0, /* IsPersistent */ true));
-						MappedBulkArchive.Reset(new FLargeMemoryWriter(0, /* IsPersistent */ true));
-					}
-
-					bool bAlignBulkData = false;
-					int64 BulkDataAlignment = 0;
-
-					if (TargetPlatform)
-					{
-						bAlignBulkData = TargetPlatform->SupportsFeature(ETargetPlatformFeatures::MemoryMappedFiles);
-						BulkDataAlignment = TargetPlatform->GetMemoryMappingAlignment();
-					}
-
-					for (FLinkerSave::FBulkDataStorageInfo& BulkDataStorageInfo : Linker->BulkDataToAppend)
-					{
-						BulkDataFeedback.EnterProgressFrame();
-
-						// Set bulk data flags to what they were during initial serialization (they might have changed after that)
-						const uint32 OldBulkDataFlags = BulkDataStorageInfo.BulkData->GetBulkDataFlags();
-						uint32 ModifiedBulkDataFlags = BulkDataStorageInfo.BulkDataFlags | ExtraBulkDataFlags;
-						const bool bBulkItemIsOptional = (ModifiedBulkDataFlags & BULKDATA_OptionalPayload) != 0;
-						bool bBulkItemIsMapped = bAlignBulkData && ((ModifiedBulkDataFlags & BULKDATA_MemoryMappedPayload) != 0);
-
-						if (bBulkItemIsMapped && bBulkItemIsOptional)
-						{
-							UE_LOG(LogSavePackage, Warning, TEXT("%s has bulk data that is both mapped and optional. This is not currently supported. Will not be mapped."), Filename);
-							ModifiedBulkDataFlags &= ~BULKDATA_MemoryMappedPayload;
-							bBulkItemIsMapped = false;
-						}
-
-						BulkDataStorageInfo.BulkData->ClearBulkDataFlags(0xFFFFFFFF);
-						BulkDataStorageInfo.BulkData->SetBulkDataFlags(ModifiedBulkDataFlags);
-
-						FArchive* TargetArchive = Linker.Get();
-
-						if (bShouldUseSeparateBulkFile)
-						{
-							check(MappedBulkArchive && OptionalBulkArchive && BulkArchive);
-
-							if (bBulkItemIsOptional)
-							{
-								TargetArchive = OptionalBulkArchive.Get();
-							}
-							else if (bBulkItemIsMapped)
-							{
-								TargetArchive = MappedBulkArchive.Get();
-							}
-							else
-							{
-								TargetArchive = BulkArchive.Get();
-							}
-						}
-
-						// Pad archive for proper alignment for memory mapping
-
-						if (bBulkItemIsMapped && BulkDataAlignment > 0)
-						{
-							const int64 BulkStartOffset = TargetArchive->Tell();
-
-							if (!IsAligned(BulkStartOffset, BulkDataAlignment))
-							{
-								const int64 AlignedOffset = Align(BulkStartOffset, BulkDataAlignment);
-
-								int64 Padding = AlignedOffset - BulkStartOffset;
-								check(Padding > 0);
-
-								uint64 Zero64 = 0;
-								while(Padding >= 8)
-								{
-									*TargetArchive << Zero64;
-									Padding -= 8;
-								}
 
-								uint8 Zero8 = 0;
-								while(Padding > 0)
-								{
-									*TargetArchive << Zero8;
-									Padding--;
-								}
-
-								check(TargetArchive->Tell() == AlignedOffset);
-							}
-						}
-
-						const int64 BulkStartOffset = TargetArchive->Tell();
-
-						int64 StoredBulkStartOffset = BulkStartOffset - StartOfBulkDataArea;
-
-						BulkDataStorageInfo.BulkData->SerializeBulkData(*TargetArchive, BulkDataStorageInfo.BulkData->Lock(LOCK_READ_ONLY));
-
-						int64 BulkEndOffset = TargetArchive->Tell();
-						const int64 LinkerEndOffset = Linker->Tell();
-
-						int64 SizeOnDisk = BulkEndOffset - BulkStartOffset;
-				
-						Linker->Seek(BulkDataStorageInfo.BulkDataFlagsPos);
-						*Linker << ModifiedBulkDataFlags;
-
-						Linker->Seek(BulkDataStorageInfo.BulkDataOffsetInFilePos);
-						*Linker << StoredBulkStartOffset;
-
-						Linker->Seek(BulkDataStorageInfo.BulkDataSizeOnDiskPos);
-						if (ModifiedBulkDataFlags & BULKDATA_Size64Bit)
-						{
-							*Linker << SizeOnDisk;
-						}
-						else
-						{
-							check(SizeOnDisk < (1LL << 31));
-							int32 SizeOnDiskAsInt32 = SizeOnDisk;
-							*Linker << SizeOnDiskAsInt32;
-						}
-
-						Linker->Seek(LinkerEndOffset);
-
-						// Restore BulkData flags to before serialization started
-						BulkDataStorageInfo.BulkData->ClearBulkDataFlags(0xFFFFFFFF);
-						BulkDataStorageInfo.BulkData->SetBulkDataFlags(OldBulkDataFlags);
-						BulkDataStorageInfo.BulkData->Unlock();
-					}
-
-					if (BulkArchive)
-					{
-						check(OptionalBulkArchive);
-						check(MappedBulkArchive);
-
-						const bool bWriteBulkToDisk = !bDiffing;
-
-						if (SavePackageContext && bWriteBulkToDisk)
-						{
-							auto ToIoBuffer = [](FLargeMemoryWriter* Writer)
-							{
-								const int64 TotalSize = Writer->TotalSize();
-								return FIoBuffer(FIoBuffer::AssumeOwnership, Writer->ReleaseOwnership(), TotalSize);
-							};
-
-							FPackageStoreWriter::FBulkDataInfo BulkInfo;
-							BulkInfo.PackageName = InOuter->GetFName();
-							BulkInfo.LooseFilePath = Filename;
-							BulkInfo.BulkdataType = FPackageStoreWriter::FBulkDataInfo::Standard;
-
-							SavePackageContext->PackageStoreWriter.WriteBulkdata(BulkInfo, ToIoBuffer(BulkArchive.Get()));
-
-							BulkInfo.BulkdataType = FPackageStoreWriter::FBulkDataInfo::Optional;
-							SavePackageContext->PackageStoreWriter.WriteBulkdata(BulkInfo, ToIoBuffer(OptionalBulkArchive.Get()));
-
-							BulkInfo.BulkdataType = FPackageStoreWriter::FBulkDataInfo::Mmap;
-							SavePackageContext->PackageStoreWriter.WriteBulkdata(BulkInfo, ToIoBuffer(MappedBulkArchive.Get()));
-						}
-						else
-						{
-							auto WriteBulkData = [&](FLargeMemoryWriter* Archive, const TCHAR* BulkFileExtension)
-								{
-									if (const int64 DataSize = Archive->TotalSize())
-									{
-										if (bComputeHash)
-										{
-											CookedPackageHash.Update(Archive->GetData(), DataSize);
-										}
-
-										TotalPackageSizeUncompressed += DataSize;
-
-										if (bWriteBulkToDisk)
-										{
-											FLargeMemoryPtr DataPtr(Archive->ReleaseOwnership());
-
-											const FString ArchiveFilename = FPaths::ChangeExtension(Filename, BulkFileExtension);
-
-											AsyncWriteFile(MoveTemp(DataPtr), DataSize, *ArchiveFilename);
-										}
-									}
-								};
-
-							WriteBulkData(BulkArchive.Get(),			TEXT(".ubulk"));		// Regular separate bulk data file
-							WriteBulkData(OptionalBulkArchive.Get(),	TEXT(".uptnl"));		// Optional bulk data
-							WriteBulkData(MappedBulkArchive.Get(),		TEXT(".m.ubulk"));		// Memory-mapped bulk data
-						}
-					}
-				}
+				SlowTask.EnterProgressFrame(1, NSLOCTEXT("Core", "SerializingBulkData", "Serializing bulk data"));
 
-				Linker->BulkDataToAppend.Empty();
+				SaveBulkData(Linker.Get(), InOuter, Filename, TargetPlatform, SavePackageContext, bTextFormat, bDiffing,
+							 bComputeHash ? &CookedPackageHash : nullptr, TotalPackageSizeUncompressed );
 
 #if WITH_EDITOR
 				if (bIsCooking && AdditionalFilesFromExports.Num() > 0)
@@ -6348,6 +6130,234 @@ static void SaveAssetRegistryData(UPackage* InOuter, FLinkerSave* Linker, FStruc
 	}
 }
 
+void SaveBulkData(FLinkerSave* Linker, const UPackage* InOuter, const TCHAR* Filename, const ITargetPlatform* TargetPlatform,
+				  FSavePackageContext* SavePackageContext, const bool bTextFormat, const bool bDiffing, FMD5* CookedPackageHash, int64 TotalPackageSizeUncompressed)
+{
+	// Now we write all the bulkdata that is supposed to be at the end of the package
+	// and fix up the offset
+	const int64 StartOfBulkDataArea = Linker->Tell();
+	Linker->Summary.BulkDataStartOffset = StartOfBulkDataArea;
+
+	check(!bTextFormat || Linker->BulkDataToAppend.Num() == 0);
+
+	if (!bTextFormat && Linker->BulkDataToAppend.Num() > 0)
+	{
+		COOK_STAT(FScopedDurationTimer SaveTimer(SavePackageStats::SerializeBulkDataTimeSec));
+
+		FScopedSlowTask BulkDataFeedback(Linker->BulkDataToAppend.Num());
+
+		TUniquePtr<FLargeMemoryWriter> BulkArchive;
+		TUniquePtr<FLargeMemoryWriter> OptionalBulkArchive;
+		TUniquePtr<FLargeMemoryWriter> MappedBulkArchive;
+
+		uint32 ExtraBulkDataFlags = 0;
+
+		static const struct FUseSeparateBulkDataFiles
+		{
+			bool bEnable = false;
+
+			FUseSeparateBulkDataFiles()
+			{
+				GConfig->GetBool(TEXT("Core.System"), TEXT("UseSeperateBulkDataFiles"), /* out */ bEnable, GEngineIni);
+
+				if (IsEventDrivenLoaderEnabledInCookedBuilds())
+				{
+					// Always split bulk data when splitting cooked files
+					bEnable = true;
+				}
+			}
+		} ShouldUseSeparateBulkDataFiles;
+
+		const bool bShouldUseSeparateBulkFile = ShouldUseSeparateBulkDataFiles.bEnable && Linker->IsCooking();
+
+		if (bShouldUseSeparateBulkFile)
+		{
+			ExtraBulkDataFlags = BULKDATA_PayloadInSeperateFile;
+
+			BulkArchive.Reset(new FLargeMemoryWriter(0, /* IsPersistent */ true));
+			OptionalBulkArchive.Reset(new FLargeMemoryWriter(0, /* IsPersistent */ true));
+			MappedBulkArchive.Reset(new FLargeMemoryWriter(0, /* IsPersistent */ true));
+		}
+
+		bool bAlignBulkData = false;
+		int64 BulkDataAlignment = 0;
+
+		if (TargetPlatform)
+		{
+			bAlignBulkData = TargetPlatform->SupportsFeature(ETargetPlatformFeatures::MemoryMappedFiles);
+			BulkDataAlignment = TargetPlatform->GetMemoryMappingAlignment();
+		}
+
+		for (FLinkerSave::FBulkDataStorageInfo& BulkDataStorageInfo : Linker->BulkDataToAppend)
+		{
+			BulkDataFeedback.EnterProgressFrame();
+
+			// Set bulk data flags to what they were during initial serialization (they might have changed after that)
+			const uint32 OldBulkDataFlags = BulkDataStorageInfo.BulkData->GetBulkDataFlags();
+			uint32 ModifiedBulkDataFlags = BulkDataStorageInfo.BulkDataFlags | ExtraBulkDataFlags;
+			const bool bBulkItemIsOptional = (ModifiedBulkDataFlags & BULKDATA_OptionalPayload) != 0;
+			bool bBulkItemIsMapped = bAlignBulkData && ((ModifiedBulkDataFlags & BULKDATA_MemoryMappedPayload) != 0);
+
+			if (bBulkItemIsMapped && bBulkItemIsOptional)
+			{
+				UE_LOG(LogSavePackage, Warning, TEXT("%s has bulk data that is both mapped and optional. This is not currently supported. Will not be mapped."), Filename);
+				ModifiedBulkDataFlags &= ~BULKDATA_MemoryMappedPayload;
+				bBulkItemIsMapped = false;
+			}
+
+			BulkDataStorageInfo.BulkData->ClearBulkDataFlags(0xFFFFFFFF);
+			BulkDataStorageInfo.BulkData->SetBulkDataFlags(ModifiedBulkDataFlags);
+
+			FArchive* TargetArchive = Linker;
+
+			if (bShouldUseSeparateBulkFile)
+			{
+				check(MappedBulkArchive && OptionalBulkArchive && BulkArchive);
+
+				if (bBulkItemIsOptional)
+				{
+					TargetArchive = OptionalBulkArchive.Get();
+				}
+				else if (bBulkItemIsMapped)
+				{
+					TargetArchive = MappedBulkArchive.Get();
+				}
+				else
+				{
+					TargetArchive = BulkArchive.Get();
+				}
+			}
+
+			// Pad archive for proper alignment for memory mapping
+			if (bBulkItemIsMapped && BulkDataAlignment > 0)
+			{
+				const int64 BulkStartOffset = TargetArchive->Tell();
+
+				if (!IsAligned(BulkStartOffset, BulkDataAlignment))
+				{
+					const int64 AlignedOffset = Align(BulkStartOffset, BulkDataAlignment);
+
+					int64 Padding = AlignedOffset - BulkStartOffset;
+					check(Padding > 0);
+
+					uint64 Zero64 = 0;
+					while (Padding >= 8)
+					{
+						*TargetArchive << Zero64;
+						Padding -= 8;
+					}
+
+					uint8 Zero8 = 0;
+					while (Padding > 0)
+					{
+						*TargetArchive << Zero8;
+						Padding--;
+					}
+
+					check(TargetArchive->Tell() == AlignedOffset);
+				}
+			}
+
+			const int64 BulkStartOffset = TargetArchive->Tell();
+
+			int64 StoredBulkStartOffset = BulkStartOffset - StartOfBulkDataArea;
+
+			BulkDataStorageInfo.BulkData->SerializeBulkData(*TargetArchive, BulkDataStorageInfo.BulkData->Lock(LOCK_READ_ONLY));
+
+			int64 BulkEndOffset = TargetArchive->Tell();
+			const int64 LinkerEndOffset = Linker->Tell();
+
+			int64 SizeOnDisk = BulkEndOffset - BulkStartOffset;
+
+			Linker->Seek(BulkDataStorageInfo.BulkDataFlagsPos);
+			*Linker << ModifiedBulkDataFlags;
+
+			Linker->Seek(BulkDataStorageInfo.BulkDataOffsetInFilePos);
+			*Linker << StoredBulkStartOffset;
+
+			Linker->Seek(BulkDataStorageInfo.BulkDataSizeOnDiskPos);
+			if (ModifiedBulkDataFlags & BULKDATA_Size64Bit)
+			{
+				*Linker << SizeOnDisk;
+			}
+			else
+			{
+				check(SizeOnDisk < (1LL << 31));
+				int32 SizeOnDiskAsInt32 = SizeOnDisk;
+				*Linker << SizeOnDiskAsInt32;
+			}
+
+			Linker->Seek(LinkerEndOffset);
+
+			// Restore BulkData flags to before serialization started
+			BulkDataStorageInfo.BulkData->ClearBulkDataFlags(0xFFFFFFFF);
+			BulkDataStorageInfo.BulkData->SetBulkDataFlags(OldBulkDataFlags);
+			BulkDataStorageInfo.BulkData->Unlock();
+		}
+
+		if (BulkArchive)
+		{
+			check(OptionalBulkArchive);
+			check(MappedBulkArchive);
+
+			const bool bWriteBulkToDisk = !bDiffing;
+
+			if (SavePackageContext && bWriteBulkToDisk)
+			{
+				auto ToIoBuffer = [](FLargeMemoryWriter* Writer)
+				{
+					const int64 TotalSize = Writer->TotalSize();
+					return FIoBuffer(FIoBuffer::AssumeOwnership, Writer->ReleaseOwnership(), TotalSize);
+				};
+
+				FPackageStoreWriter::FBulkDataInfo BulkInfo;
+				BulkInfo.PackageName = InOuter->GetFName();
+				BulkInfo.LooseFilePath = Filename;
+				BulkInfo.BulkdataType = FPackageStoreWriter::FBulkDataInfo::Standard;
+
+				SavePackageContext->PackageStoreWriter.WriteBulkdata(BulkInfo, ToIoBuffer(BulkArchive.Get()));
+
+				BulkInfo.BulkdataType = FPackageStoreWriter::FBulkDataInfo::Optional;
+				SavePackageContext->PackageStoreWriter.WriteBulkdata(BulkInfo, ToIoBuffer(OptionalBulkArchive.Get()));
+
+				BulkInfo.BulkdataType = FPackageStoreWriter::FBulkDataInfo::Mmap;
+				SavePackageContext->PackageStoreWriter.WriteBulkdata(BulkInfo, ToIoBuffer(MappedBulkArchive.Get()));
+			}
+			else
+			{
+				auto WriteBulkData = [&](FLargeMemoryWriter* Archive, const TCHAR* BulkFileExtension)
+				{
+					if (const int64 DataSize = Archive->TotalSize())
+					{
+						if (CookedPackageHash != nullptr)
+						{
+							CookedPackageHash->Update(Archive->GetData(), DataSize);
+						}
+
+						TotalPackageSizeUncompressed += DataSize;
+
+						if (bWriteBulkToDisk)
+						{
+							FLargeMemoryPtr DataPtr(Archive->ReleaseOwnership());
+
+							const FString ArchiveFilename = FPaths::ChangeExtension(Filename, BulkFileExtension);
+
+							AsyncWriteFile(MoveTemp(DataPtr), DataSize, *ArchiveFilename);
+						}
+					}
+				};
+
+				WriteBulkData(BulkArchive.Get(), TEXT(".ubulk"));			// Regular separate bulk data file
+				WriteBulkData(OptionalBulkArchive.Get(), TEXT(".uptnl"));	// Optional bulk data
+				WriteBulkData(MappedBulkArchive.Get(), TEXT(".m.ubulk"));	// Memory-mapped bulk data
+				WriteBulkData(MappedBulkArchive.Get(), TEXT(".m.ubulk"));	// Memory-mapped bulk data
+			}
+		}
+	}
+
+	Linker->BulkDataToAppend.Empty();
+}
+
 static void SaveWorldLevelInfo(UPackage* InOuter, FLinkerSave* Linker, FStructuredArchive::FRecord Record)
 {
 	Linker->Summary.WorldTileInfoDataOffset = 0;
