[AsyncLoader2] Removed copy of package file summary buffer

diff --git a/Engine/Source/Runtime/CoreUObject/Private/Serialization/AsyncLoading2.cpp b/Engine/Source/Runtime/CoreUObject/Private/Serialization/AsyncLoading2.cpp
--- a/Engine/Source/Runtime/CoreUObject/Private/Serialization/AsyncLoading2.cpp
+++ b/Engine/Source/Runtime/CoreUObject/Private/Serialization/AsyncLoading2.cpp
@@ -961,7 +961,7 @@ private:
 	uint32 ImportNodeCount = 0;
 	uint32 ExportNodeCount = 0;
 	
-	TUniquePtr<uint8[]> PackageSummaryBuffer;
+	FIoBuffer PackageSummaryIoBuffer;
 	TArray<FIoBuffer> ExportIoBuffers;
 
 	// FZenLinkerLoad
@@ -2155,9 +2155,9 @@ EAsyncPackageState::Type FAsyncPackage2::Event_StartImportPackages(FAsyncPackage
 		TRACE_CPUPROFILER_EVENT_SCOPE(SetupSerializedArcs);
 		const FGlobalNameMap& GlobalNameMap = Package->AsyncLoadingThread.GlobalNameMap;
 
-		const FPackageSummary* PackageSummary = reinterpret_cast<const FPackageSummary*>(Package->PackageSummaryBuffer.Get());
+		const FPackageSummary* PackageSummary = reinterpret_cast<const FPackageSummary*>(Package->PackageSummaryIoBuffer.Data());
 		uint64 ZenHeaderDataSize = PackageSummary->GraphDataSize;
-		const uint8* ZenHeaderData = Package->PackageSummaryBuffer.Get() + PackageSummary->GraphDataOffset;
+		const uint8* ZenHeaderData = Package->PackageSummaryIoBuffer.Data() + PackageSummary->GraphDataOffset;
 		FSimpleArchive ZenHeaderArchive(ZenHeaderData, ZenHeaderDataSize);
 		int32 InternalArcCount;
 		ZenHeaderArchive << InternalArcCount;
@@ -3148,8 +3148,7 @@ void FAsyncPackage2::ProcessIoRequest(const FIoChunkId& ChunkId, TIoStatusOr<FIo
 	{
 	case EChunkType::PackageSummary:
 	{
-		PackageSummaryBuffer = MakeUnique<uint8[]>(IoBuffer.DataSize());
-		FMemory::Memcpy(PackageSummaryBuffer.Get(), IoBuffer.Data(), IoBuffer.DataSize());
+		PackageSummaryIoBuffer = IoBuffer;
 		GetNode(EEventLoadNode2::Package_LoadSummary)->ReleaseBarrier();
 		break;
 	}
@@ -4471,7 +4470,7 @@ EAsyncPackageState::Type FAsyncPackage2::FinishLinker()
 	SCOPED_LOADTIMER(LinkerLoad_FinalizeCreation);
 
 	const FGlobalNameMap& GlobalNameMap = AsyncLoadingThread.GlobalNameMap;
-	const FPackageSummary* PackageSummary = reinterpret_cast<const FPackageSummary*>(PackageSummaryBuffer.Get());
+	const FPackageSummary* PackageSummary = reinterpret_cast<const FPackageSummary*>(PackageSummaryIoBuffer.Data());
 	const int32 ExportCount = AsyncLoadingThread.GetPackageExportCount(GlobalPackageId);
 	check(ExportCount);
 
@@ -4487,7 +4486,7 @@ EAsyncPackageState::Type FAsyncPackage2::FinishLinker()
 
 	{
 		SCOPED_LOADTIMER(LinkerLoad_SerializeExportMap);
-		const uint8* ZenExportData = PackageSummaryBuffer.Get() + PackageSummary->ExportMapOffset;
+		const uint8* ZenExportData = PackageSummaryIoBuffer.Data() + PackageSummary->ExportMapOffset;
 		FSimpleArchive ZenExportMapArchive(ZenExportData, PackageSummary->ExportMapSize);
 
 		ExportMap.AddZeroed(ExportCount);
