Fix FBulkdata2 to work with loose cooked files after recent changes to FLinker ownership.

diff --git a/Engine/Source/Runtime/CoreUObject/Private/Serialization/BulkData2.cpp b/Engine/Source/Runtime/CoreUObject/Private/Serialization/BulkData2.cpp
--- a/Engine/Source/Runtime/CoreUObject/Private/Serialization/BulkData2.cpp
+++ b/Engine/Source/Runtime/CoreUObject/Private/Serialization/BulkData2.cpp
@@ -217,21 +217,25 @@ void FBulkDataBase::Serialize(FArchive& Ar, UObject* Owner, int32 /*Index*/, boo
 
 		Ar << BulkDataOffsetInFile;
 
-		// fix up the file offset, but only if not stored inline
-		if (Owner != nullptr && Owner->GetLinker() != nullptr && !IsInlined())
+		// Assuming that Owner/Package/Linker are all valid, the old BulkData system would
+		// generally fail if any of these were nullptr but had plenty of inconsistent checks
+		// scattered throughout.
+		check(Owner != nullptr);
+		UPackage* Package = Owner->GetOutermost();
+		check(Package != nullptr);
+
+		FLinkerLoad* Linker = FLinkerLoad::FindExistingLinkerForPackage(Package);
+		check(Linker != nullptr);
+
+		// Fix up the file offset, but only if not stored inline (TODO: Look into getting rid of this and cook with the correct values?)
+		if (!IsInlined())
 		{
-			BulkDataOffsetInFile += Owner->GetLinker()->Summary.BulkDataStartOffset;
+			BulkDataOffsetInFile += Linker->Summary.BulkDataStartOffset;
 		}
 
 		FArchive* CacheableArchive = Ar.GetCacheableArchive();
-		if (Ar.IsAllowingLazyLoading() && Owner != nullptr && CacheableArchive != nullptr)
+		if (Ar.IsAllowingLazyLoading() && CacheableArchive != nullptr)
 		{
-			UPackage* Package = Owner->GetOutermost();
-			check(Package != nullptr);
-
-			FLinkerLoad* Linker = FLinkerLoad::FindExistingLinkerForPackage(Package);
-			check(Linker != nullptr);
-
 			FString Filename = ConvertFilenameFromFlags(Linker->Filename);
 
 			if (IsInlined())
