Optimize IsEventDrivenLoaderEnabledInCookedBuilds for commandlets like the cooker

diff --git a/Engine/Source/Runtime/CoreUObject/Private/Serialization/AsyncPackageLoader.cpp b/Engine/Source/Runtime/CoreUObject/Private/Serialization/AsyncPackageLoader.cpp
--- a/Engine/Source/Runtime/CoreUObject/Private/Serialization/AsyncPackageLoader.cpp
+++ b/Engine/Source/Runtime/CoreUObject/Private/Serialization/AsyncPackageLoader.cpp
@@ -419,7 +419,7 @@ bool IsEventDrivenLoaderEnabledInCookedBuilds()
 
 #if WITH_EDITOR	
 	// when building from the UE4 Editor, s.EventDrivenLoaderEnabled can be changed from Project Settings, so we need to test it at every call
-	if (GIsEditor)
+	if (GIsEditor && !IsRunningCommandlet())
 	{
 		EventDrivenLoaderEnabledInCookedBuilds.SetEventDrivenLoaderEnabled();
 	}
