noMaster fix

diff --git a/Engine/Source/Runtime/CoreUObject/Private/Serialization/LoadTimeTracePrivate.h b/Engine/Source/Runtime/CoreUObject/Private/Serialization/LoadTimeTracePrivate.h
--- a/Engine/Source/Runtime/CoreUObject/Private/Serialization/LoadTimeTracePrivate.h
+++ b/Engine/Source/Runtime/CoreUObject/Private/Serialization/LoadTimeTracePrivate.h
@@ -7,6 +7,7 @@
 #if LOADTIMEPROFILERTRACE_ENABLED
 
 class UObject;
+class FGCObject;
 
 struct FLoadTimeProfilerTracePrivate
 {
