Fixed issue with an incorrect CustomVersion container being propagated in AsyncLoading2

diff --git a/Engine/Source/Runtime/CoreUObject/Private/Serialization/AsyncLoading2.cpp b/Engine/Source/Runtime/CoreUObject/Private/Serialization/AsyncLoading2.cpp
--- a/Engine/Source/Runtime/CoreUObject/Private/Serialization/AsyncLoading2.cpp
+++ b/Engine/Source/Runtime/CoreUObject/Private/Serialization/AsyncLoading2.cpp
@@ -2946,7 +2946,6 @@ void FAsyncPackage2::EventDrivenSerializeExport(int32 LocalExportIndex)
 		Ar.SetUE4Ver(Summary.GetFileVersionUE4());
 		Ar.SetLicenseeUE4Ver(Summary.GetFileVersionLicenseeUE4());
 		Ar.SetEngineVer(Summary.SavedByEngineVersion);
-		Ar.SetCustomVersions(SummaryVersions);
 		Ar.SetIsLoading(true);
 		Ar.SetIsPersistent(true);
 		if (Summary.PackageFlags & PKG_FilterEditorOnly)
