AsyncLoading2: Skip package load requests for non-existing packages

diff --git a/Engine/Source/Runtime/CoreUObject/Private/Serialization/AsyncLoading2.cpp b/Engine/Source/Runtime/CoreUObject/Private/Serialization/AsyncLoading2.cpp
--- a/Engine/Source/Runtime/CoreUObject/Private/Serialization/AsyncLoading2.cpp
+++ b/Engine/Source/Runtime/CoreUObject/Private/Serialization/AsyncLoading2.cpp
@@ -1735,6 +1735,11 @@ void FAsyncLoadingThread2Impl::InitializeLoading()
 void FAsyncLoadingThread2Impl::QueuePackage(FAsyncPackageDesc& Package)
 {
 	//TRACE_CPUPROFILER_EVENT_SCOPE(QueuePackage);
+	if (!PackageNameToGlobalPackageId.Find(Package.NameToLoad))
+	{
+		UE_LOG(LogStreaming, Warning, TEXT("QueuePackage: Skipping unknown package: %s"), *Package.NameToLoad.ToString());
+		return;
+	}
 	{
 		FScopeLock QueueLock(&QueueCritical);
 		++QueuedPackagesCounter;
