AsyncLoading2: Simplify and optimize FinishExternalReadDependencies and PostLoad processing
Track and process external read dependencies in the async package instead of in the linker.
Remove the post load related arrays DeferredPostLoadObjects, DeferredFinalizeObjects, PackageObjLoaded.
Remove use of LoadContext->PRIVATE_GetObjectsLoadedInternalUseOnly.

diff --git a/Engine/Source/Runtime/CoreUObject/Private/Serialization/AsyncLoading2.cpp b/Engine/Source/Runtime/CoreUObject/Private/Serialization/AsyncLoading2.cpp
--- a/Engine/Source/Runtime/CoreUObject/Private/Serialization/AsyncLoading2.cpp
+++ b/Engine/Source/Runtime/CoreUObject/Private/Serialization/AsyncLoading2.cpp
@@ -867,7 +867,7 @@ struct FAsyncPackage2
 private:
 
 	/** Checks if all dependencies (imported packages) of this package have been fully loaded */
-	static bool AreAllDependenciesFullyLoadedInternal(FAsyncPackage2* Package, TSet<UPackage*>& VisitedPackages, FString& OutError);
+	bool AreAllDependenciesFullyLoadedInternal(FAsyncPackage2* Package, TSet<UPackage*>& VisitedPackages, FString& OutError);
 
 	struct FCompletionCallback
 	{
@@ -922,12 +922,6 @@ private:
 	double						LoadStartTime;
 	/** Estimated load percentage.																		*/
 	float						LoadPercentage;
-	/** Objects to be post loaded on the game thread */
-	TArray<UObject*> DeferredPostLoadObjects;
-	/** Objects to be finalized on the game thread */
-	TArray<UObject*> DeferredFinalizeObjects;
-	/** Objects loaded while loading this package */
-	TArray<UObject*> PackageObjLoaded;
 	/** Packages that were loaded synchronously while async loading this package or packages added by verify import */
 	TArray<FLinkerLoad*> DelayedLinkerClosePackages;
 	/** Objects to create GC clusters from */
@@ -967,6 +961,7 @@ private:
 
 	// FZenLinkerLoad
 	TArray<FObjectExport> ExportMap;
+	TArray<FExternalReadCallback> ExternalReadDependencies;
 	int32 GlobalImportCount = 0;
 	int32 LocalImportCount = 0;
 	int32* LocalImportIndices = nullptr;
@@ -2964,7 +2959,7 @@ void FAsyncPackage2::EventDrivenSerializeExport(int32 LocalExportIndex)
 		Ar.LocalImportIndices = LocalImportIndices;
 		Ar.GlobalImportObjects = GlobalImportObjects;
 		Ar.ExportMap = &ExportMap;
-		Ar.ExternalReadDependencies = &Linker->ExternalReadDependencies;
+		Ar.ExternalReadDependencies = &ExternalReadDependencies;
 
 		Object->ClearFlags(RF_NeedLoad);
 
@@ -3044,20 +3039,6 @@ EAsyncPackageState::Type FAsyncPackage2::Event_StartPostload(FAsyncPackage2* Pac
 	LLM_SCOPE(ELLMTag::AsyncLoading);
 
 	check(Package->bAllExportsSerialized);
-	{
-		check(Package->PackageObjLoaded.Num() == 0);
-		Package->PackageObjLoaded.Reserve(Package->ExportMap.Num());
-		for (int32 LocalExportIndex = 0; LocalExportIndex < Package->ExportMap.Num(); LocalExportIndex++)
-		{
-			FObjectExport& Export = Package->ExportMap[LocalExportIndex];
-			UObject* Object = Export.Object;
-			if (Object && (Object->HasAnyFlags(RF_NeedPostLoad) || Package->Linker->bDynamicClassLinker || Object->HasAnyInternalFlags(EInternalObjectFlags::AsyncLoading)))
-			{
-				check(Object->IsValidLowLevelFast());
-				Package->PackageObjLoaded.Add(Object);
-			}
-		}
-	}
 	Package->EventDrivenLoadingComplete();
 	return EAsyncPackageState::Complete;
 }
@@ -3303,36 +3284,33 @@ bool FAsyncPackage2::AreAllDependenciesFullyLoadedInternal(FAsyncPackage2* Packa
 		}
 		VisitedPackages.Add(ImportPackage);
 
-		if (FLinkerLoad* ImportPackageLinker = FLinkerLoad::FindExistingLinkerForPackage(ImportPackage))
+		FAsyncPackage2* AsyncRoot = AsyncLoadingThread.FindAsyncPackage(ImportPackage->GetFName());
+		if (AsyncRoot)
 		{
-			if (ImportPackageLinker->AsyncRoot)
+			if (!AsyncRoot->bAllExportsSerialized)
 			{
-				FAsyncPackage2* AsyncRoot = (FAsyncPackage2*)ImportPackageLinker->AsyncRoot;
-				if (!AsyncRoot->bAllExportsSerialized)
-				{
-					OutError = FString::Printf(TEXT("%s Doesn't have all exports Serialized"), *Package->GetPackageName().ToString());
-					return false;
-				}
-				if (AsyncRoot->DeferredPostLoadIndex < AsyncRoot->DeferredPostLoadObjects.Num())
+				OutError = FString::Printf(TEXT("%s Doesn't have all exports Serialized"), *Package->GetPackageName().ToString());
+				return false;
+			}
+			if (AsyncRoot->DeferredPostLoadIndex < AsyncRoot->ExportMap.Num())
+			{
+				OutError = FString::Printf(TEXT("%s Doesn't have all objects processed by DeferredPostLoad"), *Package->GetPackageName().ToString());
+				return false;
+			}
+			for (FObjectExport& Export : AsyncRoot->ExportMap)
+			{
+				if (Export.Object->HasAnyFlags(RF_NeedPostLoad|RF_NeedLoad))
 				{
-					OutError = FString::Printf(TEXT("%s Doesn't have all objects processed by DeferredPostLoad"), *Package->GetPackageName().ToString());
+					OutError = FString::Printf(TEXT("%s has not been %s"), *Export.Object->GetFullName(), 
+						Export.Object->HasAnyFlags(RF_NeedLoad) ? TEXT("Serialized") : TEXT("PostLoaded"));
 					return false;
 				}
-				for (FObjectExport& Export : AsyncRoot->ExportMap)
-				{
-					if (Export.Object && Export.Object->HasAnyFlags(RF_NeedPostLoad|RF_NeedLoad))
-					{
-						OutError = FString::Printf(TEXT("%s has not been %s"), *Export.Object->GetFullName(), 
-							Export.Object->HasAnyFlags(RF_NeedLoad) ? TEXT("Serialized") : TEXT("PostLoaded"));
-						return false;
-					}
-				}
+			}
 
-				if (!AreAllDependenciesFullyLoadedInternal(AsyncRoot, VisitedPackages, OutError))
-				{
-					OutError = Package->GetPackageName().ToString() + TEXT("->") + OutError;
-					return false;
-				}
+			if (!AreAllDependenciesFullyLoadedInternal(AsyncRoot, VisitedPackages, OutError))
+			{
+				OutError = Package->GetPackageName().ToString() + TEXT("->") + OutError;
+				return false;
 			}
 		}
 	}
@@ -4354,29 +4332,11 @@ void FAsyncPackage2::DetachLinker()
 
 void FAsyncPackage2::FlushObjectLinkerCache()
 {
-	for (UObject* Obj : PackageObjLoaded)
-	{
-		if (Obj)
-		{
-			FLinkerLoad* ObjLinker = Obj->GetLinker();
-			if (ObjLinker)
-			{
-				ObjLinker->FlushCache();
-			}
-		}
 	}
-}
 
 #if WITH_EDITOR 
 void FAsyncPackage2::GetLoadedAssets(TArray<FWeakObjectPtr>& AssetList)
 {
-	for (UObject* Obj : PackageObjLoaded)
-	{
-		if (Obj && !Obj->IsPendingKill() && Obj->IsAsset())
-		{
-			AssetList.AddUnique(Obj);
-		}
-	}
 }
 #endif
 
@@ -4534,29 +4494,11 @@ EAsyncPackageState::Type FAsyncPackage2::FinishLinker()
 
 EAsyncPackageState::Type FAsyncPackage2::FinishExternalReadDependencies()
 {
-	if (FAsyncLoadingThreadState2::Get().IsTimeLimitExceeded())
+	for (FExternalReadCallback& ReadCallback : ExternalReadDependencies)
 	{
-		return EAsyncPackageState::TimeOut;
+		ReadCallback(0.0);
 	}
-
-	FLinkerLoad* VisitedLinkerLoad = nullptr;
-	while (FinishExternalReadDependenciesIndex < PackageObjLoaded.Num())
-	{
-		UObject* Obj = PackageObjLoaded[FinishExternalReadDependenciesIndex];
-		//FLinkerLoad* LinkerLoad = Obj ? Obj->GetLinker() : nullptr;
-		UPackage* Package = Obj ? Obj->GetOutermost() : nullptr;
-		FLinkerLoad* LinkerLoad = Package ? FLinkerLoad::FindExistingLinkerForPackage(Package) : nullptr;
-		if (LinkerLoad && LinkerLoad != VisitedLinkerLoad)
-		{
-			if (!LinkerLoad->FinishExternalReadDependencies(0.0) || FAsyncLoadingThreadState2::Get().IsTimeLimitExceeded())
-			{
-				return EAsyncPackageState::TimeOut;
-			}
-			VisitedLinkerLoad = LinkerLoad;
-		}
-		FinishExternalReadDependenciesIndex++;
-	}
-		
+	ExternalReadDependencies.Empty();
 	return EAsyncPackageState::Complete;
 }
 
@@ -4575,67 +4517,37 @@ EAsyncPackageState::Type FAsyncPackage2::PostLoadObjects()
 	TGuardValue<bool> GuardIsRoutingPostLoad(ThreadContext.IsRoutingPostLoad, true);
 
 	FUObjectSerializeContext* LoadContext = GetSerializeContext();
-	TArray<UObject*>& ThreadObjLoaded = LoadContext->PRIVATE_GetObjectsLoadedInternalUseOnly();
-	if (ThreadObjLoaded.Num())
-	{
-		// New objects have been loaded. They need to go through PreLoad first so exit now and come back after they've been preloaded.
-		PackageObjLoaded.Append(ThreadObjLoaded);
-		ThreadObjLoaded.Reset();
-		return EAsyncPackageState::TimeOut;
-	}
 
-	int32 PreLoadIndex = PackageObjLoaded.Num();
-	
 	const bool bAsyncPostLoadEnabled = true;
 	const bool bIsMultithreaded = AsyncLoadingThread.IsMultithreaded();
 
 	// PostLoad objects.
-	while (PostLoadIndex < PackageObjLoaded.Num() && PostLoadIndex < PreLoadIndex && !FAsyncLoadingThreadState2::Get().IsTimeLimitExceeded())
+	while (PostLoadIndex < ExportMap.Num() && !FAsyncLoadingThreadState2::Get().IsTimeLimitExceeded())
 	{
-		UObject* Object = PackageObjLoaded[PostLoadIndex++];
-		if (Object)
-		{
-			if (!Object->IsReadyForAsyncPostLoad())
-			{
-				--PostLoadIndex;
-				break;
-			}
-			else if (!bIsMultithreaded || (bAsyncPostLoadEnabled && CanPostLoadOnAsyncLoadingThread(Object)))
-			{
-				check(!Object->HasAnyFlags(RF_NeedLoad));
+		FObjectExport& Export = ExportMap[PostLoadIndex++];
+		UObject* Object = Export.Object;
 
-				ThreadContext.CurrentlyPostLoadedObjectByALT = Object;
-				{
-					TRACE_LOADTIME_POSTLOAD_EXPORT_SCOPE(Object);
-					Object->ConditionalPostLoad();
-					Object->AtomicallyClearInternalFlags(EInternalObjectFlags::AsyncLoading);
-				}
-				ThreadContext.CurrentlyPostLoadedObjectByALT = nullptr;
+		check(Object);
+		check(!Object->HasAnyFlags(RF_NeedLoad));
+		if (!Object->HasAnyFlags(RF_NeedPostLoad))
+		{
+			continue;
+		}
 
-				if (ThreadObjLoaded.Num())
-				{
-					// New objects have been loaded. They need to go through PreLoad first so exit now and come back after they've been preloaded.
-					PackageObjLoaded.Append(ThreadObjLoaded);
-					ThreadObjLoaded.Reset();
-					return EAsyncPackageState::TimeOut;
-				}
-			}
-			else
+		check(Object->IsReadyForAsyncPostLoad());
+		if (!bIsMultithreaded || (bAsyncPostLoadEnabled && CanPostLoadOnAsyncLoadingThread(Object)))
+		{
+			ThreadContext.CurrentlyPostLoadedObjectByALT = Object;
 			{
-				DeferredPostLoadObjects.Add(Object);
+				TRACE_LOADTIME_POSTLOAD_EXPORT_SCOPE(Object);
+				Object->ConditionalPostLoad();
+				Object->AtomicallyClearInternalFlags(EInternalObjectFlags::AsyncLoading);
 			}
-			// All object must be finalized on the game thread
-			DeferredFinalizeObjects.Add(Object);
-			check(Object->IsValidLowLevelFast());
+			ThreadContext.CurrentlyPostLoadedObjectByALT = nullptr;
 		}
 	}
 
-	PackageObjLoaded.Append(ThreadObjLoaded);
-	ThreadObjLoaded.Reset();
-
-	// New objects might have been loaded during PostLoad.
-	EAsyncPackageState::Type Result = (PreLoadIndex == PackageObjLoaded.Num()) && (PostLoadIndex == PackageObjLoaded.Num()) ? EAsyncPackageState::Complete : EAsyncPackageState::TimeOut;
-	return Result;
+	return (PostLoadIndex == ExportMap.Num()) ? EAsyncPackageState::Complete : EAsyncPackageState::TimeOut;
 }
 
 void CreateClustersFromPackage(FLinkerLoad* PackageLinker, TArray<UObject*>& OutClusterObjects);
@@ -4651,86 +4563,42 @@ EAsyncPackageState::Type FAsyncPackage2::PostLoadDeferredObjects()
 	FAsyncLoadingTickScope2 InAsyncLoadingTick(AsyncLoadingThread);
 
 	FUObjectSerializeContext* LoadContext = GetSerializeContext();
-	TArray<UObject*>& ObjLoadedInPostLoad = LoadContext->PRIVATE_GetObjectsLoadedInternalUseOnly();
-	TArray<UObject*> ObjLoadedInPostLoadLocal;
 
 	STAT(double PostLoadStartTime = FPlatformTime::Seconds());
 
-	while (DeferredPostLoadIndex < DeferredPostLoadObjects.Num() && 
+	while (DeferredPostLoadIndex < ExportMap.Num() && 
 		!AsyncLoadingThread.IsAsyncLoadingSuspended() &&
 		!FAsyncLoadingThreadState2::Get().IsTimeLimitExceeded())
 	{
-		UObject* Object = DeferredPostLoadObjects[DeferredPostLoadIndex++];
-		
-		check(Object);
-
-		PackageScope.ThreadContext.CurrentlyPostLoadedObjectByALT = Object;
-		{
-			TRACE_LOADTIME_POSTLOAD_EXPORT_SCOPE(Object);
-			Object->ConditionalPostLoad();
-		}
-		PackageScope.ThreadContext.CurrentlyPostLoadedObjectByALT = nullptr;
+		FObjectExport& Export = ExportMap[DeferredPostLoadIndex++];
+		UObject* Object = Export.Object;
 
-		if (ObjLoadedInPostLoad.Num())
+		check(Object);
+		check(!Object->HasAnyFlags(RF_NeedLoad));
+		if (Object->HasAnyFlags(RF_NeedPostLoad))
 		{
-			// If there were any LoadObject calls inside of PostLoad, we need to pre-load those objects here. 
-			// There's no going back to the async tick loop from here.
-			UE_LOG(LogStreaming, Warning, TEXT("Detected %d objects loaded in PostLoad while streaming, this may cause hitches as we're blocking async loading to pre-load them."), ObjLoadedInPostLoad.Num());
-			
-			// Copy to local array because ObjLoadedInPostLoad can change while we're iterating over it
-			ObjLoadedInPostLoadLocal.Append(ObjLoadedInPostLoad);
-			ObjLoadedInPostLoad.Reset();
-
-			while (ObjLoadedInPostLoadLocal.Num())
+			PackageScope.ThreadContext.CurrentlyPostLoadedObjectByALT = Object;
 			{
-				// Make sure all objects loaded in PostLoad get post-loaded too
-				DeferredPostLoadObjects.Append(ObjLoadedInPostLoadLocal);
-
-				// Preload (aka serialize) the objects loaded in PostLoad.
-				for (UObject* PreLoadObject : ObjLoadedInPostLoadLocal)
-				{
-					if (PreLoadObject && PreLoadObject->GetLinker())
-					{
-						PreLoadObject->GetLinker()->Preload(PreLoadObject);
-					}
-				}
-
-				// Other objects could've been loaded while we were preloading, continue until we've processed all of them.
-				ObjLoadedInPostLoadLocal.Reset();
-				ObjLoadedInPostLoadLocal.Append(ObjLoadedInPostLoad);
-				ObjLoadedInPostLoad.Reset();
-			}			
+				TRACE_LOADTIME_POSTLOAD_EXPORT_SCOPE(Object);
+				Object->ConditionalPostLoad();
+			}
+			PackageScope.ThreadContext.CurrentlyPostLoadedObjectByALT = nullptr;
 		}
-
-		UpdateLoadPercentage();
+		Object->AtomicallyClearInternalFlags(EInternalObjectFlags::AsyncLoading);
 	}
 
-	// New objects might have been loaded during PostLoad.
-	Result = (DeferredPostLoadIndex == DeferredPostLoadObjects.Num()) ? EAsyncPackageState::Complete : EAsyncPackageState::TimeOut;
-	/*if (Result == EAsyncPackageState::Complete)
-	{
-		for (int32 LocalExportIndex = 0; LocalExportIndex < Linker->ExportMap.Num(); LocalExportIndex++)
-		{
-			UObject* Object = Linker->ExportMap[LocalExportIndex].Object;
-			if (Object && Object->HasAnyFlags(RF_NeedLoad | RF_NeedPostLoad))
-			{
-				Result = EAsyncPackageState::TimeOut;
-			}
-		}
-	}*/
-	
+	Result = (DeferredPostLoadIndex == ExportMap.Num()) ? EAsyncPackageState::Complete : EAsyncPackageState::TimeOut;
+
 	if (Result == EAsyncPackageState::Complete)
 	{
 		TArray<UObject*> CDODefaultSubobjects;
 		// Clear async loading flags (we still want RF_Async, but EInternalObjectFlags::AsyncLoading can be cleared)
-		while (DeferredFinalizeIndex < DeferredFinalizeObjects.Num() &&
-			(DeferredPostLoadIndex % 100 != 0 || (!AsyncLoadingThread.IsAsyncLoadingSuspended() && !FAsyncLoadingThreadState2::Get().IsTimeLimitExceeded())))
+		while (DeferredFinalizeIndex < ExportMap.Num() &&
+			!AsyncLoadingThread.IsAsyncLoadingSuspended() &&
+			!FAsyncLoadingThreadState2::Get().IsTimeLimitExceeded())
 		{
-			UObject* Object = DeferredFinalizeObjects[DeferredFinalizeIndex++];
-			if (Object)
-			{
-				Object->AtomicallyClearInternalFlags(EInternalObjectFlags::AsyncLoading);
-			}
+			FObjectExport& Export = ExportMap[DeferredFinalizeIndex++];
+			UObject* Object = Export.Object;
 
 			// CDO need special handling, no matter if it's listed in DeferredFinalizeObjects or created here for DynamicClass
 			UObject* CDOToHandle = nullptr;
@@ -4768,17 +4636,15 @@ EAsyncPackageState::Type FAsyncPackage2::PostLoadDeferredObjects()
 				CDODefaultSubobjects.Reset();
 			}
 		}
-		if (DeferredFinalizeIndex == DeferredFinalizeObjects.Num())
+		if (DeferredFinalizeIndex == ExportMap.Num())
 		{
-			DeferredFinalizeIndex = 0;
-			DeferredFinalizeObjects.Reset();
 			Result = EAsyncPackageState::Complete;
 		}
 		else
 		{
 			Result = EAsyncPackageState::TimeOut;
 		}
-	
+
 		// Mark package as having been fully loaded and update load time.
 		if (Result == EAsyncPackageState::Complete && LinkerRoot && !bLoadHasFailed)
 		{
@@ -4830,30 +4696,14 @@ EAsyncPackageState::Type FAsyncPackage2::FinishObjects()
 
 	FUObjectSerializeContext* LoadContext = GetSerializeContext();
 	check(!Linker || LoadContext == Linker->GetSerializeContext());		
-	TArray<UObject*>& ThreadObjLoaded = LoadContext->PRIVATE_GetObjectsLoadedInternalUseOnly();
 
 	EAsyncLoadingResult::Type LoadingResult;
 	if (!bLoadHasFailed)
 	{
-		ThreadObjLoaded.Reset();
 		LoadingResult = EAsyncLoadingResult::Succeeded;
 	}
 	else
 	{		
-		PackageObjLoaded.Append(ThreadObjLoaded);
-
-		// Cleanup objects from this package only
-		for (int32 ObjectIndex = PackageObjLoaded.Num() - 1; ObjectIndex >= 0; --ObjectIndex)
-		{
-			UObject* Object = PackageObjLoaded[ObjectIndex];
-			if (Object && Object->GetOutermost()->GetFName() == Desc.Name)
-			{
-				Object->ClearFlags(RF_NeedPostLoad | RF_NeedLoad | RF_NeedPostLoadSubobjects);
-				Object->MarkPendingKill();
-				PackageObjLoaded[ObjectIndex] = nullptr;
-			}
-		}
-
 		// Clean up UPackage so it can't be found later
 		if (LinkerRoot && !LinkerRoot->IsRooted())
 		{
@@ -4950,27 +4800,6 @@ void FAsyncPackage2::Cancel()
 	CallCompletionCallbacks(true, Result);
 	CallCompletionCallbacks(false, Result);
 
-	FUObjectSerializeContext* LoadContext = GetSerializeContext();
-	if (LoadContext)
-	{
-		TArray<UObject*>& ThreadObjLoaded = LoadContext->PRIVATE_GetObjectsLoadedInternalUseOnly();
-		if (ThreadObjLoaded.Num())
-		{
-			PackageObjLoaded.Append(ThreadObjLoaded);
-			ThreadObjLoaded.Reset();
-		}
-	}
-
-	{
-		// Clear load flags from any referenced objects
-		ClearFlagsAndDissolveClustersFromLoadedObjects(PackageObjLoaded);
-		ClearFlagsAndDissolveClustersFromLoadedObjects(DeferredFinalizeObjects);
-	
-		// Release references
-		PackageObjLoaded.Empty();
-		DeferredFinalizeObjects.Empty();
-	}
-
 	if (LinkerRoot)
 	{
 		if (Linker)
@@ -5000,15 +4829,6 @@ void FAsyncPackage2::UpdateLoadPercentage()
 	// PostLoadCount is just an estimate to prevent packages to go to 100% too quickly
 	// We may never reach 100% this way, but it's better than spending most of the load package time at 100%
 	float NewLoadPercentage = 0.0f;
-	if (Linker)
-	{
-		const int32 PostLoadCount = FMath::Max(DeferredPostLoadObjects.Num(), LocalImportCount);
-		NewLoadPercentage = 100.f * (ExportIndex + DeferredPostLoadIndex) / (Linker->ExportMap.Num() + PostLoadCount);		
-	}
-	else if (DeferredPostLoadObjects.Num() > 0)
-	{
-		NewLoadPercentage = static_cast<float>(DeferredPostLoadIndex) / DeferredPostLoadObjects.Num();
-	}
 	// It's also possible that we got so many objects to PostLoad that LoadPercantage will actually drop
 	LoadPercentage = FMath::Max(NewLoadPercentage, LoadPercentage);
 }
