First pass for creating a manifest describing all BulkData encountered during a cook for ZenCreator so that it can eventually be allowing it to be used by the IoDispatcher. Currently it is ignoring the existence of optional and duplicated non optional bulkdata which can be added later.

When cooking the package we now create a manifest of all potential bulkdata access patterns to the output ubulk files so that ZenCreator knows aobut them.

At the moment BulkData does not have access to the FGlobalNameMap and so using the PackageNameIndex to generate it's FIoChunkId will not work, so currently I am just setting it to zero. In future work BulkDatas will create this token during the cooking phase to reduce cpu work at load time and ideally gain access to FGlobalNameMap so that it can use it in it's generation.

Currently the new BulkData flag BULKDATA_CookedForIoDispatcher serves two purposes, firstly it allows the old BulkData to check if it is loading up cooked data with the new index, or if it's loading older bulkdata from a uasset, secondly at runtime for the new BulkData system (in a later submit) we can use it to decide if we load from cooked loose files or the IODispatcher system. These use cases will probably change over the next few weeks as I clean up the code.

diff --git a/Engine/Source/Editor/UnrealEd/Private/CookOnTheFlyServer.cpp b/Engine/Source/Editor/UnrealEd/Private/CookOnTheFlyServer.cpp
--- a/Engine/Source/Editor/UnrealEd/Private/CookOnTheFlyServer.cpp
+++ b/Engine/Source/Editor/UnrealEd/Private/CookOnTheFlyServer.cpp
@@ -6712,21 +6712,22 @@ void UCookOnTheFlyServer::InitializePackageStore(const TArray<FName>& TargetPlat
 		return;
 	}
 
-	const FString NameMapFilename = FPaths::RootDir() / TEXT("megafile.unamemap");
-	const FString NameMapSandboxFilename = ConvertToFullSandboxPath(*NameMapFilename, true);
+	const FString RootPath = FPaths::RootDir();
+	const FString RootPathSandbox = ConvertToFullSandboxPath(*RootPath, true);
 
-	FString NameMapCookedSandboxFilename;
 	SavePackageContexts.Reserve(TargetPlatformNames.Num());
 
 	for (const FName PlatformName : TargetPlatformNames)
 	{
-		NameMapCookedSandboxFilename = NameMapSandboxFilename.Replace(TEXT("[Platform]"), *PlatformName.ToString());
+		const FString ResolvedRootPath = RootPathSandbox.Replace(TEXT("[Platform]"), *PlatformName.ToString());
+		const FString NameMapCookedSandboxFilename = ResolvedRootPath / TEXT("megafile.unamemap");
 
 		// just leak all memory for now
-		FPackageStoreNameMapSaver* NameMapSaver = new FPackageStoreNameMapSaver(*NameMapCookedSandboxFilename);
-		FPackageHeaderSaver* PackageHeaderSaver = new FPackageHeaderSaver(*NameMapSaver);
-		FLooseFileWriter* LooseFileWriter		= new FLooseFileWriter();
-		FSavePackageContext* SavePackageContext = new FSavePackageContext(*PackageHeaderSaver, *LooseFileWriter);
+		FPackageStoreNameMapSaver* NameMapSaver			= new FPackageStoreNameMapSaver(*NameMapCookedSandboxFilename);
+		FPackageStoreBulkDataManifest* BulkDataManifest	= new FPackageStoreBulkDataManifest(ResolvedRootPath);
+		FPackageHeaderSaver* PackageHeaderSaver			= new FPackageHeaderSaver(*NameMapSaver);
+		FLooseFileWriter* LooseFileWriter				= new FLooseFileWriter();
+		FSavePackageContext* SavePackageContext			= new FSavePackageContext(*PackageHeaderSaver, *LooseFileWriter, *BulkDataManifest);
 
 		SavePackageContexts.Add(SavePackageContext);
 	}
@@ -6741,14 +6742,16 @@ void UCookOnTheFlyServer::FinalizePackageStore()
 
 	SCOPE_TIMER(FinalizePackageStore);
 
-	UE_LOG(LogCook, Display, TEXT("Saving name map(s)..."));
+	UE_LOG(LogCook, Display, TEXT("Saving PackageStoreHeaders(s)..."));
 	const TArray<ITargetPlatform*>& TargetPlatforms = GetCookingTargetPlatforms();
 	for (int32 PlatformIndex = 0; PlatformIndex < TargetPlatforms.Num(); ++PlatformIndex)
 	{
 		INameMapSaver& NameMapSaver = SavePackageContexts[PlatformIndex]->HeaderSaver.NameMapSaver;
 		NameMapSaver.End();
+		
+		SavePackageContexts[PlatformIndex]->BulkDataManifest.Save();
 	}
-	UE_LOG(LogCook, Display, TEXT("Done saving name map(s)"));
+	UE_LOG(LogCook, Display, TEXT("Done saving PackageStoreHeaders(s)"));
 }
 
 void UCookOnTheFlyServer::InitializeTargetPlatforms()
diff --git a/Engine/Source/Runtime/CoreUObject/Private/Serialization/BulkData.cpp b/Engine/Source/Runtime/CoreUObject/Private/Serialization/BulkData.cpp
--- a/Engine/Source/Runtime/CoreUObject/Private/Serialization/BulkData.cpp
+++ b/Engine/Source/Runtime/CoreUObject/Private/Serialization/BulkData.cpp
@@ -945,10 +945,10 @@ void FUntypedBulkData::Serialize( FArchive& Ar, UObject* Owner, int32 Idx, bool
 		}
 
 		// We're loading from the persistent archive.
-		if( Ar.IsLoading() )
+		if (Ar.IsLoading())
 		{
 			Filename = TEXT("");
-			
+
 			// @todo when Landscape (and others?) only Lock/Unlock once, we can enable this
 			if (false) // FPlatformProperties::RequiresCookedData())
 			{
@@ -978,6 +978,13 @@ void FUntypedBulkData::Serialize( FArchive& Ar, UObject* Owner, int32 Idx, bool
 
 			Ar << BulkDataOffsetInFile;
 
+			if ((BulkDataFlags & BULKDATA_CookedForIoDispatcher) != 0)
+			{
+				// Load the ChunkID that will be used by the IoDispatcher, but this version of BulkData will not use it!
+				uint16 ChunkIndex;
+				Ar << ChunkIndex;
+			}
+
 			// determine whether the payload is stored inline or at the end of the file
 			const bool bPayloadInline = !(BulkDataFlags&BULKDATA_PayloadAtEndOfFile);
 
@@ -1229,7 +1236,6 @@ void FUntypedBulkData::Serialize( FArchive& Ar, UObject* Owner, int32 Idx, bool
 				BulkDataOffsetInFile = INDEX_NONE;
 				// And serialize the placeholder which is going to be overwritten later.
 				Ar << BulkDataOffsetInFile;
-
 			}
 
 				// try to get the linkersave object
@@ -1243,6 +1249,16 @@ void FUntypedBulkData::Serialize( FArchive& Ar, UObject* Owner, int32 Idx, bool
 				bStoreInline = true;
 			}
 
+			// Write the ChunkID that will be used by the IoDispatcher
+			{
+				FArchive::FScopeSetDebugSerializationFlags S(Ar, DSF_IgnoreDiff);
+
+				uint16 ChunkIndex = (bStoreInline || !Ar.IsCooking()) ? INDEX_NONE : LinkerSave->BulkDataToAppend.Num();
+				Ar << ChunkIndex;
+
+				BulkDataFlags |= BULKDATA_CookedForIoDispatcher;
+			}
+
 			if (!bStoreInline)
 			{
 				// set the flag indicating where the payload is stored
@@ -1535,9 +1551,9 @@ void FUntypedBulkData::SerializeBulkData( FArchive& Ar, void* Data )
 			Ar.Serialize( Data, GetBulkDataSize() );
 		}
 	}
-	// Serialize an element at a time via the virtual SerializeElement function potentialy allowing and dealing with 
+	// Serialize an element at a time via the virtual SerializeElement function potentially allowing and dealing with 
 	// endian conversion. Dealing with compression makes this a bit more complex as SerializeCompressed expects the 
-	// full data to be compresed en block and not piecewise.
+	// full data to be compressed en block and not piecewise.
 	else
 	{
 		// Serialize data compressed.
diff --git a/Engine/Source/Runtime/CoreUObject/Private/UObject/SavePackage.cpp b/Engine/Source/Runtime/CoreUObject/Private/UObject/SavePackage.cpp
--- a/Engine/Source/Runtime/CoreUObject/Private/UObject/SavePackage.cpp
+++ b/Engine/Source/Runtime/CoreUObject/Private/UObject/SavePackage.cpp
@@ -567,6 +567,110 @@ void FPackageStoreNameMapSaver::End()
 	}
 }
 
+void FPackageStoreBulkDataManifest::PackageDesc::AddData(uint16 InIndex, uint64 InOffset, uint64 InSize)
+{
+	// The index is supposed to be unique
+	auto func = [=](const BulkDataDesc& Entry) { return Entry.Index == InIndex; };
+	check(Data.FindByPredicate(func) == false);
+
+	BulkDataDesc& Entry = Data.Emplace_GetRef();
+
+	Entry.Index = InIndex;
+	Entry.Offset = InOffset;
+	Entry.Size = InSize;
+}
+
+FPackageStoreBulkDataManifest::FPackageStoreBulkDataManifest(const FString& InRootPath)
+{
+	RootPath = InRootPath;
+	FPaths::NormalizeFilename(RootPath);
+
+	if (!RootPath.EndsWith("/"))
+	{
+		RootPath.Append("/");
+	}
+	
+	Filename = RootPath + "megafile.ubulkmanifest";
+}
+
+FArchive& operator<<(FArchive& Ar, FPackageStoreBulkDataManifest::PackageDesc::BulkDataDesc& Entry)
+{
+	Ar << Entry.Index;
+	Ar << Entry.Offset;
+	Ar << Entry.Size; 
+	
+	return Ar;
+}
+
+FArchive& operator<<(FArchive& Ar, FPackageStoreBulkDataManifest::PackageDesc& Entry)
+{
+	Ar << Entry.Data;
+
+	return Ar;
+}
+
+bool FPackageStoreBulkDataManifest::Load()
+{
+	Data.Empty();
+
+	TUniquePtr<FArchive> BinArchive(IFileManager::Get().CreateFileReader(*Filename));
+	if (BinArchive != nullptr)
+	{
+		*BinArchive << Data;
+		return true;
+	}
+	else
+	{
+		return false;
+	}
+}
+
+void FPackageStoreBulkDataManifest::Save()
+{
+	TUniquePtr<FArchive> BinArchive(IFileManager::Get().CreateFileWriter(*Filename));
+	*BinArchive << Data;
+}
+
+
+const FPackageStoreBulkDataManifest::PackageDesc* FPackageStoreBulkDataManifest::Find(const FString& PackageFilename) const
+{
+	const  FString NormalizedFilename = FixFilename(PackageFilename);
+	return Data.Find(NormalizedFilename);
+}
+
+void FPackageStoreBulkDataManifest::AddFileAccess(const FString& PackageFilename, uint16 InIndex, uint64 InOffset, uint64 InSize)
+{
+	const FString NormalizedFilename = FixFilename(PackageFilename);
+	
+	PackageDesc& Entry = GetOrCreateFileAccess(NormalizedFilename);
+	Entry.AddData(InIndex, InOffset, InSize);
+}
+
+FPackageStoreBulkDataManifest::PackageDesc& FPackageStoreBulkDataManifest::GetOrCreateFileAccess(const FString& PackageFilename)
+{
+	if (PackageDesc* Entry = Data.Find(PackageFilename))
+	{
+		return *Entry;
+	}
+	else
+	{
+		return Data.Add(PackageFilename);
+	}
+}
+
+FString FPackageStoreBulkDataManifest::FixFilename(const FString& InFilename) const
+{
+	FString OutFilename = InFilename;
+	FPaths::NormalizeFilename(OutFilename);
+
+	if (!FPaths::IsRelative(OutFilename))
+	{
+		FPaths::MakePathRelativeTo(OutFilename, *RootPath);
+	}
+
+	return OutFilename;
+}
+
 #if WITH_EDITOR
 static void AddReplacementsNames(INameMapSaver& NameMapSaver, UObject* Obj, const ITargetPlatform* TargetPlatform)
 {
@@ -3657,6 +3761,7 @@ FSavePackageResultStruct UPackage::Save(UPackage* InOuter, UObject* Base, EObjec
 			FSinglePackageNameMapSaver SinglePackageNameMapSaver;
 			FPackageHeaderSaver LocalSaver(SinglePackageNameMapSaver);
 			FPackageHeaderSaver& PackageHeaderSaver = SavePackageContext ? SavePackageContext->HeaderSaver : LocalSaver;
+
 			INameMapSaver& NameMapSaver = PackageHeaderSaver.NameMapSaver;
 			NameMapSaver.BeginPackage();
 
@@ -6188,8 +6293,9 @@ void SaveBulkData(FLinkerSave* Linker, const UPackage* InOuter, const TCHAR* Fil
 			BulkDataAlignment = TargetPlatform->GetMemoryMappingAlignment();
 		}
 
-		for (FLinkerSave::FBulkDataStorageInfo& BulkDataStorageInfo : Linker->BulkDataToAppend)
+		for (int iBulkDataIndex = 0; iBulkDataIndex < Linker->BulkDataToAppend.Num(); ++iBulkDataIndex)
 		{
+			FLinkerSave::FBulkDataStorageInfo& BulkDataStorageInfo = Linker->BulkDataToAppend[iBulkDataIndex];
 			BulkDataFeedback.EnterProgressFrame();
 
 			// Set bulk data flags to what they were during initial serialization (they might have changed after that)
@@ -6293,6 +6399,21 @@ void SaveBulkData(FLinkerSave* Linker, const UPackage* InOuter, const TCHAR* Fil
 			BulkDataStorageInfo.BulkData->ClearBulkDataFlags(0xFFFFFFFF);
 			BulkDataStorageInfo.BulkData->SetBulkDataFlags(OldBulkDataFlags);
 			BulkDataStorageInfo.BulkData->Unlock();
+
+			if ((BulkDataStorageInfo.BulkDataFlags & BULKDATA_CookedForIoDispatcher) != 0 && SavePackageContext != nullptr)
+			{
+				// TODO: iBulkDataIndex should be based on the number of BulkData objects we are adding to the manifest not
+				// it's position in Linker->BulkDataToAppend so that inlined BulkData objects are ignored and we won't have 
+				// missing index  values. 
+				// We can't do this yet because the BulkData object is writing out its index for runtime in ::Serialize and 
+				// so it's position in Linker->BulkDataToAppend is the only info we have.
+				// This can be fixed by generating the FIoChunkId at this point and writing it back to the bulkdata before it
+				// is written to disk as do for the data values above this point.
+				SavePackageContext->BulkDataManifest.AddFileAccess
+				(
+					Filename, iBulkDataIndex, BulkStartOffset, SizeOnDisk
+				);
+			}
 		}
 
 		if (BulkArchive)
diff --git a/Engine/Source/Runtime/CoreUObject/Public/Serialization/BulkData.h b/Engine/Source/Runtime/CoreUObject/Public/Serialization/BulkData.h
--- a/Engine/Source/Runtime/CoreUObject/Public/Serialization/BulkData.h
+++ b/Engine/Source/Runtime/CoreUObject/Public/Serialization/BulkData.h
@@ -51,11 +51,13 @@ enum EBulkDataFlags
 	/** This payload is optional and may not be on device */
 	BULKDATA_OptionalPayload					= 1 << 11,
 	/** This payload will be memory mapped, this requires alignment, no compression etc. */
-	BULKDATA_MemoryMappedPayload = 1 << 12,
+	BULKDATA_MemoryMappedPayload				= 1 << 12,
 	/** Bulk data size is 64 bits long */
 	BULKDATA_Size64Bit							= 1 << 13,
 	/** Duplicate non-optional payload in optional bulk data. */
-	BULKDATA_DuplicateNonOptionalPayload		= 1 << 14
+	BULKDATA_DuplicateNonOptionalPayload		= 1 << 14,
+	/** Indicates that the data was cooked for use with IoDispatcher. */
+	BULKDATA_CookedForIoDispatcher				= 1 << 15
 };
 
 /**
@@ -498,6 +500,15 @@ public:
 		return	(GetBulkDataFlags() & BULKDATA_PayloadInSeperateFile) != 0;
 	}
 
+	/**
+	* Returns whether this bulk data is accessed via the IoDispatcher or not.
+	* @return false as the old BulkData API does not support it
+	*/
+	bool IsUsingIODispatcher() const
+	{
+		return false;
+	}
+
 	/**
 	 * Enables the given flags without affecting any previously set flags.
 	 *
diff --git a/Engine/Source/Runtime/CoreUObject/Public/UObject/SavePackage.h b/Engine/Source/Runtime/CoreUObject/Public/UObject/SavePackage.h
--- a/Engine/Source/Runtime/CoreUObject/Public/UObject/SavePackage.h
+++ b/Engine/Source/Runtime/CoreUObject/Public/UObject/SavePackage.h
@@ -68,10 +68,63 @@ private:
 	const FString Filename;
 };
 
+class IBulkDataManifest 
+{
+public:
+	virtual ~IBulkDataManifest() = default;
+
+	virtual void Save() = 0;
+	virtual void AddFileAccess(const FString& PackageFilename, uint16 InIndex, uint64 InOffset, uint64 InSize) = 0;
+};
+
+class FPackageStoreBulkDataManifest : public IBulkDataManifest
+{
+public:
+	COREUOBJECT_API FPackageStoreBulkDataManifest(const FString& InRootPath);
+
+	bool Load();
+	virtual void Save() override;
+	
+	const FString& GetFilename() const { return Filename; }
+
+	class PackageDesc
+	{
+	public:
+		struct BulkDataDesc
+		{
+			uint16 Index;
+			uint64 Offset;
+			uint64 Size;
+		};
+
+		void AddData(uint16 InIndex, uint64 InOffset, uint64 InSize);
+		const TArray<BulkDataDesc>& GetDataArray() const { return Data; }
+	private:
+		friend FArchive& operator<<(FArchive& Ar, PackageDesc& Entry);
+		TArray<BulkDataDesc> Data;
+	};
+
+	const PackageDesc* Find(const FString& PackageName) const;
+
+private:
+	virtual void AddFileAccess(const FString& PackageFilename, uint16 InIndex, uint64 InOffset, uint64 InSize) override;
+
+	PackageDesc& GetOrCreateFileAccess(const FString& PackageFilename);
+
+	FString FixFilename(const FString& InFileName) const;
+
+	FString RootPath;
+	FString Filename;
+	TMap<FString, PackageDesc> Data;
+};
+
 class FPackageHeaderSaver
 {
 public:
-	FPackageHeaderSaver(INameMapSaver& InNameMapSaver) : NameMapSaver(InNameMapSaver) {}
+	FPackageHeaderSaver(INameMapSaver& InNameMapSaver)
+		: NameMapSaver(InNameMapSaver) 
+	{}
+
 	INameMapSaver& NameMapSaver;
 };
 
@@ -138,11 +191,14 @@ private:
 class FSavePackageContext
 {
 public:
-	FSavePackageContext(FPackageHeaderSaver& InHeaderSaver, FPackageStoreWriter& InPackageStoreWriter) 
-	: HeaderSaver(InHeaderSaver), PackageStoreWriter(InPackageStoreWriter) 
+	FSavePackageContext(FPackageHeaderSaver& InHeaderSaver, FPackageStoreWriter& InPackageStoreWriter, IBulkDataManifest& InBulkDataManifest)
+	: HeaderSaver(InHeaderSaver)
+	, PackageStoreWriter(InPackageStoreWriter) 
+	, BulkDataManifest(InBulkDataManifest)
 	{
 	}
 
 	FPackageHeaderSaver& HeaderSaver;
 	FPackageStoreWriter& PackageStoreWriter;
+	IBulkDataManifest& BulkDataManifest;
 };
