AsyncLoading2: Replace FObjectExport with a smaller in-place loaded FExportMapEntry.
FObjectExport is only required for create.
Keep separate array export objects.
Delay FName resolving to CreateExport.
Move some memory allocations to the FAsyncPackage2 constructor.

diff --git a/Engine/Source/Runtime/CoreUObject/Private/Serialization/AsyncLoading2.cpp b/Engine/Source/Runtime/CoreUObject/Private/Serialization/AsyncLoading2.cpp
--- a/Engine/Source/Runtime/CoreUObject/Private/Serialization/AsyncLoading2.cpp
+++ b/Engine/Source/Runtime/CoreUObject/Private/Serialization/AsyncLoading2.cpp
@@ -141,7 +141,7 @@ public:
 		}
 		else if (Index.IsExport())
 		{
-			Object = (*ExportMap)[Index.ToExport()].Object;
+			Object = (*Exports)[Index.ToExport()];
 		}
 		else
 		{
@@ -207,7 +207,7 @@ private:
 	int32* LocalImportIndices = nullptr;
 	UObject** GlobalImportObjects = nullptr;
 	const TArray<FNameEntryId>* NameMap = nullptr;
-	const TArray<FObjectExport>* ExportMap = nullptr;
+	const TArray<UObject*>* Exports = nullptr;
 	TArray<FExternalReadCallback>* ExternalReadDependencies;
 };
 
@@ -232,6 +232,17 @@ struct FGlobalPackageId
 	}
 };
 
+struct FExportMapEntry
+{
+	int32 ObjectName[2];
+	FPackageIndex OuterIndex;
+	FPackageIndex ClassIndex;
+	FPackageIndex SuperIndex;
+	FPackageIndex TemplateIndex;
+	FPackageIndex GlobalImportIndex;
+	EObjectFlags ObjectFlags;
+};
+
 struct FGlobalImport
 {
 	FName ObjectName;
@@ -271,7 +282,6 @@ struct FPackageSummary
 {
 	uint32 PackageFlags;
 	int32 ExportMapOffset;
-	int32 ExportMapSize;
 	int32 GraphDataOffset;
 	int32 GraphDataSize;
 	int32 BulkDataStartOffset;
@@ -901,8 +911,6 @@ private:
 	/** Current index into linkers import table used to spread creation over several frames				*/
 	int32							ImportIndex;
 	/** Current index into linkers export table used to spread creation over several frames				*/
-	int32							ExportIndex;
-	/** Current index into ObjLoaded array used to spread routing PreLoad over several frames			*/
 	int32							FinishExternalReadDependenciesIndex;
 	/** Current index into ObjLoaded array used to spread routing PostLoad over several frames			*/
 	int32							PostLoadIndex;
@@ -960,10 +968,12 @@ private:
 	TArray<FIoBuffer> ExportIoBuffers;
 
 	// FZenLinkerLoad
-	TArray<FObjectExport> ExportMap;
 	TArray<FExternalReadCallback> ExternalReadDependencies;
 	int32 GlobalImportCount = 0;
 	int32 LocalImportCount = 0;
+	int32 ExportCount = 0;
+	const FExportMapEntry* ExportMap = nullptr;
+	TArray<UObject*> Exports;
 	int32* LocalImportIndices = nullptr;
 	FName* GlobalImportNames = nullptr;
 	FPackageIndex* GlobalImportOuters = nullptr;
@@ -2457,7 +2467,6 @@ EAsyncPackageState::Type FAsyncPackage2::Event_SetupExports(FAsyncPackage2* Pack
 {
 	TRACE_CPUPROFILER_EVENT_SCOPE(Event_SetupExports);
 
-	Package->ExportIndex = Package->ExportMap.Num();
 	check(Package->AsyncPackageLoadingState == EAsyncPackageLoadingState2::SetupExports);
 	Package->AsyncPackageLoadingState = EAsyncPackageLoadingState2::ProcessNewImportsAndExports;
 	return EAsyncPackageState::Complete;
@@ -2501,10 +2510,13 @@ EAsyncPackageState::Type FAsyncPackage2::Event_SerializeExport(FAsyncPackage2* P
 
 	FScopedAsyncPackageEvent2 Scope(Package);
 
-	Package->EventDrivenSerializeExport(LocalExportIndex);
-	FObjectExport& Export = Package->ExportMap[LocalExportIndex];
-	UObject* Object = Export.Object;
-	check(!Object || !Object->HasAnyFlags(RF_NeedLoad));
+	UObject* Object = Package->Exports[LocalExportIndex];;
+	check(Object);
+	if (Object->HasAnyFlags(RF_NeedLoad))
+	{
+		Package->EventDrivenSerializeExport(LocalExportIndex);
+	}
+	check(!Object->HasAnyFlags(RF_NeedLoad));
 	return EAsyncPackageState::Complete;
 }
 
@@ -2571,7 +2583,7 @@ UObject* FAsyncPackage2::EventDrivenIndexToObject(FPackageIndex Index, bool bChe
 	}
 	if (Index.IsExport())
 	{
-		Result = ExportMap[Index.ToExport()].Object;
+		Result = Exports[Index.ToExport()];
 	}
 	else if (Index.IsImport())
 	{
@@ -2606,27 +2618,34 @@ UObject* FAsyncPackage2::EventDrivenIndexToObject(FPackageIndex Index, bool bChe
 void FAsyncPackage2::EventDrivenCreateExport(int32 LocalExportIndex)
 {
 	//SCOPED_LOADTIMER(Package_CreateExports);
-	FObjectExport& Export = ExportMap[LocalExportIndex];
+	const FExportMapEntry& Export = ExportMap[LocalExportIndex];
+	UObject*& Object = Exports[LocalExportIndex];
+
+	FName ObjectName;
+	{
+		TRACE_CPUPROFILER_EVENT_SCOPE(CreateExportNameFixup);
+		const FGlobalNameMap& GlobalNameMap = AsyncLoadingThread.GlobalNameMap;
+		ObjectName = GlobalNameMap.GetName(Export.ObjectName[0], Export.ObjectName[1]);
+	}
 
-	TRACE_LOADTIME_CREATE_EXPORT_SCOPE(this, &Export.Object);
+	TRACE_LOADTIME_CREATE_EXPORT_SCOPE(this, &Object);
 
 	LLM_SCOPE(ELLMTag::AsyncLoading);
 	LLM_SCOPED_TAG_WITH_OBJECT_IN_SET(GetLinkerRoot(), ELLMTagSet::Assets);
-	LLM_SCOPED_TAG_WITH_OBJECT_IN_SET((Export.DynamicType == FObjectExport::EDynamicType::DynamicType) ? UDynamicClass::StaticClass() : 
-		CastEventDrivenIndexToObject<UClass>(Export.ClassIndex, false), ELLMTagSet::AssetClasses);
+	// LLM_SCOPED_TAG_WITH_OBJECT_IN_SET((Export.DynamicType == FObjectExport::EDynamicType::DynamicType) ? UDynamicClass::StaticClass() : 
+	// 	CastEventDrivenIndexToObject<UClass>(Export.ClassIndex, false), ELLMTagSet::AssetClasses);
 
 
 	// Check whether we already loaded the object and if not whether the context flags allow loading it.
 	//check(!Export.Object || Export.Object->HasAnyFlags(RF_ClassDefaultObject)); // we should not have this yet, unless it is a CDO
-	check(!Export.Object); // we should not have this yet
-	if (!Export.Object && !Export.bExportLoadFailed)
+	check(!Object); // we should not have this yet
 	{
 		FUObjectSerializeContext* LoadContext = GetSerializeContext();
 
-		if (!Linker->FilterExport(Export)) // for some acceptable position, it was not "not for" 
 		{
-			check(Export.ObjectName != NAME_None || !(Export.ObjectFlags&RF_Public));
+			// check(Export.ObjectName != NAME_None || !(Export.ObjectFlags&RF_Public));
 			check(LoadContext->HasStartedLoading());
+#if 0
 			if (Export.DynamicType == FObjectExport::EDynamicType::DynamicType)
 			{
 //native blueprint 
@@ -2665,7 +2684,6 @@ void FAsyncPackage2::EventDrivenCreateExport(int32 LocalExportIndex)
 				if (!LoadClass)
 				{
 					UE_LOG(LogStreaming, Error, TEXT("Could not find class %s to create %s"), *Linker->ImpExp(Export.ClassIndex).ObjectName.ToString(), *Export.ObjectName.ToString());
-					Export.bExportLoadFailed = true;
 					return;
 				}
 				Export.Object = LoadClass->GetDefaultObject(true);
@@ -2675,6 +2693,7 @@ void FAsyncPackage2::EventDrivenCreateExport(int32 LocalExportIndex)
 				}
 			}
 			else
+#endif
 			{
 				UClass* LoadClass = nullptr;
 				if (Export.ClassIndex.IsNull())
@@ -2687,8 +2706,7 @@ void FAsyncPackage2::EventDrivenCreateExport(int32 LocalExportIndex)
 				}
 				if (!LoadClass)
 				{
-					UE_LOG(LogStreaming, Error, TEXT("Could not find class %s to create %s"), *Linker->ImpExp(Export.ClassIndex).ObjectName.ToString(), *Export.ObjectName.ToString());
-					Export.bExportLoadFailed = true;
+					UE_LOG(LogStreaming, Error, TEXT("Could not find class %s to create %s"), *Linker->ImpExp(Export.ClassIndex).ObjectName.ToString(), *ObjectName.ToString());
 					return;
 				}
 				UObject* ThisParent = nullptr;
@@ -2696,11 +2714,6 @@ void FAsyncPackage2::EventDrivenCreateExport(int32 LocalExportIndex)
 				{
 					ThisParent = EventDrivenIndexToObject(Export.OuterIndex, false);
 				}
-				else if (Export.bForcedExport)
-				{
-					// see FLinkerLoad::CreateExport, there may be some more we can do here
-					check(!Export.bForcedExport); // this is leftover from seekfree loading I think
-				}
 				else
 				{
 					check(LinkerRoot);
@@ -2709,51 +2722,49 @@ void FAsyncPackage2::EventDrivenCreateExport(int32 LocalExportIndex)
 				check(!dynamic_cast<UObjectRedirector*>(ThisParent));
 				if (!ThisParent)
 				{
-					UE_LOG(LogStreaming, Error, TEXT("Could not find outer %s to create %s"), *Linker->ImpExp(Export.OuterIndex).ObjectName.ToString(), *Export.ObjectName.ToString());
-					Export.bExportLoadFailed = true;
+					UE_LOG(LogStreaming, Error, TEXT("Could not find outer %s to create %s"), *Linker->ImpExp(Export.OuterIndex).ObjectName.ToString(), *ObjectName.ToString());
 					return;
 				}
 
 				// Try to find existing object first in case we're a forced export to be able to reconcile. Also do it for the
 				// case of async loading as we cannot in-place replace objects.
 
-				UObject* ActualObjectWithTheName = StaticFindObjectFastInternal(NULL, ThisParent, Export.ObjectName, true);
+				UObject* ActualObjectWithTheName = StaticFindObjectFastInternal(NULL, ThisParent, ObjectName, true);
 
 				// Always attempt to find object in memory first
 				if (ActualObjectWithTheName && (ActualObjectWithTheName->GetClass() == LoadClass))
 				{
-					Export.Object = ActualObjectWithTheName;
+					Object = ActualObjectWithTheName;
 				}
 
 				// Object is found in memory.
-				if (Export.Object)
+				if (Object)
 				{
-					check(!Export.bForcedExport);
-					//Export.Object->SetLinker(Linker, LocalExportIndex);
+					//Object->SetLinker(Linker, LocalExportIndex);
 
 					// If this object was allocated but never loaded (components created by a constructor, CDOs, etc) make sure it gets loaded
 					// Do this for all subobjects created in the native constructor.
-					const EObjectFlags ObjectFlags = Export.Object->GetFlags();
+					const EObjectFlags ObjectFlags = Object->GetFlags();
 					if (!(ObjectFlags & RF_LoadCompleted))
 					{
-						UE_LOG(LogStreaming, VeryVerbose, TEXT("Note2: %s was constructed during load and is an export and so needs loading."), *Export.Object->GetFullName());
+						UE_LOG(LogStreaming, VeryVerbose, TEXT("Note2: %s was constructed during load and is an export and so needs loading."), *Object->GetFullName());
 						ensure(!(ObjectFlags & (RF_NeedLoad | RF_WasLoaded))); // If export exist but is not completed, it is expected to have been created from a native constructor and not from EventDrivenCreateExport, but who knows...?
 						if (ObjectFlags & RF_ClassDefaultObject)
 						{
 							// never call PostLoadSubobjects on class default objects, this matches the behavior of the old linker where
 							// StaticAllocateObject prevents setting of RF_NeedPostLoad and RF_NeedPostLoadSubobjects, but FLinkerLoad::Preload
 							// assigns RF_NeedPostLoad for blueprint CDOs:
-							Export.Object->SetFlags(RF_NeedLoad | RF_NeedPostLoad | RF_WasLoaded);
+							Object->SetFlags(RF_NeedLoad | RF_NeedPostLoad | RF_WasLoaded);
 						}
 						else
 						{
-							Export.Object->SetFlags(RF_NeedLoad | RF_NeedPostLoad | RF_NeedPostLoadSubobjects | RF_WasLoaded);
+							Object->SetFlags(RF_NeedLoad | RF_NeedPostLoad | RF_NeedPostLoadSubobjects | RF_WasLoaded);
 						}
-						AddOwnedObjectWithAsyncFlag(Export.Object, /*bForceAdd*/ false);
+						AddOwnedObjectWithAsyncFlag(Object, /*bForceAdd*/ false);
 					}
 					else
 					{
-						AddOwnedObjectWithAsyncFlag(Export.Object, /*bForceAdd*/ true);
+						AddOwnedObjectWithAsyncFlag(Object, /*bForceAdd*/ true);
 					}
 				}
 				else
@@ -2761,8 +2772,7 @@ void FAsyncPackage2::EventDrivenCreateExport(int32 LocalExportIndex)
 					if (ActualObjectWithTheName && !ActualObjectWithTheName->GetClass()->IsChildOf(LoadClass))
 					{
 						UE_LOG(LogLinker, Error, TEXT("Failed import: class '%s' name '%s' outer '%s'. There is another object (of '%s' class) at the path."),
-							*LoadClass->GetName(), *Export.ObjectName.ToString(), *ThisParent->GetName(), *ActualObjectWithTheName->GetClass()->GetName());
-						Export.bExportLoadFailed = true; // I am not sure if this is an actual fail or not...looked like it in the original code
+							*LoadClass->GetName(), *ObjectName.ToString(), *ThisParent->GetName(), *ActualObjectWithTheName->GetClass()->GetName());
 						return;
 					}
 
@@ -2771,8 +2781,7 @@ void FAsyncPackage2::EventDrivenCreateExport(int32 LocalExportIndex)
 					UObject* Template = EventDrivenIndexToObject(Export.TemplateIndex, true);
 					if (!Template)
 					{
-						UE_LOG(LogStreaming, Error, TEXT("Cannot construct %s in %s because we could not find its template %s"), *Export.ObjectName.ToString(), *Linker->GetArchiveName(), *Linker->GetImportPathName(Export.TemplateIndex));
-						Export.bExportLoadFailed = true;
+						UE_LOG(LogStreaming, Error, TEXT("Cannot construct %s in %s because we could not find its template %s"), *ObjectName.ToString(), *Linker->GetArchiveName(), *Linker->GetImportPathName(Export.TemplateIndex));
 						return;
 					}
 					// we also need to ensure that the template has set up any instances
@@ -2785,8 +2794,6 @@ void FAsyncPackage2::EventDrivenCreateExport(int32 LocalExportIndex)
 					EObjectFlags ObjectLoadFlags = Export.ObjectFlags;
 					ObjectLoadFlags = EObjectFlags(ObjectLoadFlags | RF_NeedLoad | RF_NeedPostLoad | RF_NeedPostLoadSubobjects | RF_WasLoaded);
 
-					FName NewName = Export.ObjectName;
-
 					// If we are about to create a CDO, we need to ensure that all parent sub-objects are loaded
 					// to get default value initialization to work.
 		#if DO_CHECK
@@ -2800,14 +2807,12 @@ void FAsyncPackage2::EventDrivenCreateExport(int32 LocalExportIndex)
 							check(SuperCDO);
 							if (SuperClass->HasAnyFlags(RF_NeedLoad))
 							{
-								UE_LOG(LogStreaming, Fatal, TEXT("Super %s had RF_NeedLoad while creating %s"), *SuperClass->GetFullName(), *Export.ObjectName.ToString());
-								Export.bExportLoadFailed = true;
+								UE_LOG(LogStreaming, Fatal, TEXT("Super %s had RF_NeedLoad while creating %s"), *SuperClass->GetFullName(), *ObjectName.ToString());
 								return;
 							}
 							if (SuperCDO->HasAnyFlags(RF_NeedLoad))
 							{
-								UE_LOG(LogStreaming, Fatal, TEXT("Super CDO %s had RF_NeedLoad while creating %s"), *SuperCDO->GetFullName(), *Export.ObjectName.ToString());
-								Export.bExportLoadFailed = true;
+								UE_LOG(LogStreaming, Fatal, TEXT("Super CDO %s had RF_NeedLoad while creating %s"), *SuperCDO->GetFullName(), *ObjectName.ToString());
 								return;
 							}
 							TArray<UObject*> SuperSubObjects;
@@ -2817,8 +2822,7 @@ void FAsyncPackage2::EventDrivenCreateExport(int32 LocalExportIndex)
 							{
 								if (SubObject->HasAnyFlags(RF_NeedLoad))
 								{
-									UE_LOG(LogStreaming, Fatal, TEXT("Super CDO subobject %s had RF_NeedLoad while creating %s"), *SubObject->GetFullName(), *Export.ObjectName.ToString());
-									Export.bExportLoadFailed = true;
+									UE_LOG(LogStreaming, Fatal, TEXT("Super CDO subobject %s had RF_NeedLoad while creating %s"), *SubObject->GetFullName(), *ObjectName.ToString());
 									return;
 								}
 							}
@@ -2831,71 +2835,68 @@ void FAsyncPackage2::EventDrivenCreateExport(int32 LocalExportIndex)
 		#endif
 					if (LoadClass->HasAnyFlags(RF_NeedLoad))
 					{
-						UE_LOG(LogStreaming, Fatal, TEXT("LoadClass %s had RF_NeedLoad while creating %s"), *LoadClass->GetFullName(), *Export.ObjectName.ToString());
-						Export.bExportLoadFailed = true;
+						UE_LOG(LogStreaming, Fatal, TEXT("LoadClass %s had RF_NeedLoad while creating %s"), *LoadClass->GetFullName(), *ObjectName.ToString());
 						return;
 					}
 					{
 						UObject* LoadCDO = LoadClass->GetDefaultObject();
 						if (LoadCDO->HasAnyFlags(RF_NeedLoad))
 						{
-							UE_LOG(LogStreaming, Fatal, TEXT("Class CDO %s had RF_NeedLoad while creating %s"), *LoadCDO->GetFullName(), *Export.ObjectName.ToString());
-							Export.bExportLoadFailed = true;
+							UE_LOG(LogStreaming, Fatal, TEXT("Class CDO %s had RF_NeedLoad while creating %s"), *LoadCDO->GetFullName(), *ObjectName.ToString());
 							return;
 						}
 					}
 					if (Template->HasAnyFlags(RF_NeedLoad))
 					{
-						UE_LOG(LogStreaming, Fatal, TEXT("Template %s had RF_NeedLoad while creating %s"), *Template->GetFullName(), *Export.ObjectName.ToString());
-						Export.bExportLoadFailed = true;
+						UE_LOG(LogStreaming, Fatal, TEXT("Template %s had RF_NeedLoad while creating %s"), *Template->GetFullName(), *ObjectName.ToString());
 						return;
 					}
 
-					Export.Object = StaticConstructObject_Internal
-						(
-							LoadClass,
-							ThisParent,
-							NewName,
-							ObjectLoadFlags,
-							EInternalObjectFlags::None,
-							Template,
-							false,
-							nullptr,
-							true
-						);
+					{
+						TRACE_CPUPROFILER_EVENT_SCOPE(CreateExportConstruct);
+						Object = StaticConstructObject_Internal
+							(
+							 LoadClass,
+							 ThisParent,
+							 ObjectName,
+							 ObjectLoadFlags,
+							 EInternalObjectFlags::None,
+							 Template,
+							 false,
+							 nullptr,
+							 true
+							);
+					}
 
 					if (GIsInitialLoad || GUObjectArray.IsOpenForDisregardForGC())
 					{
-						Export.Object->AddToRoot();
+						Object->AddToRoot();
 					}
 
-					AddOwnedObjectWithAsyncFlag(Export.Object, /*bForceAdd*/ false);
-					//Export.Object->SetLinker(Linker, LocalExportIndex);
-					check(Export.Object->GetClass() == LoadClass);
-					check(NewName == Export.ObjectName);
+					AddOwnedObjectWithAsyncFlag(Object, /*bForceAdd*/ false);
+					//Object->SetLinker(Linker, LocalExportIndex);
+					check(Object->GetClass() == LoadClass);
+					check(Object->GetFName() == ObjectName);
 				}
 			}
 		}
 	}
-	if (!Export.Object)
-	{
-		Export.bExportLoadFailed = true;
-	}
-	check(Export.Object || Export.bExportLoadFailed);
+	check(Object);
 }
 
 void FAsyncPackage2::EventDrivenSerializeExport(int32 LocalExportIndex)
 {
 	//SCOPED_LOADTIMER(Package_PreLoadObjects);
 
-	FObjectExport& Export = ExportMap[LocalExportIndex];
+	const FExportMapEntry& Export = ExportMap[LocalExportIndex];
+	UObject* Object = Exports[LocalExportIndex];
 
 	LLM_SCOPE(ELLMTag::UObject);
 	LLM_SCOPED_TAG_WITH_OBJECT_IN_SET(GetLinkerRoot(), ELLMTagSet::Assets);
-	LLM_SCOPED_TAG_WITH_OBJECT_IN_SET((Export.DynamicType == FObjectExport::EDynamicType::DynamicType) ? UDynamicClass::StaticClass() :
-		CastEventDrivenIndexToObject<UClass>(Export.ClassIndex, false), ELLMTagSet::AssetClasses);
+	// LLM_SCOPED_TAG_WITH_OBJECT_IN_SET((Export.DynamicType == FObjectExport::EDynamicType::DynamicType) ? UDynamicClass::StaticClass() :
+	// 	CastEventDrivenIndexToObject<UClass>(Export.ClassIndex, false), ELLMTagSet::AssetClasses);
 
-	UObject* Object = Export.Object;
+#if 0
 	if (Object && Linker->bDynamicClassLinker)
 	{
 		//native blueprint 
@@ -2909,10 +2910,10 @@ void FAsyncPackage2::EventDrivenSerializeExport(int32 LocalExportIndex)
 		Object->ClearFlags(RF_NeedLoad);
 	}
 	else if (Object && Object->HasAnyFlags(RF_NeedLoad))
+#endif
 	{
 		//check(Object->GetLinker() == Linker);
 		//check(Object->GetLinkerIndex() == LocalExportIndex);
-		UClass* Cls = nullptr;
 
 		// If this is a struct, make sure that its parent struct is completely loaded
 		if (UStruct* Struct = dynamic_cast<UStruct*>(Object))
@@ -2924,8 +2925,7 @@ void FAsyncPackage2::EventDrivenSerializeExport(int32 LocalExportIndex)
 				if (!SuperStruct)
 				{
 					// see FLinkerLoad::CreateExport, there may be some more we can do here
-					UE_LOG(LogStreaming, Fatal, TEXT("Could not find SuperStruct %s to create %s"), *Linker->ImpExp(Export.SuperIndex).ObjectName.ToString(), *Export.ObjectName.ToString());
-					Export.bExportLoadFailed = true;
+					// UE_LOG(LogStreaming, Fatal, TEXT("Could not find SuperStruct %s to create %s"), *Linker->ImpExp(Export.SuperIndex).ObjectName.ToString(), *Export.ObjectName.ToString());
 					return;
 				}
 			}
@@ -2958,7 +2958,7 @@ void FAsyncPackage2::EventDrivenSerializeExport(int32 LocalExportIndex)
 		Ar.NameMap = &AsyncLoadingThread.GlobalNameMap.GetNameEntries();
 		Ar.LocalImportIndices = LocalImportIndices;
 		Ar.GlobalImportObjects = GlobalImportObjects;
-		Ar.ExportMap = &ExportMap;
+		Ar.Exports = &Exports;
 		Ar.ExternalReadDependencies = &ExternalReadDependencies;
 
 		Object->ClearFlags(RF_NeedLoad);
@@ -3009,9 +3009,9 @@ EAsyncPackageState::Type FAsyncPackage2::Event_StartIO(FAsyncPackage2* Package,
 	TRACE_CPUPROFILER_EVENT_SCOPE(Event_StartIO);
 
 	check(Package->AsyncPackageLoadingState == EAsyncPackageLoadingState2::ProcessNewImportsAndExports);
-	FObjectExport& Export = Package->ExportMap[LocalExportIndex];
+	UObject* Object = Package->Exports[LocalExportIndex];
 	if (Package->Linker->bDynamicClassLinker || //native blueprint, there is no IO for these
-		(Export.Object && !Export.Object->HasAnyFlags(RF_NeedLoad)))
+		(Object && !Object->HasAnyFlags(RF_NeedLoad)))
 	{
 		Package->GetNode(EEventLoadNode2::ImportOrExport_Serialize, FPackageIndex::FromExport(LocalExportIndex))->ReleaseBarrier();
 	}
@@ -3292,17 +3292,18 @@ bool FAsyncPackage2::AreAllDependenciesFullyLoadedInternal(FAsyncPackage2* Packa
 				OutError = FString::Printf(TEXT("%s Doesn't have all exports Serialized"), *Package->GetPackageName().ToString());
 				return false;
 			}
-			if (AsyncRoot->DeferredPostLoadIndex < AsyncRoot->ExportMap.Num())
+			if (AsyncRoot->DeferredPostLoadIndex < AsyncRoot->ExportCount)
 			{
 				OutError = FString::Printf(TEXT("%s Doesn't have all objects processed by DeferredPostLoad"), *Package->GetPackageName().ToString());
 				return false;
 			}
-			for (FObjectExport& Export : AsyncRoot->ExportMap)
+			for (int32 LocalExportIndex = 0; LocalExportIndex < AsyncRoot->ExportCount; ++LocalExportIndex)
 			{
-				if (Export.Object->HasAnyFlags(RF_NeedPostLoad|RF_NeedLoad))
+				UObject* Object = AsyncRoot->Exports[LocalExportIndex];
+				if (Object->HasAnyFlags(RF_NeedPostLoad|RF_NeedLoad))
 				{
-					OutError = FString::Printf(TEXT("%s has not been %s"), *Export.Object->GetFullName(), 
-						Export.Object->HasAnyFlags(RF_NeedLoad) ? TEXT("Serialized") : TEXT("PostLoaded"));
+					OutError = FString::Printf(TEXT("%s has not been %s"), *Object->GetFullName(), 
+						Object->HasAnyFlags(RF_NeedLoad) ? TEXT("Serialized") : TEXT("PostLoaded"));
 					return false;
 				}
 			}
@@ -4118,7 +4119,6 @@ FAsyncPackage2::FAsyncPackage2(const FAsyncPackageDesc& InDesc, int32 InSerialNu
 , Linker(nullptr)
 , LinkerRoot(nullptr)
 , ImportIndex(0)
-, ExportIndex(0)
 , FinishExternalReadDependenciesIndex(0)
 , PostLoadIndex(0)
 , DeferredPostLoadIndex(0)
@@ -4149,19 +4149,23 @@ FAsyncPackage2::FAsyncPackage2(const FAsyncPackageDesc& InDesc, int32 InSerialNu
 	PackageChunkId = CreateChunkId(*NameIndex, Desc.NameToLoad.GetNumber(), 0, EChunkType::None);
 
 	{
+		TRACE_CPUPROFILER_EVENT_SCOPE(InitZenLinker);
+		ExportCount = AsyncLoadingThread.GetPackageExportCount(GlobalPackageId);
 		LocalImportIndices = AsyncLoadingThread.GetPackageSlimports(GlobalPackageId, LocalImportCount);
 		GlobalImportNames = AsyncLoadingThread.GetGlobalImportNames(GlobalImportCount);
 		GlobalImportOuters = AsyncLoadingThread.GetGlobalImportOuters(GlobalImportCount);
 		GlobalImportPackages = AsyncLoadingThread.GetGlobalImportPackages(GlobalImportCount);
 		GlobalImportObjects = AsyncLoadingThread.GetGlobalImportObjects(GlobalImportCount);
 		AddGlobalImportObjectReferences();
+
+		Exports.AddZeroed(ExportCount);
+		ExportIoBuffers.SetNum(ExportCount);
+		OwnedObjects.Reserve(ExportCount + 1); // ExportCount + UPackage
 	}
 
 	{
 		TRACE_CPUPROFILER_EVENT_SCOPE(CreateNodes);
-		uint32 ImportCount = AsyncLoadingThread.GetPackageImportCount(GlobalPackageId);
-		ImportNodeCount = ImportCount * EEventLoadNode2::Import_NumPhases;
-		uint32 ExportCount = AsyncLoadingThread.GetPackageExportCount(GlobalPackageId);
+		ImportNodeCount = LocalImportCount * EEventLoadNode2::Import_NumPhases;
 		ExportNodeCount = ExportCount * EEventLoadNode2::Export_NumPhases;
 
 		PackageNodes = GraphAllocator.AllocNodes(EEventLoadNode2::Package_NumPhases + ImportNodeCount + ExportNodeCount);
@@ -4197,7 +4201,7 @@ FAsyncPackage2::FAsyncPackage2(const FAsyncPackageDesc& InDesc, int32 InSerialNu
 
 		// Add nodes for all imports and exports.
 		ImportNodes = PackageNodes + EEventLoadNode2::Package_NumPhases;
-		for (uint32 LocalImportIndex = 0; LocalImportIndex < ImportCount; ++LocalImportIndex)
+		for (int32 LocalImportIndex = 0; LocalImportIndex < LocalImportCount; ++LocalImportIndex)
 		{
 			uint32 NodeIndex = EEventLoadNode2::Import_NumPhases * LocalImportIndex;
 			FEventLoadNode2* CreateImportNode = ImportNodes + NodeIndex + EEventLoadNode2::ImportOrExport_Create;
@@ -4210,7 +4214,7 @@ FAsyncPackage2::FAsyncPackage2(const FAsyncPackageDesc& InDesc, int32 InSerialNu
 			ExportsSerializedNode->DependsOn(SerializeImportNode);
 		}
 		ExportNodes = ImportNodes + ImportNodeCount;
-		for (uint32 LocalExportIndex = 0; LocalExportIndex < ExportCount; ++LocalExportIndex)
+		for (int32 LocalExportIndex = 0; LocalExportIndex < ExportCount; ++LocalExportIndex)
 		{
 			uint32 NodeIndex = EEventLoadNode2::Export_NumPhases * LocalExportIndex;
 			FEventLoadNode2* CreateExportNode = ExportNodes + NodeIndex + EEventLoadNode2::ImportOrExport_Create;
@@ -4390,9 +4394,6 @@ EAsyncPackageState::Type FAsyncPackage2::CreateLinker()
 {
 	check(Linker == nullptr);
 
-	const int32 ExportCount = AsyncLoadingThread.GetPackageExportCount(GlobalPackageId);
-	OwnedObjects.Reserve(ExportCount + 1); // ExportCount + UPackage
-
 	// Try to find existing package or create it if not already present.
 	UPackage* Package = nullptr;
 	{
@@ -4440,10 +4441,8 @@ EAsyncPackageState::Type FAsyncPackage2::FinishLinker()
 
 	SCOPED_LOADTIMER(LinkerLoad_FinalizeCreation);
 
-	const FGlobalNameMap& GlobalNameMap = AsyncLoadingThread.GlobalNameMap;
 	const FPackageSummary* PackageSummary = reinterpret_cast<const FPackageSummary*>(PackageSummaryIoBuffer.Data());
-	const int32 ExportCount = AsyncLoadingThread.GetPackageExportCount(GlobalPackageId);
-	check(ExportCount);
+	ExportMap = reinterpret_cast<const FExportMapEntry*>(PackageSummaryIoBuffer.Data() + PackageSummary->ExportMapOffset);
 
 	TRACE_LOADTIME_PACKAGE_SUMMARY(this, PackageSummarySize, AsyncLoadingThread.GetPackageImportCount(GlobalPackageId), ExportCount);
 
@@ -4457,35 +4456,6 @@ EAsyncPackageState::Type FAsyncPackage2::FinishLinker()
 		Linker->UpdateFromPackageFileSummary();
 	}
 
-	{
-		SCOPED_LOADTIMER(LinkerLoad_SerializeExportMap);
-		const uint8* ZenExportData = PackageSummaryIoBuffer.Data() + PackageSummary->ExportMapOffset;
-		FSimpleArchive ZenExportMapArchive(ZenExportData, PackageSummary->ExportMapSize);
-
-		ExportMap.AddZeroed(ExportCount);
-		for (int32 LocalExportIndex = 0; LocalExportIndex < ExportCount; ++LocalExportIndex)
-		{
-			FPackageIndex Index = FPackageIndex::FromExport(LocalExportIndex);
-			FObjectExport& Export = ExportMap[LocalExportIndex];
-			int32 ObjectNameIndex;
-			int32 ObjectNameNumber;
-			int32 GlobalImportIndex;
-
-			ZenExportMapArchive << ObjectNameIndex << ObjectNameNumber;
-			ZenExportMapArchive << Export.OuterIndex;
-			ZenExportMapArchive << Export.ClassIndex;
-			ZenExportMapArchive << Export.SuperIndex;
-			ZenExportMapArchive << Export.TemplateIndex;
-			ZenExportMapArchive << (uint32&)Export.ObjectFlags;
-			ZenExportMapArchive << GlobalImportIndex;
-
-			Export.ObjectName = GlobalNameMap.GetName(ObjectNameIndex, ObjectNameNumber);
-			Export.ThisIndex = Index;
-		}
-
-		ExportIoBuffers.SetNum(ExportCount);
-	}
-
 	// Add this linker to the object manager's linker array.
 	// FLinkerManager::Get().AddLoader(Linker);
 
@@ -4522,10 +4492,9 @@ EAsyncPackageState::Type FAsyncPackage2::PostLoadObjects()
 	const bool bIsMultithreaded = AsyncLoadingThread.IsMultithreaded();
 
 	// PostLoad objects.
-	while (PostLoadIndex < ExportMap.Num() && !FAsyncLoadingThreadState2::Get().IsTimeLimitExceeded())
+	while (PostLoadIndex < ExportCount && !FAsyncLoadingThreadState2::Get().IsTimeLimitExceeded())
 	{
-		FObjectExport& Export = ExportMap[PostLoadIndex++];
-		UObject* Object = Export.Object;
+		UObject* Object = Exports[PostLoadIndex++];
 
 		check(Object);
 		check(!Object->HasAnyFlags(RF_NeedLoad));
@@ -4547,7 +4516,7 @@ EAsyncPackageState::Type FAsyncPackage2::PostLoadObjects()
 		}
 	}
 
-	return (PostLoadIndex == ExportMap.Num()) ? EAsyncPackageState::Complete : EAsyncPackageState::TimeOut;
+	return (PostLoadIndex == ExportCount) ? EAsyncPackageState::Complete : EAsyncPackageState::TimeOut;
 }
 
 void CreateClustersFromPackage(FLinkerLoad* PackageLinker, TArray<UObject*>& OutClusterObjects);
@@ -4566,12 +4535,11 @@ EAsyncPackageState::Type FAsyncPackage2::PostLoadDeferredObjects()
 
 	STAT(double PostLoadStartTime = FPlatformTime::Seconds());
 
-	while (DeferredPostLoadIndex < ExportMap.Num() && 
+	while (DeferredPostLoadIndex < ExportCount && 
 		!AsyncLoadingThread.IsAsyncLoadingSuspended() &&
 		!FAsyncLoadingThreadState2::Get().IsTimeLimitExceeded())
 	{
-		FObjectExport& Export = ExportMap[DeferredPostLoadIndex++];
-		UObject* Object = Export.Object;
+		UObject* Object = Exports[DeferredPostLoadIndex++];
 
 		check(Object);
 		check(!Object->HasAnyFlags(RF_NeedLoad));
@@ -4587,18 +4555,17 @@ EAsyncPackageState::Type FAsyncPackage2::PostLoadDeferredObjects()
 		Object->AtomicallyClearInternalFlags(EInternalObjectFlags::AsyncLoading);
 	}
 
-	Result = (DeferredPostLoadIndex == ExportMap.Num()) ? EAsyncPackageState::Complete : EAsyncPackageState::TimeOut;
+	Result = (DeferredPostLoadIndex == ExportCount) ? EAsyncPackageState::Complete : EAsyncPackageState::TimeOut;
 
 	if (Result == EAsyncPackageState::Complete)
 	{
 		TArray<UObject*> CDODefaultSubobjects;
 		// Clear async loading flags (we still want RF_Async, but EInternalObjectFlags::AsyncLoading can be cleared)
-		while (DeferredFinalizeIndex < ExportMap.Num() &&
+		while (DeferredFinalizeIndex < ExportCount &&
 			!AsyncLoadingThread.IsAsyncLoadingSuspended() &&
 			!FAsyncLoadingThreadState2::Get().IsTimeLimitExceeded())
 		{
-			FObjectExport& Export = ExportMap[DeferredFinalizeIndex++];
-			UObject* Object = Export.Object;
+			UObject* Object = Exports[DeferredFinalizeIndex++];
 
 			// CDO need special handling, no matter if it's listed in DeferredFinalizeObjects or created here for DynamicClass
 			UObject* CDOToHandle = nullptr;
@@ -4636,7 +4603,7 @@ EAsyncPackageState::Type FAsyncPackage2::PostLoadDeferredObjects()
 				CDODefaultSubobjects.Reset();
 			}
 		}
-		if (DeferredFinalizeIndex == ExportMap.Num())
+		if (DeferredFinalizeIndex == ExportCount)
 		{
 			Result = EAsyncPackageState::Complete;
 		}
