AsyncLoading2: Minor cleanup.
Rename some import related variables.
Add and remove some cpu scopes.

diff --git a/Engine/Source/Runtime/CoreUObject/Private/Serialization/AsyncLoading2.cpp b/Engine/Source/Runtime/CoreUObject/Private/Serialization/AsyncLoading2.cpp
--- a/Engine/Source/Runtime/CoreUObject/Private/Serialization/AsyncLoading2.cpp
+++ b/Engine/Source/Runtime/CoreUObject/Private/Serialization/AsyncLoading2.cpp
@@ -145,7 +145,7 @@ public:
 		}
 		else
 		{
-			Object = GlobalImportObjects[LocalImportIndices[Index.ToImport()]];
+			Object = GlobalImportObjects[ImportMap[Index.ToImport()]];
 		}
 		return *this;
 	}
@@ -204,7 +204,7 @@ private:
 	friend FAsyncPackage2;
 
 	UObject* TemplateForGetArchetypeFromLoader = nullptr;
-	int32* LocalImportIndices = nullptr;
+	int32* ImportMap = nullptr;
 	UObject** GlobalImportObjects = nullptr;
 	const TArray<FNameEntryId>* NameMap = nullptr;
 	const TArray<UObject*>* Exports = nullptr;
@@ -698,18 +698,18 @@ struct FAsyncPackage2
 
 	void AddGlobalImportObjectReferences()
 	{
-		for (int32 LocalImportIndex = 0; LocalImportIndex < LocalImportCount; ++LocalImportIndex)
+		for (int32 LocalImportIndex = 0; LocalImportIndex < ImportCount; ++LocalImportIndex)
 		{
-			const int32 GlobalImportIndex = LocalImportIndices[LocalImportIndex];
+			const int32 GlobalImportIndex = ImportMap[LocalImportIndex];
 			++GlobalImportObjectRefCounts[GlobalImportIndex];
 		}
 	}
 
 	void ReleaseGlobalImportObjectReferences()
 	{
-		for (int32 LocalImportIndex = 0; LocalImportIndex < LocalImportCount; ++LocalImportIndex)
+		for (int32 LocalImportIndex = 0; LocalImportIndex < ImportCount; ++LocalImportIndex)
 		{
-			const int32 GlobalImportIndex = LocalImportIndices[LocalImportIndex];
+			const int32 GlobalImportIndex = ImportMap[LocalImportIndex];
 			--GlobalImportObjectRefCounts[GlobalImportIndex];
 		}
 	}
@@ -909,8 +909,6 @@ private:
 	/** Call backs called when we finished loading this package											*/
 	TArray<FCompletionCallback>	CompletionCallbacks;
 	/** Current index into linkers import table used to spread creation over several frames				*/
-	int32							ImportIndex;
-	/** Current index into linkers export table used to spread creation over several frames				*/
 	int32							FinishExternalReadDependenciesIndex;
 	/** Current index into ObjLoaded array used to spread routing PostLoad over several frames			*/
 	int32							PostLoadIndex;
@@ -970,11 +968,11 @@ private:
 	// FZenLinkerLoad
 	TArray<FExternalReadCallback> ExternalReadDependencies;
 	int32 GlobalImportCount = 0;
-	int32 LocalImportCount = 0;
+	int32 ImportCount = 0;
 	int32 ExportCount = 0;
 	const FExportMapEntry* ExportMap = nullptr;
 	TArray<UObject*> Exports;
-	int32* LocalImportIndices = nullptr;
+	int32* ImportMap = nullptr;
 	FName* GlobalImportNames = nullptr;
 	FPackageIndex* GlobalImportOuters = nullptr;
 	FPackageIndex* GlobalImportPackages = nullptr;
@@ -2256,11 +2254,11 @@ void FAsyncPackage2::ImportPackagesRecursive()
 #if 1
 	FMemStack& MemStack = FMemStack::Get();
 	FMemMark Mark(MemStack);
-	TArrayView<FPackageIndex> LocalImportedPackages(new(MemStack)FPackageIndex[LocalImportCount], LocalImportCount);
+	TArrayView<FPackageIndex> LocalImportedPackages(new(MemStack)FPackageIndex[ImportCount], ImportCount);
 
-	for (int32 LocalImportIndex = 0; LocalImportIndex < LocalImportCount; ++LocalImportIndex)
+	for (int32 LocalImportIndex = 0; LocalImportIndex < ImportCount; ++LocalImportIndex)
 	{
-		const int32 GlobalImportIndex = LocalImportIndices[LocalImportIndex];
+		const int32 GlobalImportIndex = ImportMap[LocalImportIndex];
 		const FPackageIndex ImportedPackageIndex = GlobalImportPackages[GlobalImportIndex];
 		UObject* ImportedObject = FindExistingSlimport(GlobalImportIndex);
 		const UPackage* Package = (UPackage*)GlobalImportObjects[ImportedPackageIndex.ToImport()];
@@ -2310,9 +2308,9 @@ void FAsyncPackage2::ImportPackagesRecursive()
 	FPackageIndex LastImportedPackageIndex;
 
 	// slimports are sorted, first comes a new package, then N additional objects
-	for (int32 LocalImportIndex = 0; LocalImportIndex < LocalImportCount; ++LocalImportIndex)
+	for (int32 LocalImportIndex = 0; LocalImportIndex < ImportCount; ++LocalImportIndex)
 	{
-		const int32 GlobalImportIndex = LocalImportIndices[LocalImportIndex];
+		const int32 GlobalImportIndex = ImportMap[LocalImportIndex];
 		const FPackageIndex ImportedPackageIndex = GlobalImportPackages[GlobalImportIndex];
 
 		const bool bIsANewImportPackage = ImportedPackageIndex != LastImportedPackageIndex;
@@ -2391,7 +2389,6 @@ EAsyncPackageState::Type FAsyncPackage2::Event_SetupImports(FAsyncPackage2* Pack
 		verify(Package->SetupSlimports_Event() == EAsyncPackageState::Complete);
 	}
 	check(Package->AsyncPackageLoadingState == EAsyncPackageLoadingState2::SetupImports);
-	check(Package->ImportIndex == Package->LocalImportCount);
 	Package->AsyncPackageLoadingState = EAsyncPackageLoadingState2::SetupExports;
 	return EAsyncPackageState::Complete;
 }
@@ -2400,14 +2397,12 @@ EAsyncPackageState::Type FAsyncPackage2::SetupSlimports_Event()
 {
 	if (!GIsInitialLoad)
 	{
-		ImportIndex = LocalImportCount;
 		return EAsyncPackageState::Complete;
 	}
 
-	while (ImportIndex < LocalImportCount)
+	for (int32 LocalImportIndex = 0; LocalImportIndex < ImportCount; ++LocalImportIndex)
 	{
-		const int32 LocalImportIndex = ImportIndex++;
-		const int32 GlobalImportIndex = LocalImportIndices[LocalImportIndex];
+		const int32 GlobalImportIndex = ImportMap[LocalImportIndex];
 
 		// skip packages
 		if (GlobalImportOuters[GlobalImportIndex].IsNull())
@@ -2460,7 +2455,7 @@ EAsyncPackageState::Type FAsyncPackage2::SetupSlimports_Event()
 		}
 	}
 
-	return ImportIndex == LocalImportCount ? EAsyncPackageState::Complete : EAsyncPackageState::TimeOut;
+	return EAsyncPackageState::Complete;
 }
 
 EAsyncPackageState::Type FAsyncPackage2::Event_SetupExports(FAsyncPackage2* Package, int32)
@@ -2483,7 +2478,7 @@ EAsyncPackageState::Type FAsyncPackage2::Event_ImportSerialized(FAsyncPackage2*
 {
 	TRACE_CPUPROFILER_EVENT_SCOPE(Event_ImportSerialized);
 
-	int32 GlobalImportIndex = Package->LocalImportIndices[LocalImportIndex];
+	int32 GlobalImportIndex = Package->ImportMap[LocalImportIndex];
 	UObject* Object = Package->GlobalImportObjects[GlobalImportIndex];
 	if (Object)
 	{
@@ -2524,8 +2519,8 @@ void FAsyncPackage2::LinkSlimport(int32 LocalImportIndex, int32 GlobalImportInde
 {
 	if (GlobalImportIndex == -1)
 	{
-		check(LocalImportIndex >= 0 && LocalImportIndex < LocalImportCount);
-		GlobalImportIndex = LocalImportIndices[LocalImportIndex];
+		check(LocalImportIndex >= 0 && LocalImportIndex < ImportCount);
+		GlobalImportIndex = ImportMap[LocalImportIndex];
 	}
 
 	UObject*& Object = GlobalImportObjects[GlobalImportIndex];
@@ -2587,7 +2582,7 @@ UObject* FAsyncPackage2::EventDrivenIndexToObject(FPackageIndex Index, bool bChe
 	}
 	else if (Index.IsImport())
 	{
-		int32 GlobalImportIndex = LocalImportIndices[Index.ToImport()];
+		int32 GlobalImportIndex = ImportMap[Index.ToImport()];
 		Result = FindExistingSlimport(GlobalImportIndex);
 		check(Result);
 	}
@@ -2623,7 +2618,7 @@ void FAsyncPackage2::EventDrivenCreateExport(int32 LocalExportIndex)
 
 	FName ObjectName;
 	{
-		TRACE_CPUPROFILER_EVENT_SCOPE(CreateExportNameFixup);
+		TRACE_CPUPROFILER_EVENT_SCOPE(ObjectNameFixup);
 		const FGlobalNameMap& GlobalNameMap = AsyncLoadingThread.GlobalNameMap;
 		ObjectName = GlobalNameMap.GetName(Export.ObjectName[0], Export.ObjectName[1]);
 	}
@@ -2853,7 +2848,7 @@ void FAsyncPackage2::EventDrivenCreateExport(int32 LocalExportIndex)
 					}
 
 					{
-						TRACE_CPUPROFILER_EVENT_SCOPE(CreateExportConstruct);
+						TRACE_CPUPROFILER_EVENT_SCOPE(ConstructObject);
 						Object = StaticConstructObject_Internal
 							(
 							 LoadClass,
@@ -2956,7 +2951,7 @@ void FAsyncPackage2::EventDrivenSerializeExport(int32 LocalExportIndex)
 
 		// FSimpleExportArchive special fields
 		Ar.NameMap = &AsyncLoadingThread.GlobalNameMap.GetNameEntries();
-		Ar.LocalImportIndices = LocalImportIndices;
+		Ar.ImportMap = ImportMap;
 		Ar.GlobalImportObjects = GlobalImportObjects;
 		Ar.Exports = &Exports;
 		Ar.ExternalReadDependencies = &ExternalReadDependencies;
@@ -2979,18 +2974,18 @@ void FAsyncPackage2::EventDrivenSerializeExport(int32 LocalExportIndex)
 
 		if (Object->HasAnyFlags(RF_ClassDefaultObject))
 		{
+			TRACE_CPUPROFILER_EVENT_SCOPE(SerializeDefaultObject);
 			Object->GetClass()->SerializeDefaultObject(Object, Ar);
 		}
 		else
 		{
+			TRACE_CPUPROFILER_EVENT_SCOPE(SerializeObject);
 			Object->Serialize(Ar);
 		}
 
 		Object->SetFlags(RF_LoadCompleted);
 		LoadContext->SerializedObject = PrevSerializedObject;
 
-		//UE_CLOG(Ar.Tell() != Export.SerialSize, LogStreaming, Warning, TEXT("Serialize mismatch, ObjectName='%s'"), *Object->GetFullName());
-
 #if DO_CHECK
 		if (Object->HasAnyFlags(RF_ClassDefaultObject) && Object->GetClass()->HasAnyClassFlags(CLASS_CompiledFromBlueprint))
 		{
@@ -4118,7 +4113,6 @@ FAsyncPackage2::FAsyncPackage2(const FAsyncPackageDesc& InDesc, int32 InSerialNu
 : Desc(InDesc)
 , Linker(nullptr)
 , LinkerRoot(nullptr)
-, ImportIndex(0)
 , FinishExternalReadDependenciesIndex(0)
 , PostLoadIndex(0)
 , DeferredPostLoadIndex(0)
@@ -4151,7 +4145,7 @@ FAsyncPackage2::FAsyncPackage2(const FAsyncPackageDesc& InDesc, int32 InSerialNu
 	{
 		TRACE_CPUPROFILER_EVENT_SCOPE(InitZenLinker);
 		ExportCount = AsyncLoadingThread.GetPackageExportCount(GlobalPackageId);
-		LocalImportIndices = AsyncLoadingThread.GetPackageSlimports(GlobalPackageId, LocalImportCount);
+		ImportMap = AsyncLoadingThread.GetPackageSlimports(GlobalPackageId, ImportCount);
 		GlobalImportNames = AsyncLoadingThread.GetGlobalImportNames(GlobalImportCount);
 		GlobalImportOuters = AsyncLoadingThread.GetGlobalImportOuters(GlobalImportCount);
 		GlobalImportPackages = AsyncLoadingThread.GetGlobalImportPackages(GlobalImportCount);
@@ -4165,7 +4159,7 @@ FAsyncPackage2::FAsyncPackage2(const FAsyncPackageDesc& InDesc, int32 InSerialNu
 
 	{
 		TRACE_CPUPROFILER_EVENT_SCOPE(CreateNodes);
-		ImportNodeCount = LocalImportCount * EEventLoadNode2::Import_NumPhases;
+		ImportNodeCount = ImportCount * EEventLoadNode2::Import_NumPhases;
 		ExportNodeCount = ExportCount * EEventLoadNode2::Export_NumPhases;
 
 		PackageNodes = GraphAllocator.AllocNodes(EEventLoadNode2::Package_NumPhases + ImportNodeCount + ExportNodeCount);
@@ -4201,7 +4195,7 @@ FAsyncPackage2::FAsyncPackage2(const FAsyncPackageDesc& InDesc, int32 InSerialNu
 
 		// Add nodes for all imports and exports.
 		ImportNodes = PackageNodes + EEventLoadNode2::Package_NumPhases;
-		for (int32 LocalImportIndex = 0; LocalImportIndex < LocalImportCount; ++LocalImportIndex)
+		for (int32 LocalImportIndex = 0; LocalImportIndex < ImportCount; ++LocalImportIndex)
 		{
 			uint32 NodeIndex = EEventLoadNode2::Import_NumPhases * LocalImportIndex;
 			FEventLoadNode2* CreateImportNode = ImportNodes + NodeIndex + EEventLoadNode2::ImportOrExport_Create;
@@ -4439,8 +4433,6 @@ EAsyncPackageState::Type FAsyncPackage2::FinishLinker()
 {
 	LLM_SCOPE(ELLMTag::AsyncLoading);
 
-	SCOPED_LOADTIMER(LinkerLoad_FinalizeCreation);
-
 	const FPackageSummary* PackageSummary = reinterpret_cast<const FPackageSummary*>(PackageSummaryIoBuffer.Data());
 	ExportMap = reinterpret_cast<const FExportMapEntry*>(PackageSummaryIoBuffer.Data() + PackageSummary->ExportMapOffset);
 
