Fixed Clang compile error

#rnx

diff --git a/Engine/Source/Runtime/CoreUObject/Private/UObject/SavePackage.cpp b/Engine/Source/Runtime/CoreUObject/Private/UObject/SavePackage.cpp
--- a/Engine/Source/Runtime/CoreUObject/Private/UObject/SavePackage.cpp
+++ b/Engine/Source/Runtime/CoreUObject/Private/UObject/SavePackage.cpp
@@ -571,7 +571,7 @@ void FPackageStoreBulkDataManifest::PackageDesc::AddData(uint16 InIndex, uint64
 {
 	// The index is supposed to be unique
 	auto func = [=](const BulkDataDesc& Entry) { return Entry.Index == InIndex; };
-	check(Data.FindByPredicate(func) == false);
+	check(Data.FindByPredicate(func) == nullptr);
 
 	BulkDataDesc& Entry = Data.Emplace_GetRef();
 
