UnrealInsights: Remove some temp allocations in load time tracing

diff --git a/Engine/Source/Runtime/CoreUObject/Private/Serialization/AsyncLoading.cpp b/Engine/Source/Runtime/CoreUObject/Private/Serialization/AsyncLoading.cpp
--- a/Engine/Source/Runtime/CoreUObject/Private/Serialization/AsyncLoading.cpp
+++ b/Engine/Source/Runtime/CoreUObject/Private/Serialization/AsyncLoading.cpp
@@ -5188,7 +5188,7 @@ FAsyncPackage::FAsyncPackage(FAsyncLoadingThread& InThread, const FAsyncPackageD
 , FinishObjectsTime(0.0)
 #endif // PERF_TRACK_DETAILED_ASYNC_STATS
 {
-	TRACE_LOADTIME_NEW_ASYNC_PACKAGE(this, *InDesc.Name.ToString());
+	TRACE_LOADTIME_NEW_ASYNC_PACKAGE(this, InDesc.Name);
 	AddRequestID(InDesc.RequestID);
 }
 
diff --git a/Engine/Source/Runtime/CoreUObject/Private/Serialization/AsyncLoading2.cpp b/Engine/Source/Runtime/CoreUObject/Private/Serialization/AsyncLoading2.cpp
--- a/Engine/Source/Runtime/CoreUObject/Private/Serialization/AsyncLoading2.cpp
+++ b/Engine/Source/Runtime/CoreUObject/Private/Serialization/AsyncLoading2.cpp
@@ -4135,7 +4135,7 @@ FAsyncPackage2::FAsyncPackage2(const FAsyncPackageDesc& InDesc, int32 InSerialNu
 , bAllExportsSerialized(false)
 {
 	TRACE_CPUPROFILER_EVENT_SCOPE(NewAsyncPackage);
-	TRACE_LOADTIME_NEW_ASYNC_PACKAGE(this, *InDesc.Name.ToString());
+	TRACE_LOADTIME_NEW_ASYNC_PACKAGE(this, InDesc.Name);
 	AddRequestID(InDesc.RequestID);
 
 	const int32* NameIndex = AsyncLoadingThread.GlobalNameMap.GetIndex(Desc.NameToLoad);
diff --git a/Engine/Source/Runtime/CoreUObject/Private/Serialization/LoadTimeTrace.cpp b/Engine/Source/Runtime/CoreUObject/Private/Serialization/LoadTimeTrace.cpp
--- a/Engine/Source/Runtime/CoreUObject/Private/Serialization/LoadTimeTrace.cpp
+++ b/Engine/Source/Runtime/CoreUObject/Private/Serialization/LoadTimeTrace.cpp
@@ -9,6 +9,7 @@
 #include "Misc/CString.h"
 #include "HAL/PlatformTime.h"
 #include "HAL/PlatformTLS.h"
+#include "UObject/NameTypes.h"
 #include "UObject/Object.h"
 #include "UObject/UObjectGlobals.h"
 #include "Misc/CommandLine.h"
@@ -176,12 +177,13 @@ void FLoadTimeProfilerTracePrivate::OutputEndRequest(uint64 RequestId)
 		<< EndRequest.RequestId(RequestId);
 }
 
-void FLoadTimeProfilerTracePrivate::OutputNewAsyncPackage(const FGCObject* AsyncPackage, const TCHAR* PackageName)
+void FLoadTimeProfilerTracePrivate::OutputNewAsyncPackage(const FGCObject* AsyncPackage, const FName& PackageName)
 {
-	uint16 NameSize = (FCString::Strlen(PackageName) + 1) * sizeof(TCHAR);
+	TCHAR Buffer[FName::StringBufferSize];
+	uint16 NameSize = (PackageName.ToString(Buffer) + 1) * sizeof(TCHAR);
 	UE_TRACE_LOG(LoadTime, NewAsyncPackage, NameSize)
 		<< NewAsyncPackage.AsyncPackage(AsyncPackage)
-		<< NewAsyncPackage.Attachment(PackageName, NameSize);
+		<< NewAsyncPackage.Attachment(Buffer, NameSize);
 }
 
 void FLoadTimeProfilerTracePrivate::OutputBeginLoadAsyncPackage(const FGCObject* AsyncPackage)
@@ -227,6 +229,15 @@ void FLoadTimeProfilerTracePrivate::OutputAsyncPackageImportDependency(const FGC
 		<< AsyncPackageImportDependency.ImportedAsyncPackage(ImportedPackage);
 }
 
+void FLoadTimeProfilerTracePrivate::OutputClassInfo(const UClass* Class, const FName& Name)
+{
+	TCHAR Buffer[FName::StringBufferSize];
+	uint16 NameSize = (Name.ToString(Buffer) + 1) * sizeof(TCHAR);
+	UE_TRACE_LOG(LoadTime, ClassInfo, NameSize)
+		<< ClassInfo.Class(Class)
+		<< ClassInfo.Attachment(Buffer, NameSize);
+}
+
 void FLoadTimeProfilerTracePrivate::OutputClassInfo(const UClass* Class, const TCHAR* Name)
 {
 	uint16 NameSize = (FCString::Strlen(Name) + 1) * sizeof(TCHAR);
@@ -284,4 +295,4 @@ FLoadTimeProfilerTracePrivate::FPostLoadExportScope::~FPostLoadExportScope()
 		<< EndPostLoadExport.ThreadId(FPlatformTLS::GetCurrentThreadId());
 }
 
-#endif
\ No newline at end of file
+#endif
diff --git a/Engine/Source/Runtime/CoreUObject/Private/Serialization/LoadTimeTracePrivate.h b/Engine/Source/Runtime/CoreUObject/Private/Serialization/LoadTimeTracePrivate.h
--- a/Engine/Source/Runtime/CoreUObject/Private/Serialization/LoadTimeTracePrivate.h
+++ b/Engine/Source/Runtime/CoreUObject/Private/Serialization/LoadTimeTracePrivate.h
@@ -8,6 +8,7 @@
 
 class UObject;
 class FGCObject;
+class FName;
 
 struct FLoadTimeProfilerTracePrivate
 {
@@ -15,7 +16,7 @@ struct FLoadTimeProfilerTracePrivate
 	static void OutputStartAsyncLoading();
 	static void OutputSuspendAsyncLoading();
 	static void OutputResumeAsyncLoading();
-	static void OutputNewAsyncPackage(const FGCObject* AsyncPackage, const TCHAR* PackageName);
+	static void OutputNewAsyncPackage(const FGCObject* AsyncPackage, const FName& PackageName);
 	static void OutputBeginLoadAsyncPackage(const FGCObject* AsyncPackage);
 	static void OutputEndLoadAsyncPackage(const FGCObject* AsyncPackage);
 	static void OutputDestroyAsyncPackage(const FGCObject* AsyncPackage);
@@ -24,6 +25,7 @@ struct FLoadTimeProfilerTracePrivate
 	static void OutputPackageSummary(const FGCObject* AsyncPackage, uint32 TotalHeaderSize, uint32 ImportCount, uint32 ExportCount);
 	static void OutputAsyncPackageImportDependency(const FGCObject* Package, const FGCObject* ImportedPackage);
 	static void OutputAsyncPackageRequestAssociation(const FGCObject* AsyncPackage, uint64 RequestId);
+	static void OutputClassInfo(const UClass* Class, const FName& Name);
 	static void OutputClassInfo(const UClass* Class, const TCHAR* Name);
 
 	struct FCreateExportScope
@@ -119,4 +121,4 @@ struct FLoadTimeProfilerTracePrivate
 #define TRACE_LOADTIME_POSTLOAD_EXPORT_SCOPE(...)
 #define TRACE_LOADTIME_CLASS_INFO(...)
 
-#endif
\ No newline at end of file
+#endif
diff --git a/Engine/Source/Runtime/CoreUObject/Private/UObject/Class.cpp b/Engine/Source/Runtime/CoreUObject/Private/UObject/Class.cpp
--- a/Engine/Source/Runtime/CoreUObject/Private/UObject/Class.cpp
+++ b/Engine/Source/Runtime/CoreUObject/Private/UObject/Class.cpp
@@ -4294,11 +4294,7 @@ UClass::UClass(const FObjectInitializer& ObjectInitializer)
 	// If you add properties here, please update the other constructors and PurgeClass()
 
 	SetCppTypeInfoStatic(&DefaultCppClassTypeInfoStatic);
-#if LOADTIMEPROFILERTRACE_ENABLED
-	TCHAR Buffer[NAME_SIZE];
-	GetFName().GetPlainNameString(Buffer);
-	TRACE_LOADTIME_CLASS_INFO(this, Buffer);
-#endif
+	TRACE_LOADTIME_CLASS_INFO(this, GetFName());
 }
 
 /**
