AsyncLoading2: Remove inheritance of FGCObject from FAsyncPackage2.
Change FUObjectThreadContext::AsyncPackage and IEDLBootNotificationManager API from FGCObject* to void*.
Let FlinkerLoad::AsyncRoot be a FAsyncPackage* again.

diff --git a/Engine/Source/Runtime/CoreUObject/Private/Serialization/AsyncLoading.cpp b/Engine/Source/Runtime/CoreUObject/Private/Serialization/AsyncLoading.cpp
--- a/Engine/Source/Runtime/CoreUObject/Private/Serialization/AsyncLoading.cpp
+++ b/Engine/Source/Runtime/CoreUObject/Private/Serialization/AsyncLoading.cpp
@@ -119,11 +119,11 @@ static FThreadSafeCounter GPackageRequestID;
 struct FAsyncPackageScope
 {
 	/** Outer scope package */
-	FGCObject* PreviousPackage;
+	void* PreviousPackage;
 	/** Cached ThreadContext so we don't have to access it again */
 	FUObjectThreadContext& ThreadContext;
 
-	FAsyncPackageScope(FGCObject* InPackage)
+	FAsyncPackageScope(void* InPackage)
 		: ThreadContext(FUObjectThreadContext::Get())
 	{
 		PreviousPackage = ThreadContext.AsyncPackage;
@@ -5108,7 +5108,7 @@ void FAsyncLoadingThread::NotifyConstructedDuringAsyncLoading(UObject* Object, b
 		}
 	}
 
-void FAsyncLoadingThread::FireCompletedCompiledInImport(FGCObject* AsyncPackage, FPackageIndex Import)
+void FAsyncLoadingThread::FireCompletedCompiledInImport(void* AsyncPackage, FPackageIndex Import)
 {
 	FEventLoadNodePtr NodeToFire;
 	NodeToFire.WaitingPackage = FCheckedWeakAsyncPackagePtr(static_cast<FAsyncPackage*>(AsyncPackage));
diff --git a/Engine/Source/Runtime/CoreUObject/Private/Serialization/AsyncLoading2.cpp b/Engine/Source/Runtime/CoreUObject/Private/Serialization/AsyncLoading2.cpp
--- a/Engine/Source/Runtime/CoreUObject/Private/Serialization/AsyncLoading2.cpp
+++ b/Engine/Source/Runtime/CoreUObject/Private/Serialization/AsyncLoading2.cpp
@@ -662,7 +662,6 @@ uint32 FAsyncLoadingThreadState2::TlsSlot;
 */
 
 struct FAsyncPackage2
-	: public FGCObject
 {
 	friend struct FScopedAsyncPackageEvent2;
 
@@ -806,13 +805,6 @@ struct FAsyncPackage2
 		return ReentryCount > 1;
 	}
 
-	/** FGCObject Interface */
-	virtual void AddReferencedObjects(FReferenceCollector& Collector) override {}
-	virtual FString GetReferencerName() const override
-	{
-		return FString::Printf(TEXT("FAsyncPackage %s"), *GetPackageName().ToString());
-	}
-
 	void AddOwnedObjectFromCallback(UObject* Object, bool bSubObject)
 	{
 		if (bSubObject)
@@ -1405,7 +1397,7 @@ public:
 
 	void NotifyConstructedDuringAsyncLoading(UObject* Object, bool bSubObject);
 
-	void FireCompletedCompiledInImport(FGCObject* AsyncPackage, FPackageIndex Import);
+	void FireCompletedCompiledInImport(void* AsyncPackage, FPackageIndex Import);
 
 	/**
 	* [ASYNC THREAD] Finds an existing async package in the AsyncPackages by its name.
@@ -1584,11 +1576,11 @@ private:
 struct FAsyncPackageScope2
 {
 	/** Outer scope package */
-	FGCObject* PreviousPackage;
+	void* PreviousPackage;
 	/** Cached ThreadContext so we don't have to access it again */
 	FUObjectThreadContext& ThreadContext;
 
-	FAsyncPackageScope2(FGCObject* InPackage)
+	FAsyncPackageScope2(void* InPackage)
 		: ThreadContext(FUObjectThreadContext::Get())
 	{
 		PreviousPackage = ThreadContext.AsyncPackage;
@@ -3118,6 +3110,7 @@ EAsyncPackageState::Type FAsyncPackage2::Event_Tick(FAsyncPackage2* Package, int
 
 EAsyncPackageState::Type FAsyncPackage2::Event_Delete(FAsyncPackage2* Package, int32)
 {
+	TRACE_CPUPROFILER_EVENT_SCOPE(Event_Delete);
 	delete Package;
 	return EAsyncPackageState::Complete;
 }
@@ -4092,7 +4085,7 @@ void FAsyncLoadingThread2Impl::NotifyConstructedDuringAsyncLoading(UObject* Obje
 	AsyncPackage2->AddOwnedObjectFromCallback(Object, /*bForceAdd*/ bSubObject);
 }
 
-void FAsyncLoadingThread2Impl::FireCompletedCompiledInImport(FGCObject* AsyncPackage, FPackageIndex Import)
+void FAsyncLoadingThread2Impl::FireCompletedCompiledInImport(void* AsyncPackage, FPackageIndex Import)
 {
 	static_cast<FAsyncPackage2*>(AsyncPackage)->GetNode(EEventLoadNode2::ImportOrExport_Create, Import)->ReleaseBarrier();
 }
@@ -4306,8 +4299,6 @@ void FAsyncPackage2::ResetLoader()
 	// Reset loader.
 	if (Linker)
 	{
-		check(Linker->AsyncRoot == this || Linker->AsyncRoot == nullptr);
-		Linker->AsyncRoot = nullptr;
 		// Flush cache and queue for delete
 		Linker->FlushCache();
 		Linker->Detach();
@@ -4322,8 +4313,6 @@ void FAsyncPackage2::DetachLinker()
 	{
 		Linker->FlushCache();
 		checkf(bLoadHasFinished || bLoadHasFailed, TEXT("FAsyncPackage::DetachLinker called before load finished on package \"%s\""), *this->GetPackageName().ToString());
-		check(Linker->AsyncRoot == this || Linker->AsyncRoot == nullptr);
-		Linker->AsyncRoot = nullptr;
 		Linker = nullptr;
 	}
 }
@@ -4414,7 +4403,6 @@ EAsyncPackageState::Type FAsyncPackage2::CreateLinker()
 		Linker->bLockoutLegacyOperations = true;
 		Linker->SetIsLoading(true);
 		Linker->SetIsPersistent(true);
-		Linker->AsyncRoot = this;
 	}
 	Package->LinkerLoad = Linker;
 
@@ -4710,24 +4698,6 @@ EAsyncPackageState::Type FAsyncPackage2::FinishObjects()
 
 void FAsyncPackage2::CloseDelayedLinkers()
 {
-	// Close any linkers that have been open as a result of blocking load while async loading
-	for (FLinkerLoad* LinkerToClose : DelayedLinkerClosePackages)
-	{
-		if (LinkerToClose->LinkerRoot != nullptr)
-		{
-			check(LinkerToClose);
-			//check(LinkerToClose->LinkerRoot);
-			FLinkerLoad* LinkerToReset = FLinkerLoad::FindExistingLinkerForPackage(LinkerToClose->LinkerRoot);
-			check(LinkerToReset == LinkerToClose);
-			if (LinkerToReset && LinkerToReset->AsyncRoot)
-			{
-				UE_LOG(LogStreaming, Error, TEXT("Linker cannot be reset right now...leaking %s"), *LinkerToReset->GetArchiveName());
-				continue;
-			}
-		}
-		check(LinkerToClose->LinkerRoot == nullptr);
-		check(LinkerToClose->AsyncRoot == nullptr);
-	}
 }
 
 void FAsyncPackage2::CallCompletionCallbacks(bool bInternal, EAsyncLoadingResult::Type LoadingResult)
@@ -5024,7 +4994,7 @@ void FAsyncLoadingThread2::NotifyConstructedDuringAsyncLoading(UObject* Object,
 	Impl->NotifyConstructedDuringAsyncLoading(Object, bSubObject);
 }
 
-void FAsyncLoadingThread2::FireCompletedCompiledInImport(FGCObject* AsyncPackage, FPackageIndex Import)
+void FAsyncLoadingThread2::FireCompletedCompiledInImport(void* AsyncPackage, FPackageIndex Import)
 {
 	Impl->FireCompletedCompiledInImport(AsyncPackage, Import);
 }
diff --git a/Engine/Source/Runtime/CoreUObject/Private/Serialization/AsyncLoading2.h b/Engine/Source/Runtime/CoreUObject/Private/Serialization/AsyncLoading2.h
--- a/Engine/Source/Runtime/CoreUObject/Private/Serialization/AsyncLoading2.h
+++ b/Engine/Source/Runtime/CoreUObject/Private/Serialization/AsyncLoading2.h
@@ -24,7 +24,7 @@ public:
 	bool IsMultithreaded() override;
 	bool IsInAsyncLoadThread() override;
 	void NotifyConstructedDuringAsyncLoading(UObject* Object, bool bSubObject) override;
-	void FireCompletedCompiledInImport(FGCObject* AsyncPackage, FPackageIndex Import) override;
+	void FireCompletedCompiledInImport(void* AsyncPackage, FPackageIndex Import) override;
 	int32 LoadPackage(const FString& InPackageName, const FGuid* InGuid, const TCHAR* InPackageToLoadFrom, FLoadPackageAsyncDelegate InCompletionDelegate, EPackageFlags InPackageFlags, int32 InPIEInstanceID, int32 InPackagePriority) override;
 	EAsyncPackageState::Type ProcessLoading(bool bUseTimeLimit, bool bUseFullTimeLimit, float TimeLimit) override;
 	EAsyncPackageState::Type ProcessLoadingUntilComplete(TFunctionRef<bool()> CompletionPredicate, float TimeLimit) override;
diff --git a/Engine/Source/Runtime/CoreUObject/Private/Serialization/AsyncLoadingThread.h b/Engine/Source/Runtime/CoreUObject/Private/Serialization/AsyncLoadingThread.h
--- a/Engine/Source/Runtime/CoreUObject/Private/Serialization/AsyncLoadingThread.h
+++ b/Engine/Source/Runtime/CoreUObject/Private/Serialization/AsyncLoadingThread.h
@@ -291,7 +291,7 @@ public:
 
 	void NotifyConstructedDuringAsyncLoading(UObject* Object, bool bSubObject) override;
 
-	void FireCompletedCompiledInImport(FGCObject* AsyncPacakge, FPackageIndex Import) override;
+	void FireCompletedCompiledInImport(void* AsyncPacakge, FPackageIndex Import) override;
 
 	/** [EDL] Event queue */
 	FAsyncLoadEventQueue EventQueue;
diff --git a/Engine/Source/Runtime/CoreUObject/Private/Serialization/AsyncPackageLoader.cpp b/Engine/Source/Runtime/CoreUObject/Private/Serialization/AsyncPackageLoader.cpp
--- a/Engine/Source/Runtime/CoreUObject/Private/Serialization/AsyncPackageLoader.cpp
+++ b/Engine/Source/Runtime/CoreUObject/Private/Serialization/AsyncPackageLoader.cpp
@@ -22,7 +22,7 @@ struct FEDLBootObjectState
 
 struct FEDLBootWaitingPackage
 {
-	FGCObject* Package;
+	void* Package;
 	FPackageIndex Import;
 };
 
@@ -37,7 +37,7 @@ struct FEDLBootNotificationManager
 	FCriticalSection EDLBootNotificationManagerLock;
 
 	// return true if you are waiting for this compiled in object
-	bool AddWaitingPackage(FGCObject* Pkg, FName PackageName, FName ObjectName, FPackageIndex Import) override
+	bool AddWaitingPackage(void* Pkg, FName PackageName, FName ObjectName, FPackageIndex Import) override
 	{
 		if (PackageName == GLongCoreUObjectPackageName)
 		{
diff --git a/Engine/Source/Runtime/CoreUObject/Private/Serialization/LoadTimeTrace.cpp b/Engine/Source/Runtime/CoreUObject/Private/Serialization/LoadTimeTrace.cpp
--- a/Engine/Source/Runtime/CoreUObject/Private/Serialization/LoadTimeTrace.cpp
+++ b/Engine/Source/Runtime/CoreUObject/Private/Serialization/LoadTimeTrace.cpp
@@ -28,7 +28,7 @@ UE_TRACE_EVENT_BEGIN(LoadTime, ResumeAsyncLoading)
 UE_TRACE_EVENT_END()
 
 UE_TRACE_EVENT_BEGIN(LoadTime, PackageSummary)
-	UE_TRACE_EVENT_FIELD(const FGCObject*, AsyncPackage)
+	UE_TRACE_EVENT_FIELD(const void*, AsyncPackage)
 	UE_TRACE_EVENT_FIELD(uint32, TotalHeaderSize)
 	UE_TRACE_EVENT_FIELD(uint32, ImportCount)
 	UE_TRACE_EVENT_FIELD(uint32, ExportCount)
@@ -36,7 +36,7 @@ UE_TRACE_EVENT_END()
 
 UE_TRACE_EVENT_BEGIN(LoadTime, BeginCreateExport)
 	UE_TRACE_EVENT_FIELD(uint64, Cycle)
-	UE_TRACE_EVENT_FIELD(const FGCObject*, AsyncPackage)
+	UE_TRACE_EVENT_FIELD(const void*, AsyncPackage)
 	UE_TRACE_EVENT_FIELD(uint32, ThreadId)
 UE_TRACE_EVENT_END()
 
@@ -82,31 +82,31 @@ UE_TRACE_EVENT_BEGIN(LoadTime, EndRequest)
 UE_TRACE_EVENT_END()
 
 UE_TRACE_EVENT_BEGIN(LoadTime, NewAsyncPackage)
-	UE_TRACE_EVENT_FIELD(const FGCObject*, AsyncPackage)
+	UE_TRACE_EVENT_FIELD(const void*, AsyncPackage)
 UE_TRACE_EVENT_END()
 
 UE_TRACE_EVENT_BEGIN(LoadTime, BeginLoadAsyncPackage)
 	UE_TRACE_EVENT_FIELD(uint64, Cycle)
-	UE_TRACE_EVENT_FIELD(const FGCObject*, AsyncPackage)
+	UE_TRACE_EVENT_FIELD(const void*, AsyncPackage)
 UE_TRACE_EVENT_END()
 
 UE_TRACE_EVENT_BEGIN(LoadTime, EndLoadAsyncPackage)
 	UE_TRACE_EVENT_FIELD(uint64, Cycle)
-	UE_TRACE_EVENT_FIELD(const FGCObject*, AsyncPackage)
+	UE_TRACE_EVENT_FIELD(const void*, AsyncPackage)
 UE_TRACE_EVENT_END()
 
 UE_TRACE_EVENT_BEGIN(LoadTime, DestroyAsyncPackage)
-	UE_TRACE_EVENT_FIELD(const FGCObject*, AsyncPackage)
+	UE_TRACE_EVENT_FIELD(const void*, AsyncPackage)
 UE_TRACE_EVENT_END()
 
 UE_TRACE_EVENT_BEGIN(LoadTime, AsyncPackageRequestAssociation)
-	UE_TRACE_EVENT_FIELD(const FGCObject*, AsyncPackage)
+	UE_TRACE_EVENT_FIELD(const void*, AsyncPackage)
 	UE_TRACE_EVENT_FIELD(uint64, RequestId)
 UE_TRACE_EVENT_END()
 
 UE_TRACE_EVENT_BEGIN(LoadTime, AsyncPackageImportDependency)
-	UE_TRACE_EVENT_FIELD(const FGCObject*, AsyncPackage)
-	UE_TRACE_EVENT_FIELD(const FGCObject*, ImportedAsyncPackage)
+	UE_TRACE_EVENT_FIELD(const void*, AsyncPackage)
+	UE_TRACE_EVENT_FIELD(const void*, ImportedAsyncPackage)
 UE_TRACE_EVENT_END()
 
 UE_TRACE_EVENT_BEGIN(LoadTime, ClassInfo, Always)
@@ -177,7 +177,7 @@ void FLoadTimeProfilerTracePrivate::OutputEndRequest(uint64 RequestId)
 		<< EndRequest.RequestId(RequestId);
 }
 
-void FLoadTimeProfilerTracePrivate::OutputNewAsyncPackage(const FGCObject* AsyncPackage, const FName& PackageName)
+void FLoadTimeProfilerTracePrivate::OutputNewAsyncPackage(const void* AsyncPackage, const FName& PackageName)
 {
 	TCHAR Buffer[FName::StringBufferSize];
 	uint16 NameSize = (PackageName.ToString(Buffer) + 1) * sizeof(TCHAR);
@@ -186,27 +186,27 @@ void FLoadTimeProfilerTracePrivate::OutputNewAsyncPackage(const FGCObject* Async
 		<< NewAsyncPackage.Attachment(Buffer, NameSize);
 }
 
-void FLoadTimeProfilerTracePrivate::OutputBeginLoadAsyncPackage(const FGCObject* AsyncPackage)
+void FLoadTimeProfilerTracePrivate::OutputBeginLoadAsyncPackage(const void* AsyncPackage)
 {
 	UE_TRACE_LOG(LoadTime, BeginLoadAsyncPackage)
 		<< BeginLoadAsyncPackage.Cycle(FPlatformTime::Cycles64())
 		<< BeginLoadAsyncPackage.AsyncPackage(AsyncPackage);
 }
 
-void FLoadTimeProfilerTracePrivate::OutputEndLoadAsyncPackage(const FGCObject* AsyncPackage)
+void FLoadTimeProfilerTracePrivate::OutputEndLoadAsyncPackage(const void* AsyncPackage)
 {
 	UE_TRACE_LOG(LoadTime, EndLoadAsyncPackage)
 		<< EndLoadAsyncPackage.Cycle(FPlatformTime::Cycles64())
 		<< EndLoadAsyncPackage.AsyncPackage(AsyncPackage);
 }
 
-void FLoadTimeProfilerTracePrivate::OutputDestroyAsyncPackage(const FGCObject* AsyncPackage)
+void FLoadTimeProfilerTracePrivate::OutputDestroyAsyncPackage(const void* AsyncPackage)
 {
 	UE_TRACE_LOG(LoadTime, DestroyAsyncPackage)
 		<< DestroyAsyncPackage.AsyncPackage(AsyncPackage);
 }
 
-void FLoadTimeProfilerTracePrivate::OutputPackageSummary(const FGCObject* AsyncPackage, uint32 TotalHeaderSize, uint32 ImportCount, uint32 ExportCount)
+void FLoadTimeProfilerTracePrivate::OutputPackageSummary(const void* AsyncPackage, uint32 TotalHeaderSize, uint32 ImportCount, uint32 ExportCount)
 {
 	UE_TRACE_LOG(LoadTime, PackageSummary)
 		<< PackageSummary.AsyncPackage(AsyncPackage)
@@ -215,14 +215,14 @@ void FLoadTimeProfilerTracePrivate::OutputPackageSummary(const FGCObject* AsyncP
 		<< PackageSummary.ExportCount(ExportCount);
 }
 
-void FLoadTimeProfilerTracePrivate::OutputAsyncPackageRequestAssociation(const FGCObject* AsyncPackage, uint64 RequestId)
+void FLoadTimeProfilerTracePrivate::OutputAsyncPackageRequestAssociation(const void* AsyncPackage, uint64 RequestId)
 {
 	UE_TRACE_LOG(LoadTime, AsyncPackageRequestAssociation)
 		<< AsyncPackageRequestAssociation.AsyncPackage(AsyncPackage)
 		<< AsyncPackageRequestAssociation.RequestId(RequestId);
 }
 
-void FLoadTimeProfilerTracePrivate::OutputAsyncPackageImportDependency(const FGCObject* Package, const FGCObject* ImportedPackage)
+void FLoadTimeProfilerTracePrivate::OutputAsyncPackageImportDependency(const void* Package, const void* ImportedPackage)
 {
 	UE_TRACE_LOG(LoadTime, AsyncPackageImportDependency)
 		<< AsyncPackageImportDependency.AsyncPackage(Package)
@@ -246,7 +246,7 @@ void FLoadTimeProfilerTracePrivate::OutputClassInfo(const UClass* Class, const T
 		<< ClassInfo.Attachment(Name, NameSize);
 }
 
-FLoadTimeProfilerTracePrivate::FCreateExportScope::FCreateExportScope(const FGCObject* AsyncPackage, const UObject* const* InObject)
+FLoadTimeProfilerTracePrivate::FCreateExportScope::FCreateExportScope(const void* AsyncPackage, const UObject* const* InObject)
 	: Object(InObject)
 {
 	UE_TRACE_LOG(LoadTime, BeginCreateExport)
diff --git a/Engine/Source/Runtime/CoreUObject/Private/Serialization/LoadTimeTracePrivate.h b/Engine/Source/Runtime/CoreUObject/Private/Serialization/LoadTimeTracePrivate.h
--- a/Engine/Source/Runtime/CoreUObject/Private/Serialization/LoadTimeTracePrivate.h
+++ b/Engine/Source/Runtime/CoreUObject/Private/Serialization/LoadTimeTracePrivate.h
@@ -7,7 +7,6 @@
 #if LOADTIMEPROFILERTRACE_ENABLED
 
 class UObject;
-class FGCObject;
 class FName;
 
 struct FLoadTimeProfilerTracePrivate
@@ -16,21 +15,21 @@ struct FLoadTimeProfilerTracePrivate
 	static void OutputStartAsyncLoading();
 	static void OutputSuspendAsyncLoading();
 	static void OutputResumeAsyncLoading();
-	static void OutputNewAsyncPackage(const FGCObject* AsyncPackage, const FName& PackageName);
-	static void OutputBeginLoadAsyncPackage(const FGCObject* AsyncPackage);
-	static void OutputEndLoadAsyncPackage(const FGCObject* AsyncPackage);
-	static void OutputDestroyAsyncPackage(const FGCObject* AsyncPackage);
+	static void OutputNewAsyncPackage(const void* AsyncPackage, const FName& PackageName);
+	static void OutputBeginLoadAsyncPackage(const void* AsyncPackage);
+	static void OutputEndLoadAsyncPackage(const void* AsyncPackage);
+	static void OutputDestroyAsyncPackage(const void* AsyncPackage);
 	static void OutputBeginRequest(uint64 RequestId);
 	static void OutputEndRequest(uint64 RequestId);
-	static void OutputPackageSummary(const FGCObject* AsyncPackage, uint32 TotalHeaderSize, uint32 ImportCount, uint32 ExportCount);
-	static void OutputAsyncPackageImportDependency(const FGCObject* Package, const FGCObject* ImportedPackage);
-	static void OutputAsyncPackageRequestAssociation(const FGCObject* AsyncPackage, uint64 RequestId);
+	static void OutputPackageSummary(const void* AsyncPackage, uint32 TotalHeaderSize, uint32 ImportCount, uint32 ExportCount);
+	static void OutputAsyncPackageImportDependency(const void* Package, const void* ImportedPackage);
+	static void OutputAsyncPackageRequestAssociation(const void* AsyncPackage, uint64 RequestId);
 	static void OutputClassInfo(const UClass* Class, const FName& Name);
 	static void OutputClassInfo(const UClass* Class, const TCHAR* Name);
 
 	struct FCreateExportScope
 	{
-		FCreateExportScope(const FGCObject* AsyncPackage, const UObject* const* InObject);
+		FCreateExportScope(const void* AsyncPackage, const UObject* const* InObject);
 		~FCreateExportScope();
 
 	private:
diff --git a/Engine/Source/Runtime/CoreUObject/Public/Serialization/AsyncPackageLoader.h b/Engine/Source/Runtime/CoreUObject/Public/Serialization/AsyncPackageLoader.h
--- a/Engine/Source/Runtime/CoreUObject/Public/Serialization/AsyncPackageLoader.h
+++ b/Engine/Source/Runtime/CoreUObject/Public/Serialization/AsyncPackageLoader.h
@@ -43,14 +43,13 @@ void ClearFlagsAndDissolveClustersFromLoadedObjects(T& LoadedObjects)
 
 class IAsyncPackageLoader;
 class FPackageIndex;
-class FGCObject;
 
 class IEDLBootNotificationManager
 {
 public:
 	virtual ~IEDLBootNotificationManager() = default;
 
-	virtual bool AddWaitingPackage(FGCObject* Pkg, FName PackageName, FName ObjectName, FPackageIndex Import) = 0;
+	virtual bool AddWaitingPackage(void* Pkg, FName PackageName, FName ObjectName, FPackageIndex Import) = 0;
 	virtual bool ConstructWaitingBootObjects() = 0;
 	virtual bool FireCompletedCompiledInImports(bool bFinalRun = false) = 0;
 	virtual bool IsWaitingForSomething() = 0;
@@ -180,7 +179,7 @@ public:
 
 	virtual void NotifyConstructedDuringAsyncLoading(UObject* Object, bool bSubObject) = 0;
 
-	virtual void FireCompletedCompiledInImport(FGCObject* AsyncPackage, FPackageIndex Import) = 0;
+	virtual void FireCompletedCompiledInImport(void* AsyncPackage, FPackageIndex Import) = 0;
 };
 
 // Stats for ChartCreation.cpp
diff --git a/Engine/Source/Runtime/CoreUObject/Public/UObject/LinkerLoad.h b/Engine/Source/Runtime/CoreUObject/Public/UObject/LinkerLoad.h
--- a/Engine/Source/Runtime/CoreUObject/Public/UObject/LinkerLoad.h
+++ b/Engine/Source/Runtime/CoreUObject/Public/UObject/LinkerLoad.h
@@ -176,7 +176,7 @@ public:
 	}
 
 	/** The async package associated with this linker */
-	class FGCObject* AsyncRoot;
+	struct FAsyncPackage* AsyncRoot;
 #if WITH_EDITOR
 	/** Bulk data that does not need to be loaded when the linker is loaded.												*/
 	TArray<FUntypedBulkData*> BulkDataLoaders;
diff --git a/Engine/Source/Runtime/CoreUObject/Public/UObject/UObjectThreadContext.h b/Engine/Source/Runtime/CoreUObject/Public/UObject/UObjectThreadContext.h
--- a/Engine/Source/Runtime/CoreUObject/Public/UObject/UObjectThreadContext.h
+++ b/Engine/Source/Runtime/CoreUObject/Public/UObject/UObjectThreadContext.h
@@ -77,7 +77,7 @@ public:
 	/* Object that is currently being constructed with ObjectInitializer */
 	UObject* ConstructedObject;
 	/** Async Package currently processing objects */
-	class FGCObject* AsyncPackage;
+	void* AsyncPackage;
 
 #if !(UE_BUILD_SHIPPING || UE_BUILD_TEST)
 	/** Stack to ensure that PostInitProperties is routed through Super:: calls. **/
@@ -258,4 +258,4 @@ public:
 		return RefCount;
 	}
 	//~ TRefCountPtr interface
-};
\ No newline at end of file
+};
