AsyncLoading2: Add FPackageStore to extract some stuff from FAsyncLoadingThreadWorker.
Add FPackageImportStore to extract some stuff from FAsyncPackage2.
Rename FGlobalImportRuntime to FGlobalImportStore.
Rename FPackageStoreEntrySerialized to FPackageStoreEntry.
Remove FPackageStoreEntryRuntime.

diff --git a/Engine/Source/Runtime/CoreUObject/Private/Serialization/AsyncLoading2.cpp b/Engine/Source/Runtime/CoreUObject/Private/Serialization/AsyncLoading2.cpp
--- a/Engine/Source/Runtime/CoreUObject/Private/Serialization/AsyncLoading2.cpp
+++ b/Engine/Source/Runtime/CoreUObject/Private/Serialization/AsyncLoading2.cpp
@@ -107,116 +107,6 @@ public:
 	}
 };
 
-class FSimpleExportArchive
-	: public FSimpleArchive
-{
-public:
-	FSimpleExportArchive(const uint8* BufferPtr, uint64 BufferSize) : FSimpleArchive(BufferPtr, BufferSize) {}
-
-	using FArchive::operator<<; // For visibility of the overloads we don't override
-
-	//~ Begin FArchive::FArchiveUObject Interface
-	virtual FArchive& operator<<(FSoftObjectPath& Value) override { return FArchiveUObject::SerializeSoftObjectPath(*this, Value); }
-	virtual FArchive& operator<<(FWeakObjectPtr& Value) override { return FArchiveUObject::SerializeWeakObjectPtr(*this, Value); }
-	//~ End FArchive::FArchiveUObject Interface
-
-	//~ Begin FArchive::FLinkerLoad Interface
-	UObject* GetArchetypeFromLoader(const UObject* Obj) { return TemplateForGetArchetypeFromLoader; }
-
-	virtual bool AttachExternalReadDependency(FExternalReadCallback& ReadCallback) override
-	{
-		ExternalReadDependencies->Add(ReadCallback);
-		return true;
-	}
-
-	virtual FArchive& operator<<( UObject*& Object ) override
-	{
-		FPackageIndex Index;
-		FArchive& Ar = *this;
-		Ar << Index;
-
-		if (Index.IsNull())
-		{
-			Object = nullptr;
-		}
-		else if (Index.IsExport())
-		{
-			Object = (*Exports)[Index.ToExport()];
-		}
-		else
-		{
-			Object = GlobalImportObjects[ImportMap[Index.ToImport()]];
-		}
-		return *this;
-	}
-
-	inline virtual FArchive& operator<<(FLazyObjectPtr& LazyObjectPtr) override
-	{
-		FArchive& Ar = *this;
-		FUniqueObjectGuid ID;
-		Ar << ID;
-		LazyObjectPtr = ID;
-		return Ar;
-	}
-
-	inline virtual FArchive& operator<<(FSoftObjectPtr& Value) override
-	{
-		FArchive& Ar = *this;
-		FSoftObjectPath ID;
-		ID.Serialize(Ar);
-		Value = ID;
-		return Ar;
-	}
-
-	void BadNameIndexError(int32 NameIndex)
-	{
-		UE_LOG(LogLinker, Error, TEXT("Bad name index %i/%i"), NameIndex, NameMap->Num());
-	}
-
-	inline virtual FArchive& operator<<(FName& Name) override
-	{
-		FArchive& Ar = *this;
-		int32 NameIndex;
-		Ar << NameIndex;
-		int32 Number = 0;
-		Ar << Number;
-
-		if (NameMap->IsValidIndex(NameIndex))
-		{
-			// if the name wasn't loaded (because it wasn't valid in this context)
-			FNameEntryId MappedName = (*NameMap)[NameIndex];
-
-			// simply create the name from the NameMap's name and the serialized instance number
-			Name = FName::CreateFromDisplayId(MappedName, Number);
-		}
-		else
-		{
-			Name = FName();
-			BadNameIndexError(NameIndex);
-			ArIsError = true;
-			ArIsCriticalError = true;
-		}
-		return *this;
-	}
-	//~ End FArchive::FLinkerLoad Interface
-
-private:
-	friend FAsyncPackage2;
-
-	UObject* TemplateForGetArchetypeFromLoader = nullptr;
-	int32* ImportMap = nullptr;
-	UObject** GlobalImportObjects = nullptr;
-	const TArray<FNameEntryId>* NameMap = nullptr;
-	const TArray<UObject*>* Exports = nullptr;
-	TArray<FExternalReadCallback>* ExternalReadDependencies;
-};
-
-struct FPackageStoreEntryRuntime
-{
-	int32* Slimports = nullptr;
-	UPackage* Package = nullptr;
-};
-
 struct FGlobalPackageId
 {
 	int32 Id;
@@ -252,23 +142,7 @@ struct FGlobalImport
 	int32 Pad;
 };
 
-struct FGlobalImportRuntime
-{
-	/** Persistent data */
-	int32 Count = 0;
-	FName* Names = nullptr;
-	FPackageIndex* Outers = nullptr;
-	FPackageIndex* Packages = nullptr;
-	UObject** Objects = nullptr;
-	/** Reference tracking for GC management */
-	TAtomic<int32>* RefCounts = nullptr;
-	TArray<UObject*> KeepAliveObjects;
-	bool bNeedToRunGarbageCollect = false;
-	void OnPreGarbageCollect(bool bInIsLoadingPackages);
-	void OnPostGarbageCollect();
-};
-	
-struct FPackageStoreEntrySerialized
+struct FPackageStoreEntry
 {
 	FName Name;
 	FName FileName;
@@ -364,6 +238,348 @@ private:
 	TMap<FNameEntryId, int32> EntryToIndex;
 };
 
+struct FGlobalImportStore
+{
+	/** Persistent data */
+	int32 Count = 0;
+	FName* Names = nullptr;
+	FPackageIndex* Outers = nullptr;
+	FPackageIndex* Packages = nullptr;
+	UObject** Objects = nullptr;
+	/** Reference tracking for GC management */
+	TAtomic<int32>* RefCounts = nullptr;
+	TArray<UObject*> KeepAliveObjects;
+	bool bNeedToRunGarbageCollect = false;
+	void OnPreGarbageCollect(bool bInIsLoadingPackages);
+	void OnPostGarbageCollect();
+};
+
+class FPackageStore
+{
+	TMap<FName, FGlobalPackageId> PackageNameToGlobalPackageId;
+	FGlobalImportStore ImportStore;
+	FPackageStoreEntry* StoreEntries = nullptr;
+	int32* Slimports = nullptr;
+	int32 SlimportCount = 0;
+	int32 PackageCount = 0;
+
+public:
+	void Load(const FString& RootDir, FGlobalNameMap& GlobalNameMap)
+	{
+		{
+			TRACE_CPUPROFILER_EVENT_SCOPE(LoadPackageStoreToc);
+
+			TUniquePtr<FArchive> StoreTocArchive(IFileManager::Get().CreateFileReader(*(RootDir / TEXT("megafile.ustoretoc"))));
+			check(StoreTocArchive);
+
+			const int32 PackageByteCount = StoreTocArchive->TotalSize();
+			PackageCount = PackageByteCount / sizeof(FPackageStoreEntry);
+			StoreEntries = reinterpret_cast<FPackageStoreEntry*>(FMemory::Malloc(PackageByteCount));
+
+			// In-place loading
+			StoreTocArchive->Serialize(StoreEntries, PackageByteCount);
+			PackageNameToGlobalPackageId.Reserve(PackageCount);
+
+			// FName fixup
+			{
+				TRACE_CPUPROFILER_EVENT_SCOPE(LoadPackageStoreFNameFixup);
+				for (int I = 0; I < PackageCount; ++I)
+				{
+					StoreEntries[I].Name = GlobalNameMap.FromSerializedName(StoreEntries[I].Name);
+					StoreEntries[I].FileName = GlobalNameMap.FromSerializedName(StoreEntries[I].FileName);
+				}
+			}
+
+			// Global package index
+			{
+				TRACE_CPUPROFILER_EVENT_SCOPE(LoadPackageStoreFNameFixup);
+				for (int I = 0; I < PackageCount; ++I)
+				{
+					PackageNameToGlobalPackageId.Add(StoreEntries[I].Name, FGlobalPackageId { I });
+				}
+			}
+		}
+
+		{
+			TRACE_CPUPROFILER_EVENT_SCOPE(LoadPackageStoreSlimports);
+
+			TUniquePtr<FArchive> SlimportArchive(IFileManager::Get().CreateFileReader(*(RootDir / TEXT("megafile.uslimport"))));
+			check(SlimportArchive);
+
+			const int32 SlimportByteCount = SlimportArchive->TotalSize();
+			SlimportCount = SlimportByteCount / sizeof(int32);
+			Slimports = reinterpret_cast<int32*>(FMemory::Malloc(SlimportByteCount));
+			SlimportArchive->Serialize(Slimports, SlimportByteCount);
+		}
+
+		{
+			TRACE_CPUPROFILER_EVENT_SCOPE(LoadPackageStoreGlimports);
+
+			TUniquePtr<FArchive> ImportArchive(IFileManager::Get().CreateFileReader(*(RootDir / TEXT("megafile.uglimport"))));
+			check(ImportArchive);
+
+			const int32 ImportByteCount = ImportArchive->TotalSize();
+			ImportStore.Count = ImportByteCount / sizeof(FGlobalImport);
+			FGlobalImport* Imports = reinterpret_cast<FGlobalImport*>(FMemory::Malloc(ImportByteCount));
+			ImportArchive->Serialize(Imports, ImportByteCount);
+
+			TRACE_CPUPROFILER_EVENT_SCOPE(LoadPackageStoreGlimportsFixup);
+			{
+				ImportStore.Names = new FName[ImportStore.Count];
+				ImportStore.Outers = new FPackageIndex[ImportStore.Count];
+				ImportStore.Packages = new FPackageIndex[ImportStore.Count];
+				ImportStore.Objects = new UObject*[ImportStore.Count];
+				ImportStore.RefCounts = new TAtomic<int32>[ImportStore.Count];
+
+				for (int I = 0; I < ImportStore.Count; ++I)
+				{
+					FGlobalImport& Import = Imports[I];
+					Import.ObjectName = GlobalNameMap.FromSerializedName(Import.ObjectName);
+					ImportStore.Names[I] = Import.ObjectName;
+					ImportStore.Outers[I] = Import.OuterIndex;
+					ImportStore.Packages[I] = Import.OutermostIndex;
+					ImportStore.Objects[I] = nullptr;
+					ImportStore.RefCounts[I] = 0;
+				}
+			}
+		}
+
+	}
+
+	inline FGlobalImportStore& GetImportStore()
+	{
+		return ImportStore;
+	}
+
+	inline FGlobalPackageId* GetGlobalPackageId(FName Name)
+	{
+		FGlobalPackageId* Id = PackageNameToGlobalPackageId.Find(Name);
+		return Id;
+	}
+
+	inline UObject** GetGlobalImportObjects(int32& OutCount)
+	{
+		OutCount = ImportStore.Count;
+		return ImportStore.Objects;
+	}
+
+	inline FPackageIndex* GetGlobalImportOuters(int32& OutCount) const
+	{
+		OutCount = ImportStore.Count;
+		return ImportStore.Outers;
+	}
+
+	inline FName* GetGlobalImportNames(int32& OutCount) const
+	{
+		OutCount = ImportStore.Count;
+		return ImportStore.Names;
+	}
+
+	inline FPackageIndex* GetGlobalImportPackages(int32& OutCount) const 
+	{
+		OutCount = ImportStore.Count;
+		return ImportStore.Packages;
+	}
+
+	inline TAtomic<int32>* GetGlobalImportObjectRefCounts()
+	{
+		return ImportStore.RefCounts;
+	}
+
+	inline int32* GetPackageSlimports(FGlobalPackageId GlobalPackageId, int32& OutCount) const 
+	{
+		FPackageStoreEntry& Entry = StoreEntries[GlobalPackageId.Id];
+		OutCount = Entry.SlimportCount;
+		return Slimports + Entry.SlimportOffset / sizeof(int32);
+	}
+
+	inline int32 GetPackageImportCount(FGlobalPackageId GlobalPackageId) const 
+	{
+		return StoreEntries[GlobalPackageId.Id].ImportCount;
+	}
+
+	inline int32 GetPackageExportCount(FGlobalPackageId GlobalPackageId) const 
+	{
+		return StoreEntries[GlobalPackageId.Id].ExportCount;
+	}
+
+	inline FString GetPackageFileName(FGlobalPackageId GlobalPackageId)
+	{
+		return StoreEntries[GlobalPackageId.Id].FileName.ToString();
+	}
+
+};
+
+struct FPackageImportStore
+{
+	int32 Count = 0;
+	const int32* ImportMap = nullptr;
+	const FName* GlobalImportNames = nullptr;
+	const FPackageIndex* GlobalImportOuters = nullptr;
+	const FPackageIndex* GlobalImportPackages = nullptr;
+	UObject** GlobalImportObjects = nullptr;
+	TAtomic<int32>* GlobalImportObjectRefCounts = nullptr;
+
+	FPackageImportStore(FPackageStore& PackageStore, FGlobalPackageId GlobalPackageId)
+	{
+		TRACE_CPUPROFILER_EVENT_SCOPE(NewPackageImportStore);
+		int32 GlobalImportCount = 0;
+		ImportMap = PackageStore.GetPackageSlimports(GlobalPackageId, Count);
+		GlobalImportNames = PackageStore.GetGlobalImportNames(GlobalImportCount);
+		GlobalImportOuters = PackageStore.GetGlobalImportOuters(GlobalImportCount);
+		GlobalImportPackages = PackageStore.GetGlobalImportPackages(GlobalImportCount);
+		GlobalImportObjects = PackageStore.GetGlobalImportObjects(GlobalImportCount);
+		GlobalImportObjectRefCounts = PackageStore.GetGlobalImportObjectRefCounts();
+		AddGlobalImportObjectReferences();
+	}
+
+	~FPackageImportStore()
+	{
+		Clear();
+	}
+
+	UObject* GetObject(int32 LocalImportIndex)
+	{
+		return GlobalImportObjects[ImportMap[LocalImportIndex]];
+	}
+
+	UObject* GetObject(FPackageIndex Index)
+	{
+		return GlobalImportObjects[ImportMap[Index.ToImport()]];
+	}
+
+	void Clear()
+	{
+		if (ImportMap)
+		{
+			ReleaseGlobalImportObjectReferences();
+			ImportMap = nullptr;
+		}
+	}
+
+private:
+	void AddGlobalImportObjectReferences()
+	{
+		for (int32 LocalImportIndex = 0; LocalImportIndex < Count; ++LocalImportIndex)
+		{
+			const int32 GlobalImportIndex = ImportMap[LocalImportIndex];
+			++GlobalImportObjectRefCounts[GlobalImportIndex];
+		}
+	}
+
+	void ReleaseGlobalImportObjectReferences()
+	{
+		for (int32 LocalImportIndex = 0; LocalImportIndex < Count; ++LocalImportIndex)
+		{
+			const int32 GlobalImportIndex = ImportMap[LocalImportIndex];
+			--GlobalImportObjectRefCounts[GlobalImportIndex];
+		}
+	}
+};
+	
+class FSimpleExportArchive
+	: public FSimpleArchive
+{
+public:
+	FSimpleExportArchive(const uint8* BufferPtr, uint64 BufferSize) : FSimpleArchive(BufferPtr, BufferSize) {}
+
+	using FArchive::operator<<; // For visibility of the overloads we don't override
+
+	//~ Begin FArchive::FArchiveUObject Interface
+	virtual FArchive& operator<<(FSoftObjectPath& Value) override { return FArchiveUObject::SerializeSoftObjectPath(*this, Value); }
+	virtual FArchive& operator<<(FWeakObjectPtr& Value) override { return FArchiveUObject::SerializeWeakObjectPtr(*this, Value); }
+	//~ End FArchive::FArchiveUObject Interface
+
+	//~ Begin FArchive::FLinkerLoad Interface
+	UObject* GetArchetypeFromLoader(const UObject* Obj) { return TemplateForGetArchetypeFromLoader; }
+
+	virtual bool AttachExternalReadDependency(FExternalReadCallback& ReadCallback) override
+	{
+		ExternalReadDependencies->Add(ReadCallback);
+		return true;
+	}
+
+	virtual FArchive& operator<<( UObject*& Object ) override
+	{
+		FPackageIndex Index;
+		FArchive& Ar = *this;
+		Ar << Index;
+
+		if (Index.IsNull())
+		{
+			Object = nullptr;
+		}
+		else if (Index.IsExport())
+		{
+			Object = (*Exports)[Index.ToExport()];
+		}
+		else
+		{
+			Object = ImportStore->GetObject(Index);
+		}
+		return *this;
+	}
+
+	inline virtual FArchive& operator<<(FLazyObjectPtr& LazyObjectPtr) override
+	{
+		FArchive& Ar = *this;
+		FUniqueObjectGuid ID;
+		Ar << ID;
+		LazyObjectPtr = ID;
+		return Ar;
+	}
+
+	inline virtual FArchive& operator<<(FSoftObjectPtr& Value) override
+	{
+		FArchive& Ar = *this;
+		FSoftObjectPath ID;
+		ID.Serialize(Ar);
+		Value = ID;
+		return Ar;
+	}
+
+	void BadNameIndexError(int32 NameIndex)
+	{
+		UE_LOG(LogLinker, Error, TEXT("Bad name index %i/%i"), NameIndex, NameMap->Num());
+	}
+
+	inline virtual FArchive& operator<<(FName& Name) override
+	{
+		FArchive& Ar = *this;
+		int32 NameIndex;
+		Ar << NameIndex;
+		int32 Number = 0;
+		Ar << Number;
+
+		if (NameMap->IsValidIndex(NameIndex))
+		{
+			// if the name wasn't loaded (because it wasn't valid in this context)
+			FNameEntryId MappedName = (*NameMap)[NameIndex];
+
+			// simply create the name from the NameMap's name and the serialized instance number
+			Name = FName::CreateFromDisplayId(MappedName, Number);
+		}
+		else
+		{
+			Name = FName();
+			BadNameIndexError(NameIndex);
+			ArIsError = true;
+			ArIsCriticalError = true;
+		}
+		return *this;
+	}
+	//~ End FArchive::FLinkerLoad Interface
+
+private:
+	friend FAsyncPackage2;
+
+	UObject* TemplateForGetArchetypeFromLoader = nullptr;
+	FPackageImportStore* ImportStore = nullptr;
+	const TArray<FNameEntryId>* NameMap = nullptr;
+	const TArray<UObject*>* Exports = nullptr;
+	TArray<FExternalReadCallback>* ExternalReadDependencies;
+};
+
 enum class EChunkType : uint8
 {
 	None,
@@ -692,25 +908,7 @@ struct FAsyncPackage2
 			ImportedAsyncPackage->ReleaseRef();
 		}
 		ImportedAsyncPackages.Empty();
-		ReleaseGlobalImportObjectReferences();
-	}
-
-	void AddGlobalImportObjectReferences()
-	{
-		for (int32 LocalImportIndex = 0; LocalImportIndex < ImportCount; ++LocalImportIndex)
-		{
-			const int32 GlobalImportIndex = ImportMap[LocalImportIndex];
-			++GlobalImportObjectRefCounts[GlobalImportIndex];
-		}
-	}
-
-	void ReleaseGlobalImportObjectReferences()
-	{
-		for (int32 LocalImportIndex = 0; LocalImportIndex < ImportCount; ++LocalImportIndex)
-		{
-			const int32 GlobalImportIndex = ImportMap[LocalImportIndex];
-			--GlobalImportObjectRefCounts[GlobalImportIndex];
-		}
+		ImportStore.Clear();
 	}
 
 	/** Marks a specific request as complete */
@@ -959,17 +1157,10 @@ private:
 
 	// FZenLinkerLoad
 	TArray<FExternalReadCallback> ExternalReadDependencies;
-	int32 GlobalImportCount = 0;
-	int32 ImportCount = 0;
 	int32 ExportCount = 0;
 	const FExportMapEntry* ExportMap = nullptr;
 	TArray<UObject*> Exports;
-	int32* ImportMap = nullptr;
-	FName* GlobalImportNames = nullptr;
-	FPackageIndex* GlobalImportOuters = nullptr;
-	FPackageIndex* GlobalImportPackages = nullptr;
-	UObject** GlobalImportObjects = nullptr;
-	TAtomic<int32>* GlobalImportObjectRefCounts = nullptr;
+	FPackageImportStore ImportStore;
 public:
 
 	FAsyncLoadingThread2Impl& GetAsyncLoadingThread()
@@ -1238,65 +1429,7 @@ private:
 	FIoDispatcher IoDispatcher;
 	TUniquePtr<FIoQueue> IoQueue;
 
-	/** Package store */
-	FPackageStoreEntryRuntime* StoreEntriesRuntime = nullptr;
-	FPackageStoreEntrySerialized* StoreEntriesSerialized = nullptr;
-	int32* Slimports = nullptr;
-	int32 SlimportCount = 0;
-	FGlobalImportRuntime GlobalImportRuntime;
-	TMap<FName, FGlobalPackageId> PackageNameToGlobalPackageId;
-	int32 PackageCount = 0;
-
-	inline UObject** GetGlobalImportObjects(int32& OutCount)
-	{
-		OutCount = GlobalImportRuntime.Count;
-		return GlobalImportRuntime.Objects;
-	}
-
-	inline FPackageIndex* GetGlobalImportOuters(int32& OutCount) const
-	{
-		OutCount = GlobalImportRuntime.Count;
-		return GlobalImportRuntime.Outers;
-	}
-
-	inline FName* GetGlobalImportNames(int32& OutCount) const
-	{
-		OutCount = GlobalImportRuntime.Count;
-		return GlobalImportRuntime.Names;
-	}
-
-	inline FPackageIndex* GetGlobalImportPackages(int32& OutCount) const 
-	{
-		OutCount = GlobalImportRuntime.Count;
-		return GlobalImportRuntime.Packages;
-	}
-
-	inline TAtomic<int32>* GetGlobalImportObjectRefCounts()
-	{
-		return GlobalImportRuntime.RefCounts;
-	}
-
-	inline int32* GetPackageSlimports(FGlobalPackageId GlobalPackageId, int32& OutCount) const 
-	{
-		OutCount = StoreEntriesSerialized[GlobalPackageId.Id].SlimportCount;
-		return StoreEntriesRuntime[GlobalPackageId.Id].Slimports;
-	}
-
-	inline int32 GetPackageImportCount(FGlobalPackageId GlobalPackageId) const 
-	{
-		return StoreEntriesSerialized[GlobalPackageId.Id].ImportCount;
-	}
-
-	inline int32 GetPackageExportCount(FGlobalPackageId GlobalPackageId) const 
-	{
-		return StoreEntriesSerialized[GlobalPackageId.Id].ExportCount;
-	}
-
-	inline FString GetPackageFileName(FGlobalPackageId GlobalPackageId)
-	{
-		return StoreEntriesSerialized[GlobalPackageId.Id].FileName.ToString();
-	}
-
+	FPackageStore PackageStore;
 public:
 
 	//~ Begin FRunnable Interface.
@@ -1548,8 +1681,9 @@ private:
 
 	FAsyncPackage2* CreateAsyncPackage(const FAsyncPackageDesc& Desc)
 	{
-		FGlobalPackageId GlobalPackageId = PackageNameToGlobalPackageId[Desc.NameToLoad];
-		return new FAsyncPackage2(Desc, AsyncPackageSerialNumber.Increment(), *this, EDLBootNotificationManager, GraphAllocator, EventSpecs.GetData(), GlobalPackageId);
+		FGlobalPackageId* GlobalPackageId = PackageStore.GetGlobalPackageId(Desc.NameToLoad);
+		check(GlobalPackageId);
+		return new FAsyncPackage2(Desc, AsyncPackageSerialNumber.Increment(), *this, EDLBootNotificationManager, GraphAllocator, EventSpecs.GetData(), *GlobalPackageId);
 	}
 
 	/**
@@ -1636,6 +1770,11 @@ void FAsyncLoadingThread2Impl::InitializeLoading()
 		GlobalNameMap.Load(GlobalNameMapFilePath);
 	}
 
+	{
+		TRACE_CPUPROFILER_EVENT_SCOPE(LoadPackageStore);
+		PackageStore.Load(RootDir, GlobalNameMap);
+	}
+
 	{
 		TRACE_CPUPROFILER_EVENT_SCOPE(InitIoDispatcher);
 
@@ -1653,84 +1792,13 @@ void FAsyncLoadingThread2Impl::InitializeLoading()
 		}));
 	}
 
-	{
-		TRACE_CPUPROFILER_EVENT_SCOPE(LoadPackageStoreToc);
-
-		TUniquePtr<FArchive> StoreTocArchive(IFileManager::Get().CreateFileReader(*(RootDir / TEXT("megafile.ustoretoc"))));
-		check(StoreTocArchive)
-
-		const int32 PackageByteCount = StoreTocArchive->TotalSize();
-		PackageCount = PackageByteCount / sizeof(FPackageStoreEntrySerialized);
-		StoreEntriesSerialized = reinterpret_cast<FPackageStoreEntrySerialized*>(FMemory::Malloc(PackageByteCount));
-		// In-place loading
-		StoreTocArchive->Serialize(StoreEntriesSerialized, PackageByteCount);
-		// FName fixup
-		PackageNameToGlobalPackageId.Reserve(PackageCount);
-		for (int I = 0; I < PackageCount; ++I)
-		{
-			StoreEntriesSerialized[I].Name = GlobalNameMap.FromSerializedName(StoreEntriesSerialized[I].Name);
-			StoreEntriesSerialized[I].FileName = GlobalNameMap.FromSerializedName(StoreEntriesSerialized[I].FileName);
-			PackageNameToGlobalPackageId.Add(StoreEntriesSerialized[I].Name, FGlobalPackageId { I });
-		}
-	}
-
-	{
-		TRACE_CPUPROFILER_EVENT_SCOPE(LoadPackageStoreSlimports);
-
-		TUniquePtr<FArchive> SlimportArchive(IFileManager::Get().CreateFileReader(*(RootDir / TEXT("megafile.uslimport"))));
-		check(SlimportArchive);
-
-		const int32 SlimportByteCount = SlimportArchive->TotalSize();
-		SlimportCount = SlimportByteCount / sizeof(int32);
-		Slimports = reinterpret_cast<int32*>(FMemory::Malloc(SlimportByteCount));
-		SlimportArchive->Serialize(Slimports, SlimportByteCount);
-	}
-
-	{
-		TRACE_CPUPROFILER_EVENT_SCOPE(LoadPackageStoreGlimports);
-
-		TUniquePtr<FArchive> ImportArchive(IFileManager::Get().CreateFileReader(*(RootDir / TEXT("megafile.uglimport"))));
-		check(ImportArchive);
-
-		const int32 ImportByteCount = ImportArchive->TotalSize();
-		GlobalImportRuntime.Count = ImportByteCount / sizeof(FGlobalImport);
-		FGlobalImport* Imports = reinterpret_cast<FGlobalImport*>(FMemory::Malloc(ImportByteCount));
-		// In-place loading
-		ImportArchive->Serialize(Imports, ImportByteCount);
-
-		GlobalImportRuntime.Names = new FName[GlobalImportRuntime.Count];
-		GlobalImportRuntime.Outers = new FPackageIndex[GlobalImportRuntime.Count];
-		GlobalImportRuntime.Packages = new FPackageIndex[GlobalImportRuntime.Count];
-		GlobalImportRuntime.Objects = new UObject*[GlobalImportRuntime.Count];
-		GlobalImportRuntime.RefCounts = new TAtomic<int32>[GlobalImportRuntime.Count];
-
-		for (int I = 0; I < GlobalImportRuntime.Count; ++I)
-		{
-			FGlobalImport& Import = Imports[I];
-			Import.ObjectName = GlobalNameMap.FromSerializedName(Import.ObjectName);
-			GlobalImportRuntime.Names[I] = Import.ObjectName;
-			GlobalImportRuntime.Outers[I] = Import.OuterIndex;
-			GlobalImportRuntime.Packages[I] = Import.OutermostIndex;
-			GlobalImportRuntime.Objects[I] = nullptr;
-			GlobalImportRuntime.RefCounts[I] = 0;
-		}
-	}
-
-	StoreEntriesRuntime = new FPackageStoreEntryRuntime[PackageCount];
-	for (int I = 0; I < PackageCount; ++I)
-	{
-		FPackageStoreEntrySerialized& EntrySerialized = StoreEntriesSerialized[I];
-		FPackageStoreEntryRuntime& EntryRuntime = StoreEntriesRuntime[I];
-		EntryRuntime.Slimports = Slimports + EntrySerialized.SlimportOffset / sizeof(int32);
-	}
-
 	AsyncThreadReady.Increment();
 }
 
 void FAsyncLoadingThread2Impl::QueuePackage(FAsyncPackageDesc& Package)
 {
 	//TRACE_CPUPROFILER_EVENT_SCOPE(QueuePackage);
-	if (!PackageNameToGlobalPackageId.Find(Package.NameToLoad))
+	if (!PackageStore.GetGlobalPackageId(Package.NameToLoad))
 	{
 		UE_LOG(LogStreaming, Warning, TEXT("QueuePackage: Skipping unknown package: %s"), *Package.NameToLoad.ToString());
 		return;
@@ -2203,8 +2271,8 @@ EAsyncPackageState::Type FAsyncPackage2::Event_StartImportPackages(FAsyncPackage
 
 static UObject* GFindExistingSlimport(int32 GlobalImportIndex,
 	UObject** GlobalImportObjects,
-	FPackageIndex* GlobalImportOuters,
-	FName* GlobalImportNames)
+	const FPackageIndex* GlobalImportOuters,
+	const FName* GlobalImportNames)
 {
 	UObject*& Object = GlobalImportObjects[GlobalImportIndex];
 	if (!Object)
@@ -2230,9 +2298,9 @@ static UObject* GFindExistingSlimport(int32 GlobalImportIndex,
 UObject* FAsyncPackage2::FindExistingSlimport(int32 GlobalImportIndex)
 {
 	return GFindExistingSlimport(GlobalImportIndex,
-		GlobalImportObjects,
-		GlobalImportOuters,
-		GlobalImportNames);
+		ImportStore.GlobalImportObjects,
+		ImportStore.GlobalImportOuters,
+		ImportStore.GlobalImportNames);
 }
 
 void FAsyncPackage2::ImportPackagesRecursive()
@@ -2246,14 +2314,14 @@ void FAsyncPackage2::ImportPackagesRecursive()
 #if 1
 	FMemStack& MemStack = FMemStack::Get();
 	FMemMark Mark(MemStack);
-	TArrayView<FPackageIndex> LocalImportedPackages(new(MemStack)FPackageIndex[ImportCount], ImportCount);
+	TArrayView<FPackageIndex> LocalImportedPackages(new(MemStack)FPackageIndex[ImportStore.Count], ImportStore.Count);
 
-	for (int32 LocalImportIndex = 0; LocalImportIndex < ImportCount; ++LocalImportIndex)
+	for (int32 LocalImportIndex = 0; LocalImportIndex < ImportStore.Count; ++LocalImportIndex)
 	{
-		const int32 GlobalImportIndex = ImportMap[LocalImportIndex];
-		const FPackageIndex ImportedPackageIndex = GlobalImportPackages[GlobalImportIndex];
+		const int32 GlobalImportIndex = ImportStore.ImportMap[LocalImportIndex];
+		const FPackageIndex ImportedPackageIndex = ImportStore.GlobalImportPackages[GlobalImportIndex];
 		UObject* ImportedObject = FindExistingSlimport(GlobalImportIndex);
-		const UPackage* Package = (UPackage*)GlobalImportObjects[ImportedPackageIndex.ToImport()];
+		const UPackage* Package = (UPackage*)ImportStore.GlobalImportObjects[ImportedPackageIndex.ToImport()];
 
 		if (Package && Package->HasAnyPackageFlags(PKG_CompiledIn))
 		{
@@ -2275,7 +2343,7 @@ void FAsyncPackage2::ImportPackagesRecursive()
 
 		LocalImportedPackages[LocalImportIndex] = ImportedPackageIndex;
 
-		const FName ImportedPackageName = GlobalImportNames[ImportedPackageIndex.ToImport()];
+		const FName ImportedPackageName = ImportStore.GlobalImportNames[ImportedPackageIndex.ToImport()];
 		FAsyncPackageDesc Info(INDEX_NONE, ImportedPackageName);
 		Info.Priority = Desc.Priority;
 		bool bInserted;
@@ -2392,20 +2460,20 @@ EAsyncPackageState::Type FAsyncPackage2::SetupSlimports_Event()
 		return EAsyncPackageState::Complete;
 	}
 
-	for (int32 LocalImportIndex = 0; LocalImportIndex < ImportCount; ++LocalImportIndex)
+	for (int32 LocalImportIndex = 0; LocalImportIndex < ImportStore.Count; ++LocalImportIndex)
 	{
-		const int32 GlobalImportIndex = ImportMap[LocalImportIndex];
+		const int32 GlobalImportIndex = ImportStore.ImportMap[LocalImportIndex];
 
 		// skip packages
-		if (GlobalImportOuters[GlobalImportIndex].IsNull())
+		if (ImportStore.GlobalImportOuters[GlobalImportIndex].IsNull())
 		{
 			continue;
 		}
 
 		// find package of import object
-		FPackageIndex PackageIndex = GlobalImportPackages[GlobalImportIndex];
-		const FName& PackageName = GlobalImportNames[PackageIndex.ToImport()];
-		UObject*& Package = GlobalImportObjects[PackageIndex.ToImport()];
+		FPackageIndex PackageIndex = ImportStore.GlobalImportPackages[GlobalImportIndex];
+		const FName& PackageName = ImportStore.GlobalImportNames[PackageIndex.ToImport()];
+		UObject*& Package = ImportStore.GlobalImportObjects[PackageIndex.ToImport()];
 		UPackage* ImportPackage = Package ? CastChecked<UPackage>(Package) : nullptr;
 		if (!ImportPackage)
 		{
@@ -2424,7 +2492,7 @@ EAsyncPackageState::Type FAsyncPackage2::SetupSlimports_Event()
 			while (true)
 			{
 				check(!OuterMostIndex.IsNull() && OuterMostIndex.IsImport());
-				FPackageIndex NextOuterMostIndex = GlobalImportOuters[OuterMostIndex.ToImport()];
+				FPackageIndex NextOuterMostIndex = ImportStore.GlobalImportOuters[OuterMostIndex.ToImport()];
 
 				if (NextOuterMostIndex.IsNull())
 				{
@@ -2433,9 +2501,9 @@ EAsyncPackageState::Type FAsyncPackage2::SetupSlimports_Event()
 				OuterMostNonPackageIndex = OuterMostIndex;
 				OuterMostIndex = NextOuterMostIndex;
 			}
-			FName OuterMostNonPackageObjectName = GlobalImportNames[OuterMostNonPackageIndex.ToImport()];
-			check(GlobalImportOuters[OuterMostIndex.ToImport()].IsNull());
-			check(PackageName == GlobalImportNames[OuterMostIndex.ToImport()]);
+			FName OuterMostNonPackageObjectName = ImportStore.GlobalImportNames[OuterMostNonPackageIndex.ToImport()];
+			check(ImportStore.GlobalImportOuters[OuterMostIndex.ToImport()].IsNull());
+			check(PackageName == ImportStore.GlobalImportNames[OuterMostIndex.ToImport()]);
 			// OuterMostNonPackageIndex is used here because if it is a CDO or subobject, etc,
 			// we wait for the outermost thing that is not a package
 			const bool bWaitingForCompiledInImport = EDLBootNotificationManager.AddWaitingPackage(
@@ -2470,12 +2538,9 @@ EAsyncPackageState::Type FAsyncPackage2::Event_ImportSerialized(FAsyncPackage2*
 {
 	TRACE_CPUPROFILER_EVENT_SCOPE(Event_ImportSerialized);
 
-	int32 GlobalImportIndex = Package->ImportMap[LocalImportIndex];
-	UObject* Object = Package->GlobalImportObjects[GlobalImportIndex];
-	if (Object)
-	{
-		checkf(!Object->HasAnyFlags(RF_NeedLoad), TEXT("%s had RF_NeedLoad yet it was marked as serialized."), *Object->GetFullName()); // otherwise is isn't done serializing, is it?
-	}
+	UObject* Object = Package->ImportStore.GetObject(LocalImportIndex);
+	check(Object);
+	checkf(!Object->HasAnyFlags(RF_NeedLoad), TEXT("%s had RF_NeedLoad yet it was marked as serialized."), *Object->GetFullName()); // otherwise is isn't done serializing, is it?
 	return EAsyncPackageState::Complete;
 }
 
@@ -2511,40 +2576,40 @@ void FAsyncPackage2::LinkSlimport(int32 LocalImportIndex, int32 GlobalImportInde
 {
 	if (GlobalImportIndex == -1)
 	{
-		check(LocalImportIndex >= 0 && LocalImportIndex < ImportCount);
-		GlobalImportIndex = ImportMap[LocalImportIndex];
+		check(LocalImportIndex >= 0 && LocalImportIndex < ImportStore.Count);
+		GlobalImportIndex = ImportStore.ImportMap[LocalImportIndex];
 	}
 
-	UObject*& Object = GlobalImportObjects[GlobalImportIndex];
+	UObject*& Object = ImportStore.GlobalImportObjects[GlobalImportIndex];
 
 	if (Object)
 	{
 		return;
 	}
 
-	const FName ObjectName = GlobalImportNames[GlobalImportIndex];
-	const FPackageIndex PackageIndex = GlobalImportPackages[GlobalImportIndex];
+	const FName ObjectName = ImportStore.GlobalImportNames[GlobalImportIndex];
+	const FPackageIndex PackageIndex = ImportStore.GlobalImportPackages[GlobalImportIndex];
 
-	UObject*& ImportPackage = GlobalImportObjects[PackageIndex.ToImport()];
+	UObject*& ImportPackage = ImportStore.GlobalImportObjects[PackageIndex.ToImport()];
 	// UPackage* ImportPackage = (UPackage*)FindExistingSlimport(PackageIndex.ToImport());
 
 	if (!ImportPackage)
 	{
-		const FName PackageName = GlobalImportNames[PackageIndex.ToImport()];
+		const FName PackageName = ImportStore.GlobalImportNames[PackageIndex.ToImport()];
 		ImportPackage = FindObjectFast<UPackage>(NULL, PackageName, false, false);
 		check(ImportPackage);
 	}
 	if (PackageIndex.ToImport() != GlobalImportIndex)
 	{
 		UObject* Outer = ImportPackage;
-		FPackageIndex OuterIndex = GlobalImportOuters[GlobalImportIndex];
+		FPackageIndex OuterIndex = ImportStore.GlobalImportOuters[GlobalImportIndex];
 		if (PackageIndex != OuterIndex)
 		{
-			if (!GlobalImportObjects[OuterIndex.ToImport()])
+			if (!ImportStore.GlobalImportObjects[OuterIndex.ToImport()])
 			{
 				LinkSlimport(/*dummy_local_index*/-1, OuterIndex.ToImport());
 			}
-			Outer = GlobalImportObjects[OuterIndex.ToImport()];
+			Outer = ImportStore.GlobalImportObjects[OuterIndex.ToImport()];
 		}
 		check(Outer);
 
@@ -2574,7 +2639,7 @@ UObject* FAsyncPackage2::EventDrivenIndexToObject(FPackageIndex Index, bool bChe
 	}
 	else if (Index.IsImport())
 	{
-		int32 GlobalImportIndex = ImportMap[Index.ToImport()];
+		int32 GlobalImportIndex = ImportStore.ImportMap[Index.ToImport()];
 		Result = FindExistingSlimport(GlobalImportIndex);
 		check(Result);
 	}
@@ -2943,8 +3008,7 @@ void FAsyncPackage2::EventDrivenSerializeExport(int32 LocalExportIndex)
 
 		// FSimpleExportArchive special fields
 		Ar.NameMap = &AsyncLoadingThread.GlobalNameMap.GetNameEntries();
-		Ar.ImportMap = ImportMap;
-		Ar.GlobalImportObjects = GlobalImportObjects;
+		Ar.ImportStore = &ImportStore;
 		Ar.Exports = &Exports;
 		Ar.ExternalReadDependencies = &ExternalReadDependencies;
 
@@ -3770,7 +3834,7 @@ uint32 FAsyncLoadingThread2Impl::Run()
 		}
 		else
 		{
-			GlobalImportRuntime.bNeedToRunGarbageCollect |= IsAsyncLoadingPackages();
+			PackageStore.GetImportStore().bNeedToRunGarbageCollect |= IsAsyncLoadingPackages();
 			bool bDidSomething = false;
 			{
 				FGCScopeGuard GCGuard;
@@ -3963,7 +4027,7 @@ static void VerifyLoadFlagsWhenFinishedLoading()
 	}
 	UE_LOG(LogStreaming, Log, TEXT("Verified load flags when finished loading"));
 }
-void FGlobalImportRuntime::OnPreGarbageCollect(bool bInIsLoadingPackages)
+void FGlobalImportStore::OnPreGarbageCollect(bool bInIsLoadingPackages)
 {
 	if (!bNeedToRunGarbageCollect && !bInIsLoadingPackages)
 	{
@@ -4029,11 +4093,11 @@ void FGlobalImportRuntime::OnPreGarbageCollect(bool bInIsLoadingPackages)
 	}
 #endif
 
-	UE_LOG(LogStreaming, Log, TEXT("FGlobalImportRuntime::OnPreGarbageCollect - Marked %d objects, cleared %d object pointers"),
+	UE_LOG(LogStreaming, Log, TEXT("FGlobalImportStore::OnPreGarbageCollect - Marked %d objects, cleared %d object pointers"),
 		KeepAliveObjects.Num(), NumCleared);
 }
 
-void FGlobalImportRuntime::OnPostGarbageCollect()
+void FGlobalImportStore::OnPostGarbageCollect()
 {
 	if (KeepAliveObjects.Num() == 0)
 	{
@@ -4048,7 +4112,7 @@ void FGlobalImportRuntime::OnPostGarbageCollect()
 
 	const int32 UnmarkedCount = KeepAliveObjects.Num();
 	KeepAliveObjects.Reset();
-	UE_LOG(LogStreaming, Log, TEXT("FGlobalImportRuntime::UpdateGlobalImportsPostGC - Unmarked %d objects"),
+	UE_LOG(LogStreaming, Log, TEXT("FGlobalImportStore::UpdateGlobalImportsPostGC - Unmarked %d objects"),
 		UnmarkedCount);
 }
 
@@ -4056,13 +4120,13 @@ void FAsyncLoadingThread2Impl::OnPreGarbageCollect()
 {
 	TRACE_CPUPROFILER_EVENT_SCOPE(AltPreGC);
 	const bool bIsAsyncLoadingPackages = IsAsyncLoadingPackages();
-	GlobalImportRuntime.OnPreGarbageCollect(bIsAsyncLoadingPackages);
+	PackageStore.GetImportStore().OnPreGarbageCollect(bIsAsyncLoadingPackages);
 }
 
 void FAsyncLoadingThread2Impl::OnPostGarbageCollect()
 {
 	TRACE_CPUPROFILER_EVENT_SCOPE(AltPostGC);
-	GlobalImportRuntime.OnPostGarbageCollect();
+	PackageStore.GetImportStore().OnPostGarbageCollect();
 }
 
 /**
@@ -4122,7 +4186,7 @@ FAsyncPackage2::FAsyncPackage2(const FAsyncPackageDesc& InDesc, int32 InSerialNu
 , GraphAllocator(InGraphAllocator)
 , GlobalPackageId(InGlobalPackageId)
 // Begin EDL specific properties
-, GlobalImportObjectRefCounts(AsyncLoadingThread.GetGlobalImportObjectRefCounts())
+, ImportStore(AsyncLoadingThread.PackageStore, InGlobalPackageId)
 , AsyncPackageLoadingState(EAsyncPackageLoadingState2::NewPackage)
 , SerialNumber(InSerialNumber)
 , bAllExportsSerialized(false)
@@ -4136,15 +4200,7 @@ FAsyncPackage2::FAsyncPackage2(const FAsyncPackageDesc& InDesc, int32 InSerialNu
 	PackageChunkId = CreateChunkId(*NameIndex, Desc.NameToLoad.GetNumber(), 0, EChunkType::None);
 
 	{
-		TRACE_CPUPROFILER_EVENT_SCOPE(InitZenLinker);
-		ExportCount = AsyncLoadingThread.GetPackageExportCount(GlobalPackageId);
-		ImportMap = AsyncLoadingThread.GetPackageSlimports(GlobalPackageId, ImportCount);
-		GlobalImportNames = AsyncLoadingThread.GetGlobalImportNames(GlobalImportCount);
-		GlobalImportOuters = AsyncLoadingThread.GetGlobalImportOuters(GlobalImportCount);
-		GlobalImportPackages = AsyncLoadingThread.GetGlobalImportPackages(GlobalImportCount);
-		GlobalImportObjects = AsyncLoadingThread.GetGlobalImportObjects(GlobalImportCount);
-		AddGlobalImportObjectReferences();
-
+		ExportCount = AsyncLoadingThread.PackageStore.GetPackageExportCount(GlobalPackageId);
 		Exports.AddZeroed(ExportCount);
 		ExportIoBuffers.SetNum(ExportCount);
 		OwnedObjects.Reserve(ExportCount + 1); // ExportCount + UPackage
@@ -4152,7 +4208,7 @@ FAsyncPackage2::FAsyncPackage2(const FAsyncPackageDesc& InDesc, int32 InSerialNu
 
 	{
 		TRACE_CPUPROFILER_EVENT_SCOPE(CreateNodes);
-		ImportNodeCount = ImportCount * EEventLoadNode2::Import_NumPhases;
+		ImportNodeCount = ImportStore.Count * EEventLoadNode2::Import_NumPhases;
 		ExportNodeCount = ExportCount * EEventLoadNode2::Export_NumPhases;
 
 		PackageNodes = GraphAllocator.AllocNodes(EEventLoadNode2::Package_NumPhases + ImportNodeCount + ExportNodeCount);
@@ -4188,7 +4244,7 @@ FAsyncPackage2::FAsyncPackage2(const FAsyncPackageDesc& InDesc, int32 InSerialNu
 
 		// Add nodes for all imports and exports.
 		ImportNodes = PackageNodes + EEventLoadNode2::Package_NumPhases;
-		for (int32 LocalImportIndex = 0; LocalImportIndex < ImportCount; ++LocalImportIndex)
+		for (int32 LocalImportIndex = 0; LocalImportIndex < ImportStore.Count; ++LocalImportIndex)
 		{
 			uint32 NodeIndex = EEventLoadNode2::Import_NumPhases * LocalImportIndex;
 			FEventLoadNode2* CreateImportNode = ImportNodes + NodeIndex + EEventLoadNode2::ImportOrExport_Create;
@@ -4398,7 +4454,7 @@ EAsyncPackageState::Type FAsyncPackage2::CreateLinker()
 
 	{
 		uint32 LinkerFlags = LOAD_None | LOAD_Async | LOAD_NoVerify;
-		Linker = new FLinkerLoad(Package, *AsyncLoadingThread.GetPackageFileName(GlobalPackageId), LinkerFlags);
+		Linker = new FLinkerLoad(Package, *AsyncLoadingThread.PackageStore.GetPackageFileName(GlobalPackageId), LinkerFlags);
 		Linker->bIsAsyncLoader = false;
 		Linker->bLockoutLegacyOperations = true;
 		Linker->SetIsLoading(true);
@@ -4424,7 +4480,7 @@ EAsyncPackageState::Type FAsyncPackage2::FinishLinker()
 	const FPackageSummary* PackageSummary = reinterpret_cast<const FPackageSummary*>(PackageSummaryIoBuffer.Data());
 	ExportMap = reinterpret_cast<const FExportMapEntry*>(PackageSummaryIoBuffer.Data() + PackageSummary->ExportMapOffset);
 
-	TRACE_LOADTIME_PACKAGE_SUMMARY(this, PackageSummarySize, AsyncLoadingThread.GetPackageImportCount(GlobalPackageId), ExportCount);
+	TRACE_LOADTIME_PACKAGE_SUMMARY(this, PackageSummarySize, ImportStore.Count, ExportCount);
 
 	{
 		SCOPED_LOADTIMER(LinkerLoad_SerializePackageFileSummary);
