Cleaned up code that calls CreateStreamingRequest/CreateStreamingRequestForRange in the Bulkdata systems. Previous all places calling these methods would have two versions, one with a filename and one without, it makes more sense to have a macro that will implement or discard that single parameter rather than having large blocks of duplicated code.

This meant removing some of the custom code in the texture mipstreaming (which wouldn't have worked anyway as the wrong file name was being used) for adjusting offsets in the case of loading from .uexp files. I moved these checks into BulkData itself so that any code calling it would get the same fixes.

TextureMip and Audio streaming will now work if we have seperate bulkdata files disabled for the old Bulkdata api, static and dynamic mesh mip streaming will not as the code makes assumptions about the file paths that can not easily be fixed.

The overall goal of this was to make implementing the runtime for Bulkdata2 easier as a lot of the special case code I moved inside Bulkdata can be entirely ignored by the new system.

diff --git a/Engine/Source/Runtime/CoreUObject/Private/Serialization/BulkData.cpp b/Engine/Source/Runtime/CoreUObject/Private/Serialization/BulkData.cpp
--- a/Engine/Source/Runtime/CoreUObject/Private/Serialization/BulkData.cpp
+++ b/Engine/Source/Runtime/CoreUObject/Private/Serialization/BulkData.cpp
@@ -1608,10 +1608,25 @@ FBulkDataIORequest* FUntypedBulkData::CreateStreamingRequest(int64 OffsetInBulkD
 {
 	check(Filename.IsEmpty() == false);
 
-	UE_CLOG(IsStoredCompressedOnDisk(), LogSerialization, Fatal, TEXT("Package level compression is no longer supported (%s)."), *Filename);
-	UE_CLOG(GetBulkDataSize() <= 0, LogSerialization, Error, TEXT("(%s) has invalid bulk data size."), *Filename);
+	// If we are loading from a .uexp file then we need to adjust the filename and offset stored by BulkData in order to use them
+	// to access the data in the .uexp file. To keep the method const (due to a large number of places calling this, assuming that
+	// it is const) we take a copy of these data values which if needed can be adjusted and use them instead.
+	FString AdjustedFilename = Filename; 
+	int64 AdjustedBulkDataOffsetInFile = BulkDataOffsetInFile;
 
-	IAsyncReadFileHandle* IORequestHandle = FPlatformFileManager::Get().GetPlatformFile().OpenAsyncRead(*Filename);
+	// Fix up the Filename/Offset to work with streaming if we are loading from a .uexp file
+	if (AdjustedFilename.EndsWith(TEXT(".uasset")) || AdjustedFilename.EndsWith(TEXT(".umap")))
+	{
+		AdjustedBulkDataOffsetInFile -= IFileManager::Get().FileSize(*AdjustedFilename);
+
+		AdjustedFilename = FPaths::GetBaseFilename(AdjustedFilename, false) + TEXT(".uexp");
+		UE_LOG(LogSerialization, Error, TEXT("Streaming from the .uexp file '%s' this MUST be in a ubulk instead for best performance."), *AdjustedFilename);
+	}
+
+	UE_CLOG(IsStoredCompressedOnDisk(), LogSerialization, Fatal, TEXT("Package level compression is no longer supported (%s)."), *AdjustedFilename);
+	UE_CLOG(GetBulkDataSize() <= 0, LogSerialization, Error, TEXT("(%s) has invalid bulk data size."), *AdjustedFilename);
+
+	IAsyncReadFileHandle* IORequestHandle = FPlatformFileManager::Get().GetPlatformFile().OpenAsyncRead(*AdjustedFilename);
 	check(IORequestHandle); // this generally cannot fail because it is async
 
 	if (IORequestHandle == nullptr)
@@ -1619,7 +1634,7 @@ FBulkDataIORequest* FUntypedBulkData::CreateStreamingRequest(int64 OffsetInBulkD
 		return nullptr;
 	}
 
-	const int64 OffsetInFile = GetBulkDataOffsetInFile() + OffsetInBulkData;
+	const int64 OffsetInFile = AdjustedBulkDataOffsetInFile + OffsetInBulkData;
 
 	IAsyncReadRequest* ReadRequest = IORequestHandle->ReadRequest(OffsetInFile, BytesToRead, Priority, CompleteCallback, UserSuppliedMemory);
 	if (ReadRequest != nullptr)
diff --git a/Engine/Source/Runtime/CoreUObject/Public/Serialization/BulkData.h b/Engine/Source/Runtime/CoreUObject/Public/Serialization/BulkData.h
--- a/Engine/Source/Runtime/CoreUObject/Public/Serialization/BulkData.h
+++ b/Engine/Source/Runtime/CoreUObject/Public/Serialization/BulkData.h
@@ -19,6 +19,12 @@
 // Enable the following to use the more compact FBulkDataStreamingToken in places where it is implemented
 #define USE_BULKDATA_STREAMING_TOKEN !USE_NEW_BULKDATA
 
+#if USE_BULKDATA_STREAMING_TOKEN
+	#define STREAMINGTOKEN_PARAM(param) param,
+#else
+	#define STREAMINGTOKEN_PARAM(param) 
+#endif
+
 /**
  * Flags serialized with the bulk data.
  */
diff --git a/Engine/Source/Runtime/Engine/Private/Streaming/SkeletalMeshUpdate.cpp b/Engine/Source/Runtime/Engine/Private/Streaming/SkeletalMeshUpdate.cpp
--- a/Engine/Source/Runtime/Engine/Private/Streaming/SkeletalMeshUpdate.cpp
+++ b/Engine/Source/Runtime/Engine/Private/Streaming/SkeletalMeshUpdate.cpp
@@ -412,20 +412,12 @@ void FSkeletalMeshStreamIn_IO::SetIORequest(const FContext& Context, const FStri
 		// but that won't do anything because the tick would not try to acquire the lock since it is already locked.
 		TaskSynchronization.Increment();
 
-#if USE_BULKDATA_STREAMING_TOKEN
 		IORequest = FUntypedBulkData::CreateStreamingRequestForRange(
-			IOFilename,
+			STREAMINGTOKEN_PARAM(IOFilename)
 			FirstLOD.BulkDataStreamingToken,
 			LastLOD.BulkDataStreamingToken,
 			bHighPrioIORequest ? AIOP_BelowNormal : AIOP_Low,
 			&AsyncFileCallback);
-#else
-		IORequest = BulkDataUtils::CreateStreamingRequestForRange(
-			FirstLOD.StreamingBulkData,
-			LastLOD.StreamingBulkData,
-			bHighPrioIORequest ? AIOP_BelowNormal : AIOP_Low,
-			&AsyncFileCallback);
-#endif
 	}
 	else
 	{
diff --git a/Engine/Source/Runtime/Engine/Private/Streaming/StaticMeshUpdate.cpp b/Engine/Source/Runtime/Engine/Private/Streaming/StaticMeshUpdate.cpp
--- a/Engine/Source/Runtime/Engine/Private/Streaming/StaticMeshUpdate.cpp
+++ b/Engine/Source/Runtime/Engine/Private/Streaming/StaticMeshUpdate.cpp
@@ -386,20 +386,12 @@ void FStaticMeshStreamIn_IO::SetIORequest(const FContext& Context, const FString
 		// but that won't do anything because the tick would not try to acquire the lock since it is already locked.
 		TaskSynchronization.Increment();
 
-#if USE_BULKDATA_STREAMING_TOKEN
 		IORequest = FUntypedBulkData::CreateStreamingRequestForRange(
-			IOFilename,
+			STREAMINGTOKEN_PARAM(IOFilename)
 			FirstLOD.BulkDataStreamingToken,
 			LastLOD.BulkDataStreamingToken,
 			bHighPrioIORequest ? AIOP_BelowNormal : AIOP_Low,
 			&AsyncFileCallback);
-#else
-		IORequest = BulkDataUtils::CreateStreamingRequestForRange(
-			FirstLOD.StreamingBulkData,
-			LastLOD.StreamingBulkData,
-			bHighPrioIORequest ? AIOP_BelowNormal : AIOP_Low,
-			&AsyncFileCallback);
-#endif
 	}
 	else
 	{
diff --git a/Engine/Source/Runtime/Engine/Private/Streaming/Texture2DStreamIn_IO.cpp b/Engine/Source/Runtime/Engine/Private/Streaming/Texture2DStreamIn_IO.cpp
--- a/Engine/Source/Runtime/Engine/Private/Streaming/Texture2DStreamIn_IO.cpp
+++ b/Engine/Source/Runtime/Engine/Private/Streaming/Texture2DStreamIn_IO.cpp
@@ -17,7 +17,7 @@ Texture2DStreamIn.cpp: Stream in helper for 2D textures using texture streaming
 FTexture2DStreamIn_IO::FTexture2DStreamIn_IO(UTexture2D* InTexture, int32 InRequestedMips, bool InPrioritizedIORequest)
 	: FTexture2DStreamIn(InTexture, InRequestedMips)
 	, bPrioritizedIORequest(InPrioritizedIORequest)
-	, IOFileOffset(0)
+
 {
 	FMemory::Memzero(IORequests, sizeof(IORequests));
 }
@@ -42,19 +42,9 @@ void FTexture2DStreamIn_IO::SetIOFilename(const FContext& Context)
 	{
 		const FTexture2DMipMap& MipMap = OwnerMips[PendingFirstMip];
 	
-#if !TEXTURE2DMIPMAP_USE_COMPACT_BULKDATA
-		FString IOFilename;
-		IOFilename = MipMap.BulkData.GetFilename();
-#else
+#if TEXTURE2DMIPMAP_USE_COMPACT_BULKDATA
 		verify(Context.Texture->GetMipDataFilename(PendingFirstMip, IOFilename));
 #endif
-		if (IOFilename.EndsWith(TEXT(".uasset")) || IOFilename.EndsWith(TEXT(".umap")))
-		{
-			IOFileOffset = -IFileManager::Get().FileSize(*IOFilename);
-			check(IOFileOffset < 0);
-			IOFilename = FPaths::GetBaseFilename(IOFilename, false) + TEXT(".uexp");
-			UE_LOG(LogTexture, Error, TEXT("Streaming from the .uexp file '%s' this MUST be in a ubulk instead for best performance."), *IOFilename);
-		}
 	}
 }
 
@@ -75,22 +65,13 @@ void FTexture2DStreamIn_IO::SetIORequests(const FContext& Context)
 		// but that won't do anything because the tick would not try to acquire the lock since it is already locked.
 		TaskSynchronization.Increment();
 
-#if !TEXTURE2DMIPMAP_USE_COMPACT_BULKDATA
 		IORequests[MipIndex] = MipMap.BulkData.CreateStreamingRequest(
-			IOFileOffset,
+			TEXTURE2DMIPMAP_PARAM(IOFilename) // Only used if TEXTURE2DMIPMAP_USE_COMPACT_BULKDATA is enabled
+			0,
 			MipMap.BulkData.GetBulkDataSize(),
 			bPrioritizedIORequest ? AIOP_BelowNormal : AIOP_Low,
 			&AsyncFileCallBack,
 			(uint8*)MipData[MipIndex]);
-#else
-		IORequests[MipIndex] = MipMap.BulkData.CreateStreamingRequest(
-			IOFilename,
-			IOFileOffset,
-			MipMap.BulkData.GetBulkDataSize(),
-			bPrioritizedIORequest ? AIOP_BelowNormal : AIOP_Low,
-			&AsyncFileCallBack,
-			(uint8*)MipData[MipIndex]);
-#endif
 	}
 }
 
diff --git a/Engine/Source/Runtime/Engine/Private/Streaming/Texture2DStreamIn_IO.h b/Engine/Source/Runtime/Engine/Private/Streaming/Texture2DStreamIn_IO.h
--- a/Engine/Source/Runtime/Engine/Private/Streaming/Texture2DStreamIn_IO.h
+++ b/Engine/Source/Runtime/Engine/Private/Streaming/Texture2DStreamIn_IO.h
@@ -72,7 +72,6 @@ private:
 	FString IOFilename;
 #endif
 
-	int64 IOFileOffset;
 	FAsyncFileCallBack AsyncFileCallBack;
 };
 
diff --git a/Engine/Source/Runtime/Engine/Private/Texture2D.cpp b/Engine/Source/Runtime/Engine/Private/Texture2D.cpp
--- a/Engine/Source/Runtime/Engine/Private/Texture2D.cpp
+++ b/Engine/Source/Runtime/Engine/Private/Texture2D.cpp
@@ -348,10 +348,19 @@ void FTexture2DMipMap::FCompactByteBulkData::GetCopy(void** Dest, bool bDiscardI
 	}
 }
 
-FBulkDataIORequest* FTexture2DMipMap::FCompactByteBulkData::CreateStreamingRequest(const FString& Filename, int64 OffsetInBulkData, int64 BytesToRead, EAsyncIOPriorityAndFlags Priority, FAsyncFileCallBack* CompleteCallback, uint8* UserSuppliedMemory) const
+FBulkDataIORequest* FTexture2DMipMap::FCompactByteBulkData::CreateStreamingRequest(FString Filename, int64 OffsetInBulkData, int64 BytesToRead, EAsyncIOPriorityAndFlags Priority, FAsyncFileCallBack* CompleteCallback, uint8* UserSuppliedMemory) const
 {
 	check(Filename.IsEmpty() == false);
 
+	// Fix up the Filename/Offset to work with streaming if we are loading from a .uexp file
+	if (Filename.EndsWith(TEXT(".uasset")) || Filename.EndsWith(TEXT(".umap")))
+	{
+		OffsetInBulkData -= IFileManager::Get().FileSize(*Filename);
+
+		Filename = FPaths::GetBaseFilename(Filename, false) + TEXT(".uexp");
+		UE_LOG(LogTexture, Error, TEXT("Streaming from the .uexp file '%s' this MUST be in a ubulk instead for best performance."), *Filename);
+	}
+
 	UE_CLOG(IsStoredCompressedOnDisk(), LogSerialization, Fatal, TEXT("Package level compression is no longer supported (%s)."), *Filename);
 	UE_CLOG(GetBulkDataSize() <= 0, LogSerialization, Error, TEXT("(%s) has invalid bulk data size."), *Filename);
 
@@ -509,8 +518,12 @@ bool UTexture2D::GetMipDataFilename(const int32 MipIndex, FString& OutBulkDataFi
 			OutBulkDataFilename = PlatformData->Mips[MipIndex].BulkData.GetFilename();
 #else
 			OutBulkDataFilename = PlatformData->CachedPackageFileName;
-			const bool UseOptionalBulkDataFileName = PlatformData->Mips[MipIndex].BulkData.IsOptional();
-			OutBulkDataFilename = FPaths::ChangeExtension(OutBulkDataFilename, UseOptionalBulkDataFileName ? TEXT(".uptnl") : TEXT(".ubulk"));
+
+			if (PlatformData->Mips[MipIndex].BulkData.InSeperateFile())
+			{	
+				const bool UseOptionalBulkDataFileName = PlatformData->Mips[MipIndex].BulkData.IsOptional();
+				OutBulkDataFilename = FPaths::ChangeExtension(OutBulkDataFilename, UseOptionalBulkDataFileName ? TEXT(".uptnl") : TEXT(".ubulk"));
+			}
 #endif
 			return true;
 		}
diff --git a/Engine/Source/Runtime/Engine/Public/TextureResource.h b/Engine/Source/Runtime/Engine/Public/TextureResource.h
--- a/Engine/Source/Runtime/Engine/Public/TextureResource.h
+++ b/Engine/Source/Runtime/Engine/Public/TextureResource.h
@@ -30,11 +30,16 @@ class IVirtualTexture;
 #define MAX_TEXTURE_SOURCE_SLICES 6
 
 #ifndef FORCE_ENABLE_TEXTURE_STREAMING
-#define FORCE_ENABLE_TEXTURE_STREAMING 0
+	#define FORCE_ENABLE_TEXTURE_STREAMING 0
 #endif
 
 #define TEXTURE2DMIPMAP_USE_COMPACT_BULKDATA (!(WITH_EDITORONLY_DATA) && !(UE_SERVER) && FORCE_ENABLE_TEXTURE_STREAMING && PLATFORM_SUPPORTS_TEXTURE_STREAMING && !USE_NEW_BULKDATA)
 
+#if TEXTURE2DMIPMAP_USE_COMPACT_BULKDATA
+	#define TEXTURE2DMIPMAP_PARAM(param) param,
+#else
+	#define TEXTURE2DMIPMAP_PARAM(param) 
+#endif
 
 /**
  * A 2D texture mip-map.
@@ -79,7 +84,7 @@ struct FTexture2DMipMap
 		void* Realloc(int32 NumBytes);
 		void GetCopy(void** Dest, bool bDiscardInternalCopy = true);
 
-		FBulkDataIORequest* CreateStreamingRequest(const FString& Filename, int64 OffsetInBulkData, int64 BytesToRead, EAsyncIOPriorityAndFlags Priority, FAsyncFileCallBack* CompleteCallback, uint8* UserSuppliedMemory) const;
+		FBulkDataIORequest* CreateStreamingRequest(FString Filename, int64 OffsetInBulkData, int64 BytesToRead, EAsyncIOPriorityAndFlags Priority, FAsyncFileCallBack* CompleteCallback, uint8* UserSuppliedMemory) const;
 
 		/**
 		 * FCompactByteBulkData doesn't support GetFilename. Use FTexturePlatformData::CachedPackageFileName
