Further cleanup to the Bulkdata API.

Make sure that Bulkdata2 understands how to process BULKDATA_CookedForIoDispatcher during serialization
Remove BulkDataUtils namespace and just call FBulkDataInterface::CreateStreamingRequestForRange directly.
StaticMeshResources and SkeletalMeshLODRenderData use the same name for both FBulkDataStreamingToken and normal bulkdata options, this also let us use TChooseClass to select the correct type and reduces code bloat.

diff --git a/Engine/Source/Runtime/CoreUObject/Private/Serialization/BulkData.cpp b/Engine/Source/Runtime/CoreUObject/Private/Serialization/BulkData.cpp
--- a/Engine/Source/Runtime/CoreUObject/Private/Serialization/BulkData.cpp
+++ b/Engine/Source/Runtime/CoreUObject/Private/Serialization/BulkData.cpp
@@ -1676,7 +1676,7 @@ FBulkDataIORequest* FUntypedBulkData::CreateStreamingRequestForRange(const FStri
 	}
 
 	const int64 ReadOffset = Start.GetOffset();
-	const int64 ReadSize = (End.GetOffset() + End.GetSize()) - ReadOffset;
+	const int64 ReadSize = (End.GetOffset() + End.GetBulkDataSize()) - ReadOffset;
 
 	check(ReadSize > 0);
 
@@ -2024,18 +2024,6 @@ void FFloatBulkDataOld::SerializeElement( FArchive& Ar, void* Data, int64 Elemen
 	Ar << FloatData;
 }
 
-FBulkDataIORequest* BulkDataUtils::CreateStreamingRequestForRange(const FBulkDataInterface& Start, const FBulkDataInterface& End, EAsyncIOPriorityAndFlags Priority, FAsyncFileCallBack* CompleteCallback)
-{
-	check(Start.GetFilename() == End.GetFilename());
-
-	const int64 ReadOffset = Start.GetBulkDataOffsetInFile();
-	const int64 ReadSize = (End.GetBulkDataOffsetInFile() + End.GetBulkDataSize()) - ReadOffset;
-
-	check(ReadSize > 0);
-
-	return Start.CreateStreamingRequest(0, ReadSize, Priority, CompleteCallback, nullptr);
-}
-
 void FFormatContainer::Serialize(FArchive& Ar, UObject* Owner, const TArray<FName>* FormatsToSave, bool bSingleUse, uint32 InAlignment, bool bInline, bool bMapped)
 {
 	if (Ar.IsLoading())
diff --git a/Engine/Source/Runtime/CoreUObject/Private/Serialization/BulkData2.cpp b/Engine/Source/Runtime/CoreUObject/Private/Serialization/BulkData2.cpp
--- a/Engine/Source/Runtime/CoreUObject/Private/Serialization/BulkData2.cpp
+++ b/Engine/Source/Runtime/CoreUObject/Private/Serialization/BulkData2.cpp
@@ -217,6 +217,13 @@ void FBulkDataBase::Serialize(FArchive& Ar, UObject* Owner, int32 /*Index*/, boo
 
 		Ar << BulkDataOffsetInFile;
 
+		// Currently not used!
+		if ((BulkDataFlags & BULKDATA_CookedForIoDispatcher) != 0)
+		{
+			uint16 ChunkIndex;
+			Ar << ChunkIndex;
+		}
+
 		// Assuming that Owner/Package/Linker are all valid, the old BulkData system would
 		// generally fail if any of these were nullptr but had plenty of inconsistent checks
 		// scattered throughout.
@@ -464,7 +471,6 @@ FBulkDataIORequest* FBulkDataBase::CreateStreamingRequest(EAsyncIOPriorityAndFla
 
 FBulkDataIORequest* FBulkDataBase::CreateStreamingRequest(int64 OffsetInBulkData, int64 BytesToRead, EAsyncIOPriorityAndFlags Priority, FAsyncFileCallBack* CompleteCallback, uint8* UserSuppliedMemory) const
 {
-#if 1
 	FileTokenSystem::Data FileData = FileTokenSystem::GetFileData(Token);
 
 	check(FileData.Filename.IsEmpty() == false);
@@ -492,9 +498,18 @@ FBulkDataIORequest* FBulkDataBase::CreateStreamingRequest(int64 OffsetInBulkData
 		delete IORequestHandle;
 		return nullptr;
 	}
-#else
-	return nullptr;
-#endif
+}
+
+FBulkDataIORequest* FBulkDataBase::CreateStreamingRequestForRange(const FBulkDataBase& Start, const FBulkDataBase& End, EAsyncIOPriorityAndFlags Priority, FAsyncFileCallBack* CompleteCallback)
+{
+	check(Start.GetFilename() == End.GetFilename());
+
+	const int64 ReadOffset = Start.GetBulkDataOffsetInFile();
+	const int64 ReadSize = (End.GetBulkDataOffsetInFile() + End.GetBulkDataSize()) - ReadOffset;
+
+	check(ReadSize > 0);
+
+	return Start.CreateStreamingRequest(0, ReadSize, Priority, CompleteCallback, nullptr);
 }
 
 void FBulkDataBase::ForceBulkDataResident()
diff --git a/Engine/Source/Runtime/CoreUObject/Public/Serialization/BulkData.h b/Engine/Source/Runtime/CoreUObject/Public/Serialization/BulkData.h
--- a/Engine/Source/Runtime/CoreUObject/Public/Serialization/BulkData.h
+++ b/Engine/Source/Runtime/CoreUObject/Public/Serialization/BulkData.h
@@ -223,7 +223,7 @@ public:
 		return OffsetInFile;
 	}
 
-	uint32 GetSize() const
+	uint32 GetBulkDataSize() const
 	{
 		return BulkDataSize;
 	}
@@ -232,6 +232,8 @@ private:
 	uint32 OffsetInFile{ InvalidOffset };
 	uint32 BulkDataSize{ 0 };
 };
+#else
+class FBulkDataStreamingToken; // Forward declared but not implemented
 #endif
 
 /**
@@ -961,22 +963,6 @@ using FIntBulkData			= FIntBulkData2;
 using FFloatBulkData		= FFloatBulkData2;
 #endif
 
-namespace BulkDataUtils
-{
-	/**
-	* Create an async read request for a range of bulk data
-	* The request will read all data between the two given BulkData objects. They must both represent areas of data in the same file!
-	* The memory to be read into will be automatically allocated the size of which can be retrieved by calling FBulkDataIORequest::GetSize()
-	*
-	* @param Start				The bulk data to start reading from.
-	* @param End				The bulk data to finish reading from.
-	* @param Priority			Priority and flags of the request. If this includes AIOP_FLAG_PRECACHE, then memory will never be returned. The request should always be canceled and waited for, even for a precache request.
-	* @param CompleteCallback	Called from an arbitrary thread when the request is complete. Can be nullptr, if non-null, must remain valid until it is called. It will always be called.
-	* @return					A request for the read. This is owned by the caller and must be deleted by the caller.
-	**/
-	COREUOBJECT_API FBulkDataIORequest* CreateStreamingRequestForRange(const FBulkDataInterface& Start, const FBulkDataInterface& End, EAsyncIOPriorityAndFlags Priority, FAsyncFileCallBack* CompleteCallback);
-}
-
 class FFormatContainer
 {
 	friend class UBodySetup;
diff --git a/Engine/Source/Runtime/CoreUObject/Public/Serialization/BulkData2.h b/Engine/Source/Runtime/CoreUObject/Public/Serialization/BulkData2.h
--- a/Engine/Source/Runtime/CoreUObject/Public/Serialization/BulkData2.h
+++ b/Engine/Source/Runtime/CoreUObject/Public/Serialization/BulkData2.h
@@ -70,6 +70,7 @@ public:
 
 	FBulkDataIORequest* CreateStreamingRequest(EAsyncIOPriorityAndFlags Priority, FAsyncFileCallBack* CompleteCallback, uint8* UserSuppliedMemory) const;
 	FBulkDataIORequest* CreateStreamingRequest(int64 OffsetInBulkData, int64 BytesToRead, EAsyncIOPriorityAndFlags Priority, FAsyncFileCallBack* CompleteCallback, uint8* UserSuppliedMemory) const;
+	static FBulkDataIORequest* CreateStreamingRequestForRange(const FBulkDataBase& Start, const FBulkDataBase& End, EAsyncIOPriorityAndFlags Priority, FAsyncFileCallBack* CompleteCallback);
 
 	void RemoveBulkData();
 
diff --git a/Engine/Source/Runtime/Engine/Private/SkeletalMeshLODRenderData.cpp b/Engine/Source/Runtime/Engine/Private/SkeletalMeshLODRenderData.cpp
--- a/Engine/Source/Runtime/Engine/Private/SkeletalMeshLODRenderData.cpp
+++ b/Engine/Source/Runtime/Engine/Private/SkeletalMeshLODRenderData.cpp
@@ -1036,15 +1036,16 @@ void FSkeletalMeshLODRenderData::Serialize(FArchive& Ar, UObject* Owner, int32 I
 #endif
 			{
 #if USE_BULKDATA_STREAMING_TOKEN
-				FByteBulkData StreamingBulkData;
-				StreamingBulkData.Serialize(Ar, Owner, Idx, false);
+				FByteBulkData TmpBulkData;
+				TmpBulkData.Serialize(Ar, Owner, Idx, false);
+				bIsLODOptional = TmpBulkData.IsOptional();
 
-				BulkDataStreamingToken = StreamingBulkData.CreateStreamingToken();
+				StreamingBulkData = TmpBulkData.CreateStreamingToken();
 #else
 				StreamingBulkData.Serialize(Ar, Owner, Idx, false);
+				bIsLODOptional = StreamingBulkData.IsOptional();
 #endif
-				bIsLODOptional = !!(StreamingBulkData.GetBulkDataFlags() & BULKDATA_OptionalPayload);
-
+			
 				if (StreamingBulkData.GetBulkDataSize() == 0)
 				{
 					bDiscardBulkData = true;
diff --git a/Engine/Source/Runtime/Engine/Private/StaticMesh.cpp b/Engine/Source/Runtime/Engine/Private/StaticMesh.cpp
--- a/Engine/Source/Runtime/Engine/Private/StaticMesh.cpp
+++ b/Engine/Source/Runtime/Engine/Private/StaticMesh.cpp
@@ -674,15 +674,14 @@ void FStaticMeshLODResources::Serialize(FArchive& Ar, UObject* Owner, int32 Inde
 #if USE_BULKDATA_STREAMING_TOKEN
 				FByteBulkData TmpBulkData;
 				TmpBulkData.Serialize(Ar, Owner, Index, false);
-				bIsOptionalLOD = !!(TmpBulkData.GetBulkDataFlags() & BULKDATA_OptionalPayload);
+				bIsOptionalLOD = TmpBulkData.IsOptional();
 
-				BulkDataStreamingToken = TmpBulkData.CreateStreamingToken();			
-				BulkDataSize = BulkDataStreamingToken.GetSize();
+				StreamingBulkData = TmpBulkData.CreateStreamingToken();			
 #else
 				StreamingBulkData.Serialize(Ar, Owner, Index, false);
-				bIsOptionalLOD = !!(StreamingBulkData.GetBulkDataFlags() & BULKDATA_OptionalPayload);
-				BulkDataSize = (uint32)StreamingBulkData.GetBulkDataSize();
+				bIsOptionalLOD = StreamingBulkData.IsOptional();
 #endif
+				BulkDataSize = (uint32)StreamingBulkData.GetBulkDataSize();
 
 #if WITH_EDITORONLY_DATA
 				// Streaming CPU data in editor build isn't supported yet because tools and utils need access
diff --git a/Engine/Source/Runtime/Engine/Private/Streaming/SkeletalMeshUpdate.cpp b/Engine/Source/Runtime/Engine/Private/Streaming/SkeletalMeshUpdate.cpp
--- a/Engine/Source/Runtime/Engine/Private/Streaming/SkeletalMeshUpdate.cpp
+++ b/Engine/Source/Runtime/Engine/Private/Streaming/SkeletalMeshUpdate.cpp
@@ -412,10 +412,10 @@ void FSkeletalMeshStreamIn_IO::SetIORequest(const FContext& Context, const FStri
 		// but that won't do anything because the tick would not try to acquire the lock since it is already locked.
 		TaskSynchronization.Increment();
 
-		IORequest = FUntypedBulkData::CreateStreamingRequestForRange(
+		IORequest = FBulkDataInterface::CreateStreamingRequestForRange(
 			STREAMINGTOKEN_PARAM(IOFilename)
-			FirstLOD.BulkDataStreamingToken,
-			LastLOD.BulkDataStreamingToken,
+			FirstLOD.StreamingBulkData,
+			LastLOD.StreamingBulkData,
 			bHighPrioIORequest ? AIOP_BelowNormal : AIOP_Low,
 			&AsyncFileCallback);
 	}
diff --git a/Engine/Source/Runtime/Engine/Private/Streaming/StaticMeshUpdate.cpp b/Engine/Source/Runtime/Engine/Private/Streaming/StaticMeshUpdate.cpp
--- a/Engine/Source/Runtime/Engine/Private/Streaming/StaticMeshUpdate.cpp
+++ b/Engine/Source/Runtime/Engine/Private/Streaming/StaticMeshUpdate.cpp
@@ -386,10 +386,10 @@ void FStaticMeshStreamIn_IO::SetIORequest(const FContext& Context, const FString
 		// but that won't do anything because the tick would not try to acquire the lock since it is already locked.
 		TaskSynchronization.Increment();
 
-		IORequest = FUntypedBulkData::CreateStreamingRequestForRange(
+		IORequest = FBulkDataInterface::CreateStreamingRequestForRange(
 			STREAMINGTOKEN_PARAM(IOFilename)
-			FirstLOD.BulkDataStreamingToken,
-			LastLOD.BulkDataStreamingToken,
+			FirstLOD.StreamingBulkData,
+			LastLOD.StreamingBulkData,
 			bHighPrioIORequest ? AIOP_BelowNormal : AIOP_Low,
 			&AsyncFileCallback);
 	}
diff --git a/Engine/Source/Runtime/Engine/Public/Rendering/SkeletalMeshLODRenderData.h b/Engine/Source/Runtime/Engine/Public/Rendering/SkeletalMeshLODRenderData.h
--- a/Engine/Source/Runtime/Engine/Public/Rendering/SkeletalMeshLODRenderData.h
+++ b/Engine/Source/Runtime/Engine/Public/Rendering/SkeletalMeshLODRenderData.h
@@ -130,11 +130,7 @@ public:
 
 	uint32 BuffersSize;
 
-#if USE_BULKDATA_STREAMING_TOKEN
-	FBulkDataStreamingToken BulkDataStreamingToken;
-#else
-	FByteBulkData StreamingBulkData;
-#endif
+	typename TChooseClass<USE_BULKDATA_STREAMING_TOKEN, FBulkDataStreamingToken, FByteBulkData>::Result StreamingBulkData;
 
 	/** Whether buffers of this LOD is inlined (i.e. stored in .uexp instead of .ubulk) */
 	uint32 bStreamedDataInlined : 1;
diff --git a/Engine/Source/Runtime/Engine/Public/StaticMeshResources.h b/Engine/Source/Runtime/Engine/Public/StaticMeshResources.h
--- a/Engine/Source/Runtime/Engine/Public/StaticMeshResources.h
+++ b/Engine/Source/Runtime/Engine/Public/StaticMeshResources.h
@@ -382,11 +382,7 @@ struct FStaticMeshLODResources
 	/** Sum of all vertex and index buffer sizes. Calculated in SerializeBuffers */
 	uint32 BuffersSize;
 
-#if USE_BULKDATA_STREAMING_TOKEN
-	FBulkDataStreamingToken BulkDataStreamingToken;
-#else
-	FByteBulkData StreamingBulkData;
-#endif
+	typename TChooseClass<USE_BULKDATA_STREAMING_TOKEN, FBulkDataStreamingToken, FByteBulkData>::Result StreamingBulkData;
 
 #if STATS
 	uint32 StaticMeshIndexMemory;
