FBulkDataIORequest now derives from an interface IBulkDataIORequest which can be found in BulkData2.h
Currently this changes no functionality but will allow the IoDispatcher requests to be encapsulated in a different implementation with the calling code not needing to know about it.
This change is just moving code around.

diff --git a/Engine/Plugins/Experimental/GeometryCache/Source/GeometryCache/Private/StreamingGeometryCacheData.h b/Engine/Plugins/Experimental/GeometryCache/Source/GeometryCache/Private/StreamingGeometryCacheData.h
--- a/Engine/Plugins/Experimental/GeometryCache/Source/GeometryCache/Private/StreamingGeometryCacheData.h
+++ b/Engine/Plugins/Experimental/GeometryCache/Source/GeometryCache/Private/StreamingGeometryCacheData.h
@@ -8,7 +8,7 @@
 #include "Containers/Queue.h"
 
 class IAsyncReadRequest;
-struct FBulkDataIORequest;
+class IBulkDataIORequest;
 class UGeometryCacheTrackStreamable;
 class UGeometryCacheComponent;
 struct FStreamedGeometryCacheChunk;
@@ -21,7 +21,7 @@ struct FResidentChunk
 	void* Memory;
 	int32 DataSize;
 	int32 Refcount;
-	FBulkDataIORequest* IORequest; //null when resident, nonnull when being loaded
+	IBulkDataIORequest* IORequest; //null when resident, nonnull when being loaded
 };
 
 /**
diff --git a/Engine/Source/Runtime/CoreUObject/Private/Serialization/BulkData.cpp b/Engine/Source/Runtime/CoreUObject/Private/Serialization/BulkData.cpp
--- a/Engine/Source/Runtime/CoreUObject/Private/Serialization/BulkData.cpp
+++ b/Engine/Source/Runtime/CoreUObject/Private/Serialization/BulkData.cpp
@@ -1597,14 +1597,14 @@ void FUntypedBulkData::SerializeBulkData( FArchive& Ar, void* Data )
 	}
 }
 
-FBulkDataIORequest* FUntypedBulkData::CreateStreamingRequest(EAsyncIOPriorityAndFlags Priority, FAsyncFileCallBack* CompleteCallback, uint8* UserSuppliedMemory) const
+IBulkDataIORequest* FUntypedBulkData::CreateStreamingRequest(EAsyncIOPriorityAndFlags Priority, FAsyncFileCallBack* CompleteCallback, uint8* UserSuppliedMemory) const
 {
 	const int64 DataSize = GetBulkDataSize();
 
 	return CreateStreamingRequest(0, DataSize, Priority, CompleteCallback, UserSuppliedMemory);
 }
 
-FBulkDataIORequest* FUntypedBulkData::CreateStreamingRequest(int64 OffsetInBulkData, int64 BytesToRead, EAsyncIOPriorityAndFlags Priority, FAsyncFileCallBack* CompleteCallback, uint8* UserSuppliedMemory) const
+IBulkDataIORequest* FUntypedBulkData::CreateStreamingRequest(int64 OffsetInBulkData, int64 BytesToRead, EAsyncIOPriorityAndFlags Priority, FAsyncFileCallBack* CompleteCallback, uint8* UserSuppliedMemory) const
 {
 	check(Filename.IsEmpty() == false);
 
@@ -1661,7 +1661,7 @@ FBulkDataStreamingToken FUntypedBulkData::CreateStreamingToken() const
 	return FBulkDataStreamingToken((uint32)GetBulkDataOffsetInFile(), (uint32)GetBulkDataSize());
 }
 
-FBulkDataIORequest* FUntypedBulkData::CreateStreamingRequestForRange(const FString& Filename, const FBulkDataStreamingToken& Start, const FBulkDataStreamingToken& End, EAsyncIOPriorityAndFlags Priority, FAsyncFileCallBack* CompleteCallback)
+IBulkDataIORequest* FUntypedBulkData::CreateStreamingRequestForRange(const FString& Filename, const FBulkDataStreamingToken& Start, const FBulkDataStreamingToken& End, EAsyncIOPriorityAndFlags Priority, FAsyncFileCallBack* CompleteCallback)
 {
 	check(Filename.IsEmpty() == false);
 	check(Start.IsValid());
diff --git a/Engine/Source/Runtime/CoreUObject/Private/Serialization/BulkData2.cpp b/Engine/Source/Runtime/CoreUObject/Private/Serialization/BulkData2.cpp
--- a/Engine/Source/Runtime/CoreUObject/Private/Serialization/BulkData2.cpp
+++ b/Engine/Source/Runtime/CoreUObject/Private/Serialization/BulkData2.cpp
@@ -462,14 +462,14 @@ bool FBulkDataBase::IsMemoryMapped() const
 	return (BulkDataFlags & BULKDATA_MemoryMappedPayload) != 0;
 }
 
-FBulkDataIORequest* FBulkDataBase::CreateStreamingRequest(EAsyncIOPriorityAndFlags Priority, FAsyncFileCallBack* CompleteCallback, uint8* UserSuppliedMemory) const
+IBulkDataIORequest* FBulkDataBase::CreateStreamingRequest(EAsyncIOPriorityAndFlags Priority, FAsyncFileCallBack* CompleteCallback, uint8* UserSuppliedMemory) const
 {
 	const int64 DataSize = GetBulkDataSize();
 
 	return CreateStreamingRequest(0, DataSize, Priority, CompleteCallback, UserSuppliedMemory);
 }
 
-FBulkDataIORequest* FBulkDataBase::CreateStreamingRequest(int64 OffsetInBulkData, int64 BytesToRead, EAsyncIOPriorityAndFlags Priority, FAsyncFileCallBack* CompleteCallback, uint8* UserSuppliedMemory) const
+IBulkDataIORequest* FBulkDataBase::CreateStreamingRequest(int64 OffsetInBulkData, int64 BytesToRead, EAsyncIOPriorityAndFlags Priority, FAsyncFileCallBack* CompleteCallback, uint8* UserSuppliedMemory) const
 {
 	FileTokenSystem::Data FileData = FileTokenSystem::GetFileData(Token);
 
@@ -500,7 +500,7 @@ FBulkDataIORequest* FBulkDataBase::CreateStreamingRequest(int64 OffsetInBulkData
 	}
 }
 
-FBulkDataIORequest* FBulkDataBase::CreateStreamingRequestForRange(const FBulkDataBase& Start, const FBulkDataBase& End, EAsyncIOPriorityAndFlags Priority, FAsyncFileCallBack* CompleteCallback)
+IBulkDataIORequest* FBulkDataBase::CreateStreamingRequestForRange(const FBulkDataBase& Start, const FBulkDataBase& End, EAsyncIOPriorityAndFlags Priority, FAsyncFileCallBack* CompleteCallback)
 {
 	check(Start.GetFilename() == End.GetFilename());
 
diff --git a/Engine/Source/Runtime/CoreUObject/Public/Serialization/BulkData.h b/Engine/Source/Runtime/CoreUObject/Public/Serialization/BulkData.h
--- a/Engine/Source/Runtime/CoreUObject/Public/Serialization/BulkData.h
+++ b/Engine/Source/Runtime/CoreUObject/Public/Serialization/BulkData.h
@@ -149,16 +149,9 @@ private:
 	IMappedFileRegion* MappedRegion;
 };
 
-/**
- * Represents an IO request from the FUntypedBulkData streaming API.
- *
- * It functions pretty much the same as IAsyncReadRequest expect that it also holds
- * the file handle as well. 
- *
- * 
- */
-struct COREUOBJECT_API FBulkDataIORequest
+class FBulkDataIORequest : public IBulkDataIORequest
 {
+public:
 	FBulkDataIORequest(IAsyncReadFileHandle* InFileHandle, IAsyncReadRequest* InReadRequest, int64 BytesToRead)
 		: FileHandle(InFileHandle)
 		, ReadRequest(InReadRequest)
@@ -167,20 +160,20 @@ struct COREUOBJECT_API FBulkDataIORequest
 
 	}
 
-	~FBulkDataIORequest();
+	virtual ~FBulkDataIORequest();
 
-	bool PollCompletion() const;
-	bool WaitCompletion( float TimeLimitSeconds = 0.0f ) const;
+	virtual bool PollCompletion() const override;
+	virtual bool WaitCompletion( float TimeLimitSeconds = 0.0f ) const override;
 
-	uint8* GetReadResults() const;
-	int64 GetSize() const;
+	virtual uint8* GetReadResults() const override;
+	virtual int64 GetSize() const override;
 
-	bool operator == (const IAsyncReadRequest* InReadRequest)
+	virtual bool operator == (const IAsyncReadRequest* InReadRequest) override
 	{
 		return ReadRequest == InReadRequest;
 	}
 
-	void Cancel() const;
+	virtual void Cancel() const override;
 
 private:
 	IAsyncReadFileHandle* FileHandle;
@@ -684,7 +677,7 @@ public:
 	* @param UserSuppliedMemory A pointer to memory for the IO request to be written to, it is up to the caller to make sure that it is large enough. If the pointer is null then the system will allocate memory instead.
 	* @return					A request for the read. This is owned by the caller and must be deleted by the caller.
 	**/
-	FBulkDataIORequest* CreateStreamingRequest(EAsyncIOPriorityAndFlags Priority, FAsyncFileCallBack* CompleteCallback, uint8* UserSuppliedMemory) const;
+	IBulkDataIORequest* CreateStreamingRequest(EAsyncIOPriorityAndFlags Priority, FAsyncFileCallBack* CompleteCallback, uint8* UserSuppliedMemory) const;
 
 	/**
 	* Create an async read request for the bulk data.
@@ -697,7 +690,7 @@ public:
 	* @param UserSuppliedMemory A pointer to memory for the IO request to be written to, it is up to the caller to make sure that it is large enough. If the pointer is null then the system will allocate memory instead.
 	* @return					A request for the read. This is owned by the caller and must be deleted by the caller.
 	**/
-	FBulkDataIORequest* CreateStreamingRequest(int64 OffsetInBulkData, int64 BytesToRead, EAsyncIOPriorityAndFlags Priority, FAsyncFileCallBack* CompleteCallback, uint8* UserSuppliedMemory) const;
+	IBulkDataIORequest* CreateStreamingRequest(int64 OffsetInBulkData, int64 BytesToRead, EAsyncIOPriorityAndFlags Priority, FAsyncFileCallBack* CompleteCallback, uint8* UserSuppliedMemory) const;
 
 #if USE_BULKDATA_STREAMING_TOKEN 
 
@@ -713,7 +706,7 @@ public:
 	* Create an async read request for a range of bulk data streaming tokens
 	* The request will read all data between the two given streaming tokens objects. They must both represent areas of data in the file!
 	* There is no way to validate this and it is up to the caller to make sure that it is correct.
-	* The memory to be read into will be automatically allocated the size of which can be retrieved by calling FBulkDataIORequest::GetSize()
+	* The memory to be read into will be automatically allocated the size of which can be retrieved by calling IBulkDataIORequest::GetSize()
 	*
 	* @param Filename			The file to read from.
 	* @param Start				The bulk data to start reading from.
@@ -722,7 +715,7 @@ public:
 	* @param CompleteCallback	Called from an arbitrary thread when the request is complete. Can be nullptr, if non-null, must remain valid until it is called. It will always be called.
 	* @return					A request for the read. This is owned by the caller and must be deleted by the caller.
 	**/
-	static FBulkDataIORequest* CreateStreamingRequestForRange(const FString& Filename, const FBulkDataStreamingToken& Start, const FBulkDataStreamingToken& End, EAsyncIOPriorityAndFlags Priority, FAsyncFileCallBack* CompleteCallback);
+	static IBulkDataIORequest* CreateStreamingRequestForRange(const FString& Filename, const FBulkDataStreamingToken& Start, const FBulkDataStreamingToken& End, EAsyncIOPriorityAndFlags Priority, FAsyncFileCallBack* CompleteCallback);
 
 #endif
 
diff --git a/Engine/Source/Runtime/CoreUObject/Public/Serialization/BulkData2.h b/Engine/Source/Runtime/CoreUObject/Public/Serialization/BulkData2.h
--- a/Engine/Source/Runtime/CoreUObject/Public/Serialization/BulkData2.h
+++ b/Engine/Source/Runtime/CoreUObject/Public/Serialization/BulkData2.h
@@ -12,9 +12,30 @@
 // not been able to find any way to actually execute/test.
 #define BULKDATA_NOT_IMPLEMENTED_FOR_RUNTIME PLATFORM_BREAK()
 
-struct FBulkDataIORequest;
 struct FOwnedBulkDataPtr;
 
+/**
+ * Represents an IO request from the BulkData streaming API.
+ *
+ * It functions pretty much the same as IAsyncReadRequest expect that it also holds
+ * the file handle as well.
+ *
+ *
+ */
+class COREUOBJECT_API IBulkDataIORequest
+{
+public:
+	virtual bool PollCompletion() const = 0;
+	virtual bool WaitCompletion(float TimeLimitSeconds = 0.0f) const = 0;
+
+	virtual uint8* GetReadResults() const = 0;
+	virtual int64 GetSize() const = 0;
+
+	virtual bool operator == (const IAsyncReadRequest* InReadRequest) = 0;
+
+	virtual void Cancel() const = 0;
+};
+
 /**
  * @documentation @todo documentation
  */
@@ -68,9 +89,9 @@ public:
 	bool IsSingleUse() const;
 	bool IsMemoryMapped() const;
 
-	FBulkDataIORequest* CreateStreamingRequest(EAsyncIOPriorityAndFlags Priority, FAsyncFileCallBack* CompleteCallback, uint8* UserSuppliedMemory) const;
-	FBulkDataIORequest* CreateStreamingRequest(int64 OffsetInBulkData, int64 BytesToRead, EAsyncIOPriorityAndFlags Priority, FAsyncFileCallBack* CompleteCallback, uint8* UserSuppliedMemory) const;
-	static FBulkDataIORequest* CreateStreamingRequestForRange(const FBulkDataBase& Start, const FBulkDataBase& End, EAsyncIOPriorityAndFlags Priority, FAsyncFileCallBack* CompleteCallback);
+	IBulkDataIORequest* CreateStreamingRequest(EAsyncIOPriorityAndFlags Priority, FAsyncFileCallBack* CompleteCallback, uint8* UserSuppliedMemory) const;
+	IBulkDataIORequest* CreateStreamingRequest(int64 OffsetInBulkData, int64 BytesToRead, EAsyncIOPriorityAndFlags Priority, FAsyncFileCallBack* CompleteCallback, uint8* UserSuppliedMemory) const;
+	static IBulkDataIORequest* CreateStreamingRequestForRange(const FBulkDataBase& Start, const FBulkDataBase& End, EAsyncIOPriorityAndFlags Priority, FAsyncFileCallBack* CompleteCallback);
 
 	void RemoveBulkData();
 
diff --git a/Engine/Source/Runtime/Engine/Private/Streaming/SkeletalMeshUpdate.h b/Engine/Source/Runtime/Engine/Private/Streaming/SkeletalMeshUpdate.h
--- a/Engine/Source/Runtime/Engine/Private/Streaming/SkeletalMeshUpdate.h
+++ b/Engine/Source/Runtime/Engine/Private/Streaming/SkeletalMeshUpdate.h
@@ -192,7 +192,7 @@ protected:
 	/** Called by FAsyncCancelIORequestsTask to cancel inflight IO request if any */
 	void CancelIORequest();
 
-	struct FBulkDataIORequest* IORequest;
+	class IBulkDataIORequest* IORequest;
 	FAsyncFileCallBack AsyncFileCallback;
 	bool bHighPrioIORequest;
 };
diff --git a/Engine/Source/Runtime/Engine/Private/Streaming/StaticMeshUpdate.h b/Engine/Source/Runtime/Engine/Private/Streaming/StaticMeshUpdate.h
--- a/Engine/Source/Runtime/Engine/Private/Streaming/StaticMeshUpdate.h
+++ b/Engine/Source/Runtime/Engine/Private/Streaming/StaticMeshUpdate.h
@@ -183,7 +183,7 @@ protected:
 	/** Called by FAsyncCancelIORequestsTask to cancel inflight IO request if any */
 	void CancelIORequest();
 
-	struct FBulkDataIORequest* IORequest;
+	class IBulkDataIORequest* IORequest;
 	FAsyncFileCallBack AsyncFileCallback;
 	bool bHighPrioIORequest;
 };
diff --git a/Engine/Source/Runtime/Engine/Private/Streaming/Texture2DStreamIn_IO.cpp b/Engine/Source/Runtime/Engine/Private/Streaming/Texture2DStreamIn_IO.cpp
--- a/Engine/Source/Runtime/Engine/Private/Streaming/Texture2DStreamIn_IO.cpp
+++ b/Engine/Source/Runtime/Engine/Private/Streaming/Texture2DStreamIn_IO.cpp
@@ -25,7 +25,7 @@ FTexture2DStreamIn_IO::FTexture2DStreamIn_IO(UTexture2D* InTexture, int32 InRequ
 FTexture2DStreamIn_IO::~FTexture2DStreamIn_IO()
 {
 #if DO_CHECK
-	for (FBulkDataIORequest* IORequest : IORequests)
+	for (const IBulkDataIORequest* IORequest : IORequests)
 	{
 		check(!IORequest);
 	}
@@ -79,7 +79,7 @@ void FTexture2DStreamIn_IO::CancelIORequests()
 {
 	for (int32 MipIndex = 0; MipIndex < IORequests.Num(); ++MipIndex)
 	{
-		FBulkDataIORequest* IORequest = IORequests[MipIndex];
+		IBulkDataIORequest* IORequest = IORequests[MipIndex];
 		if (IORequest)
 		{
 			// Calling cancel will trigger the SetAsyncFileCallback() which will also try a tick but will fail.
@@ -95,7 +95,7 @@ void FTexture2DStreamIn_IO::ClearIORequests(const FContext& Context)
 
 	for (int32 MipIndex = PendingFirstMip; MipIndex < CurrentFirstMip; ++MipIndex)
 	{
-		FBulkDataIORequest* IORequest = IORequests[MipIndex];
+		IBulkDataIORequest* IORequest = IORequests[MipIndex];
 		IORequests[MipIndex] = nullptr;
 
 		if (IORequest != nullptr)
@@ -154,7 +154,7 @@ void FTexture2DStreamIn_IO::Abort()
 
 bool FTexture2DStreamIn_IO::HasPendingIORequests()
 {
-	for (FBulkDataIORequest* IORequest : IORequests)
+	for (IBulkDataIORequest* IORequest : IORequests)
 	{
 		if (IORequest != nullptr)
 		{
diff --git a/Engine/Source/Runtime/Engine/Private/Streaming/Texture2DStreamIn_IO.h b/Engine/Source/Runtime/Engine/Private/Streaming/Texture2DStreamIn_IO.h
--- a/Engine/Source/Runtime/Engine/Private/Streaming/Texture2DStreamIn_IO.h
+++ b/Engine/Source/Runtime/Engine/Private/Streaming/Texture2DStreamIn_IO.h
@@ -9,7 +9,7 @@ Texture2DStreamIn.h: Stream in helper for 2D textures using texture streaming fi
 #include "CoreMinimal.h"
 #include "Texture2DStreamIn.h"
 
-struct FBulkDataIORequest;
+class IBulkDataIORequest;
 
 // Base StreamIn framework exposing MipData
 class FTexture2DStreamIn_IO : public FTexture2DStreamIn
@@ -63,7 +63,7 @@ private:
 
 
 	// Request for loading into each mip.
-	TArray<FBulkDataIORequest*, TInlineAllocator<MAX_TEXTURE_MIP_COUNT> > IORequests;
+	TArray<IBulkDataIORequest*, TInlineAllocator<MAX_TEXTURE_MIP_COUNT> > IORequests;
 
 	bool bPrioritizedIORequest;
 
diff --git a/Engine/Source/Runtime/Engine/Private/Texture2D.cpp b/Engine/Source/Runtime/Engine/Private/Texture2D.cpp
--- a/Engine/Source/Runtime/Engine/Private/Texture2D.cpp
+++ b/Engine/Source/Runtime/Engine/Private/Texture2D.cpp
@@ -348,7 +348,7 @@ void FTexture2DMipMap::FCompactByteBulkData::GetCopy(void** Dest, bool bDiscardI
 	}
 }
 
-FBulkDataIORequest* FTexture2DMipMap::FCompactByteBulkData::CreateStreamingRequest(FString Filename, int64 OffsetInBulkData, int64 BytesToRead, EAsyncIOPriorityAndFlags Priority, FAsyncFileCallBack* CompleteCallback, uint8* UserSuppliedMemory) const
+IBulkDataIORequest* FTexture2DMipMap::FCompactByteBulkData::CreateStreamingRequest(FString Filename, int64 OffsetInBulkData, int64 BytesToRead, EAsyncIOPriorityAndFlags Priority, FAsyncFileCallBack* CompleteCallback, uint8* UserSuppliedMemory) const
 {
 	check(Filename.IsEmpty() == false);
 
diff --git a/Engine/Source/Runtime/Engine/Public/Animation/AnimationStreaming.h b/Engine/Source/Runtime/Engine/Public/Animation/AnimationStreaming.h
--- a/Engine/Source/Runtime/Engine/Public/Animation/AnimationStreaming.h
+++ b/Engine/Source/Runtime/Engine/Public/Animation/AnimationStreaming.h
@@ -23,7 +23,7 @@ struct FLoadedAnimationChunk
 {
 	TAtomic<FCompressedAnimSequence*> CompressedAnimData;
 
-	struct FBulkDataIORequest* IORequest;
+	class IBulkDataIORequest* IORequest;
 	double RequestStart;
 
 	uint32	Index;
diff --git a/Engine/Source/Runtime/Engine/Public/AudioStreaming.h b/Engine/Source/Runtime/Engine/Public/AudioStreaming.h
--- a/Engine/Source/Runtime/Engine/Public/AudioStreaming.h
+++ b/Engine/Source/Runtime/Engine/Public/AudioStreaming.h
@@ -102,7 +102,7 @@ struct FWaveRequest
 struct FLoadedAudioChunk
 {
 	uint8*	Data;
-	struct FBulkDataIORequest* IORequest;
+	class IBulkDataIORequest* IORequest;
 	int32	DataSize;
 	int32	AudioDataSize;
 	uint32	Index;
diff --git a/Engine/Source/Runtime/Engine/Public/AudioStreamingCache.h b/Engine/Source/Runtime/Engine/Public/AudioStreamingCache.h
--- a/Engine/Source/Runtime/Engine/Public/AudioStreamingCache.h
+++ b/Engine/Source/Runtime/Engine/Public/AudioStreamingCache.h
@@ -147,7 +147,7 @@ private:
 #endif
 
 		// Handle to our async read request operation.
-		TUniquePtr<FBulkDataIORequest> ReadRequest;
+		TUniquePtr<IBulkDataIORequest> ReadRequest;
 
 #if DEBUG_STREAM_CACHE
 		FCacheElementDebugInfo DebugInfo;
diff --git a/Engine/Source/Runtime/Engine/Public/TextureResource.h b/Engine/Source/Runtime/Engine/Public/TextureResource.h
--- a/Engine/Source/Runtime/Engine/Public/TextureResource.h
+++ b/Engine/Source/Runtime/Engine/Public/TextureResource.h
@@ -84,7 +84,7 @@ struct FTexture2DMipMap
 		void* Realloc(int32 NumBytes);
 		void GetCopy(void** Dest, bool bDiscardInternalCopy = true);
 
-		FBulkDataIORequest* CreateStreamingRequest(FString Filename, int64 OffsetInBulkData, int64 BytesToRead, EAsyncIOPriorityAndFlags Priority, FAsyncFileCallBack* CompleteCallback, uint8* UserSuppliedMemory) const;
+		IBulkDataIORequest* CreateStreamingRequest(FString Filename, int64 OffsetInBulkData, int64 BytesToRead, EAsyncIOPriorityAndFlags Priority, FAsyncFileCallBack* CompleteCallback, uint8* UserSuppliedMemory) const;
 
 		/**
 		 * FCompactByteBulkData doesn't support GetFilename. Use FTexturePlatformData::CachedPackageFileName
