Compilation fix for Linux builds, add the missing virtual destructor to IBulkDataIORequest.

diff --git a/Engine/Source/Runtime/CoreUObject/Public/Serialization/BulkData2.h b/Engine/Source/Runtime/CoreUObject/Public/Serialization/BulkData2.h
--- a/Engine/Source/Runtime/CoreUObject/Public/Serialization/BulkData2.h
+++ b/Engine/Source/Runtime/CoreUObject/Public/Serialization/BulkData2.h
@@ -25,6 +25,8 @@ struct FOwnedBulkDataPtr;
 class COREUOBJECT_API IBulkDataIORequest
 {
 public:
+	virtual ~IBulkDataIORequest() {}
+
 	virtual bool PollCompletion() const = 0;
 	virtual bool WaitCompletion(float TimeLimitSeconds = 0.0f) const = 0;
 
