Fixed incorrect UObject archetype caching when the object's Outer is still pending load.

diff --git a/Engine/Source/Runtime/CoreUObject/Private/UObject/UObjectArchetype.cpp b/Engine/Source/Runtime/CoreUObject/Private/UObject/UObjectArchetype.cpp
--- a/Engine/Source/Runtime/CoreUObject/Private/UObject/UObjectArchetype.cpp
+++ b/Engine/Source/Runtime/CoreUObject/Private/UObject/UObjectArchetype.cpp
@@ -156,7 +156,20 @@ UObject* UObject::GetArchetype() const
 		Archetype = GetArchetypeFromRequiredInfo(GetClass(), GetOuter(), GetFName(), GetFlags());
 		if (Archetype)
 		{
-			ArchetypeAnnotation.AddAnnotation(this, GUObjectArray.ObjectToIndex(Archetype));
+			bool bAddToCache = true;
+			if (GetOuter() &&
+				GetOuter()->GetClass() != UPackage::StaticClass() &&
+				HasAnyFlags(RF_InheritableComponentTemplate) &&
+				GetOuter()->IsA<UClass>() &&
+				GetOuter()->HasAnyFlags(RF_NeedLoad))
+			{
+				bAddToCache = false;
+			}
+
+			if (bAddToCache)
+			{
+				ArchetypeAnnotation.AddAnnotation(this, GUObjectArray.ObjectToIndex(Archetype));
+			}
 		}
 	}		
 	else
