Preparing to clean up up tagged property saving before for adding new unversioned serialization path.

* Added check that custom property lists are not used during tagged property saving. Tagged properties loading code doesn't handle custom properties so it doesn't make sense to either save or byte-count with custom properties. The few known uses cases all use binary serialization. Custom property list is intended as an optimization for binary serialization and used by blueprint instancing and nativization.
* Added check STRUCT_Atomic tagged structs, which means serialize densely, don't get sparsely delta-serialized. Some structs, like FVector_NetQuantize100, inherit the flag from an atomic base struct but do not inherit the STRUCT_SerializeNative flag which would prevent tagged properties from being called. It still seems like these structs never actually get delta-serialized with defaults.

Submitting these checks to validate my assumptions before deleting code.

diff --git a/Engine/Source/Runtime/CoreUObject/Private/UObject/Class.cpp b/Engine/Source/Runtime/CoreUObject/Private/UObject/Class.cpp
--- a/Engine/Source/Runtime/CoreUObject/Private/UObject/Class.cpp
+++ b/Engine/Source/Runtime/CoreUObject/Private/UObject/Class.cpp
@@ -1328,6 +1328,9 @@ void UStruct::SerializeTaggedProperties(FStructuredArchive::FSlot Slot, uint8* D
 		FStructuredArchive::FRecord PropertiesRecord = Slot.EnterRecord();
 
 		check(UnderlyingArchive.IsSaving() || UnderlyingArchive.IsCountingMemory());
+		checkf(!UnderlyingArchive.ArUseCustomPropertyList, 
+				TEXT("Custom property lists only work with binary serialization, not tagged property serialization. "
+					 "Attempted for struct '%s' and archive '%s'. "), *GetFName().ToString(), *UnderlyingArchive.GetArchiveName());
 
 		UScriptStruct* DefaultsScriptStruct = dynamic_cast<UScriptStruct*>(DefaultsStruct);
 
@@ -1338,6 +1341,10 @@ void UStruct::SerializeTaggedProperties(FStructuredArchive::FSlot Slot, uint8* D
 			bUseAtomicSerialization = DefaultsScriptStruct->ShouldSerializeAtomically(UnderlyingArchive);
 		}
 
+		checkf(!(Defaults && UnderlyingArchive.DoDelta() && !UnderlyingArchive.IsTransacting() && bUseAtomicSerialization), 
+				TEXT("Dense/'atomic' structs shouldn't be sparsely delta-serialized. Struct '%s', archive: '%s'."),
+				*DefaultsScriptStruct->GetFName().ToString(), *UnderlyingArchive.GetArchiveName());
+
 		// Save tagged properties.
 
 		// Iterate over properties in the order they were linked and serialize them.
