BulkData2 can now use the FIoDispatcher when it is enabled for the given platform.

NOTE: This is still disabled by default.

Inline data does not currently use the FIoDispatcher (I am not sure if we want to support reloading the inline data at runtime, we currently have no active use cases and it would inflate the toc) so uses the same fallback members in the union as we use when not running with the FIoDispatcher. Using a union allows us to keep the memory footprint very small. BulkDataSize was moved from the FileToken system to the Fallback members so that inline data does not need to register a token which prevents pointless memory allocations.
BulkData can be polled by calling ::IsUsingIODispatcher to see if it is using the FIoDispatcher or not.
NOTE: If the Zenloader is in use then BulkData will no longer need to try and access a FLinkerLoad and will not try to find any filenames.

Implemented a FIoDispatcher specific implementation FBulkDataIoDispatcherRequest, this mirrors the IPlatformFIle::OpenAsyncRead api as it moved the IO request to a AsyncTask and implements all of the same functionality and provides the same callback interface although that callback expects a IAsyncReadRequest pointer which it cannot prrovide. It also fakes the ability to use user supplied memory and issue multiple FIoDispatcher requests and have them returns in the same continious memory buffer. These features will probably be moved to FIoDispatcher and could be removed from this code over time.

The FIoDispatcher is passed to the BulkData system at the point of creation by FAsyncLoadingThread2Impl and held in a static member.
Changing ENABLE_IO_DISPATCHER to 0 will cause the system to read the loose files and not use IoDispatcher even if the data exists.
Removed BulkDataSize from the FileTokenSystem.
Removed the == operator from IBulkDataIORequest as it is no longer used, some methods had const removed as the FIoDispatcher implementation would not be const.

The old 16bit ChunkID has been replaced by a proper FIoChunkID which is generated by using info from the GlobalName map when cooking the bulkdata (BulkData will save out dummy values that are filled in later when actually saving the package, such as offsets etc, the ChunkID is now one of these values)
Previously I was writing out the uint16 ChunkID during saving of BulkData but it should've been when saving AND cooking to avoid writing it to the uasset or DDC, in order to fix this I had to rename BULKDATA_CookedForIoDispatcher to BULKDATA_BadDataVersion so I can specifically check for it on load and discard the value. Once things stablize I can increase the DDC keys for the data types that write BulkData to the DDC and remove BULKDATA_BadDataVersion.

::CreateStreamingRequestForRange now takes for the full range of reads requested rather than just the start and end in the form of a TArray (BulkDataRangeArray). This is so we can make a unique request for each read when using the FIoDispatcher and keeps the old and new code paths the same in calling code.

Known issues: Optional data will still not function correctly and Virtual Textures are still using older code paths and require the loose ubulk files still. The IOS memory mapping systems have no actual support. This has added more code paths with BULKDATA_NOT_IMPLEMENTED_FOR_RUNTIME which will need to be handled properly soon (either fixed up or log warnings for improper use)

diff --git a/Engine/Source/Runtime/Core/Private/IO/IoDispatcher.cpp b/Engine/Source/Runtime/Core/Private/IO/IoDispatcher.cpp
--- a/Engine/Source/Runtime/Core/Private/IO/IoDispatcher.cpp
+++ b/Engine/Source/Runtime/Core/Private/IO/IoDispatcher.cpp
@@ -24,8 +24,9 @@
 
 DEFINE_LOG_CATEGORY_STATIC(LogIoDispatch, Log, All);
 
-#if !defined(PLATFORM_IMPLEMENTS_IO)
+const FIoChunkId FIoChunkId::InvalidChunkId = FIoChunkId::CreateEmptyId();
 
+#if !defined(PLATFORM_IMPLEMENTS_IO)
 //////////////////////////////////////////////////////////////////////////
 
 UE_TRACE_EVENT_BEGIN(IoDispatcher, BatchIssued, Always)
diff --git a/Engine/Source/Runtime/Core/Public/IO/IoDispatcher.h b/Engine/Source/Runtime/Core/Public/IO/IoDispatcher.h
--- a/Engine/Source/Runtime/Core/Public/IO/IoDispatcher.h
+++ b/Engine/Source/Runtime/Core/Public/IO/IoDispatcher.h
@@ -499,14 +499,7 @@ private:
 class FIoChunkId
 {
 public:
-	static inline FIoChunkId CreateEmptyId()
-	{
-		FIoChunkId ChunkId;
-		uint8 Data[12] = { 0 };
-		ChunkId.Set(Data, sizeof Data);
-
-		return ChunkId;
-	}
+	CORE_API static const FIoChunkId InvalidChunkId;
 
 	friend uint32 GetTypeHash(FIoChunkId InId)
 	{
@@ -524,11 +517,16 @@ public:
 		return Ar;
 	}
 
-	inline bool operator==(const FIoChunkId& Rhs) const
+	inline bool operator ==(const FIoChunkId& Rhs) const
 	{
 		return 0 == FMemory::Memcmp(Id, Rhs.Id, sizeof Id);
 	}
 
+	inline bool operator !=(const FIoChunkId& Rhs) const
+	{
+		return !(*this == Rhs);
+	}
+
 	CORE_API void GenerateFromData(const void* InData, SIZE_T InDataSize);
 
 	void Set(const void* InIdPtr, SIZE_T InSize)
@@ -537,13 +535,21 @@ public:
 		FMemory::Memcpy(Id, InIdPtr, sizeof Id);
 	}
 
-	bool IsValid() const
+	inline bool IsValid() const
 	{
-		uint8 Data[12] = { 0 };
-		return FMemory::Memcmp(Id, Data, sizeof Id) != 0;
+		return *this != InvalidChunkId;
 	}
 
 private:
+	static inline FIoChunkId CreateEmptyId()
+	{
+		FIoChunkId ChunkId;
+		uint8 Data[12] = { 0 };
+		ChunkId.Set(Data, sizeof Data);
+
+		return ChunkId;
+	}
+
 	uint8	Id[12];
 };
 
diff --git a/Engine/Source/Runtime/CoreUObject/Private/Serialization/AsyncLoading2.cpp b/Engine/Source/Runtime/CoreUObject/Private/Serialization/AsyncLoading2.cpp
--- a/Engine/Source/Runtime/CoreUObject/Private/Serialization/AsyncLoading2.cpp
+++ b/Engine/Source/Runtime/CoreUObject/Private/Serialization/AsyncLoading2.cpp
@@ -45,6 +45,7 @@
 #include "IO/IoDispatcher.h"
 #include "UObject/GCObject.h"
 #include "UObject/ObjectRedirector.h"
+#include "Serialization/BulkData.h"
 
 #if UE_BUILD_DEVELOPMENT || UE_BUILD_DEBUG
 //PRAGMA_DISABLE_OPTIMIZATION
@@ -1790,6 +1791,10 @@ void FAsyncLoadingThread2Impl::InitializeLoading()
 		{
 			AltZenaphore.NotifyOne();
 		}));
+
+#if USE_NEW_BULKDATA
+		FBulkDataBase::SetIODispatcher(&IoDispatcher);
+#endif
 	}
 
 	AsyncThreadReady.Increment();
@@ -3685,6 +3690,10 @@ FAsyncLoadingThread2Impl::~FAsyncLoadingThread2Impl()
 	{
 		ShutdownLoading();
 	}
+
+#if USE_NEW_BULKDATA
+	FBulkDataBase::SetIODispatcher(nullptr);
+#endif
 }
 
 void FAsyncLoadingThread2Impl::ShutdownLoading()
diff --git a/Engine/Source/Runtime/CoreUObject/Private/Serialization/BulkData.cpp b/Engine/Source/Runtime/CoreUObject/Private/Serialization/BulkData.cpp
--- a/Engine/Source/Runtime/CoreUObject/Private/Serialization/BulkData.cpp
+++ b/Engine/Source/Runtime/CoreUObject/Private/Serialization/BulkData.cpp
@@ -978,11 +978,19 @@ void FUntypedBulkData::Serialize( FArchive& Ar, UObject* Owner, int32 Idx, bool
 
 			Ar << BulkDataOffsetInFile;
 
+			if ((BulkDataFlags & BULKDATA_BadDataVersion) != 0)
+			{
+				uint16 DummyValue;
+				Ar << DummyValue;
+
+				BulkDataFlags &= ~BULKDATA_BadDataVersion;		
+			}
+			
 			if ((BulkDataFlags & BULKDATA_CookedForIoDispatcher) != 0)
 			{
 				// Load the ChunkID that will be used by the IoDispatcher, but this version of BulkData will not use it!
-				uint16 ChunkIndex;
-				Ar << ChunkIndex;
+				FIoChunkId ChunkId;
+ 				Ar << ChunkId;
 			}
 
 			// determine whether the payload is stored inline or at the end of the file
@@ -1213,6 +1221,7 @@ void FUntypedBulkData::Serialize( FArchive& Ar, UObject* Owner, int32 Idx, bool
 			// Only serialize status information if wanted.
 			int64 SavedBulkDataSizeOnDiskPos	= INDEX_NONE;
 			int64 SavedBulkDataOffsetInFilePos	= INDEX_NONE;
+			int64 SavedBulkDataChunkIdPos		= INDEX_NONE;
 			
 			// Keep track of position we are going to serialize placeholder BulkDataSizeOnDisk.
 			SavedBulkDataSizeOnDiskPos = Ar.Tell();
@@ -1249,12 +1258,15 @@ void FUntypedBulkData::Serialize( FArchive& Ar, UObject* Owner, int32 Idx, bool
 				bStoreInline = true;
 			}
 
-			// Write the ChunkID that will be used by the IoDispatcher
+			// Write the ChunkID that will be used by the IoDispatcher, but ONLY if we are producing cooked data
+			if(Ar.IsCooking() )
 			{
 				FArchive::FScopeSetDebugSerializationFlags S(Ar, DSF_IgnoreDiff);
 
-				uint16 ChunkIndex = (bStoreInline || !Ar.IsCooking()) ? INDEX_NONE : LinkerSave->BulkDataToAppend.Num();
-				Ar << ChunkIndex;
+				SavedBulkDataChunkIdPos = Ar.Tell();
+
+				FIoChunkId InvalidChunkId = FIoChunkId::InvalidChunkId;
+				Ar << InvalidChunkId;
 
 				BulkDataFlags |= BULKDATA_CookedForIoDispatcher;
 			}
@@ -1274,6 +1286,7 @@ void FUntypedBulkData::Serialize( FArchive& Ar, UObject* Owner, int32 Idx, bool
 				BulkStore.BulkDataOffsetInFilePos = SavedBulkDataOffsetInFilePos;
 				BulkStore.BulkDataSizeOnDiskPos = SavedBulkDataSizeOnDiskPos;
 				BulkStore.BulkDataFlagsPos = SavedBulkDataFlagsPos;
+				BulkStore.BulkDataChunkIdPos = SavedBulkDataChunkIdPos;
 				BulkStore.BulkDataFlags = BulkDataFlags;
 				BulkStore.BulkData = this;
 
@@ -1283,6 +1296,7 @@ void FUntypedBulkData::Serialize( FArchive& Ar, UObject* Owner, int32 Idx, bool
 					int64 SavedDupeBulkDataFlagsPos = INDEX_NONE;
 					int64 SavedDupeBulkDataSizeOnDiskPos = INDEX_NONE;
 					int64 SavedDupeBulkDataOffsetInFilePos = INDEX_NONE;
+					int64 SavedDupeBulkDataChunkIdPos = INDEX_NONE;
 
 					int32 SavedDupeBulkDataFlags = (BulkDataFlags & ~BULKDATA_DuplicateNonOptionalPayload) | BULKDATA_OptionalPayload;
 					{
@@ -1303,11 +1317,24 @@ void FUntypedBulkData::Serialize( FArchive& Ar, UObject* Owner, int32 Idx, bool
 							int32 BulkDataSizeOnDiskAsInt32 = BulkDataSizeOnDisk;
 							Ar << BulkDataSizeOnDiskAsInt32;
 						}
+
 						// Keep track of position we are going to serialize placeholder BulkDataOffsetInFile.
 						SavedDupeBulkDataOffsetInFilePos = Ar.Tell();
 						BulkDataOffsetInFile = INDEX_NONE;
 						// And serialize the placeholder which is going to be overwritten later.
 						Ar << BulkDataOffsetInFile;
+
+						// TODO: Fairly sure we can only be in the BULKDATA_DuplicateNonOptionalPayload branch if this
+						// is cooked data, see if we can remove the if check
+						if (Ar.IsCooking())
+						{
+							SavedBulkDataChunkIdPos = Ar.Tell();
+							
+							FIoChunkId InvalidChunkId = FIoChunkId::InvalidChunkId;
+							Ar << InvalidChunkId;
+
+							check((SavedDupeBulkDataFlagsPos& BULKDATA_CookedForIoDispatcher) != 0); // Validate that the BULKDATA_CookedForIoDispatcher flag was already set
+						}
 					}
 
 					// add duplicate bulkdata with different flag
@@ -1318,6 +1345,8 @@ void FUntypedBulkData::Serialize( FArchive& Ar, UObject* Owner, int32 Idx, bool
 					DupeBulkStore.BulkDataSizeOnDiskPos = SavedDupeBulkDataSizeOnDiskPos;
 					DupeBulkStore.BulkDataFlagsPos = SavedDupeBulkDataFlagsPos;
 					DupeBulkStore.BulkDataFlags = SavedDupeBulkDataFlags;
+					DupeBulkStore.BulkDataChunkIdPos = SavedDupeBulkDataChunkIdPos;
+
 					DupeBulkStore.BulkData = this;
 				}
 				
@@ -1661,9 +1690,14 @@ FBulkDataStreamingToken FUntypedBulkData::CreateStreamingToken() const
 	return FBulkDataStreamingToken((uint32)GetBulkDataOffsetInFile(), (uint32)GetBulkDataSize());
 }
 
-IBulkDataIORequest* FUntypedBulkData::CreateStreamingRequestForRange(const FString& Filename, const FBulkDataStreamingToken& Start, const FBulkDataStreamingToken& End, EAsyncIOPriorityAndFlags Priority, FAsyncFileCallBack* CompleteCallback)
+IBulkDataIORequest* FUntypedBulkData::CreateStreamingRequestForRange(const FString& Filename, const BulkDataRangeArray& RangeArray, EAsyncIOPriorityAndFlags Priority, FAsyncFileCallBack* CompleteCallback)
 {
 	check(Filename.IsEmpty() == false);
+	check(RangeArray.Num() > 0);
+
+	const FBulkDataStreamingToken& Start = *(RangeArray[0]);
+	const FBulkDataStreamingToken& End = *(RangeArray[RangeArray.Num()-1]);
+
 	check(Start.IsValid());
 	check(End.IsValid());
 
@@ -1710,8 +1744,7 @@ bool FBulkDataIORequest::WaitCompletion(float TimeLimitSeconds) const
 	return ReadRequest->WaitCompletion(TimeLimitSeconds);
 }
 
-
-uint8* FBulkDataIORequest::GetReadResults() const
+uint8* FBulkDataIORequest::GetReadResults()
 {
 	return ReadRequest->GetReadResults();
 }
@@ -1728,7 +1761,7 @@ int64 FBulkDataIORequest::GetSize() const
 	}
 }
 
-void FBulkDataIORequest::Cancel() const
+void FBulkDataIORequest::Cancel()
 {
 	ReadRequest->Cancel();
 }
diff --git a/Engine/Source/Runtime/CoreUObject/Private/Serialization/BulkData2.cpp b/Engine/Source/Runtime/CoreUObject/Private/Serialization/BulkData2.cpp
--- a/Engine/Source/Runtime/CoreUObject/Private/Serialization/BulkData2.cpp
+++ b/Engine/Source/Runtime/CoreUObject/Private/Serialization/BulkData2.cpp
@@ -4,16 +4,23 @@
 
 #include "HAL/FileManager.h"
 #include "HAL/PlatformFilemanager.h"
+#include "Misc/CommandLine.h"
 #include "Misc/Paths.h"
 #include "ProfilingDebugging/LoadTimeTracker.h"
 #include "Serialization/BulkData.h"
 #include "UObject/LinkerLoad.h"
 #include "UObject/Object.h"
+#include "UObject/Package.h"
+
+DEFINE_LOG_CATEGORY_STATIC(LogBulkDataRuntime, Log, All);
+
+// If set to 0 then the loose file fallback will be used even if the -ZenLoader command line flag is present.
+#define ENABLE_IO_DISPATCHER 1
 
 namespace
 {
 	// TODO: Maybe expose this and start using everywhere?
-	static constexpr const TCHAR* InlinedExt = TEXT(".uexp");			// Inlined
+	static constexpr const TCHAR* InlinedExt = TEXT(".uexp");			// Stored in the export data
 	static constexpr const TCHAR* DefaultExt = TEXT(".ubulk");			// Stored in a separate file
 	static constexpr const TCHAR* MemoryMappedExt = TEXT(".m.ubulk");	// Stored in a separate file aligned for memory mapping
 	static constexpr const TCHAR* OptionalExt = TEXT(".uptnl");			// Stored in a separate file that is optional
@@ -27,7 +34,6 @@ namespace FileTokenSystem
 	struct Data
 	{
 		FString Filename;
-		int64 BulkDataSize;
 		int64 BulkDataOffsetInFile;
 	};
 
@@ -41,7 +47,6 @@ namespace FileTokenSystem
 
 		Data& Data = TokenDataMap.Add(++TokenCounter);
 		Data.Filename = Filename;
-		Data.BulkDataSize = BulkDataSize;
 		Data.BulkDataOffsetInFile = BulkDataOffsetInFile;
 
 		return TokenCounter;
@@ -103,44 +108,288 @@ namespace FileTokenSystem
 		return Data->Filename;
 	}
 
-	uint64 GetBulkDataSize(FBulkDataBase::FileToken ID)
+	uint64 GetBulkDataOffset(FBulkDataBase::FileToken ID)
 	{
 		if (ID == FBulkDataBase::InvalidToken)
 		{
 			return 0;
 		}
-		
+
 		FReadScopeLock LockForScope(TokenLock);
 		Data* Data = TokenDataMap.Find(ID);
 		check(Data);
 
-		return Data != nullptr ? Data->BulkDataSize : 0;
+		return Data != nullptr ? Data->BulkDataOffsetInFile : 0;
 	}
+}
 
-	uint64 GetBulkDataOffset(FBulkDataBase::FileToken ID)
+FIoDispatcher* FBulkDataBase::IoDispatcher = nullptr;
+
+class FBulkDataIoDispatcherRequest : public IBulkDataIORequest
+{
+public:
+	using ChunkIdArray = TArray<FIoChunkId, TInlineAllocator<8>>;
+
+public:
+	FBulkDataIoDispatcherRequest(const FIoChunkId& InChunkID, int64 InOffsetInBulkData, int64 InBytesToRead, FAsyncFileCallBack* InCompleteCallback, uint8* InUserSuppliedMemory)
+		: UserSuppliedMemory(InUserSuppliedMemory)
 	{
-		if (ID == FBulkDataBase::InvalidToken)
+		RequestArray.Push({ InChunkID, (uint64)InOffsetInBulkData , (uint64)InBytesToRead });
+
+		if (InCompleteCallback != nullptr)
 		{
-			return 0;
+			CompleteCallback = *InCompleteCallback;
 		}
+	}
 
-		FReadScopeLock LockForScope(TokenLock);
-		Data* Data = TokenDataMap.Find(ID);
-		check(Data);
+	FBulkDataIoDispatcherRequest(const ChunkIdArray& ChunkIDs, FAsyncFileCallBack* InCompleteCallback)
+		: UserSuppliedMemory(nullptr)
+	{
+		for (const FIoChunkId& ChunkId : ChunkIDs)
+		{
+			RequestArray.Push({ChunkId, 0, (uint64)INDEX_NONE});
+		}
 
-		return Data != nullptr ? Data->BulkDataOffsetInFile : 0;
+		if (InCompleteCallback != nullptr)
+		{
+			CompleteCallback = *InCompleteCallback;
+		}
 	}
-}
+
+	virtual ~FBulkDataIoDispatcherRequest()
+	{
+		if (Task != nullptr)
+		{
+			Task->EnsureCompletion();
+			delete Task;
+		}
+
+		// Free the data is no caller has taken ownership of it and it was allocated by FBulkDataIoDispatcherRequest
+		if (UserSuppliedMemory == nullptr)
+		{
+			FMemory::Free(DataResult);
+			DataResult = nullptr;
+		}
+	}
+
+	void StartAsyncWork()
+	{
+		check(Task == nullptr);
+
+		bIsCompleted = false;
+		FPlatformMisc::MemoryBarrier();
+
+		if (FPlatformProcess::SupportsMultithreading())
+		{
+			Task = new FAsyncTask<FBulkDataIoDispatcherRequestWorker>(*this);
+			Task->StartBackgroundTask(GIOThreadPool);
+		}
+		else
+		{
+			// No multi threading support so just process immediately
+			DoWork();
+		}
+	}
+
+	virtual bool PollCompletion() const override
+	{
+		return bIsCompleted;
+	}
+
+	virtual bool WaitCompletion(float TimeLimitSeconds) const override
+	{
+		// Note that currently we do not get events from the FIoDispatcher, so we just
+		// have a basic implementation.
+		// We only have one use case for a time limited wait, every other use case is 
+		// supposed to be fully blocking so ideally we can eliminate the single use case 
+		// and just change this code entirely.
+		if (!bIsCompleted)
+		{
+			if (TimeLimitSeconds > 0.0f)
+			{
+				FPlatformProcess::Sleep(TimeLimitSeconds);
+			}
+			else
+			{
+				while (!bIsCompleted)
+				{
+					FPlatformProcess::Sleep(0.0f);
+				}
+			}
+		}
+
+		return bIsCompleted;
+	}
+
+	virtual uint8* GetReadResults() override
+	{
+		if (bIsCompleted && !bIsCanceled)
+		{
+			uint8* Result = DataResult;
+			DataResult = nullptr;
+
+			return Result;
+		}
+		else
+		{
+			return nullptr;
+		}	
+	}
+
+	virtual int64 GetSize() const override
+	{
+		if (bIsCompleted && !bIsCanceled)
+		{
+			return SizeResult;
+		}
+		else
+		{
+			return INDEX_NONE;
+		}
+	}
+
+	virtual void Cancel() override
+	{
+		if (!bIsCanceled)
+		{
+			bIsCanceled = true;
+			FPlatformMisc::MemoryBarrier();
+			// TODO: Send to IoDispatcher
+		}
+	}
+
+private:
+
+	void DoWork()
+	{
+		FIoReadOptions IoOptions;
+
+		// TODO: FIoReadOptions is ignored at the moment
+		//IoOptions.SetRange(OffsetInBulkData, BytesToRead);
+		//IoOptions.SetTargetVa((uint64)UserSuppliedMemory);
+
+		FIoBatch NewBatch = FBulkDataBase::IoDispatcher->NewBatch();
+		for (Request& Request : RequestArray)
+		{
+			Request.IoRequest = NewBatch.Read(Request.ChunkId, IoOptions);
+		}
+		NewBatch.Issue();
+
+		// Wait until the request is completed
+		while (!AreRequestsCompleted())
+		{
+			FPlatformProcess::Sleep(0.0f);
+		}
+
+		SizeResult = GetTotalBufferSize();
+
+		uint8* DstBuffer = UserSuppliedMemory != nullptr ? UserSuppliedMemory : (uint8*)FMemory::Malloc(SizeResult);
+		DataResult = DstBuffer;
+
+		for (const Request& Request : RequestArray)
+		{
+			TIoStatusOr<FIoBuffer> Result = Request.IoRequest.GetResult();
+			check(Result.IsOk());
+
+			const FIoBuffer& IoBuffer = Result.ValueOrDie();
+			uint64 BytesToRead = Request.BytesToRead == INDEX_NONE ? IoBuffer.DataSize() : Request.BytesToRead;
+		
+			FMemory::Memcpy(DstBuffer, IoBuffer.Data() + Request.OffsetInBulkData, BytesToRead);
+			DstBuffer += BytesToRead;
+		}
+		
+		FBulkDataBase::IoDispatcher->FreeBatch(NewBatch);
+
+		if (CompleteCallback)
+		{
+			CompleteCallback(bIsCanceled, nullptr);
+		}
+
+		FPlatformMisc::MemoryBarrier();
+		bIsCompleted = true;
+	}
+
+	class FBulkDataIoDispatcherRequestWorker : public FNonAbandonableTask
+	{
+		FBulkDataIoDispatcherRequest& IoRequest;
+	public:
+		FBulkDataIoDispatcherRequestWorker(FBulkDataIoDispatcherRequest& InIoRequest)
+			: IoRequest(InIoRequest)
+		{
+		}
+
+		void DoWork()
+		{
+			IoRequest.DoWork();
+		}
+
+		FORCEINLINE TStatId GetStatId() const
+		{
+			RETURN_QUICK_DECLARE_CYCLE_STAT(FBulkDataIoDispatcherRequestWorker, STATGROUP_ThreadPoolAsyncTasks);
+		}
+	};
+
+private:
+	struct Request
+	{
+		FIoChunkId ChunkId;
+		uint64 OffsetInBulkData;
+		uint64 BytesToRead;
+
+		FIoRequest IoRequest;
+	};
+
+	bool AreRequestsCompleted() const
+	{
+		for (const Request& Req : RequestArray)
+		{
+			if (!Req.IoRequest.Status().IsCompleted())
+			{
+				return false;
+			}
+		}
+
+		return true;
+	}
+
+	int64 GetTotalBufferSize()  
+	{
+		uint64 Size = 0;
+		for (Request& Req : RequestArray)
+		{
+			TIoStatusOr<FIoBuffer> Result = Req.IoRequest.GetResult();
+			Size += Result.ValueOrDie().DataSize();
+		}
+
+		return Size;
+	}
+
+	TArray<Request, TInlineAllocator<8>> RequestArray;
+
+	FAsyncFileCallBack CompleteCallback;
+	uint8* UserSuppliedMemory = nullptr;
+
+	uint8* DataResult = nullptr;
+	int64 SizeResult = 0;
+
+	// Task related members
+	FAsyncTask<FBulkDataIoDispatcherRequestWorker>* Task = nullptr;
+	bool bIsCompleted = false;
+	bool bIsCanceled = false;
+};
 
 FBulkDataBase::FBulkDataBase(FBulkDataBase&& Other)
-	: Token(Other.Token)
+	: ChunkID(Other.ChunkID) // Copies the entire union
 	, DataBuffer(Other.DataBuffer)
 	, BulkDataFlags(Other.BulkDataFlags)
 {
 	check(Other.LockStatus == LOCKSTATUS_Unlocked); // Make sure that the other object wasn't inuse
 
-	Other.Token = InvalidToken; // Prevent the other object from unregistering the token
-	Other.DataBuffer = nullptr; // Prevent the other object from deleting our data
+	if (!Other.IsUsingIODispatcher())
+	{
+		Other.Fallback.Token = InvalidToken; // Prevent the other object from unregistering the token
+	}	
+	Other.DataBuffer = nullptr;	// Prevent the other object from deleting our data
 }
 
 FBulkDataBase& FBulkDataBase::operator=(const FBulkDataBase& Other)
@@ -152,7 +401,7 @@ FBulkDataBase& FBulkDataBase::operator=(const FBulkDataBase& Other)
 
 	RemoveBulkData();
 
-	Token = FileTokenSystem::CopyFileToken(Other.Token);
+	Fallback.Token = FileTokenSystem::CopyFileToken(Other.Fallback.Token);
 
 	// Copy token
 	BulkDataFlags = Other.BulkDataFlags;
@@ -173,7 +422,10 @@ FBulkDataBase::~FBulkDataBase()
 	check(LockStatus == LOCKSTATUS_Unlocked);
 
 	FreeData();
-	FileTokenSystem::UnregisterFileToken(Token);
+	if (!IsUsingIODispatcher())
+	{
+		FileTokenSystem::UnregisterFileToken(Fallback.Token);
+	}	
 }
 
 void FBulkDataBase::Serialize(FArchive& Ar, UObject* Owner, int32 /*Index*/, bool bAttemptFileMapping, int32 ElementSize)
@@ -217,34 +469,57 @@ void FBulkDataBase::Serialize(FArchive& Ar, UObject* Owner, int32 /*Index*/, boo
 
 		Ar << BulkDataOffsetInFile;
 
-		// Currently not used!
-		if ((BulkDataFlags & BULKDATA_CookedForIoDispatcher) != 0)
+#if ENABLE_IO_DISPATCHER
+		static const bool bUseZenLoader = FParse::Param(FCommandLine::Get(), TEXT("zenloader"));
+#else
+		const bool bUseZenLoader = false;
+#endif
+
+		if ((BulkDataFlags & BULKDATA_BadDataVersion) != 0)
 		{
-			uint16 ChunkIndex;
-			Ar << ChunkIndex;
+			uint16 DummyValue;
+			Ar << DummyValue;
 		}
 
-		// Assuming that Owner/Package/Linker are all valid, the old BulkData system would
-		// generally fail if any of these were nullptr but had plenty of inconsistent checks
-		// scattered throughout.
-		check(Owner != nullptr);
-		UPackage* Package = Owner->GetOutermost();
-		check(Package != nullptr);
-
-		FLinkerLoad* Linker = FLinkerLoad::FindExistingLinkerForPackage(Package);
-		check(Linker != nullptr);
+		if ((BulkDataFlags & BULKDATA_CookedForIoDispatcher) != 0)
+		{
+			Ar << ChunkID;
+		}
 
-		// Fix up the file offset, but only if not stored inline (TODO: Look into getting rid of this and cook with the correct values?)
-		if (!IsInlined())
+		if (!bUseZenLoader || ChunkID == FIoChunkId::InvalidChunkId)
 		{
-			BulkDataOffsetInFile += Linker->Summary.BulkDataStartOffset;
+			// Remove the flag if we are not going to use it
+			BulkDataFlags &= ~BULKDATA_CookedForIoDispatcher;
+			Fallback.Token = InvalidToken;
+			Fallback.BulkDataSize = BulkDataSize;
 		}
 
+		FString Filename;
+		FLinkerLoad* Linker = nullptr;
+		if (!bUseZenLoader)
+		{
+			// Assuming that Owner/Package/Linker are all valid, the old BulkData system would
+			// generally fail if any of these were nullptr but had plenty of inconsistent checks
+			// scattered throughout.
+			check(Owner != nullptr);
+			UPackage* Package = Owner->GetOutermost();
+			check(Package != nullptr);
+
+			Linker = FLinkerLoad::FindExistingLinkerForPackage(Package);
+			check(Linker != nullptr);
+
+			// Fix up the file offset, but only if not stored inline (TODO: Look into getting rid of this and cook with the correct values?)
+			if (!IsInlined())
+			{
+				BulkDataOffsetInFile += Linker->Summary.BulkDataStartOffset;
+			}
+
+			Filename = ConvertFilenameFromFlags(Linker->Filename);
+		}
+		
 		FArchive* CacheableArchive = Ar.GetCacheableArchive();
 		if (Ar.IsAllowingLazyLoading() && CacheableArchive != nullptr)
 		{
-			FString Filename = ConvertFilenameFromFlags(Linker->Filename);
-
 			if (IsInlined())
 			{
 				// Inline data is already in the archive so serialize it immediately
@@ -258,7 +533,7 @@ void FBulkDataBase::Serialize(FArchive& Ar, UObject* Owner, int32 /*Index*/, boo
 					FString OptionalFilename = FPaths::ChangeExtension(Filename, OptionalExt);
 					if (IFileManager::Get().FileExists(*Filename))
 					{
-						SerializeDuplicateData(Ar, Owner, BulkDataFlags, BulkDataSizeOnDisk, BulkDataOffsetInFile);
+						SerializeDuplicateData(Ar, Linker, BulkDataFlags, BulkDataSizeOnDisk, BulkDataOffsetInFile);
 						Filename = OptionalFilename;
 					}
 					else
@@ -267,15 +542,17 @@ void FBulkDataBase::Serialize(FArchive& Ar, UObject* Owner, int32 /*Index*/, boo
 						// out if we read things as 32bit or 64bit)
 						uint32 DummyValue32;
 						int64 DummyValue64;
-						SerializeDuplicateData(Ar, Owner, DummyValue32, DummyValue64, DummyValue64);
+						SerializeDuplicateData(Ar, Linker, DummyValue32, DummyValue64, DummyValue64);
 
 						check(Filename.EndsWith(DefaultExt));
 					}			
 				}			
 			}
 
-			// TODO: Replace with the proper storage solution
-			Token = FileTokenSystem::RegisterFileToken(Filename, BulkDataSize, BulkDataOffsetInFile);
+			if (bUseZenLoader == false)
+			{
+				Fallback.Token = FileTokenSystem::RegisterFileToken(Filename, BulkDataSize, BulkDataOffsetInFile);
+			}		
 		}
 		else
 		{
@@ -374,10 +651,10 @@ void FBulkDataBase::GetCopy(void** DstBuffer, bool bDiscardInternalCopy)
 			}
 			else
 			{
-				int64 BulkDataSize = GetBulkDataSize();
+				const int64 BulkDataSize = GetBulkDataSize();
 
 				*DstBuffer = FMemory::Malloc(BulkDataSize, 0);
-				FMemory::Memcpy(*DstBuffer, DataBuffer, GetBulkDataSize());
+				FMemory::Memcpy(*DstBuffer, DataBuffer, BulkDataSize);
 			}
 		}
 		else
@@ -414,7 +691,15 @@ void  FBulkDataBase::ClearBulkDataFlags(uint32 BulkDataFlagsToClear)
 
 int64 FBulkDataBase::GetBulkDataSize() const
 {
-	return FileTokenSystem::GetBulkDataSize(Token);
+	if (IsUsingIODispatcher())
+	{
+		TIoStatusOr<uint64> Result = IoDispatcher->GetSizeForChunk(ChunkID);
+		return Result.ValueOrDie(); // TODO: Consider logging errors instead of relying on ::ValueOrDie
+	}
+	else
+	{
+		return Fallback.BulkDataSize;
+	}	
 }
 
 bool FBulkDataBase::IsStoredCompressedOnDisk() const
@@ -462,6 +747,11 @@ bool FBulkDataBase::IsMemoryMapped() const
 	return (BulkDataFlags & BULKDATA_MemoryMappedPayload) != 0;
 }
 
+bool FBulkDataBase::IsUsingIODispatcher() const
+{
+	return (BulkDataFlags & BULKDATA_CookedForIoDispatcher) != 0;
+}
+
 IBulkDataIORequest* FBulkDataBase::CreateStreamingRequest(EAsyncIOPriorityAndFlags Priority, FAsyncFileCallBack* CompleteCallback, uint8* UserSuppliedMemory) const
 {
 	const int64 DataSize = GetBulkDataSize();
@@ -471,45 +761,80 @@ IBulkDataIORequest* FBulkDataBase::CreateStreamingRequest(EAsyncIOPriorityAndFla
 
 IBulkDataIORequest* FBulkDataBase::CreateStreamingRequest(int64 OffsetInBulkData, int64 BytesToRead, EAsyncIOPriorityAndFlags Priority, FAsyncFileCallBack* CompleteCallback, uint8* UserSuppliedMemory) const
 {
-	FileTokenSystem::Data FileData = FileTokenSystem::GetFileData(Token);
+	if (IsUsingIODispatcher())
+	{
+		FBulkDataIoDispatcherRequest* IoRequest = new FBulkDataIoDispatcherRequest(ChunkID, OffsetInBulkData, BytesToRead, CompleteCallback, UserSuppliedMemory);	
+		IoRequest->StartAsyncWork();
 
-	check(FileData.Filename.IsEmpty() == false);
+		return IoRequest;
+	}
+	else
+	{
+		const int64 BulkDataSize = GetBulkDataSize();
+		FileTokenSystem::Data FileData = FileTokenSystem::GetFileData(Fallback.Token);
 
-	UE_CLOG(IsStoredCompressedOnDisk(), LogSerialization, Fatal, TEXT("Package level compression is no longer supported (%s)."), *FileData.Filename);
-	UE_CLOG(FileData.BulkDataSize <= 0, LogSerialization, Error, TEXT("(%s) has invalid bulk data size."), *FileData.Filename);
+		check(FileData.Filename.IsEmpty() == false);
 
-	IAsyncReadFileHandle* IORequestHandle = FPlatformFileManager::Get().GetPlatformFile().OpenAsyncRead(*FileData.Filename);
-	check(IORequestHandle); // this generally cannot fail because it is async
+		UE_CLOG(IsStoredCompressedOnDisk(), LogSerialization, Fatal, TEXT("Package level compression is no longer supported (%s)."), *FileData.Filename);
+		UE_CLOG(BulkDataSize <= 0, LogSerialization, Error, TEXT("(%s) has invalid bulk data size."), *FileData.Filename);
 
-	if (IORequestHandle == nullptr)
-	{
-		return nullptr;
+		IAsyncReadFileHandle* IORequestHandle = FPlatformFileManager::Get().GetPlatformFile().OpenAsyncRead(*FileData.Filename);
+		check(IORequestHandle); // this generally cannot fail because it is async
+
+		if (IORequestHandle == nullptr)
+		{
+			return nullptr;
+		}
+
+		const int64 OffsetInFile = FileData.BulkDataOffsetInFile + OffsetInBulkData;
+
+		IAsyncReadRequest* ReadRequest = IORequestHandle->ReadRequest(OffsetInFile, BytesToRead, Priority, CompleteCallback, UserSuppliedMemory);
+		if (ReadRequest != nullptr)
+		{
+			return new FBulkDataIORequest(IORequestHandle, ReadRequest, BytesToRead);
+		}
+		else
+		{
+			delete IORequestHandle;
+			return nullptr;
+		}
 	}
+}
+
+IBulkDataIORequest* FBulkDataBase::CreateStreamingRequestForRange(const BulkDataRangeArray& RangeArray, EAsyncIOPriorityAndFlags Priority, FAsyncFileCallBack* CompleteCallback)
+{
+	check(RangeArray.Num() > 0);
+
+	const FBulkDataBase& Start = *(RangeArray[0]);
 
-	const int64 OffsetInFile = FileData.BulkDataOffsetInFile + OffsetInBulkData;
+	check(!Start.IsInlined());
 
-	IAsyncReadRequest* ReadRequest = IORequestHandle->ReadRequest(OffsetInFile, BytesToRead, Priority, CompleteCallback, UserSuppliedMemory);
-	if (ReadRequest != nullptr)
+	if (Start.IsUsingIODispatcher())
 	{
-		return new FBulkDataIORequest(IORequestHandle, ReadRequest, BytesToRead);
+		FBulkDataIoDispatcherRequest::ChunkIdArray ChunkIds;
+		for (const FBulkDataBase* BulkData : RangeArray)
+		{
+			ChunkIds.Push(BulkData->ChunkID);
+		}
+
+		FBulkDataIoDispatcherRequest* IoRequest = new FBulkDataIoDispatcherRequest(ChunkIds, CompleteCallback);
+		IoRequest->StartAsyncWork();
+
+		return IoRequest;
 	}
 	else
-	{
-		delete IORequestHandle;
-		return nullptr;
-	}
-}
+	{	
+		const FBulkDataBase& End = *RangeArray[RangeArray.Num() - 1];
 
-IBulkDataIORequest* FBulkDataBase::CreateStreamingRequestForRange(const FBulkDataBase& Start, const FBulkDataBase& End, EAsyncIOPriorityAndFlags Priority, FAsyncFileCallBack* CompleteCallback)
-{
-	check(Start.GetFilename() == End.GetFilename());
+		check(Start.GetFilename() == End.GetFilename());
 
-	const int64 ReadOffset = Start.GetBulkDataOffsetInFile();
-	const int64 ReadSize = (End.GetBulkDataOffsetInFile() + End.GetBulkDataSize()) - ReadOffset;
+		const int64 ReadOffset = Start.GetBulkDataOffsetInFile();
+		const int64 ReadSize = (End.GetBulkDataOffsetInFile() + End.GetBulkDataSize()) - ReadOffset;
 
-	check(ReadSize > 0);
+		check(ReadSize > 0);
 
-	return Start.CreateStreamingRequest(0, ReadSize, Priority, CompleteCallback, nullptr);
+		return Start.CreateStreamingRequest(0, ReadSize, Priority, CompleteCallback, nullptr);
+	}
 }
 
 void FBulkDataBase::ForceBulkDataResident()
@@ -538,8 +863,11 @@ void FBulkDataBase::RemoveBulkData()
 
 	FreeData();
 
-	FileTokenSystem::UnregisterFileToken(Token);
-	Token = InvalidToken;
+	if (!IsUsingIODispatcher())
+	{ 
+		FileTokenSystem::UnregisterFileToken(Fallback.Token);
+		Fallback.Token = InvalidToken;
+	}
 
 	BulkDataFlags = 0;
 
@@ -547,12 +875,28 @@ void FBulkDataBase::RemoveBulkData()
 
 int64 FBulkDataBase::GetBulkDataOffsetInFile() const
 {
-	return FileTokenSystem::GetBulkDataOffset(Token);
+	if (!IsUsingIODispatcher())
+	{
+		return FileTokenSystem::GetBulkDataOffset(Fallback.Token);
+	}
+	else
+	{
+		BULKDATA_NOT_IMPLEMENTED_FOR_RUNTIME;
+		return 0;
+	}
 }
 
 FString FBulkDataBase::GetFilename() const
 {
-	return FileTokenSystem::GetFilename(Token);
+	if (!IsUsingIODispatcher())
+	{
+		return FileTokenSystem::GetFilename(Fallback.Token);
+	}
+	else
+	{
+		UE_LOG(LogBulkDataRuntime, Warning, TEXT("Attempting to get the filename for BulkData that uses the IoDispatcher, this will return an empty string"));
+		return FString("");
+	}
 }
 
 void FBulkDataBase::LoadDataDirectly(void** DstBuffer)
@@ -560,36 +904,44 @@ void FBulkDataBase::LoadDataDirectly(void** DstBuffer)
 	DECLARE_SCOPE_CYCLE_COUNTER(TEXT("FBulkDataBase::LoadDataDirectly"), STAT_UBD_LoadDataDirectly, STATGROUP_Memory);
 	check(CanLoadFromDisk());
 
-	FileTokenSystem::Data FileData = FileTokenSystem::GetFileData(Token);
+	if (!IsUsingIODispatcher())
+	{
+		const int64 BulkDataSize = GetBulkDataSize();
+		FileTokenSystem::Data FileData = FileTokenSystem::GetFileData(Fallback.Token);
 
-	// In the older code path if the bulkdata was not in a separate file we could just serialize from the package instead, but the new system
-	// does not keep a reference to the pack so we cannot do that.
-	// TODO: Maybe add a check for this condition?
+		// In the older code path if the bulkdata was not in a separate file we could just serialize from the package instead, but the new system
+		// does not keep a reference to the pack so we cannot do that.
+		// TODO: Maybe add a check for this condition?
 
-	check(!IsInlined());	// TODO: Currently inline data will have the wrong BulkDataOffsetInFile (need to subtract the size of the original
-							// .uasset or .umap file from it but we don't have access to that filepath at this point and doing it during
-							// ::Serialize is too slow!)
+		check(!IsInlined());	// TODO: Currently inline data will have the wrong BulkDataOffsetInFile (need to subtract the size of the original
+								// .uasset or .umap file from it but we don't have access to that filepath at this point and doing it during
+								// ::Serialize is too slow!)
 
-	// If the data is inlined then we already loaded is during ::Serialize, this warning should help track cases where data is being discarded then re-requested.
-	UE_CLOG(IsInlined(), LogSerialization, Warning, TEXT("Reloading inlined bulk data directly from disk, consider not discarding it in the first place. Filename: '%s'."), *FileData.Filename);
+		// If the data is inlined then we already loaded is during ::Serialize, this warning should help track cases where data is being discarded then re-requested.
+		UE_CLOG(IsInlined(), LogSerialization, Warning, TEXT("Reloading inlined bulk data directly from disk, consider not discarding it in the first place. Filename: '%s'."), *FileData.Filename);
 
-	FArchive* Ar = IFileManager::Get().CreateFileReader(*FileData.Filename, FILEREAD_Silent);
-	checkf(Ar != nullptr, TEXT("Failed to open the file to load bulk data from. Filename: '%s'."), *FileData.Filename);
+		FArchive* Ar = IFileManager::Get().CreateFileReader(*FileData.Filename, FILEREAD_Silent);
+		checkf(Ar != nullptr, TEXT("Failed to open the file to load bulk data from. Filename: '%s'."), *FileData.Filename);
 
-	// Seek to the beginning of the bulk data in the file.
-	Ar->Seek(FileData.BulkDataOffsetInFile);
+		// Seek to the beginning of the bulk data in the file.
+		Ar->Seek(FileData.BulkDataOffsetInFile);
 
-	if (*DstBuffer == nullptr)
-	{
-		*DstBuffer = FMemory::Malloc(FileData.BulkDataSize, 0);
-	}
+		if (*DstBuffer == nullptr)
+		{
+			*DstBuffer = FMemory::Malloc(BulkDataSize, 0);
+		}
 
-	SerializeBulkData(*Ar, *DstBuffer, FileData.BulkDataSize);
+		SerializeBulkData(*Ar, *DstBuffer, BulkDataSize);
 
-	delete Ar;
+		delete Ar;
+	}
+	else
+	{
+		BULKDATA_NOT_IMPLEMENTED_FOR_RUNTIME;
+	}
 }
 
-void FBulkDataBase::SerializeDuplicateData(FArchive& Ar, UObject* Owner, uint32& OutBulkDataFlags, int64& OutBulkDataSizeOnDisk, int64& OutBulkDataOffsetInFile)
+void FBulkDataBase::SerializeDuplicateData(FArchive& Ar, FLinkerLoad* Linker, uint32& OutBulkDataFlags, int64& OutBulkDataSizeOnDisk, int64& OutBulkDataOffsetInFile)
 {
 	Ar << OutBulkDataFlags;
 
@@ -607,10 +959,16 @@ void FBulkDataBase::SerializeDuplicateData(FArchive& Ar, UObject* Owner, uint32&
 
 	Ar << OutBulkDataOffsetInFile;
 
-	// fix up the file offset, but only if not stored inline
-	if (Owner != nullptr && Owner->GetLinker() != nullptr)
+	if ((BulkDataFlags & BULKDATA_CookedForIoDispatcher) != 0)
+	{
+		FIoChunkId DummyChunkID;
+		Ar << DummyChunkID;
+	}
+
+	// Fix up the file offset if we have a valid linker
+	if (Linker != nullptr)
 	{
-		OutBulkDataOffsetInFile += Owner->GetLinker()->Summary.BulkDataStartOffset;
+		OutBulkDataOffsetInFile += Linker->Summary.BulkDataStartOffset;
 	}
 }
 
diff --git a/Engine/Source/Runtime/CoreUObject/Private/UObject/SavePackage.cpp b/Engine/Source/Runtime/CoreUObject/Private/UObject/SavePackage.cpp
--- a/Engine/Source/Runtime/CoreUObject/Private/UObject/SavePackage.cpp
+++ b/Engine/Source/Runtime/CoreUObject/Private/UObject/SavePackage.cpp
@@ -580,6 +580,22 @@ void FPackageStoreBulkDataManifest::PackageDesc::AddData(uint16 InIndex, uint64
 	Entry.Size = InSize;
 }
 
+FIoChunkId IBulkDataManifest::CreateChunkId(uint32 NameIndex, uint32 NameNumber, uint16 ChunkIndex)
+{
+	uint8 Data[12] = { 0 };
+
+	*reinterpret_cast<uint32*>(&Data[0]) = NameIndex;
+	*reinterpret_cast<int32*>(&Data[4]) = NameNumber;
+	*reinterpret_cast<uint16*>(&Data[8]) = ChunkIndex;
+	*reinterpret_cast<uint8*>(&Data[10]) = static_cast<uint8>(3); 	// EZenChunkType::BulkData in ZenCreator
+																	// TODO: We need to expose this somewhere!
+
+	FIoChunkId ChunkId;
+	ChunkId.Set(Data, 12);
+
+	return ChunkId;
+}
+
 FPackageStoreBulkDataManifest::FPackageStoreBulkDataManifest(const FString& InRootPath)
 {
 	RootPath = InRootPath;
@@ -598,7 +614,7 @@ FArchive& operator<<(FArchive& Ar, FPackageStoreBulkDataManifest::PackageDesc::B
 	Ar << Entry.Index;
 	Ar << Entry.Offset;
 	Ar << Entry.Size; 
-	
+
 	return Ar;
 }
 
@@ -6293,9 +6309,9 @@ void SaveBulkData(FLinkerSave* Linker, const UPackage* InOuter, const TCHAR* Fil
 			BulkDataAlignment = TargetPlatform->GetMemoryMappingAlignment();
 		}
 
-		for (int iBulkDataIndex = 0; iBulkDataIndex < Linker->BulkDataToAppend.Num(); ++iBulkDataIndex)
+		uint16 BulkDataIndex = 1;
+		for (FLinkerSave::FBulkDataStorageInfo& BulkDataStorageInfo : Linker->BulkDataToAppend)
 		{
-			FLinkerSave::FBulkDataStorageInfo& BulkDataStorageInfo = Linker->BulkDataToAppend[iBulkDataIndex];
 			BulkDataFeedback.EnterProgressFrame();
 
 			// Set bulk data flags to what they were during initial serialization (they might have changed after that)
@@ -6393,27 +6409,32 @@ void SaveBulkData(FLinkerSave* Linker, const UPackage* InOuter, const TCHAR* Fil
 				*Linker << SizeOnDiskAsInt32;
 			}
 
-			Linker->Seek(LinkerEndOffset);
-
-			// Restore BulkData flags to before serialization started
-			BulkDataStorageInfo.BulkData->ClearBulkDataFlags(0xFFFFFFFF);
-			BulkDataStorageInfo.BulkData->SetBulkDataFlags(OldBulkDataFlags);
-			BulkDataStorageInfo.BulkData->Unlock();
-
 			if ((BulkDataStorageInfo.BulkDataFlags & BULKDATA_CookedForIoDispatcher) != 0 && SavePackageContext != nullptr)
 			{
-				// TODO: iBulkDataIndex should be based on the number of BulkData objects we are adding to the manifest not
-				// it's position in Linker->BulkDataToAppend so that inlined BulkData objects are ignored and we won't have 
-				// missing index  values. 
-				// We can't do this yet because the BulkData object is writing out its index for runtime in ::Serialize and 
-				// so it's position in Linker->BulkDataToAppend is the only info we have.
-				// This can be fixed by generating the FIoChunkId at this point and writing it back to the bulkdata before it
-				// is written to disk as do for the data values above this point.
+				const int32 PackageNameIndex = SavePackageContext->HeaderSaver.NameMapSaver.MapName(InOuter->GetFName());
+				const int32 PackageNameNumber = InOuter->GetFName().GetNumber();
+				check(BulkDataIndex < TNumericLimits<uint16>::Max());
+				const uint16 ChunkIndex = BulkDataIndex++;
+
+				FIoChunkId ChunkId = IBulkDataManifest::CreateChunkId(PackageNameIndex, PackageNameNumber, ChunkIndex);
+
 				SavePackageContext->BulkDataManifest.AddFileAccess
 				(
-					Filename, iBulkDataIndex, BulkStartOffset, SizeOnDisk
+					Filename, ChunkIndex, BulkStartOffset, SizeOnDisk
 				);
+
+				check(BulkDataStorageInfo.BulkDataChunkIdPos != INDEX_NONE);
+
+				Linker->Seek(BulkDataStorageInfo.BulkDataChunkIdPos);
+				*Linker << ChunkId;
 			}
+
+			Linker->Seek(LinkerEndOffset);
+
+			// Restore BulkData flags to before serialization started
+			BulkDataStorageInfo.BulkData->ClearBulkDataFlags(0xFFFFFFFF);
+			BulkDataStorageInfo.BulkData->SetBulkDataFlags(OldBulkDataFlags);
+			BulkDataStorageInfo.BulkData->Unlock();
 		}
 
 		if (BulkArchive)
diff --git a/Engine/Source/Runtime/CoreUObject/Public/Serialization/BulkData.h b/Engine/Source/Runtime/CoreUObject/Public/Serialization/BulkData.h
--- a/Engine/Source/Runtime/CoreUObject/Public/Serialization/BulkData.h
+++ b/Engine/Source/Runtime/CoreUObject/Public/Serialization/BulkData.h
@@ -11,7 +11,7 @@
 #include "BulkData2.h"
 
 #if WITH_EDITOR == 0 && WITH_EDITORONLY_DATA == 0 //Runtime
-	#define USE_NEW_BULKDATA 0
+	#define USE_NEW_BULKDATA 0 // Set to 1 to enable 
 #else
 	#define USE_NEW_BULKDATA 0
 #endif
@@ -62,8 +62,10 @@ enum EBulkDataFlags
 	BULKDATA_Size64Bit							= 1 << 13,
 	/** Duplicate non-optional payload in optional bulk data. */
 	BULKDATA_DuplicateNonOptionalPayload		= 1 << 14,
+	/** Indicates that an old ID is present in the data, at some point when the DDCs are flushed we can remove this. */
+	BULKDATA_BadDataVersion						= 1 << 15,
 	/** Indicates that the data was cooked for use with IoDispatcher. */
-	BULKDATA_CookedForIoDispatcher				= 1 << 15
+	BULKDATA_CookedForIoDispatcher				= 1 << 16,
 };
 
 /**
@@ -165,15 +167,10 @@ public:
 	virtual bool PollCompletion() const override;
 	virtual bool WaitCompletion( float TimeLimitSeconds = 0.0f ) const override;
 
-	virtual uint8* GetReadResults() const override;
+	virtual uint8* GetReadResults() override;
 	virtual int64 GetSize() const override;
 
-	virtual bool operator == (const IAsyncReadRequest* InReadRequest) override
-	{
-		return ReadRequest == InReadRequest;
-	}
-
-	virtual void Cancel() const override;
+	virtual void Cancel() override;
 
 private:
 	IAsyncReadFileHandle* FileHandle;
@@ -358,6 +355,7 @@ private:
 
 public:
 	friend class FLinkerLoad;
+	using BulkDataRangeArray = TArray<FBulkDataStreamingToken*, TInlineAllocator<8>>;
 
 	/*-----------------------------------------------------------------------------
 		Constructors and operators
@@ -715,8 +713,7 @@ public:
 	* @param CompleteCallback	Called from an arbitrary thread when the request is complete. Can be nullptr, if non-null, must remain valid until it is called. It will always be called.
 	* @return					A request for the read. This is owned by the caller and must be deleted by the caller.
 	**/
-	static IBulkDataIORequest* CreateStreamingRequestForRange(const FString& Filename, const FBulkDataStreamingToken& Start, const FBulkDataStreamingToken& End, EAsyncIOPriorityAndFlags Priority, FAsyncFileCallBack* CompleteCallback);
-
+	static IBulkDataIORequest* CreateStreamingRequestForRange(const FString& Filename, const BulkDataRangeArray& RangeArray, EAsyncIOPriorityAndFlags Priority, FAsyncFileCallBack* CompleteCallback);
 #endif
 
 	/*-----------------------------------------------------------------------------
diff --git a/Engine/Source/Runtime/CoreUObject/Public/Serialization/BulkData2.h b/Engine/Source/Runtime/CoreUObject/Public/Serialization/BulkData2.h
--- a/Engine/Source/Runtime/CoreUObject/Public/Serialization/BulkData2.h
+++ b/Engine/Source/Runtime/CoreUObject/Public/Serialization/BulkData2.h
@@ -4,6 +4,7 @@
 
 #include "CoreMinimal.h"
 #include "Async/AsyncFileHandle.h"
+#include "IO/IoDispatcher.h"
 
 // Any place marked with BULKDATA_NOT_IMPLEMENTED_FOR_RUNTIME marks a method that we do not support
 // but needs to exist in order for the code to compile. This pretty much means code that is editor only
@@ -30,12 +31,10 @@ public:
 	virtual bool PollCompletion() const = 0;
 	virtual bool WaitCompletion(float TimeLimitSeconds = 0.0f) const = 0;
 
-	virtual uint8* GetReadResults() const = 0;
+	virtual uint8* GetReadResults() = 0;
 	virtual int64 GetSize() const = 0;
 
-	virtual bool operator == (const IAsyncReadRequest* InReadRequest) = 0;
-
-	virtual void Cancel() const = 0;
+	virtual void Cancel() = 0;
 };
 
 /**
@@ -43,6 +42,11 @@ public:
  */
 class COREUOBJECT_API FBulkDataBase
 {
+public:
+	using BulkDataRangeArray = TArray<FBulkDataBase*, TInlineAllocator<8>>;
+
+	static void SetIODispatcher(FIoDispatcher* InIoDispatcher) { IoDispatcher = InIoDispatcher; }
+
 public:
 	using FileToken = uint64;
 	static constexpr FileToken InvalidToken = 0;
@@ -51,12 +55,12 @@ public:
 	FBulkDataBase(FBulkDataBase&& Other);
 	FBulkDataBase& operator=(const FBulkDataBase& Other);
 
-	// Implemented:
 	FBulkDataBase() = default;
 	~FBulkDataBase();
+
 protected:
-	// Unimplemented:
-	void Serialize(FArchive& Ar, UObject* Owner, int32 Index, bool bAttemptFileMapping, int32 ElementSize );
+
+	void Serialize(FArchive& Ar, UObject* Owner, int32 Index, bool bAttemptFileMapping, int32 ElementSize);
 
 public:
 	// Unimplemented:
@@ -64,18 +68,18 @@ public:
 	const void* LockReadOnly() const;
 	void Unlock();
 	bool IsLocked() const;
-	
+
 	void* Realloc(int64 InElementCount);
 	void GetCopy(void** Dest, bool bDiscardInternalCopy);
 
 	int64 GetBulkDataSize() const;
-	
+
 	void SetBulkDataFlags(uint32 BulkDataFlagsToSet);
 	void ResetBulkDataFlags(uint32 BulkDataFlagsToSet);
 	void ClearBulkDataFlags(uint32 BulkDataFlagsToClear);
 	uint32 GetBulkDataFlags() const { return BulkDataFlags; }
 
-	bool CanLoadFromDisk() const { return Token != InvalidToken; }
+	bool CanLoadFromDisk() const { return IsUsingIODispatcher() || Fallback.Token != InvalidToken; }
 
 	bool IsStoredCompressedOnDisk() const;
 	FName GetDecompressionFormat() const;
@@ -90,14 +94,16 @@ public:
 	bool InSeperateFile() const;
 	bool IsSingleUse() const;
 	bool IsMemoryMapped() const;
+	bool IsUsingIODispatcher() const;
 
 	IBulkDataIORequest* CreateStreamingRequest(EAsyncIOPriorityAndFlags Priority, FAsyncFileCallBack* CompleteCallback, uint8* UserSuppliedMemory) const;
 	IBulkDataIORequest* CreateStreamingRequest(int64 OffsetInBulkData, int64 BytesToRead, EAsyncIOPriorityAndFlags Priority, FAsyncFileCallBack* CompleteCallback, uint8* UserSuppliedMemory) const;
-	static IBulkDataIORequest* CreateStreamingRequestForRange(const FBulkDataBase& Start, const FBulkDataBase& End, EAsyncIOPriorityAndFlags Priority, FAsyncFileCallBack* CompleteCallback);
+	
+	static IBulkDataIORequest* CreateStreamingRequestForRange(const BulkDataRangeArray& RangeArray, EAsyncIOPriorityAndFlags Priority, FAsyncFileCallBack* CompleteCallback);
 
 	void RemoveBulkData();
 
-	bool IsAsyncLoadingComplete() const { return true;  }
+	bool IsAsyncLoadingComplete() const { return true; }
 
 	// Added for compatibility with the older BulkData system
 	int64 GetBulkDataOffsetInFile() const;
@@ -111,7 +117,7 @@ public:
 private:
 	void LoadDataDirectly(void** DstBuffer);
 
-	void SerializeDuplicateData(FArchive& Ar, UObject* Owner, uint32& OutBulkDataFlags, int64& OutBulkDataSizeOnDisk, int64& OutBulkDataOffsetInFile);
+	void SerializeDuplicateData(FArchive& Ar, FLinkerLoad* Linker, uint32& OutBulkDataFlags, int64& OutBulkDataSizeOnDisk, int64& OutBulkDataOffsetInFile);
 	void SerializeBulkData(FArchive& Ar, void* DstBuffer, int64 DataLength);
 
 	void AllocateData(SIZE_T SizeInBytes);
@@ -119,7 +125,24 @@ private:
 
 	FString ConvertFilenameFromFlags(const FString& Filename);
 
-	FileToken Token = InvalidToken; // TODO: Maybe store as a union of Token/InlineData size so that inline files don't need to register?
+private:
+	friend class FBulkDataIoDispatcherRequest;
+	static FIoDispatcher* IoDispatcher;
+
+	union
+	{
+		// Inline data or fallback path
+		struct
+		{
+			uint64 BulkDataSize;
+			FileToken Token;
+			
+		} Fallback;
+
+		// For IODispatcher
+		FIoChunkId ChunkID;	
+	}; // Note that the union will end up being 16 bytes with padding
+	
 	void* DataBuffer = nullptr;
 	uint32 BulkDataFlags = 0;
 	mutable uint8 LockStatus = 0; // Mutable so that the read only lock can be const
diff --git a/Engine/Source/Runtime/CoreUObject/Public/UObject/LinkerSave.h b/Engine/Source/Runtime/CoreUObject/Public/UObject/LinkerSave.h
--- a/Engine/Source/Runtime/CoreUObject/Public/UObject/LinkerSave.h
+++ b/Engine/Source/Runtime/CoreUObject/Public/UObject/LinkerSave.h
@@ -57,6 +57,8 @@ public:
 		int64 BulkDataSizeOnDiskPos;
 		/** Offset to the location where the bulk data flags are stored */
 		int64 BulkDataFlagsPos;
+		/** Offset to the location where the bulk data ChunkID is stored */
+		int64 BulkDataChunkIdPos;
 		/** Bulk data flags at the time of serialization */
 		uint32 BulkDataFlags;
 		/** The bulkdata */
diff --git a/Engine/Source/Runtime/CoreUObject/Public/UObject/SavePackage.h b/Engine/Source/Runtime/CoreUObject/Public/UObject/SavePackage.h
--- a/Engine/Source/Runtime/CoreUObject/Public/UObject/SavePackage.h
+++ b/Engine/Source/Runtime/CoreUObject/Public/UObject/SavePackage.h
@@ -75,6 +75,8 @@ public:
 
 	virtual void Save() = 0;
 	virtual void AddFileAccess(const FString& PackageFilename, uint16 InIndex, uint64 InOffset, uint64 InSize) = 0;
+
+	static FIoChunkId CreateChunkId(uint32 NameIndex, uint32 NameNumber, uint16 ChunkIndex);
 };
 
 class FPackageStoreBulkDataManifest : public IBulkDataManifest
diff --git a/Engine/Source/Runtime/Engine/Private/Streaming/SkeletalMeshUpdate.cpp b/Engine/Source/Runtime/Engine/Private/Streaming/SkeletalMeshUpdate.cpp
--- a/Engine/Source/Runtime/Engine/Private/Streaming/SkeletalMeshUpdate.cpp
+++ b/Engine/Source/Runtime/Engine/Private/Streaming/SkeletalMeshUpdate.cpp
@@ -405,8 +405,11 @@ void FSkeletalMeshStreamIn_IO::SetIORequest(const FContext& Context, const FStri
 	{
 		SetAsyncFileCallback(Context);
 
-		const FSkeletalMeshLODRenderData& FirstLOD = RenderData->LODRenderData[PendingFirstMip];
-		const FSkeletalMeshLODRenderData& LastLOD = RenderData->LODRenderData[CurrentFirstLODIdx - 1];
+		FBulkDataInterface::BulkDataRangeArray BulkDataArray;
+		for (int32 Index = PendingFirstMip; Index < CurrentFirstLODIdx; ++Index)
+		{
+			BulkDataArray.Push(&RenderData->LODRenderData[Index].StreamingBulkData);
+		}
 
 		// Increment as we push the request. If a request complete immediately, then it will call the callback
 		// but that won't do anything because the tick would not try to acquire the lock since it is already locked.
@@ -414,8 +417,7 @@ void FSkeletalMeshStreamIn_IO::SetIORequest(const FContext& Context, const FStri
 
 		IORequest = FBulkDataInterface::CreateStreamingRequestForRange(
 			STREAMINGTOKEN_PARAM(IOFilename)
-			FirstLOD.StreamingBulkData,
-			LastLOD.StreamingBulkData,
+			BulkDataArray,
 			bHighPrioIORequest ? AIOP_BelowNormal : AIOP_Low,
 			&AsyncFileCallback);
 	}
diff --git a/Engine/Source/Runtime/Engine/Private/Streaming/StaticMeshUpdate.cpp b/Engine/Source/Runtime/Engine/Private/Streaming/StaticMeshUpdate.cpp
--- a/Engine/Source/Runtime/Engine/Private/Streaming/StaticMeshUpdate.cpp
+++ b/Engine/Source/Runtime/Engine/Private/Streaming/StaticMeshUpdate.cpp
@@ -379,8 +379,11 @@ void FStaticMeshStreamIn_IO::SetIORequest(const FContext& Context, const FString
 	{
 		SetAsyncFileCallback(Context);
 
-		const FStaticMeshLODResources& FirstLOD = RenderData->LODResources[PendingFirstMip];
-		const FStaticMeshLODResources& LastLOD = RenderData->LODResources[CurrentFirstLODIdx - 1];
+		FBulkDataInterface::BulkDataRangeArray BulkDataArray;
+		for (int32 Index = PendingFirstMip; Index < CurrentFirstLODIdx; ++Index)
+		{
+			BulkDataArray.Push(&RenderData->LODResources[Index].StreamingBulkData);
+		}
 
 		// Increment as we push the request. If a request complete immediately, then it will call the callback
 		// but that won't do anything because the tick would not try to acquire the lock since it is already locked.
@@ -388,8 +391,7 @@ void FStaticMeshStreamIn_IO::SetIORequest(const FContext& Context, const FString
 
 		IORequest = FBulkDataInterface::CreateStreamingRequestForRange(
 			STREAMINGTOKEN_PARAM(IOFilename)
-			FirstLOD.StreamingBulkData,
-			LastLOD.StreamingBulkData,
+			BulkDataArray,
 			bHighPrioIORequest ? AIOP_BelowNormal : AIOP_Low,
 			&AsyncFileCallback);
 	}
