AsyncLoading2: Reduce EDL runtime graph
Remove import nodes (ImportOrExport_Create, ImportOrExport_Serialize, Import_NumPhases) from EDL runtime graph.
Replace import->export arcs with export->export arcs.
During initial load add barriers for exports instead of imports (using a separate scriptarcs file only required during initial loading).
Merge Package_LoadSummary & Package_ImportPackages (Event_FinishLinker & Event_StartImportPackages) into Package_ProcessSummary (Event_ProcessSummary).
Remove Package_PostLoad (Event_StartPostLoad) and replace post load dependency arcs with Package_ExportsSerialized->Package_Tick (Event_Tick).
Remove Export_StartIO node and execute this from CreateExport instead.

diff --git a/Engine/Source/Runtime/CoreUObject/Private/Serialization/AsyncLoading2.cpp b/Engine/Source/Runtime/CoreUObject/Private/Serialization/AsyncLoading2.cpp
--- a/Engine/Source/Runtime/CoreUObject/Private/Serialization/AsyncLoading2.cpp
+++ b/Engine/Source/Runtime/CoreUObject/Private/Serialization/AsyncLoading2.cpp
@@ -130,7 +130,7 @@ struct FExportMapEntry
 	FPackageIndex ClassIndex;
 	FPackageIndex SuperIndex;
 	FPackageIndex TemplateIndex;
-	FPackageIndex GlobalImportIndex;
+	int32 GlobalImportIndex;
 	EObjectFlags ObjectFlags;
 };
 
@@ -151,6 +151,8 @@ struct FPackageStoreEntry
 	int32 SlimportCount;
 	int32 SlimportOffset;
 	int32 ExportCount;
+	int32 ScriptArcsOffset;
+	int32 ScriptArcsCount;
 };
 
 struct FPackageSummary
@@ -261,7 +263,9 @@ class FPackageStore
 	FGlobalImportStore ImportStore;
 	FPackageStoreEntry* StoreEntries = nullptr;
 	int32* Slimports = nullptr;
+	int32* ScriptArcs = nullptr;
 	int32 SlimportCount = 0;
+	int32 ScriptArcsCount = 0;
 	int32 PackageCount = 0;
 
 public:
@@ -293,7 +297,7 @@ public:
 
 			// Global package index
 			{
-				TRACE_CPUPROFILER_EVENT_SCOPE(LoadPackageStoreFNameFixup);
+				TRACE_CPUPROFILER_EVENT_SCOPE(LoadPackageStoreGlobalIds);
 				for (int I = 0; I < PackageCount; ++I)
 				{
 					PackageNameToGlobalPackageId.Add(StoreEntries[I].Name, FGlobalPackageId { I });
@@ -345,6 +349,18 @@ public:
 			}
 		}
 
+		{
+			TRACE_CPUPROFILER_EVENT_SCOPE(LoadPackageStoreScriptArcs);
+
+			TUniquePtr<FArchive> ScriptArcsArchive(IFileManager::Get().CreateFileReader(*(RootDir / TEXT("megafile.uscriptarcs"))));
+			check(ScriptArcsArchive);
+
+			const int32 ScriptArcsByteCount = ScriptArcsArchive->TotalSize();
+			ScriptArcsCount = ScriptArcsByteCount / sizeof(int32);
+			ScriptArcs = reinterpret_cast<int32*>(FMemory::Malloc(ScriptArcsByteCount));
+			ScriptArcsArchive->Serialize(ScriptArcs, ScriptArcsByteCount);
+		}
+
 	}
 
 	inline FGlobalImportStore& GetImportStore()
@@ -394,6 +410,13 @@ public:
 		return Slimports + Entry.SlimportOffset / sizeof(int32);
 	}
 
+	inline int32* GetPackageScriptArcs(FGlobalPackageId GlobalPackageId, int32& OutCount) const 
+	{
+		FPackageStoreEntry& Entry = StoreEntries[GlobalPackageId.Id];
+		OutCount = Entry.ScriptArcsCount;
+		return ScriptArcs + Entry.ScriptArcsOffset / sizeof(int32);
+	}
+
 	inline int32 GetPackageImportCount(FGlobalPackageId GlobalPackageId) const 
 	{
 		return StoreEntries[GlobalPackageId.Id].ImportCount;
@@ -439,6 +462,9 @@ struct FPackageImportStore
 		Clear();
 	}
 
+	UObject* FindExistingSlimport(int32 LocalImportIndex);
+	UObject* FindExistingGlimport(int32 GlobalImportIndex);
+
 	UObject* GetObject(int32 LocalImportIndex)
 	{
 		return GlobalImportObjects[ImportMap[LocalImportIndex]];
@@ -516,7 +542,8 @@ public:
 		}
 		else
 		{
-			Object = ImportStore->GetObject(Index);
+			// Object = ImportStore->GetObject(Index);
+			Object = ImportStore->FindExistingSlimport(Index.ToImport());
 		}
 		return *this;
 	}
@@ -641,10 +668,6 @@ enum class EAsyncPackageLoadingState2 : uint8
 {
 	NewPackage,
 	WaitingForSummary,
-	StartImportPackages,
-	WaitingForImportPackages,
-	SetupImports,
-	SetupExports,
 	ProcessNewImportsAndExports,
 	PostLoad_Etc,
 	PackageComplete,
@@ -653,21 +676,14 @@ enum class EAsyncPackageLoadingState2 : uint8
 enum EEventLoadNode2 : uint8
 {
 	Package_CreateLinker,
-	Package_LoadSummary,
-	Package_ImportPackages,
-	Package_SetupImports,
-	Package_SetupExports,
+	Package_ProcessSummary,
 	Package_ExportsSerialized,
-	Package_PostLoad,
 	Package_Tick,
 	Package_Delete,
 	Package_NumPhases,
 
-	ImportOrExport_Create = 0,
-	ImportOrExport_Serialize,
-	Import_NumPhases,
-
-	Export_StartIO = Import_NumPhases,
+	Export_Create = 0,
+	Export_Serialize,
 	Export_NumPhases,
 };
 
@@ -1147,9 +1163,7 @@ private:
 	TSet<UPackage*> ImportedPackages;
 
 	FEventLoadNode2* PackageNodes = nullptr;
-	FEventLoadNode2* ImportNodes = nullptr;
 	FEventLoadNode2* ExportNodes = nullptr;
-	uint32 ImportNodeCount = 0;
 	uint32 ExportNodeCount = 0;
 
 	uint64 PackageSummarySize = 0;
@@ -1184,31 +1198,17 @@ public:
 	bool bAllExportsSerialized;
 
 	static EAsyncPackageState::Type Event_CreateLinker(FAsyncPackage2* Package, int32);
-	static EAsyncPackageState::Type Event_FinishLinker(FAsyncPackage2* Package, int32);
-	static EAsyncPackageState::Type Event_StartImportPackages(FAsyncPackage2* Package, int32);
-	static EAsyncPackageState::Type Event_SetupImports(FAsyncPackage2* Package, int32);
-	static EAsyncPackageState::Type Event_SetupExports(FAsyncPackage2* Package, int32);
-	static EAsyncPackageState::Type Event_LinkImport(FAsyncPackage2* Package, int32 LocalImportIndex);
-	static EAsyncPackageState::Type Event_ImportSerialized(FAsyncPackage2* Package, int32 LocalImportIndex);
+	static EAsyncPackageState::Type Event_ProcessSummary(FAsyncPackage2* Package, int32);
 	static EAsyncPackageState::Type Event_CreateExport(FAsyncPackage2* Package, int32 LocalExportIndex);
-	static EAsyncPackageState::Type Event_StartIO(FAsyncPackage2* Package, int32 LocalExportIndex);
 	static EAsyncPackageState::Type Event_SerializeExport(FAsyncPackage2* Package, int32 LocalExportIndex);
 	static EAsyncPackageState::Type Event_ExportsDone(FAsyncPackage2* Package, int32);
-	static EAsyncPackageState::Type Event_StartPostload(FAsyncPackage2* Package, int32);
 	static EAsyncPackageState::Type Event_Tick(FAsyncPackage2* Package, int32);
 	static EAsyncPackageState::Type Event_Delete(FAsyncPackage2* Package, int32);
 
 	void ProcessIoRequest(const FIoChunkId& ChunkId, TIoStatusOr<FIoBuffer>& Result);
 
-	void MarkNewObjectForLoadIfItIsAnExport(UObject *Object);
-
-	EAsyncPackageState::Type SetupSlimports_Event();
-
-	UObject* FindExistingSlimport(int32 GlobalImportIndex);
-	void LinkSlimport(int32 LocalImportIndex, int32 GlobalImportIndex = -1);
 	void EventDrivenCreateExport(int32 LocalExportIndex);
 	void EventDrivenSerializeExport(int32 LocalExportIndex);
-	void EventDrivenLoadingComplete();
 
 	UObject* EventDrivenIndexToObject(FPackageIndex Index, bool bCheckSerialized);
 	template<class T>
@@ -1223,6 +1223,8 @@ public:
 	}
 
 	FEventLoadNode2* GetNode(EEventLoadNode2 Phase, FPackageIndex ImportOrExportIndex = FPackageIndex());
+	
+	FEventLoadNode2* GetNode(int32 NodeIndex);
 
 	/** [EDL] End Event driven loader specific stuff */
 
@@ -1239,6 +1241,10 @@ public:
 	void CloseDelayedLinkers();
 
 private:
+	void CreateNodes(const FAsyncLoadEventSpec* EventSpecs);
+	void SetupScriptArcs();
+	void StartImportPackages();
+
 	/**
 	 * Begin async loading process. Simulates parts of BeginLoad.
 	 *
@@ -1541,6 +1547,7 @@ public:
 	*/
 	FORCEINLINE FAsyncPackage2* FindAsyncPackage(const FName& PackageName)
 	{
+		TRACE_CPUPROFILER_EVENT_SCOPE(FindAsyncPackage);
 		FScopeLock LockAsyncPackages(&AsyncPackagesCritical);
 		//checkSlow(IsInAsyncLoadThread());
 		return AsyncPackageNameLookup.FindRef(PackageName);
@@ -1757,8 +1764,14 @@ void FAsyncLoadingThread2Impl::InitializeLoading()
 	FString RootDir;
 	if (!FParse::Value(FCommandLine::Get(), TEXT("-zendir="), RootDir))
 	{
-		UE_LOG(LogStreaming, Error, TEXT("Failed to initialize package loader. No root directory specified"));
-		return;
+		if (IFileManager::Get().IsSandboxEnabled())
+		{
+			UE_LOG(LogStreaming, Error, TEXT("Failed to initialize package loader. No root directory specified when running with sandbox enabled"));
+		}
+		else
+		{
+			RootDir = FPaths::EngineDir() / TEXT("..");
+		}
 	}
 	
 	UE_LOG(LogStreaming, Log, TEXT("Initializing package loader using directory '%s'"), *RootDir);
@@ -1825,7 +1838,7 @@ void FAsyncLoadingThread2Impl::UpdateExistingPackagePriorities(FAsyncPackage2* I
 
 FAsyncPackage2* FAsyncLoadingThread2Impl::FindOrInsertPackage(FAsyncPackageDesc* InDesc, bool& bInserted)
 {
-	//TRACE_CPUPROFILER_EVENT_SCOPE(FindOrInsertPackage);
+	TRACE_CPUPROFILER_EVENT_SCOPE(FindOrInsertPackage);
 	FAsyncPackage2* Package = nullptr;
 	bInserted = false;
 	{
@@ -1892,7 +1905,7 @@ FEventLoadNode2::FEventLoadNode2(const FAsyncLoadEventSpec* InSpec, FAsyncPackag
 
 void FEventLoadNode2::DependsOn(FEventLoadNode2* Other)
 {
-	//TRACE_CPUPROFILER_EVENT_SCOPE(DependsOn);
+	TRACE_CPUPROFILER_EVENT_SCOPE(DependsOn);
 #if UE_BUILD_DEVELOPMENT
 	check(!bDone);
 	check(!bFired);
@@ -1916,6 +1929,7 @@ void FEventLoadNode2::DependsOn(FEventLoadNode2* Other)
 		{
 			if (Other->DependenciesCount == 1)
 			{
+				TRACE_CPUPROFILER_EVENT_SCOPE(DependsOnAlloc);
 				FEventLoadNode2* FirstDependency = Other->SingleDependent;
 				uint32 NewDependenciesCapacity = 4;
 				Other->DependenciesCapacity = NewDependenciesCapacity;
@@ -1924,6 +1938,7 @@ void FEventLoadNode2::DependsOn(FEventLoadNode2* Other)
 			}
 			else if (Other->DependenciesCount == Other->DependenciesCapacity)
 			{
+				TRACE_CPUPROFILER_EVENT_SCOPE(DependsOnRealloc);
 				FEventLoadNode2** OriginalDependents = Other->MultipleDependents;
 				uint32 OldDependenciesCapcity = Other->DependenciesCapacity;
 				SIZE_T OldDependenciesSize = OldDependenciesCapcity * sizeof(FEventLoadNode2*);
@@ -2208,31 +2223,40 @@ EAsyncPackageState::Type FAsyncPackage2::Event_CreateLinker(FAsyncPackage2* Pack
 	return EAsyncPackageState::Complete;
 }
 
-EAsyncPackageState::Type FAsyncPackage2::Event_FinishLinker(FAsyncPackage2* Package, int32)
+EAsyncPackageState::Type FAsyncPackage2::Event_ProcessSummary(FAsyncPackage2* Package, int32)
 {
-	TRACE_CPUPROFILER_EVENT_SCOPE(Event_FinishLinker);
+	TRACE_CPUPROFILER_EVENT_SCOPE(Event_ProcessSummary);
 	FScopedAsyncPackageEvent2 Scope(Package);
 	EAsyncPackageState::Type Result = Package->FinishLinker();
 	check(Result == EAsyncPackageState::Complete);
 	check(Package->Linker/* && Linker->HasFinishedInitialization()*/);
+
+	Package->StartImportPackages();
+
+	if (GIsInitialLoad)
+	{
+		Package->SetupScriptArcs();
+	}
+
 	check(Package->AsyncPackageLoadingState == EAsyncPackageLoadingState2::WaitingForSummary);
-	Package->AsyncPackageLoadingState = EAsyncPackageLoadingState2::StartImportPackages;
+	Package->AsyncPackageLoadingState = EAsyncPackageLoadingState2::ProcessNewImportsAndExports;
+
 	return EAsyncPackageState::Complete;
 }
 
-EAsyncPackageState::Type FAsyncPackage2::Event_StartImportPackages(FAsyncPackage2* Package, int32)
+void FAsyncPackage2::StartImportPackages()
 {
-	TRACE_CPUPROFILER_EVENT_SCOPE(Event_StartImportPackages);
+	TRACE_CPUPROFILER_EVENT_SCOPE(StartImportPackages);
 
 	LLM_SCOPE(ELLMTag::AsyncLoading);
 
 	{
 		TRACE_CPUPROFILER_EVENT_SCOPE(SetupSerializedArcs);
-		const FGlobalNameMap& GlobalNameMap = Package->AsyncLoadingThread.GlobalNameMap;
+		const FGlobalNameMap& GlobalNameMap = AsyncLoadingThread.GlobalNameMap;
 
-		const FPackageSummary* PackageSummary = reinterpret_cast<const FPackageSummary*>(Package->PackageSummaryIoBuffer.Data());
+		const FPackageSummary* PackageSummary = reinterpret_cast<const FPackageSummary*>(PackageSummaryIoBuffer.Data());
 		uint64 ZenHeaderDataSize = PackageSummary->GraphDataSize;
-		const uint8* ZenHeaderData = Package->PackageSummaryIoBuffer.Data() + PackageSummary->GraphDataOffset;
+		const uint8* ZenHeaderData = PackageSummaryIoBuffer.Data() + PackageSummary->GraphDataOffset;
 		FSimpleArchive ZenHeaderArchive(ZenHeaderData, ZenHeaderDataSize);
 		int32 InternalArcCount;
 		ZenHeaderArchive << InternalArcCount;
@@ -2242,7 +2266,7 @@ EAsyncPackageState::Type FAsyncPackage2::Event_StartImportPackages(FAsyncPackage
 			int32 ToNodeIndex;
 			ZenHeaderArchive << FromNodeIndex;
 			ZenHeaderArchive << ToNodeIndex;
-			Package->PackageNodes[ToNodeIndex].DependsOn(Package->PackageNodes + FromNodeIndex);
+			PackageNodes[ToNodeIndex].DependsOn(PackageNodes + FromNodeIndex);
 		}
 		int32 ImportedPackagesCount;
 		ZenHeaderArchive << ImportedPackagesCount;
@@ -2252,7 +2276,7 @@ EAsyncPackageState::Type FAsyncPackage2::Event_StartImportPackages(FAsyncPackage
 			int32 ImportedPackageNameNumber;
 			ZenHeaderArchive << ImportedPackageNameIndex << ImportedPackageNameNumber;
 			FName ImportedPackageName = GlobalNameMap.GetName(ImportedPackageNameIndex, ImportedPackageNameNumber);
-			FAsyncPackage2* ImportedPackage = Package->AsyncLoadingThread.FindAsyncPackage(ImportedPackageName);
+			FAsyncPackage2* ImportedPackage = AsyncLoadingThread.FindAsyncPackage(ImportedPackageName);
 			int32 ExternalArcCount;
 			ZenHeaderArchive << ExternalArcCount;
 			for (int32 ExternalArcIndex = 0; ExternalArcIndex < ExternalArcCount; ++ExternalArcIndex)
@@ -2263,18 +2287,14 @@ EAsyncPackageState::Type FAsyncPackage2::Event_StartImportPackages(FAsyncPackage
 				ZenHeaderArchive << ToNodeIndex;
 				if (ImportedPackage)
 				{
-					Package->PackageNodes[ToNodeIndex].DependsOn(ImportedPackage->PackageNodes + FromNodeIndex);
+					PackageNodes[ToNodeIndex].DependsOn(ImportedPackage->PackageNodes + FromNodeIndex);
 				}
 			}
 		}
 	}
-
-	check(Package->AsyncPackageLoadingState == EAsyncPackageLoadingState2::StartImportPackages);
-	Package->AsyncPackageLoadingState = EAsyncPackageLoadingState2::WaitingForImportPackages;
-	return EAsyncPackageState::Complete;
 }
 
-static UObject* GFindExistingSlimport(int32 GlobalImportIndex,
+static UObject* GFindExistingGlimport(int32 GlobalImportIndex,
 	UObject** GlobalImportObjects,
 	const FPackageIndex* GlobalImportOuters,
 	const FName* GlobalImportNames)
@@ -2290,7 +2310,7 @@ static UObject* GFindExistingSlimport(int32 GlobalImportIndex,
 		}
 		else
 		{
-			UObject* Outer = GFindExistingSlimport(OuterIndex.ToImport(), GlobalImportObjects, GlobalImportOuters, GlobalImportNames);
+			UObject* Outer = GFindExistingGlimport(OuterIndex.ToImport(), GlobalImportObjects, GlobalImportOuters, GlobalImportNames);
 			if (Outer)
 			{
 				Object = StaticFindObjectFast(UObject::StaticClass(), Outer, ObjectName, false, true);
@@ -2300,14 +2320,22 @@ static UObject* GFindExistingSlimport(int32 GlobalImportIndex,
 	return Object;
 }
 
-UObject* FAsyncPackage2::FindExistingSlimport(int32 GlobalImportIndex)
+UObject* FPackageImportStore::FindExistingGlimport(int32 GlobalImportIndex)
 {
-	return GFindExistingSlimport(GlobalImportIndex,
-		ImportStore.GlobalImportObjects,
-		ImportStore.GlobalImportOuters,
-		ImportStore.GlobalImportNames);
+	return GFindExistingGlimport(GlobalImportIndex,
+		GlobalImportObjects,
+		GlobalImportOuters,
+		GlobalImportNames);
 }
 
+UObject* FPackageImportStore::FindExistingSlimport(int32 LocalImportIndex)
+{
+	int32 GlobalImportIndex = ImportMap[LocalImportIndex];
+	return GFindExistingGlimport(GlobalImportIndex,
+		GlobalImportObjects,
+		GlobalImportOuters,
+		GlobalImportNames);
+}
 void FAsyncPackage2::ImportPackagesRecursive()
 {
 	if (bHasImportedPackagesRecursive)
@@ -2325,7 +2353,7 @@ void FAsyncPackage2::ImportPackagesRecursive()
 	{
 		const int32 GlobalImportIndex = ImportStore.ImportMap[LocalImportIndex];
 		const FPackageIndex ImportedPackageIndex = ImportStore.GlobalImportPackages[GlobalImportIndex];
-		UObject* ImportedObject = FindExistingSlimport(GlobalImportIndex);
+		UObject* ImportedObject = ImportStore.FindExistingGlimport(GlobalImportIndex);
 		const UPackage* Package = (UPackage*)ImportStore.GlobalImportObjects[ImportedPackageIndex.ToImport()];
 
 		if (Package && Package->HasAnyPackageFlags(PKG_CompiledIn))
@@ -2360,11 +2388,6 @@ void FAsyncPackage2::ImportPackagesRecursive()
 		{
 			ImportedAsyncPackage->ImportPackagesRecursive();
 		}
-
-		// we can't set up our imports until all packages we are importing have loaded their summary
-		FEventLoadNode2* SetupImportsNode = GetNode(EEventLoadNode2::Package_SetupImports);
-		FEventLoadNode2* ImportedPackageLoadSummaryNode = ImportedAsyncPackage->GetNode(EEventLoadNode2::Package_LoadSummary);
-		SetupImportsNode->DependsOn(ImportedPackageLoadSummaryNode);
 	}
 #else
 	bool bNeedToLoadPackage = false;
@@ -2435,39 +2458,25 @@ void FAsyncPackage2::ImportPackagesRecursive()
 				ImportedAsyncPackage->ImportPackagesRecursive();
 			}
 
-			FEventLoadNode2* SetupImportsNode = GetNode(EEventLoadNode2::Package_SetupImports);
-			FEventLoadNode2* ImportedPackageLoadSummaryNode = ImportedAsyncPackage->GetNode(EEventLoadNode2::Package_LoadSummary);
-			SetupImportsNode->DependsOn(ImportedPackageLoadSummaryNode);
+			FEventLoadNode2* SetupExportsNode = GetNode(EEventLoadNode2::Package_SetupExports);
+			FEventLoadNode2* ImportedPackageLoadSummaryNode = ImportedAsyncPackage->GetNode(EEventLoadNode2::Package_ProcessSummary);
+			SetupExportsNode->DependsOn(ImportedPackageLoadSummaryNode);
 		}
 	}
 #endif
 }
 
-EAsyncPackageState::Type FAsyncPackage2::Event_SetupImports(FAsyncPackage2* Package, int32)
-{
-	TRACE_CPUPROFILER_EVENT_SCOPE(Event_SetupImports);
-
-	check(Package->AsyncPackageLoadingState == EAsyncPackageLoadingState2::WaitingForImportPackages);
-	Package->AsyncPackageLoadingState = EAsyncPackageLoadingState2::SetupImports;
-	{
-		FScopedAsyncPackageEvent2 Scope(Package);
-		verify(Package->SetupSlimports_Event() == EAsyncPackageState::Complete);
-	}
-	check(Package->AsyncPackageLoadingState == EAsyncPackageLoadingState2::SetupImports);
-	Package->AsyncPackageLoadingState = EAsyncPackageLoadingState2::SetupExports;
-	return EAsyncPackageState::Complete;
-}
-
-EAsyncPackageState::Type FAsyncPackage2::SetupSlimports_Event()
+void FAsyncPackage2::SetupScriptArcs()
 {
-	if (!GIsInitialLoad)
-	{
-		return EAsyncPackageState::Complete;
-	}
+	TRACE_CPUPROFILER_EVENT_SCOPE(SetupScriptArcs);
 
-	for (int32 LocalImportIndex = 0; LocalImportIndex < ImportStore.Count; ++LocalImportIndex)
+	FPackageStore& GlobalPackageStore = AsyncLoadingThread.PackageStore;
+	int32 ScriptArcsCount = 0;
+	int32* ScriptArcs = GlobalPackageStore.GetPackageScriptArcs(GlobalPackageId, ScriptArcsCount);
+	for (int32 ScriptArcIndex = 0; ScriptArcIndex < ScriptArcsCount * 2;)
 	{
-		const int32 GlobalImportIndex = ImportStore.ImportMap[LocalImportIndex];
+		int32 GlobalImportIndex = ScriptArcs[ScriptArcIndex++];
+		int32 ToNodeIndex = ScriptArcs[ScriptArcIndex++];
 
 		// skip packages
 		if (ImportStore.GlobalImportOuters[GlobalImportIndex].IsNull())
@@ -2512,41 +2521,13 @@ EAsyncPackageState::Type FAsyncPackage2::SetupSlimports_Event()
 			// OuterMostNonPackageIndex is used here because if it is a CDO or subobject, etc,
 			// we wait for the outermost thing that is not a package
 			const bool bWaitingForCompiledInImport = EDLBootNotificationManager.AddWaitingPackage(
-				this, PackageName, OuterMostNonPackageObjectName, FPackageIndex::FromImport(LocalImportIndex));
+				this, PackageName, OuterMostNonPackageObjectName, FPackageIndex::FromImport(ToNodeIndex));
 			if (bWaitingForCompiledInImport)
 			{
-				GetNode(EEventLoadNode2::ImportOrExport_Create, FPackageIndex::FromImport(LocalImportIndex))->AddBarrier();
+				PackageNodes[ToNodeIndex].AddBarrier();
 			}
 		}
 	}
-
-	return EAsyncPackageState::Complete;
-}
-
-EAsyncPackageState::Type FAsyncPackage2::Event_SetupExports(FAsyncPackage2* Package, int32)
-{
-	TRACE_CPUPROFILER_EVENT_SCOPE(Event_SetupExports);
-
-	check(Package->AsyncPackageLoadingState == EAsyncPackageLoadingState2::SetupExports);
-	Package->AsyncPackageLoadingState = EAsyncPackageLoadingState2::ProcessNewImportsAndExports;
-	return EAsyncPackageState::Complete;
-}
-
-EAsyncPackageState::Type FAsyncPackage2::Event_LinkImport(FAsyncPackage2* Package, int32 LocalImportIndex)
-{
-	TRACE_CPUPROFILER_EVENT_SCOPE(Event_LinkImport);
-	Package->LinkSlimport(LocalImportIndex);
-	return EAsyncPackageState::Complete;
-}
-
-EAsyncPackageState::Type FAsyncPackage2::Event_ImportSerialized(FAsyncPackage2* Package, int32 LocalImportIndex)
-{
-	TRACE_CPUPROFILER_EVENT_SCOPE(Event_ImportSerialized);
-
-	UObject* Object = Package->ImportStore.GetObject(LocalImportIndex);
-	check(Object);
-	checkf(!Object->HasAnyFlags(RF_NeedLoad), TEXT("%s had RF_NeedLoad yet it was marked as serialized."), *Object->GetFullName()); // otherwise is isn't done serializing, is it?
-	return EAsyncPackageState::Complete;
 }
 
 EAsyncPackageState::Type FAsyncPackage2::Event_CreateExport(FAsyncPackage2* Package, int32 LocalExportIndex)
@@ -2577,60 +2558,6 @@ EAsyncPackageState::Type FAsyncPackage2::Event_SerializeExport(FAsyncPackage2* P
 	return EAsyncPackageState::Complete;
 }
 
-void FAsyncPackage2::LinkSlimport(int32 LocalImportIndex, int32 GlobalImportIndex)
-{
-	if (GlobalImportIndex == -1)
-	{
-		check(LocalImportIndex >= 0 && LocalImportIndex < ImportStore.Count);
-		GlobalImportIndex = ImportStore.ImportMap[LocalImportIndex];
-	}
-
-	UObject*& Object = ImportStore.GlobalImportObjects[GlobalImportIndex];
-
-	if (Object)
-	{
-		return;
-	}
-
-	const FName ObjectName = ImportStore.GlobalImportNames[GlobalImportIndex];
-	const FPackageIndex PackageIndex = ImportStore.GlobalImportPackages[GlobalImportIndex];
-
-	UObject*& ImportPackage = ImportStore.GlobalImportObjects[PackageIndex.ToImport()];
-	// UPackage* ImportPackage = (UPackage*)FindExistingSlimport(PackageIndex.ToImport());
-
-	if (!ImportPackage)
-	{
-		const FName PackageName = ImportStore.GlobalImportNames[PackageIndex.ToImport()];
-		ImportPackage = FindObjectFast<UPackage>(NULL, PackageName, false, false);
-		check(ImportPackage);
-	}
-	if (PackageIndex.ToImport() != GlobalImportIndex)
-	{
-		UObject* Outer = ImportPackage;
-		FPackageIndex OuterIndex = ImportStore.GlobalImportOuters[GlobalImportIndex];
-		if (PackageIndex != OuterIndex)
-		{
-			if (!ImportStore.GlobalImportObjects[OuterIndex.ToImport()])
-			{
-				LinkSlimport(/*dummy_local_index*/-1, OuterIndex.ToImport());
-			}
-			Outer = ImportStore.GlobalImportObjects[OuterIndex.ToImport()];
-		}
-		check(Outer);
-
-		Object = StaticFindObjectFast(UObject::StaticClass(), Outer, ObjectName, false, true);
-		/*if (!Object)
-		{
-			TRACE_CPUPROFILER_EVENT_SCOPE(LinkImport_SpinWait);
-			while (!Object)
-			{
-				Object = StaticFindObjectFast(UObject::StaticClass(), Outer, ObjectName, false, true);
-			}
-		}*/
-	}
-	//check(Object);
-}
-
 UObject* FAsyncPackage2::EventDrivenIndexToObject(FPackageIndex Index, bool bCheckSerialized)
 {
 	UObject* Result = nullptr;
@@ -2644,13 +2571,12 @@ UObject* FAsyncPackage2::EventDrivenIndexToObject(FPackageIndex Index, bool bChe
 	}
 	else if (Index.IsImport())
 	{
-		int32 GlobalImportIndex = ImportStore.ImportMap[Index.ToImport()];
-		Result = FindExistingSlimport(GlobalImportIndex);
+		Result = ImportStore.FindExistingSlimport(Index.ToImport());
 		check(Result);
 	}
 	if (bCheckSerialized && !IsFullyLoadedObj(Result))
 	{
-		FEventLoadNode2* MyDependentNode = GetNode(EEventLoadNode2::ImportOrExport_Serialize, Index);
+		FEventLoadNode2* MyDependentNode = GetNode(EEventLoadNode2::Export_Serialize, Index);
 		if (!Result)
 		{
 			UE_LOG(LogStreaming, Error, TEXT("Missing Dependency, request for %s but it hasn't been created yet."), *Linker->GetPathName(Index));
@@ -2692,6 +2618,7 @@ void FAsyncPackage2::EventDrivenCreateExport(int32 LocalExportIndex)
 	// LLM_SCOPED_TAG_WITH_OBJECT_IN_SET((Export.DynamicType == FObjectExport::EDynamicType::DynamicType) ? UDynamicClass::StaticClass() : 
 	// 	CastEventDrivenIndexToObject<UClass>(Export.ClassIndex, false), ELLMTagSet::AssetClasses);
 
+	bool bIsCompleteyLoaded = false;
 
 	// Check whether we already loaded the object and if not whether the context flags allow loading it.
 	//check(!Export.Object || Export.Object->HasAnyFlags(RF_ClassDefaultObject)); // we should not have this yet, unless it is a CDO
@@ -2802,7 +2729,8 @@ void FAsyncPackage2::EventDrivenCreateExport(int32 LocalExportIndex)
 					// If this object was allocated but never loaded (components created by a constructor, CDOs, etc) make sure it gets loaded
 					// Do this for all subobjects created in the native constructor.
 					const EObjectFlags ObjectFlags = Object->GetFlags();
-					if (!(ObjectFlags & RF_LoadCompleted))
+					bIsCompleteyLoaded = !!(ObjectFlags & RF_LoadCompleted);
+					if (!bIsCompleteyLoaded)
 					{
 						UE_LOG(LogStreaming, VeryVerbose, TEXT("Note2: %s was constructed during load and is an export and so needs loading."), *Object->GetFullName());
 						ensure(!(ObjectFlags & (RF_NeedLoad | RF_WasLoaded))); // If export exist but is not completed, it is expected to have been created from a native constructor and not from EventDrivenCreateExport, but who knows...?
@@ -2939,6 +2867,28 @@ void FAsyncPackage2::EventDrivenCreateExport(int32 LocalExportIndex)
 		}
 	}
 	check(Object);
+	if (Export.GlobalImportIndex != -1)
+	{
+		UObject*& StoreObject = ImportStore.GlobalImportObjects[Export.GlobalImportIndex];
+		if (!StoreObject)
+		{
+			StoreObject = Object;
+		}
+		else
+		{
+			check(StoreObject == Object);
+		}
+	}
+
+	if (bIsCompleteyLoaded) 
+	{
+		GetNode(EEventLoadNode2::Export_Serialize, FPackageIndex::FromExport(LocalExportIndex))->ReleaseBarrier();
+	}
+	else
+	{
+		TRACE_CPUPROFILER_EVENT_SCOPE(StartIO);
+		AsyncLoadingThread.EnqueueIoRequest(this, CreateChunkId(PackageChunkId, LocalExportIndex, EChunkType::ExportData));
+	}
 }
 
 void FAsyncPackage2::EventDrivenSerializeExport(int32 LocalExportIndex)
@@ -3060,24 +3010,6 @@ void FAsyncPackage2::EventDrivenSerializeExport(int32 LocalExportIndex)
 	LLM_PUSH_STATS_FOR_ASSET_TAGS();
 }
 
-EAsyncPackageState::Type FAsyncPackage2::Event_StartIO(FAsyncPackage2* Package, int32 LocalExportIndex)
-{
-	TRACE_CPUPROFILER_EVENT_SCOPE(Event_StartIO);
-
-	check(Package->AsyncPackageLoadingState == EAsyncPackageLoadingState2::ProcessNewImportsAndExports);
-	UObject* Object = Package->Exports[LocalExportIndex];
-	if (Package->Linker->bDynamicClassLinker || //native blueprint, there is no IO for these
-		(Object && !Object->HasAnyFlags(RF_NeedLoad)))
-	{
-		Package->GetNode(EEventLoadNode2::ImportOrExport_Serialize, FPackageIndex::FromExport(LocalExportIndex))->ReleaseBarrier();
-	}
-	else
-	{
-		Package->AsyncLoadingThread.EnqueueIoRequest(Package, CreateChunkId(Package->PackageChunkId, LocalExportIndex, EChunkType::ExportData));
-	}
-	return EAsyncPackageState::Complete;
-}
-
 EAsyncPackageState::Type FAsyncPackage2::Event_ExportsDone(FAsyncPackage2* Package, int32)
 {
 	TRACE_CPUPROFILER_EVENT_SCOPE(Event_ExportsDone);
@@ -3088,17 +3020,6 @@ EAsyncPackageState::Type FAsyncPackage2::Event_ExportsDone(FAsyncPackage2* Packa
 	return EAsyncPackageState::Complete;
 }
 
-EAsyncPackageState::Type FAsyncPackage2::Event_StartPostload(FAsyncPackage2* Package, int32)
-{
-	TRACE_CPUPROFILER_EVENT_SCOPE(Event_StartPostLoad);
-
-	LLM_SCOPE(ELLMTag::AsyncLoading);
-
-	check(Package->bAllExportsSerialized);
-	Package->EventDrivenLoadingComplete();
-	return EAsyncPackageState::Complete;
-}
-
 EAsyncPackageState::Type FAsyncPackage2::Event_Tick(FAsyncPackage2* Package, int32)
 {
 	TRACE_CPUPROFILER_EVENT_SCOPE(Event_Tick);
@@ -3196,7 +3117,7 @@ void FAsyncPackage2::ProcessIoRequest(const FIoChunkId& ChunkId, TIoStatusOr<FIo
 	{
 		PackageSummarySize = IoBuffer.DataSize();
 		PackageSummaryIoBuffer = IoBuffer;
-		GetNode(EEventLoadNode2::Package_LoadSummary)->ReleaseBarrier();
+		GetNode(EEventLoadNode2::Package_ProcessSummary)->ReleaseBarrier();
 		break;
 	}
 	case EChunkType::ExportData:
@@ -3204,18 +3125,12 @@ void FAsyncPackage2::ProcessIoRequest(const FIoChunkId& ChunkId, TIoStatusOr<FIo
 		const int32 LocalExportIndex = GetChunkIndex(ChunkId);
 		check(ExportIoBuffers.Num() > LocalExportIndex);
 		ExportIoBuffers[LocalExportIndex] = IoBuffer;
-		GetNode(EEventLoadNode2::ImportOrExport_Serialize, FPackageIndex::FromExport(LocalExportIndex))->ReleaseBarrier();
+		GetNode(EEventLoadNode2::Export_Serialize, FPackageIndex::FromExport(LocalExportIndex))->ReleaseBarrier();
 		break;
 	}
 	}
 }
 
-void FAsyncPackage2::EventDrivenLoadingComplete()
-{
-	AsyncPackageLoadingState = EAsyncPackageLoadingState2::PostLoad_Etc;
-	GetNode(Package_Tick)->ReleaseBarrier();
-}
-
 FEventLoadNode2* FAsyncPackage2::GetNode(EEventLoadNode2 Phase, FPackageIndex ImportOrExportIndex)
 {
 	if (ImportOrExportIndex.IsNull())
@@ -3224,8 +3139,8 @@ FEventLoadNode2* FAsyncPackage2::GetNode(EEventLoadNode2 Phase, FPackageIndex Im
 	}
 	else if (ImportOrExportIndex.IsImport())
 	{
-		int32 ImportNodeIndex = ImportOrExportIndex.ToImport() * EEventLoadNode2::Import_NumPhases + Phase;
-		return ImportNodes + ImportNodeIndex;
+		check(false);
+		return 0;
 	}
 	else
 	{
@@ -3234,6 +3149,11 @@ FEventLoadNode2* FAsyncPackage2::GetNode(EEventLoadNode2 Phase, FPackageIndex Im
 	}
 }
 
+FEventLoadNode2* FAsyncPackage2::GetNode(int32 NodeIndex)
+{ 
+	return &PackageNodes[NodeIndex];
+}
+
 void FAsyncLoadingThread2Impl::AddToLoadedPackages(FAsyncPackage2* Package)
 {
 	FScopeLock LoadedLock(&LoadedPackagesCritical);
@@ -3657,23 +3577,16 @@ FAsyncLoadingThread2Impl::FAsyncLoadingThread2Impl(IEDLBootNotificationManager&
 	AltEventQueues.Add(&EventQueue);
 	AltEventQueues.Add(&CreateExportsEventQueue);
 
-	EventSpecs.AddDefaulted(EEventLoadNode2::Package_NumPhases + EEventLoadNode2::Import_NumPhases + EEventLoadNode2::Export_NumPhases);
+	EventSpecs.AddDefaulted(EEventLoadNode2::Package_NumPhases + EEventLoadNode2::Export_NumPhases);
 	EventSpecs[EEventLoadNode2::Package_CreateLinker] = { &FAsyncPackage2::Event_CreateLinker, &AsyncEventQueue, false };
-	EventSpecs[EEventLoadNode2::Package_LoadSummary] = { &FAsyncPackage2::Event_FinishLinker, &AsyncEventQueue, false };
-	EventSpecs[EEventLoadNode2::Package_ImportPackages] = { &FAsyncPackage2::Event_StartImportPackages, &AsyncEventQueue, false };
-	EventSpecs[EEventLoadNode2::Package_SetupImports] = { &FAsyncPackage2::Event_SetupImports, &AsyncEventQueue, false };
-	EventSpecs[EEventLoadNode2::Package_SetupExports] = { &FAsyncPackage2::Event_SetupExports, &AsyncEventQueue, true };
+	EventSpecs[EEventLoadNode2::Package_ProcessSummary] = { &FAsyncPackage2::Event_ProcessSummary, &AsyncEventQueue, false };
 	EventSpecs[EEventLoadNode2::Package_ExportsSerialized] = { &FAsyncPackage2::Event_ExportsDone, &AsyncEventQueue, true };
-	EventSpecs[EEventLoadNode2::Package_PostLoad] = { &FAsyncPackage2::Event_StartPostload, &AsyncEventQueue, false };
 	EventSpecs[EEventLoadNode2::Package_Tick] = { &FAsyncPackage2::Event_Tick, &EventQueue, false };
 	EventSpecs[EEventLoadNode2::Package_Delete] = { &FAsyncPackage2::Event_Delete, &AsyncEventQueue, false };
 
-	EventSpecs[EEventLoadNode2::Package_NumPhases + EEventLoadNode2::ImportOrExport_Create] = { &FAsyncPackage2::Event_LinkImport, &AsyncEventQueue, false };
-	EventSpecs[EEventLoadNode2::Package_NumPhases + EEventLoadNode2::ImportOrExport_Serialize] = { &FAsyncPackage2::Event_ImportSerialized, &AsyncEventQueue, true };
 
-	EventSpecs[EEventLoadNode2::Package_NumPhases + EEventLoadNode2::Import_NumPhases + EEventLoadNode2::ImportOrExport_Create] = { &FAsyncPackage2::Event_CreateExport, &CreateExportsEventQueue, false };
-	EventSpecs[EEventLoadNode2::Package_NumPhases + EEventLoadNode2::Import_NumPhases + EEventLoadNode2::ImportOrExport_Serialize] = { &FAsyncPackage2::Event_SerializeExport, &SerializeExportsEventQueue, false };
-	EventSpecs[EEventLoadNode2::Package_NumPhases + EEventLoadNode2::Import_NumPhases + EEventLoadNode2::Export_StartIO] = { &FAsyncPackage2::Event_StartIO, &AsyncEventQueue, false };
+	EventSpecs[EEventLoadNode2::Package_NumPhases + EEventLoadNode2::Export_Create] = { &FAsyncPackage2::Event_CreateExport, &CreateExportsEventQueue, false };
+	EventSpecs[EEventLoadNode2::Package_NumPhases + EEventLoadNode2::Export_Serialize] = { &FAsyncPackage2::Event_SerializeExport, &SerializeExportsEventQueue, false };
 
 	CancelLoadingEvent = FPlatformProcess::GetSynchEventFromPool();
 	ThreadSuspendedEvent = FPlatformProcess::GetSynchEventFromPool();
@@ -4160,7 +4073,8 @@ void FAsyncLoadingThread2Impl::NotifyConstructedDuringAsyncLoading(UObject* Obje
 
 void FAsyncLoadingThread2Impl::FireCompletedCompiledInImport(void* AsyncPackage, FPackageIndex Import)
 {
-	static_cast<FAsyncPackage2*>(AsyncPackage)->GetNode(EEventLoadNode2::ImportOrExport_Create, Import)->ReleaseBarrier();
+	int32 ExportNodeIndex = Import.ToImport();
+	static_cast<FAsyncPackage2*>(AsyncPackage)->GetNode(ExportNodeIndex)->ReleaseBarrier();
 }
 
 void FAsyncLoadingThread2Impl::EnqueueIoRequest(FAsyncPackage2* Package, const FIoChunkId& ChunkId)
@@ -4214,13 +4128,17 @@ FAsyncPackage2::FAsyncPackage2(const FAsyncPackageDesc& InDesc, int32 InSerialNu
 		ExportIoBuffers.SetNum(ExportCount);
 		OwnedObjects.Reserve(ExportCount + 1); // ExportCount + UPackage
 	}
+	
+	CreateNodes(EventSpecs);
+}
 
+void FAsyncPackage2::CreateNodes(const FAsyncLoadEventSpec* EventSpecs)
+{
 	{
 		TRACE_CPUPROFILER_EVENT_SCOPE(CreateNodes);
-		ImportNodeCount = ImportStore.Count * EEventLoadNode2::Import_NumPhases;
 		ExportNodeCount = ExportCount * EEventLoadNode2::Export_NumPhases;
 
-		PackageNodes = GraphAllocator.AllocNodes(EEventLoadNode2::Package_NumPhases + ImportNodeCount + ExportNodeCount);
+		PackageNodes = GraphAllocator.AllocNodes(EEventLoadNode2::Package_NumPhases + ExportNodeCount);
 		for (int32 Phase = 0; Phase < EEventLoadNode2::Package_NumPhases; ++Phase)
 		{
 			new (PackageNodes + Phase) FEventLoadNode2(EventSpecs + Phase, this, -1);
@@ -4229,54 +4147,28 @@ FAsyncPackage2::FAsyncPackage2(const FAsyncPackageDesc& InDesc, int32 InSerialNu
 		FEventLoadNode2* CreateLinkerNode = PackageNodes + EEventLoadNode2::Package_CreateLinker;
 		CreateLinkerNode->AddBarrier();
 
-		FEventLoadNode2* LoadSummaryNode = PackageNodes + EEventLoadNode2::Package_LoadSummary;
-		LoadSummaryNode->DependsOn(CreateLinkerNode);
-		LoadSummaryNode->AddBarrier();
+		FEventLoadNode2* ProcessSummaryNode = PackageNodes + EEventLoadNode2::Package_ProcessSummary;
+		ProcessSummaryNode->AddBarrier();
 
-		FEventLoadNode2* ImportPackagesNode = PackageNodes + EEventLoadNode2::Package_ImportPackages;
-		FEventLoadNode2* SetupImportsNode = PackageNodes + EEventLoadNode2::Package_SetupImports;
-		FEventLoadNode2* SetupExportsNode = PackageNodes + EEventLoadNode2::Package_SetupExports;
 		FEventLoadNode2* ExportsSerializedNode = PackageNodes + EEventLoadNode2::Package_ExportsSerialized;
-		FEventLoadNode2* PostLoadNode = PackageNodes + EEventLoadNode2::Package_PostLoad;
 		FEventLoadNode2* TickNode = PackageNodes + EEventLoadNode2::Package_Tick;
-		TickNode->AddBarrier();
 
-		ImportPackagesNode->DependsOn(LoadSummaryNode);
-		SetupImportsNode->DependsOn(ImportPackagesNode);
-		SetupExportsNode->DependsOn(SetupImportsNode);
-		ExportsSerializedNode->DependsOn(SetupExportsNode);
-		PostLoadNode->DependsOn(ExportsSerializedNode);
+		TickNode->DependsOn(ExportsSerializedNode);
 
 		FEventLoadNode2* DeleteNode = PackageNodes + EEventLoadNode2::Package_Delete;
 		DeleteNode->AddBarrier();
 		DeleteNode->DependsOn(TickNode);
 
-		// Add nodes for all imports and exports.
-		ImportNodes = PackageNodes + EEventLoadNode2::Package_NumPhases;
-		for (int32 LocalImportIndex = 0; LocalImportIndex < ImportStore.Count; ++LocalImportIndex)
-		{
-			uint32 NodeIndex = EEventLoadNode2::Import_NumPhases * LocalImportIndex;
-			FEventLoadNode2* CreateImportNode = ImportNodes + NodeIndex + EEventLoadNode2::ImportOrExport_Create;
-			new (CreateImportNode) FEventLoadNode2(EventSpecs + EEventLoadNode2::Package_NumPhases + EEventLoadNode2::ImportOrExport_Create, this, LocalImportIndex);
-			FEventLoadNode2* SerializeImportNode = ImportNodes + NodeIndex + EEventLoadNode2::ImportOrExport_Serialize;
-			new (SerializeImportNode) FEventLoadNode2(EventSpecs + EEventLoadNode2::Package_NumPhases + EEventLoadNode2::ImportOrExport_Serialize, this, LocalImportIndex);
-
-			CreateImportNode->DependsOn(SetupExportsNode); // Need to wait for SetupExports here because of preload dependencies
-			SerializeImportNode->DependsOn(CreateImportNode);
-			ExportsSerializedNode->DependsOn(SerializeImportNode);
-		}
-		ExportNodes = ImportNodes + ImportNodeCount;
+		// Add nodes for exports.
+		ExportNodes = PackageNodes + EEventLoadNode2::Package_NumPhases;
 		for (int32 LocalExportIndex = 0; LocalExportIndex < ExportCount; ++LocalExportIndex)
 		{
 			uint32 NodeIndex = EEventLoadNode2::Export_NumPhases * LocalExportIndex;
-			FEventLoadNode2* CreateExportNode = ExportNodes + NodeIndex + EEventLoadNode2::ImportOrExport_Create;
-			new (CreateExportNode) FEventLoadNode2(EventSpecs + EEventLoadNode2::Package_NumPhases + EEventLoadNode2::Import_NumPhases + EEventLoadNode2::ImportOrExport_Create, this, LocalExportIndex);
-			FEventLoadNode2* SerializeExportNode = ExportNodes + NodeIndex + EEventLoadNode2::ImportOrExport_Serialize;
-			new (SerializeExportNode) FEventLoadNode2(EventSpecs + EEventLoadNode2::Package_NumPhases + EEventLoadNode2::Import_NumPhases + EEventLoadNode2::ImportOrExport_Serialize, this, LocalExportIndex);
-			FEventLoadNode2* StartIONode = ExportNodes + NodeIndex + EEventLoadNode2::Export_StartIO;
-			new (StartIONode) FEventLoadNode2(EventSpecs + EEventLoadNode2::Package_NumPhases + EEventLoadNode2::Import_NumPhases + EEventLoadNode2::Export_StartIO, this, LocalExportIndex);
-			CreateExportNode->DependsOn(SetupExportsNode);
-			StartIONode->DependsOn(CreateExportNode);
+			FEventLoadNode2* CreateExportNode = ExportNodes + NodeIndex + EEventLoadNode2::Export_Create;
+			new (CreateExportNode) FEventLoadNode2(EventSpecs + EEventLoadNode2::Package_NumPhases + EEventLoadNode2::Export_Create, this, LocalExportIndex);
+			FEventLoadNode2* SerializeExportNode = ExportNodes + NodeIndex + EEventLoadNode2::Export_Serialize;
+			new (SerializeExportNode) FEventLoadNode2(EventSpecs + EEventLoadNode2::Package_NumPhases + EEventLoadNode2::Export_Serialize, this, LocalExportIndex);
+			CreateExportNode->DependsOn(ProcessSummaryNode);
 			ExportsSerializedNode->DependsOn(SerializeExportNode);
 			SerializeExportNode->AddBarrier();
 		}
@@ -4289,7 +4181,7 @@ FAsyncPackage2::~FAsyncPackage2()
 
 	check(RefCount == 0);
 
-	FAsyncLoadingThreadState2::Get().DeferredFreeNodes.Add(MakeTuple(PackageNodes, EEventLoadNode2::Package_NumPhases + ImportNodeCount + ExportNodeCount));
+	FAsyncLoadingThreadState2::Get().DeferredFreeNodes.Add(MakeTuple(PackageNodes, EEventLoadNode2::Package_NumPhases + ExportNodeCount));
 
 	TRACE_LOADTIME_DESTROY_ASYNC_PACKAGE(this);
 
@@ -4445,9 +4337,13 @@ EAsyncPackageState::Type FAsyncPackage2::CreateLinker()
 	// Try to find existing package or create it if not already present.
 	UPackage* Package = nullptr;
 	{
-		Package = FindObjectFast<UPackage>(nullptr, Desc.Name);
+		{
+			TRACE_CPUPROFILER_EVENT_SCOPE(UPackageFind);
+			Package = FindObjectFast<UPackage>(nullptr, Desc.Name);
+		}
 		if (!Package)
 		{
+			TRACE_CPUPROFILER_EVENT_SCOPE(UPackageCreate);
 			Package = NewObject<UPackage>(/*Outer*/nullptr, Desc.Name, RF_Public);
 			Package->SetPackageFlags(Desc.PackageFlags);
 			Package->FileName = Desc.NameToLoad;
@@ -4471,7 +4367,10 @@ EAsyncPackageState::Type FAsyncPackage2::CreateLinker()
 	}
 	Package->LinkerLoad = Linker;
 
-	AsyncLoadingThread.EnqueueIoRequest(this, CreateChunkId(PackageChunkId, 0, EChunkType::PackageSummary));
+	{
+		TRACE_CPUPROFILER_EVENT_SCOPE(PackageSummaryStartIO);
+		AsyncLoadingThread.EnqueueIoRequest(this, CreateChunkId(PackageChunkId, 0, EChunkType::PackageSummary));
+	}
 
 	UE_LOG(LogStreaming, Verbose, TEXT("FAsyncPackage::CreateLinker for %s finished."), *Desc.NameToLoad.ToString());
 	return EAsyncPackageState::Complete;
@@ -4484,6 +4383,7 @@ EAsyncPackageState::Type FAsyncPackage2::CreateLinker()
  */
 EAsyncPackageState::Type FAsyncPackage2::FinishLinker()
 {
+	TRACE_CPUPROFILER_EVENT_SCOPE(FinishLinker);
 	LLM_SCOPE(ELLMTag::AsyncLoading);
 
 	const FPackageSummary* PackageSummary = reinterpret_cast<const FPackageSummary*>(PackageSummaryIoBuffer.Data());
@@ -4491,18 +4391,11 @@ EAsyncPackageState::Type FAsyncPackage2::FinishLinker()
 
 	TRACE_LOADTIME_PACKAGE_SUMMARY(this, PackageSummarySize, ImportStore.Count, ExportCount);
 
-	{
-		SCOPED_LOADTIMER(LinkerLoad_SerializePackageFileSummary);
-
-		FPackageFileSummary& Summary	= Linker->Summary;
-		Summary.PackageFlags			= PackageSummary->PackageFlags;
-		Summary.BulkDataStartOffset		= PackageSummary->BulkDataStartOffset;
-		Summary.SetFileVersions(GPackageFileUE4Version, GPackageFileLicenseeUE4Version, /*unversioned*/true);
-		Linker->UpdateFromPackageFileSummary();
-	}
-
-	// Add this linker to the object manager's linker array.
-	// FLinkerManager::Get().AddLoader(Linker);
+	FPackageFileSummary& Summary	= Linker->Summary;
+	Summary.PackageFlags			= PackageSummary->PackageFlags;
+	Summary.BulkDataStartOffset		= PackageSummary->BulkDataStartOffset;
+	Summary.SetFileVersions(GPackageFileUE4Version, GPackageFileLicenseeUE4Version, /*unversioned*/true);
+	Linker->UpdateFromPackageFileSummary();
 
 	return EAsyncPackageState::Complete;
 }
