Optimized GetArchetypeFromRequiredInfoImpl by reusing cached archetypes.

+ fixed bad archetype caching when the Outer is still pending load

diff --git a/Engine/Source/Runtime/CoreUObject/Private/UObject/UObjectArchetype.cpp b/Engine/Source/Runtime/CoreUObject/Private/UObject/UObjectArchetype.cpp
--- a/Engine/Source/Runtime/CoreUObject/Private/UObject/UObjectArchetype.cpp
+++ b/Engine/Source/Runtime/CoreUObject/Private/UObject/UObjectArchetype.cpp
@@ -12,6 +12,45 @@
 #include "UObject/UObjectAnnotation.h"
 #include "Stats/StatsMisc.h"
 
+#define UE_CACHE_ARCHETYPE (1 && !WITH_EDITORONLY_DATA)
+#define UE_VERIFY_CACHED_ARCHETYPE 0
+
+#if UE_CACHE_ARCHETYPE
+struct FArchetypeInfo
+{
+	/**
+	* default contructor
+	* Default constructor must be the default item
+	*/
+	FArchetypeInfo()
+		: ArchetypeIndex(INDEX_NONE)
+	{
+	}
+	/**
+	* Determine if this linker pair is the default
+	* @return true is this is a default pair. We only check the linker because CheckInvariants rules out bogus combinations
+	*/
+	FORCEINLINE bool IsDefault() const
+	{
+		return ArchetypeIndex == INDEX_NONE;
+	}
+
+	/**
+	* Constructor
+	* @param InArchetype Archetype to assign
+	*/
+	FArchetypeInfo(int32 InArchetypeIndex)
+		: ArchetypeIndex(InArchetypeIndex)
+	{
+	}
+
+	int32 ArchetypeIndex;
+};
+
+static FUObjectAnnotationDense<FArchetypeInfo, true> ArchetypeAnnotation;
+
+#endif // UE_CACHE_ARCHETYPE
+
 UObject* GetArchetypeFromRequiredInfoImpl(const UClass* Class, const UObject* Outer, FName Name, EObjectFlags ObjectFlags, bool bUseUpToDateClass)
 {
 	UObject* Result = NULL;
@@ -29,7 +68,18 @@ UObject* GetArchetypeFromRequiredInfoImpl(const UClass* Class, const UObject* Ou
 			void LockUObjectHashTables();
 			LockUObjectHashTables();
 
-			UObject* ArchetypeToSearch = GetArchetypeFromRequiredInfoImpl(Outer->GetClass(), Outer->GetOuter(), Outer->GetFName(), Outer->GetFlags(), bUseUpToDateClass);
+			UObject* ArchetypeToSearch = nullptr;
+#if UE_CACHE_ARCHETYPE
+			ArchetypeToSearch = Outer->GetArchetype();
+#if UE_VERIFY_CACHED_ARCHETYPE
+			{
+				UObject* VerifyArchetype = GetArchetypeFromRequiredInfoImpl(Outer->GetClass(), Outer->GetOuter(), Outer->GetFName(), Outer->GetFlags(), bUseUpToDateClass);
+				checkf(ArchetypeToSearch == VerifyArchetype, TEXT("Cached archetype mismatch, expected: %s, cached: %s"), *GetFullNameSafe(VerifyArchetype), *GetFullNameSafe(ArchetypeToSearch));
+			}
+#endif // UE_VERIFY_CACHED_ARCHETYPE
+#else
+			ArchetypeToSearch = GetArchetypeFromRequiredInfoImpl(Outer->GetClass(), Outer->GetOuter(), Outer->GetFName(), Outer->GetFlags(), bUseUpToDateClass);
+#endif // UE_CACHE_ARCHETYPE
 			UObject* MyArchetype = static_cast<UObject*>(FindObjectWithOuter(ArchetypeToSearch, Class, Name));
 			if (MyArchetype)
 			{
@@ -103,45 +153,6 @@ UObject* UObject::GetArchetypeFromRequiredInfo(const UClass* Class, const UObjec
 	return GetArchetypeFromRequiredInfoImpl(Class, Outer, Name, ObjectFlags, bUseUpToDateClass);
 }
 
-#define UE_CACHE_ARCHETYPE (1 && !WITH_EDITORONLY_DATA)
-#define UE_VERIFY_CACHED_ARCHETYPE 0
-
-#if UE_CACHE_ARCHETYPE
-struct FArchetypeInfo
-{
-	/**
-	* default contructor
-	* Default constructor must be the default item
-	*/
-	FArchetypeInfo() 
-		: ArchetypeIndex(INDEX_NONE)
-	{
-	}
-	/**
-	* Determine if this linker pair is the default
-	* @return true is this is a default pair. We only check the linker because CheckInvariants rules out bogus combinations
-	*/
-	FORCEINLINE bool IsDefault() const
-	{
-		return ArchetypeIndex == INDEX_NONE;
-	}
-
-	/**
-	* Constructor
-	* @param InArchetype Archetype to assign
-	*/
-	FArchetypeInfo(int32 InArchetypeIndex) 
-		: ArchetypeIndex(InArchetypeIndex)
-	{
-	}
-
-	int32 ArchetypeIndex;
-};
-
-static FUObjectAnnotationDense<FArchetypeInfo, true> ArchetypeAnnotation;
-
-#endif // UE_CACHE_ARCHETYPE
-
 //DECLARE_FLOAT_ACCUMULATOR_STAT(TEXT("UObject::GetArchetype"), STAT_FArchiveRealtimeGC_GetArchetype, STATGROUP_GC);
 
 UObject* UObject::GetArchetype() const
@@ -154,22 +165,10 @@ UObject* UObject::GetArchetype() const
 	if (ArchetypeIndex == INDEX_NONE)
 	{
 		Archetype = GetArchetypeFromRequiredInfo(GetClass(), GetOuter(), GetFName(), GetFlags());
-		if (Archetype)
+		// If the Outer is pending load we can't cache the archetype as it may be inacurate
+		if (Archetype && !(GetOuter() && GetOuter()->HasAnyFlags(RF_NeedLoad)))
 		{
-			bool bAddToCache = true;
-			if (GetOuter() &&
-				GetOuter()->GetClass() != UPackage::StaticClass() &&
-				HasAnyFlags(RF_InheritableComponentTemplate) &&
-				GetOuter()->IsA<UClass>() &&
-				GetOuter()->HasAnyFlags(RF_NeedLoad))
-			{
-				bAddToCache = false;
-			}
-
-			if (bAddToCache)
-			{
-				ArchetypeAnnotation.AddAnnotation(this, GUObjectArray.ObjectToIndex(Archetype));
-			}
+			ArchetypeAnnotation.AddAnnotation(this, GUObjectArray.ObjectToIndex(Archetype));
 		}
 	}		
 	else
@@ -178,10 +177,10 @@ UObject* UObject::GetArchetype() const
 		check(ArchetypeItem != nullptr);
 		Archetype = static_cast<UObject*>(ArchetypeItem->Object);
 #if UE_VERIFY_CACHED_ARCHETYPE
-		UObject* CurrentArchetype = GetArchetypeFromRequiredInfo(GetClass(), GetOuter(), GetFName(), GetFlags());
-		if (CurrentArchetype != Archetype)
+		UObject* ExpectedArchetype = GetArchetypeFromRequiredInfo(GetClass(), GetOuter(), GetFName(), GetFlags());
+		if (ExpectedArchetype != Archetype)
 		{
-			UE_LOG(LogClass, Fatal, TEXT("Cached archetype mismatch: %s vs current: %s"), *Archetype->GetFullName(), *CurrentArchetype->GetFullName());
+			UE_LOG(LogClass, Fatal, TEXT("Cached archetype mismatch, expected: %s, cached: %s"), *ExpectedArchetype->GetFullName(), *Archetype->GetFullName());
 		}
 #endif // UE_VERIFY_CACHED_ARCHETYPE
 	}
