Initialize the fallback token/bulkdata size to 0 by default to avoid BulkData that has not being serialized yet from thinking it might have valid data.

#rnx

diff --git a/Engine/Source/Runtime/CoreUObject/Public/Serialization/BulkData2.h b/Engine/Source/Runtime/CoreUObject/Public/Serialization/BulkData2.h
--- a/Engine/Source/Runtime/CoreUObject/Public/Serialization/BulkData2.h
+++ b/Engine/Source/Runtime/CoreUObject/Public/Serialization/BulkData2.h
@@ -55,7 +55,11 @@ public:
 	FBulkDataBase(FBulkDataBase&& Other);
 	FBulkDataBase& operator=(const FBulkDataBase& Other);
 
-	FBulkDataBase() = default;
+	FBulkDataBase()
+	{ 
+		Fallback.BulkDataSize = 0;
+		Fallback.Token = InvalidToken;
+	}
 	~FBulkDataBase();
 
 protected:
