Make a method const.

#rnx

diff --git a/Engine/Source/Runtime/CoreUObject/Private/Serialization/BulkData2.cpp b/Engine/Source/Runtime/CoreUObject/Private/Serialization/BulkData2.cpp
--- a/Engine/Source/Runtime/CoreUObject/Private/Serialization/BulkData2.cpp
+++ b/Engine/Source/Runtime/CoreUObject/Private/Serialization/BulkData2.cpp
@@ -1011,7 +1011,7 @@ void FBulkDataBase::FreeData()
 	DataBuffer = nullptr;
 }
 
-FString FBulkDataBase::ConvertFilenameFromFlags(const FString& Filename)
+FString FBulkDataBase::ConvertFilenameFromFlags(const FString& Filename) const
 {
 	if (IsInlined())
 	{
diff --git a/Engine/Source/Runtime/CoreUObject/Public/Serialization/BulkData2.h b/Engine/Source/Runtime/CoreUObject/Public/Serialization/BulkData2.h
--- a/Engine/Source/Runtime/CoreUObject/Public/Serialization/BulkData2.h
+++ b/Engine/Source/Runtime/CoreUObject/Public/Serialization/BulkData2.h
@@ -127,7 +127,7 @@ private:
 	void AllocateData(SIZE_T SizeInBytes);
 	void FreeData();
 
-	FString ConvertFilenameFromFlags(const FString& Filename);
+	FString ConvertFilenameFromFlags(const FString& Filename) const;
 
 private:
 	friend class FBulkDataIoDispatcherRequest;
