Added 'obj overhead [-detailed]' console command to print out the total memory overhead of UObject hash tables and maps

diff --git a/Engine/Source/Runtime/CoreUObject/Private/UObject/UObjectHash.cpp b/Engine/Source/Runtime/CoreUObject/Private/UObject/UObjectHash.cpp
--- a/Engine/Source/Runtime/CoreUObject/Private/UObject/UObjectHash.cpp
+++ b/Engine/Source/Runtime/CoreUObject/Private/UObject/UObjectHash.cpp
@@ -1267,3 +1267,77 @@ void LogHashOuterStatistics(FOutputDevice& Ar, const bool bShowHashBucketCollisi
 	Ar.Logf(TEXT(""));
 }
 
+void LogHashMemoryOverheadStatistics(FOutputDevice& Ar, const bool bShowIndividualStats)
+{
+	Ar.Logf(TEXT("UObject Hash Tables and Maps memory overhead"));
+	Ar.Logf(TEXT("-------------------------------------------------"));
+
+	FUObjectHashTables& HashTables = FUObjectHashTables::Get();
+	FHashTableLock HashLock(HashTables);
+
+	int64 TotalSize = 0;
+	
+	{
+		int64 Size = HashTables.Hash.GetAllocatedSize();
+		for (const TPair<int32, FHashBucket>& Pair : HashTables.Hash)
+		{
+			Size += Pair.Value.GetItemsSize();
+		}
+		if (bShowIndividualStats)
+		{
+			Ar.Logf(TEXT("Memory used for UObject Hash: %lld bytes."), Size);
+		}
+		TotalSize += Size;
+	}
+
+	{
+		int64 Size = HashTables.HashOuter.GetAllocatedSize();
+		if (bShowIndividualStats)
+		{
+			Ar.Logf(TEXT("Memory used for UObject Outer Hash: %lld bytes."), Size);
+		}
+		TotalSize += Size;
+	}
+
+	{
+		int64 Size = HashTables.ObjectOuterMap.GetAllocatedSize();
+		for (const TPair<UObjectBase*, FHashBucket>& Pair : HashTables.ObjectOuterMap)
+		{
+			Size += Pair.Value.GetItemsSize();
+		}
+		if (bShowIndividualStats)
+		{
+			Ar.Logf(TEXT("Memory used for UObject Outer Map: %lld bytes."), Size);
+		}
+		TotalSize += Size;
+	}
+
+	{
+		int64 Size = HashTables.ClassToObjectListMap.GetAllocatedSize();
+		for (const TPair<UClass*, FHashBucket>& Pair : HashTables.ClassToObjectListMap)
+		{
+			Size += Pair.Value.GetItemsSize();
+		}
+		if (bShowIndividualStats)
+		{
+			Ar.Logf(TEXT("Memory used for UClass To UObject List Map: %lld bytes."), Size);
+		}
+		TotalSize += Size;
+	}
+
+	{
+		int64 Size = HashTables.ClassToChildListMap.GetAllocatedSize();
+		for (const TPair<UClass*, TSet<UClass*> >& Pair : HashTables.ClassToChildListMap)
+		{
+			Size += Pair.Value.GetAllocatedSize();
+		}
+		if (bShowIndividualStats)
+		{
+			Ar.Logf(TEXT("Memory used for UClass To Child UClass List Map: %lld bytes."), Size);
+		}
+		TotalSize += Size;
+	}
+
+	Ar.Logf(TEXT("Total memory allocated for Object hash tables and maps: %lld bytes (%.2f MB)."), TotalSize, (double)TotalSize / 1024.0 / 1024.0);
+	Ar.Logf(TEXT(""));
+}
diff --git a/Engine/Source/Runtime/CoreUObject/Public/UObject/UObjectHash.h b/Engine/Source/Runtime/CoreUObject/Public/UObject/UObjectHash.h
--- a/Engine/Source/Runtime/CoreUObject/Public/UObject/UObjectHash.h
+++ b/Engine/Source/Runtime/CoreUObject/Public/UObject/UObjectHash.h
@@ -146,6 +146,14 @@ COREUOBJECT_API void LogHashStatistics(FOutputDevice& Ar, const bool bShowHashBu
  */
 COREUOBJECT_API void LogHashOuterStatistics(FOutputDevice& Ar, const bool bShowHashBucketCollisionInfo);
 
+/**
+ * Logs out information about the total object hash memory usage for debug purposes
+ *
+ * @param Ar the archive to write the log data to
+ * @param bShowIndividualStats whether to log each hash/map memory usage separately
+ */
+COREUOBJECT_API void LogHashMemoryOverheadStatistics(FOutputDevice& Ar, const bool bShowIndividualStats);
+
 /**
  * Adds a uobject to the global array which is used for uobject iteration
  *
diff --git a/Engine/Source/Runtime/Engine/Private/UnrealEngine.cpp b/Engine/Source/Runtime/Engine/Private/UnrealEngine.cpp
--- a/Engine/Source/Runtime/Engine/Private/UnrealEngine.cpp
+++ b/Engine/Source/Runtime/Engine/Private/UnrealEngine.cpp
@@ -7838,6 +7838,12 @@ bool UEngine::HandleObjCommand( const TCHAR* Cmd, FOutputDevice& Ar )
 		LogHashOuterStatistics(Ar, bShowHashBucketCollisionInfo);
 		return true;
 	}
+	else if (FParse::Command(&Cmd, TEXT("OVERHEAD")))
+	{
+		const bool bShowIndividualStats = FParse::Param(Cmd, TEXT("DETAILED"));
+		LogHashMemoryOverheadStatistics(Ar, bShowIndividualStats);
+		return true;
+	}
 #endif
 	else
 	{
