Included GUObjectArray memory usage in 'obj overhead' command report

diff --git a/Engine/Source/Runtime/CoreUObject/Private/UObject/UObjectHash.cpp b/Engine/Source/Runtime/CoreUObject/Private/UObject/UObjectHash.cpp
--- a/Engine/Source/Runtime/CoreUObject/Private/UObject/UObjectHash.cpp
+++ b/Engine/Source/Runtime/CoreUObject/Private/UObject/UObjectHash.cpp
@@ -1285,7 +1285,7 @@ void LogHashMemoryOverheadStatistics(FOutputDevice& Ar, const bool bShowIndividu
 		}
 		if (bShowIndividualStats)
 		{
-			Ar.Logf(TEXT("Memory used for UObject Hash: %lld bytes."), Size);
+			Ar.Logf(TEXT("Memory used by UObject Hash: %lld bytes."), Size);
 		}
 		TotalSize += Size;
 	}
@@ -1294,7 +1294,7 @@ void LogHashMemoryOverheadStatistics(FOutputDevice& Ar, const bool bShowIndividu
 		int64 Size = HashTables.HashOuter.GetAllocatedSize();
 		if (bShowIndividualStats)
 		{
-			Ar.Logf(TEXT("Memory used for UObject Outer Hash: %lld bytes."), Size);
+			Ar.Logf(TEXT("Memory used by UObject Outer Hash: %lld bytes."), Size);
 		}
 		TotalSize += Size;
 	}
@@ -1307,7 +1307,7 @@ void LogHashMemoryOverheadStatistics(FOutputDevice& Ar, const bool bShowIndividu
 		}
 		if (bShowIndividualStats)
 		{
-			Ar.Logf(TEXT("Memory used for UObject Outer Map: %lld bytes."), Size);
+			Ar.Logf(TEXT("Memory used by UObject Outer Map: %lld bytes."), Size);
 		}
 		TotalSize += Size;
 	}
@@ -1320,7 +1320,7 @@ void LogHashMemoryOverheadStatistics(FOutputDevice& Ar, const bool bShowIndividu
 		}
 		if (bShowIndividualStats)
 		{
-			Ar.Logf(TEXT("Memory used for UClass To UObject List Map: %lld bytes."), Size);
+			Ar.Logf(TEXT("Memory used by UClass To UObject List Map: %lld bytes."), Size);
 		}
 		TotalSize += Size;
 	}
@@ -1333,11 +1333,20 @@ void LogHashMemoryOverheadStatistics(FOutputDevice& Ar, const bool bShowIndividu
 		}
 		if (bShowIndividualStats)
 		{
-			Ar.Logf(TEXT("Memory used for UClass To Child UClass List Map: %lld bytes."), Size);
+			Ar.Logf(TEXT("Memory used by UClass To Child UClass List Map: %lld bytes."), Size);
 		}
 		TotalSize += Size;
 	}
 
-	Ar.Logf(TEXT("Total memory allocated for Object hash tables and maps: %lld bytes (%.2f MB)."), TotalSize, (double)TotalSize / 1024.0 / 1024.0);
+    {
+        int64 Size = GUObjectArray.GetAllocatedSize();
+        if (bShowIndividualStats)
+        {
+            Ar.Logf(TEXT("Memory used by UObjectArray: %lld bytes."), Size);
+        }
+        TotalSize += Size;
+    }
+    
+	Ar.Logf(TEXT("Total memory allocated by Object hash tables and maps: %lld bytes (%.2f MB)."), TotalSize, (double)TotalSize / 1024.0 / 1024.0);
 	Ar.Logf(TEXT(""));
 }
diff --git a/Engine/Source/Runtime/CoreUObject/Public/UObject/UObjectArray.h b/Engine/Source/Runtime/CoreUObject/Public/UObject/UObjectArray.h
--- a/Engine/Source/Runtime/CoreUObject/Public/UObject/UObjectArray.h
+++ b/Engine/Source/Runtime/CoreUObject/Public/UObject/UObjectArray.h
@@ -375,7 +375,7 @@ class FChunkedFixedUObjectArray
 		}
 		check(ChunkIndex < NumChunks && Objects[ChunkIndex]); // should have a valid pointer now
 	}
-
+    
 public:
 
 	/** Constructor : Probably not thread safe **/
@@ -526,6 +526,11 @@ public:
 	{
 		return nullptr;
 	}
+    
+    int64 GetAllocatedSize() const
+    {
+        return MaxChunks * sizeof(FUObjectItem*) + NumChunks * NumElementsPerChunk * sizeof(FUObjectItem);
+    }
 };
 
 /***
@@ -1028,6 +1033,11 @@ public:
 	{
 		return ObjObjects;
 	}
+    
+    int64 GetAllocatedSize() const
+    {
+        return ObjObjects.GetAllocatedSize();
+    }
 };
 
 /** UObject cluster. Groups UObjects into a single unit for GC. */
