Fix a typo, we should be checking against the flags themselves not the flags seek position in the archive.

#rnx

diff --git a/Engine/Source/Runtime/CoreUObject/Private/Serialization/BulkData.cpp b/Engine/Source/Runtime/CoreUObject/Private/Serialization/BulkData.cpp
--- a/Engine/Source/Runtime/CoreUObject/Private/Serialization/BulkData.cpp
+++ b/Engine/Source/Runtime/CoreUObject/Private/Serialization/BulkData.cpp
@@ -1333,7 +1333,7 @@ void FUntypedBulkData::Serialize( FArchive& Ar, UObject* Owner, int32 Idx, bool
 							FIoChunkId InvalidChunkId = FIoChunkId::InvalidChunkId;
 							Ar << InvalidChunkId;
 
-							check((SavedDupeBulkDataFlagsPos& BULKDATA_CookedForIoDispatcher) != 0); // Validate that the BULKDATA_CookedForIoDispatcher flag was already set
+							check((SavedDupeBulkDataFlags& BULKDATA_CookedForIoDispatcher) != 0); // Validate that the BULKDATA_CookedForIoDispatcher flag was already set
 						}
 					}
 
