If GC Clustering is disabled, Garbage Collector will now run with code responsible for handling clusters compiled out to avoid unnecessary checks.

+ Actor clustering will now be disabled if clustering is disabled in general
+ Added a separate setting for clustering generic assets (gc.AssetClustreringEnabled) to be able to enable actor clustering independently from asset clustering as it was possible in the past

diff --git a/Engine/Source/Runtime/CoreUObject/Private/UObject/GarbageCollection.cpp b/Engine/Source/Runtime/CoreUObject/Private/UObject/GarbageCollection.cpp
--- a/Engine/Source/Runtime/CoreUObject/Private/UObject/GarbageCollection.cpp
+++ b/Engine/Source/Runtime/CoreUObject/Private/UObject/GarbageCollection.cpp
@@ -521,7 +521,7 @@ void ShutdownGarbageCollection()
 /**
 * Handles UObject references found by TFastReferenceCollector
 */
-template <bool bParallel>
+template <bool bParallel, bool bWithClusters>
 class FGCReferenceProcessor
 {
 public:
@@ -571,6 +571,8 @@ public:
 	/** Marks all objects that can't be directly in a cluster but are referenced by it as reachable */
 	static FORCEINLINE bool MarkClusterMutableObjectsAsReachable(FUObjectCluster& Cluster, TArray<UObject*>& ObjectsToSerialize)
 	{
+		check(bWithClusters);
+
 		// This is going to be the return value and basically means that we ran across some pending kill objects
 		bool bAddClusterObjectsToSerialize = false;
 		for (int32& ReferencedMutableObjectIndex : Cluster.MutableObjects)
@@ -663,6 +665,8 @@ public:
 	/** Marks all clusters referenced by another cluster as reachable */
 	static FORCEINLINE void MarkReferencedClustersAsReachable(int32 ClusterIndex, TArray<UObject*>& ObjectsToSerialize)
 	{
+		check(bWithClusters);
+
 		// If we run across some PendingKill objects we need to add all objects from this cluster
 		// to ObjectsToSerialize so that we can properly null out all the references.
 		// It also means this cluster will have to be dissolved because we may no longer guarantee all cross-cluster references are correct.
@@ -763,7 +767,7 @@ public:
 					// Objects that are part of a GC cluster should never have the unreachable flag set!
 					checkSlow(ObjectItem->GetOwnerIndex() <= 0);
 
-					if (!ObjectItem->HasAnyFlags(EInternalObjectFlags::ClusterRoot))
+					if (!bWithClusters || !ObjectItem->HasAnyFlags(EInternalObjectFlags::ClusterRoot))
 					{
 						// Add it to the list of objects to serialize.
 						ObjectsToSerialize.Add(Object);
@@ -797,7 +801,7 @@ public:
 				// Objects that are part of a GC cluster should never have the unreachable flag set!
 				checkSlow(ObjectItem->GetOwnerIndex() <= 0);
 
-				if (!ObjectItem->HasAnyFlags(EInternalObjectFlags::ClusterRoot))
+				if (!bWithClusters || !ObjectItem->HasAnyFlags(EInternalObjectFlags::ClusterRoot))
 				{
 					// Add it to the list of objects to serialize.
 					ObjectsToSerialize.Add(Object);
@@ -809,7 +813,7 @@ public:
 				}
 			}
 		}
-		else if (ObjectItem->GetOwnerIndex() > 0 && !ObjectItem->HasAnyFlags(EInternalObjectFlags::ReachableInCluster))
+		else if (bWithClusters && (ObjectItem->GetOwnerIndex() > 0 && !ObjectItem->HasAnyFlags(EInternalObjectFlags::ReachableInCluster)))
 		{
 			bool bNeedsDoing = true;
 			if (bParallel)
@@ -890,20 +894,17 @@ public:
 	}
 };
 
-typedef FGCReferenceProcessor<true> FGCReferenceProcessorMultithreaded;
-typedef FGCReferenceProcessor<false> FGCReferenceProcessorSinglethreaded;
 
-
-template <bool bParallel>
-FGCCollector<bParallel>::FGCCollector(FGCReferenceProcessor<bParallel>& InProcessor, FGCArrayStruct& InObjectArrayStruct)
+template <bool bParallel, bool bWithClusters>
+FGCCollector<bParallel, bWithClusters>::FGCCollector(FGCReferenceProcessor<bParallel, bWithClusters>& InProcessor, FGCArrayStruct& InObjectArrayStruct)
 		: ReferenceProcessor(InProcessor)
 		, ObjectArrayStruct(InObjectArrayStruct)
 		, bAllowEliminatingReferences(true)
 {
 }
 
-template <bool bParallel>
-FORCEINLINE void FGCCollector<bParallel>::InternalHandleObjectReference(UObject*& Object, const UObject* ReferencingObject, const UProperty* ReferencingProperty)
+template <bool bParallel, bool bWithClusters>
+FORCEINLINE void FGCCollector<bParallel, bWithClusters>::InternalHandleObjectReference(UObject*& Object, const UObject* ReferencingObject, const UProperty* ReferencingProperty)
 {
 #if ENABLE_GC_OBJECT_CHECKS
 		if (Object && !Object->IsValidLowLevelFast())
@@ -917,14 +918,14 @@ FORCEINLINE void FGCCollector<bParallel>::InternalHandleObjectReference(UObject*
 		ReferenceProcessor.HandleObjectReference(ObjectArrayStruct.ObjectsToSerialize, const_cast<UObject*>(ReferencingObject), Object, bAllowEliminatingReferences);
 }
 
-template <bool bParallel>
-void FGCCollector<bParallel>::HandleObjectReference(UObject*& Object, const UObject* ReferencingObject, const UProperty* ReferencingProperty)
+template <bool bParallel, bool bWithClusters>
+void FGCCollector<bParallel, bWithClusters>::HandleObjectReference(UObject*& Object, const UObject* ReferencingObject, const UProperty* ReferencingProperty)
 {
 		InternalHandleObjectReference(Object, ReferencingObject, ReferencingProperty);
 }
 
-template <bool bParallel>
-void FGCCollector<bParallel>::HandleObjectReferences(UObject** InObjects, const int32 ObjectNum, const UObject* InReferencingObject, const UProperty* InReferencingProperty)
+template <bool bParallel, bool bWithClusters>
+void FGCCollector<bParallel, bWithClusters>::HandleObjectReferences(UObject** InObjects, const int32 ObjectNum, const UObject* InReferencingObject, const UProperty* InReferencingProperty)
 {
 		for (int32 ObjectIndex = 0; ObjectIndex < ObjectNum; ++ObjectIndex)
 		{
@@ -933,20 +934,17 @@ void FGCCollector<bParallel>::HandleObjectReferences(UObject** InObjects, const
 		}
 }
 
-typedef FGCCollector<true> FGCCollectorMultithreaded;
-typedef FGCCollector<false> FGCCollectorSinglethreaded;
-
 /*----------------------------------------------------------------------------
 	FReferenceFinder.
 ----------------------------------------------------------------------------*/
-FReferenceFinder::FReferenceFinder( TArray<UObject*>& InObjectArray, UObject* InOuter, bool bInRequireDirectOuter, bool bInShouldIgnoreArchetype, bool bInSerializeRecursively, bool bInShouldIgnoreTransient )
-	:	ObjectArray( InObjectArray )
-	,	LimitOuter( InOuter )
+FReferenceFinder::FReferenceFinder(TArray<UObject*>& InObjectArray, UObject* InOuter, bool bInRequireDirectOuter, bool bInShouldIgnoreArchetype, bool bInSerializeRecursively, bool bInShouldIgnoreTransient)
+	: ObjectArray(InObjectArray)
+	, LimitOuter(InOuter)
 	, SerializedProperty(nullptr)
-	,	bRequireDirectOuter( bInRequireDirectOuter )
-	, bShouldIgnoreArchetype( bInShouldIgnoreArchetype )
-	, bSerializeRecursively( false )
-	, bShouldIgnoreTransient( bInShouldIgnoreTransient )
+	, bRequireDirectOuter(bInRequireDirectOuter)
+	, bShouldIgnoreArchetype(bInShouldIgnoreArchetype)
+	, bSerializeRecursively(false)
+	, bShouldIgnoreTransient(bInShouldIgnoreTransient)
 {
 	bSerializeRecursively = bInSerializeRecursively && LimitOuter != NULL;
 	if (InOuter)
@@ -1010,16 +1008,49 @@ void FReferenceFinder::HandleObjectReference( UObject*& InObject, const UObject*
  */
 class FRealtimeGC : public FGarbageCollectionTracer
 {
+	typedef void(FRealtimeGC::*MarkObjectsFn)(TArray<UObject*>&, const EObjectFlags);
+	typedef void(FRealtimeGC::*ReachabilityAnalysisFn)(FGCArrayStruct*);
+
+	/** Pointers to functions used for Marking objects as unreachable */
+	MarkObjectsFn MarkObjectsFunctions[4];
+	/** Pointers to functions used for Reachability Analysis */
+	ReachabilityAnalysisFn ReachabilityAnalysisFunctions[4];
+
+	template <bool bParallel, bool bWithClusters>
+	void PerformReachabilityAnalysisOnObjectsInternal(FGCArrayStruct* ArrayStruct)
+	{
+		FGCReferenceProcessor<bParallel, bWithClusters> ReferenceProcessor;
+		TFastReferenceCollector<bParallel, FGCReferenceProcessor<bParallel, bWithClusters>, FGCCollector<bParallel, bWithClusters>, FGCArrayPool> ReferenceCollector(ReferenceProcessor, FGCArrayPool::Get());
+		ReferenceCollector.CollectReferences(*ArrayStruct);
+	}
+
+	/** Calculates GC function index based on current settings */
+	static FORCEINLINE int32 GetGCFunctionIndex(bool bParallel, bool bWithClusters)
+	{
+		return (int32(bParallel) | (int32(bWithClusters) << 1));
+	}
+
 public:
 	/** Default constructor, initializing all members. */
 	FRealtimeGC()
-	{}
+	{
+		MarkObjectsFunctions[GetGCFunctionIndex(false, false)] = &FRealtimeGC::MarkObjectsAsUnreachable<false, false>;
+		MarkObjectsFunctions[GetGCFunctionIndex(true, false)] = &FRealtimeGC::MarkObjectsAsUnreachable<true, false>;
+		MarkObjectsFunctions[GetGCFunctionIndex(false, true)] = &FRealtimeGC::MarkObjectsAsUnreachable<false, true>;
+		MarkObjectsFunctions[GetGCFunctionIndex(true, true)] = &FRealtimeGC::MarkObjectsAsUnreachable<true, true>;
+
+		ReachabilityAnalysisFunctions[GetGCFunctionIndex(false, false)] = &FRealtimeGC::PerformReachabilityAnalysisOnObjectsInternal<false, false>;
+		ReachabilityAnalysisFunctions[GetGCFunctionIndex(true, false)] = &FRealtimeGC::PerformReachabilityAnalysisOnObjectsInternal<true, false>;
+		ReachabilityAnalysisFunctions[GetGCFunctionIndex(false, true)] = &FRealtimeGC::PerformReachabilityAnalysisOnObjectsInternal<false, true>;
+		ReachabilityAnalysisFunctions[GetGCFunctionIndex(true, true)] = &FRealtimeGC::PerformReachabilityAnalysisOnObjectsInternal<true, true>;
+	}
 
 	/** 
 	 * Marks all objects that don't have KeepFlags and EInternalObjectFlags::GarbageCollectionKeepFlags as unreachable
 	 * This function is a template to speed up the case where we don't need to assemble the token stream (saves about 6ms on PS4)
 	 */
-	void MarkObjectsAsUnreachable(TArray<UObject*>& ObjectsToSerialize, const EObjectFlags KeepFlags, bool bForceSingleThreaded)
+	template <bool bParallel, bool bWithClusters>
+	void MarkObjectsAsUnreachable(TArray<UObject*>& ObjectsToSerialize, const EObjectFlags KeepFlags)
 	{
 		const EInternalObjectFlags FastKeepFlags = EInternalObjectFlags::GarbageCollectionKeepFlags;
 
@@ -1053,7 +1084,10 @@ public:
 					// Keep track of how many objects are around.
 					ObjectCountDuringMarkPhase++;
 					
-					ObjectItem->ClearFlags(EInternalObjectFlags::ReachableInCluster);
+					if (bWithClusters)
+					{
+						ObjectItem->ClearFlags(EInternalObjectFlags::ReachableInCluster);
+					}
 					// Special case handling for objects that are part of the root set.
 					if (ObjectItem->IsRootSet())
 					{
@@ -1063,15 +1097,18 @@ public:
 #if DO_GUARD_SLOW
 						checkCode(if (ObjectItem->IsPendingKill()) { UE_LOG(LogGarbage, Fatal, TEXT("Object %s is part of root set though has been marked RF_PendingKill!"), *Object->GetFullName()); });
 #endif
-						if (ObjectItem->HasAnyFlags(EInternalObjectFlags::ClusterRoot) || ObjectItem->GetOwnerIndex() > 0)
+						if (bWithClusters)
 						{
-							KeepClusterRefsList.Push(ObjectItem);
+							if (ObjectItem->HasAnyFlags(EInternalObjectFlags::ClusterRoot) || ObjectItem->GetOwnerIndex() > 0)
+							{
+								KeepClusterRefsList.Push(ObjectItem);
+							}
 						}
 
 						ObjectsToSerializeList.Push(Object);
 					}
 					// Regular objects or cluster root objects
-					else if (ObjectItem->GetOwnerIndex() <= 0)
+					else if (!bWithClusters || ObjectItem->GetOwnerIndex() <= 0)
 					{
 						bool bMarkAsUnreachable = true;
 						if (!ObjectItem->IsPendingKill())
@@ -1087,7 +1124,7 @@ public:
 								bMarkAsUnreachable = false;
 							}
 						}
-						else if (ObjectItem->HasAnyFlags(EInternalObjectFlags::ClusterRoot))
+						else if (bWithClusters && ObjectItem->HasAnyFlags(EInternalObjectFlags::ClusterRoot))
 						{
 							ClustersToDissolveList.Push(ObjectItem);
 						}
@@ -1099,9 +1136,12 @@ public:
 							checkSlow(Object->IsValidLowLevel());
 							ObjectsToSerializeList.Push(Object);
 
-							if (ObjectItem->HasAnyFlags(EInternalObjectFlags::ClusterRoot))
+							if (bWithClusters)
 							{
-								KeepClusterRefsList.Push(ObjectItem);
+								if (ObjectItem->HasAnyFlags(EInternalObjectFlags::ClusterRoot))
+								{
+									KeepClusterRefsList.Push(ObjectItem);
+								}
 							}
 						}
 						else
@@ -1113,10 +1153,11 @@ public:
 			}
 
 			GObjectCountDuringLastMarkPhase.Add(ObjectCountDuringMarkPhase);
-		}, bForceSingleThreaded);
+		}, !bParallel);
 		
 		ObjectsToSerializeList.PopAll(ObjectsToSerialize);
 
+		if (bWithClusters)
 		{
 			TArray<FUObjectItem*> ClustersToDissolve;
 			ClustersToDissolveList.PopAll(ClustersToDissolve);
@@ -1132,6 +1173,7 @@ public:
 			}
 		}
 
+		if (bWithClusters)
 		{
 			TArray<FUObjectItem*> KeepClusterRefs;
 			KeepClusterRefsList.PopAll(KeepClusterRefs);
@@ -1153,7 +1195,7 @@ public:
 						{
 							RootObjectItem->ClearFlags(EInternalObjectFlags::Unreachable);
 							// Make sure all referenced clusters are marked as reachable too
-							FGCReferenceProcessorSinglethreaded::MarkReferencedClustersAsReachable(RootObjectItem->GetClusterIndex(), ObjectsToSerialize);
+							FGCReferenceProcessor<false, bWithClusters>::MarkReferencedClustersAsReachable(RootObjectItem->GetClusterIndex(), ObjectsToSerialize);
 						}
 					}
 				}
@@ -1162,7 +1204,7 @@ public:
 					checkSlow(ObjectItem->HasAnyFlags(EInternalObjectFlags::ClusterRoot));
 					// this thing is definitely not marked unreachable, so don't test it here
 					// Make sure all referenced clusters are marked as reachable too
-					FGCReferenceProcessorSinglethreaded::MarkReferencedClustersAsReachable(ObjectItem->GetClusterIndex(), ObjectsToSerialize);
+					FGCReferenceProcessor<false, bWithClusters>::MarkReferencedClustersAsReachable(ObjectItem->GetClusterIndex(), ObjectsToSerialize);
 				}
 			}
 		}
@@ -1173,7 +1215,7 @@ public:
 	 *
 	 * @param KeepFlags		Objects with these flags will be kept regardless of being referenced or not
 	 */
-	void PerformReachabilityAnalysis(EObjectFlags KeepFlags, bool bForceSingleThreaded = false)
+	void PerformReachabilityAnalysis(EObjectFlags KeepFlags, bool bForceSingleThreaded, bool bWithClusters)
 	{
 		LLM_SCOPE(ELLMTag::GC);
 
@@ -1195,13 +1237,13 @@ public:
 
 		{
 			const double StartTime = FPlatformTime::Seconds();
-			MarkObjectsAsUnreachable(ObjectsToSerialize, KeepFlags, bForceSingleThreaded);
+			(this->*MarkObjectsFunctions[GetGCFunctionIndex(!bForceSingleThreaded, bWithClusters)])(ObjectsToSerialize, KeepFlags);
 			UE_LOG(LogGarbage, Verbose, TEXT("%f ms for Mark Phase (%d Objects To Serialize"), (FPlatformTime::Seconds() - StartTime) * 1000, ObjectsToSerialize.Num());
 		}
 
 		{
 			const double StartTime = FPlatformTime::Seconds();
-			PerformReachabilityAnalysisOnObjects(ArrayStruct, bForceSingleThreaded);
+			PerformReachabilityAnalysisOnObjects(ArrayStruct, bForceSingleThreaded, bWithClusters);
 			UE_LOG(LogGarbage, Verbose, TEXT("%f ms for Reachability Analysis"), (FPlatformTime::Seconds() - StartTime) * 1000);
 		}
         
@@ -1216,20 +1258,9 @@ public:
 #endif
 	}
 
-	virtual void PerformReachabilityAnalysisOnObjects(FGCArrayStruct* ArrayStruct, bool bForceSingleThreaded) override
+	virtual void PerformReachabilityAnalysisOnObjects(FGCArrayStruct* ArrayStruct, bool bForceSingleThreaded, bool bWithClusters) override
 	{
-		if (!bForceSingleThreaded)
-		{
-			FGCReferenceProcessorMultithreaded ReferenceProcessor;
-			TFastReferenceCollector<true, FGCReferenceProcessorMultithreaded, FGCCollectorMultithreaded, FGCArrayPool> ReferenceCollector(ReferenceProcessor, FGCArrayPool::Get());
-			ReferenceCollector.CollectReferences(*ArrayStruct);
-		}
-		else
-		{
-			FGCReferenceProcessorSinglethreaded ReferenceProcessor;
-			TFastReferenceCollector<false, FGCReferenceProcessorSinglethreaded, FGCCollectorSinglethreaded, FGCArrayPool> ReferenceCollector(ReferenceProcessor, FGCArrayPool::Get());
-			ReferenceCollector.CollectReferences(*ArrayStruct);
-		}
+		(this->*ReachabilityAnalysisFunctions[GetGCFunctionIndex(!bForceSingleThreaded, bWithClusters)])(ArrayStruct);
 	}
 };
 
@@ -1842,6 +1873,12 @@ void CollectGarbageInternal(EObjectFlags KeepFlags, bool bPerformFullPurge)
 		check(!GObjIncrementalPurgeIsInProgress);
 		check(!GObjPurgeIsRequired);
 
+		// This can happen if someone disables clusters from the console (gc.CreateGCClusters)
+		if (!GCreateGCClusters && GUObjectClusters.GetNumAllocatedClusters())
+		{
+			GUObjectClusters.DissolveClusters(true);
+		}
+
 #if VERIFY_DISREGARD_GC_ASSUMPTIONS
 		// Only verify assumptions if option is enabled. This avoids false positives in the Editor or commandlets.
 		if ((GUObjectArray.DisregardForGCEnabled() || GUObjectClusters.GetNumAllocatedClusters()) && GShouldVerifyGCAssumptions)
@@ -1863,7 +1900,7 @@ void CollectGarbageInternal(EObjectFlags KeepFlags, bool bPerformFullPurge)
 		{
 			const double StartTime = FPlatformTime::Seconds();
 			FRealtimeGC TagUsedRealtimeGC;
-			TagUsedRealtimeGC.PerformReachabilityAnalysis(KeepFlags, bForceSingleThreadedGC);
+			TagUsedRealtimeGC.PerformReachabilityAnalysis(KeepFlags, bForceSingleThreadedGC, GCreateGCClusters);
 			UE_LOG(LogGarbage, Log, TEXT("%f ms for GC"), (FPlatformTime::Seconds() - StartTime) * 1000);
 		}
 
diff --git a/Engine/Source/Runtime/CoreUObject/Private/UObject/UObjectBase.cpp b/Engine/Source/Runtime/CoreUObject/Private/UObject/UObjectBase.cpp
--- a/Engine/Source/Runtime/CoreUObject/Private/UObject/UObjectBase.cpp
+++ b/Engine/Source/Runtime/CoreUObject/Private/UObject/UObjectBase.cpp
@@ -1087,7 +1087,6 @@ void UObjectBaseInit()
 		bool bIsCookOnTheFly = FParse::Value(FCommandLine::Get(), TEXT("-filehostip="), Value);
 		if (bIsCookOnTheFly)
 		{
-			extern int32 GCreateGCClusters;
 			GCreateGCClusters = false;
 		}
 		else
diff --git a/Engine/Source/Runtime/CoreUObject/Private/UObject/UObjectClusters.cpp b/Engine/Source/Runtime/CoreUObject/Private/UObject/UObjectClusters.cpp
--- a/Engine/Source/Runtime/CoreUObject/Private/UObject/UObjectClusters.cpp
+++ b/Engine/Source/Runtime/CoreUObject/Private/UObject/UObjectClusters.cpp
@@ -29,6 +29,14 @@ static FAutoConsoleVariableRef CCreateGCClusters(
 	ECVF_Default
 	);
 
+int32 GAssetClustreringEnabled = 1;
+static FAutoConsoleVariableRef CVarAssetClustreringEnabled(
+	TEXT("gc.AssetClustreringEnabled"),
+	GAssetClustreringEnabled,
+	TEXT("If true, the engine will attempt to create clusters from asset files."),
+	ECVF_Default
+);
+
 int32 GMinGCClusterSize = 5;
 static FAutoConsoleVariableRef CMinGCClusterSize(
 	TEXT("gc.MinGCClusterSize"),
@@ -713,7 +721,7 @@ public:
 /** Looks through objects loaded with a package and creates clusters from them */
 void CreateClustersFromPackage(FLinkerLoad* PackageLinker, TArray<UObject*>& OutClusterObjects)
 {	
-	if (FPlatformProperties::RequiresCookedData() && !GIsInitialLoad && GCreateGCClusters && !GUObjectArray.IsOpenForDisregardForGC() && GUObjectArray.DisregardForGCEnabled() )
+	if (FPlatformProperties::RequiresCookedData() && !GIsInitialLoad && GCreateGCClusters && GAssetClustreringEnabled && !GUObjectArray.IsOpenForDisregardForGC() && GUObjectArray.DisregardForGCEnabled() )
 	{
 		check(PackageLinker);
 
@@ -784,6 +792,8 @@ bool UObjectBaseUtility::CanBeInCluster() const
 
 void UObjectBaseUtility::CreateCluster()
 {
+	check(GCreateGCClusters);
+
 	DECLARE_SCOPE_CYCLE_COUNTER(TEXT("UObjectBaseUtility::CreateCluster"), STAT_FArchiveRealtimeGC_CreateCluster, STATGROUP_GC);
 
 	FUObjectItem* RootItem = GUObjectArray.IndexToObject(InternalIndex);
diff --git a/Engine/Source/Runtime/CoreUObject/Private/UObject/UObjectGlobals.cpp b/Engine/Source/Runtime/CoreUObject/Private/UObject/UObjectGlobals.cpp
--- a/Engine/Source/Runtime/CoreUObject/Private/UObject/UObjectGlobals.cpp
+++ b/Engine/Source/Runtime/CoreUObject/Private/UObject/UObjectGlobals.cpp
@@ -1638,8 +1638,7 @@ void EndLoad(FUObjectSerializeContext* LoadContext)
 				}
 			}
 			// Create clusters after all objects have been loaded
-			extern int32 GCreateGCClusters;
-			if (FPlatformProperties::RequiresCookedData() && !GIsInitialLoad && GCreateGCClusters && !GUObjectArray.IsOpenForDisregardForGC())
+			if (FPlatformProperties::RequiresCookedData() && !GIsInitialLoad && GCreateGCClusters && GAssetClustreringEnabled && !GUObjectArray.IsOpenForDisregardForGC())
 			{
 				for (UObject* Obj : ObjLoaded)
 				{
diff --git a/Engine/Source/Runtime/CoreUObject/Public/UObject/GarbageCollection.h b/Engine/Source/Runtime/CoreUObject/Public/UObject/GarbageCollection.h
--- a/Engine/Source/Runtime/CoreUObject/Public/UObject/GarbageCollection.h
+++ b/Engine/Source/Runtime/CoreUObject/Public/UObject/GarbageCollection.h
@@ -451,7 +451,7 @@ public:
 	~FGCScopeGuard();
 };
 
-template <bool bParallel> class FGCReferenceProcessor;
+template <bool bParallel, bool bWithClusters> class FGCReferenceProcessor;
 
 /** Struct to hold the objects to serialize array and the list of weak references. This is allocated by ArrayPool */
 struct FGCArrayStruct
@@ -463,16 +463,16 @@ struct FGCArrayStruct
 /**
 * Specialized FReferenceCollector that uses FGCReferenceProcessor to mark objects as reachable.
 */
-template <bool bParallel>
+template <bool bParallel, bool bWithClusters>
 class FGCCollector : public FReferenceCollector
 {
-	FGCReferenceProcessor<bParallel>& ReferenceProcessor;
+	FGCReferenceProcessor<bParallel, bWithClusters>& ReferenceProcessor;
 	FGCArrayStruct& ObjectArrayStruct;
 	bool bAllowEliminatingReferences;
 
 public:
 
-	FGCCollector(FGCReferenceProcessor<bParallel>& InProcessor, FGCArrayStruct& InObjectArrayStruct);
+	FGCCollector(FGCReferenceProcessor<bParallel, bWithClusters>& InProcessor, FGCArrayStruct& InObjectArrayStruct);
 
 	virtual void HandleObjectReference(UObject*& InObject, const UObject* InReferencingObject, const UProperty* InReferencingProperty) override;
 
@@ -510,6 +510,6 @@ class FGarbageCollectionTracer
 public:
 	virtual ~FGarbageCollectionTracer() {}
 
-	virtual void PerformReachabilityAnalysisOnObjects(FGCArrayStruct* ArrayStruct, bool bForceSingleThreaded) = 0;
+	virtual void PerformReachabilityAnalysisOnObjects(FGCArrayStruct* ArrayStruct, bool bForceSingleThreaded, bool bWithClusters) = 0;
 
 };
diff --git a/Engine/Source/Runtime/CoreUObject/Public/UObject/UObjectGlobals.h b/Engine/Source/Runtime/CoreUObject/Public/UObject/UObjectGlobals.h
--- a/Engine/Source/Runtime/CoreUObject/Public/UObject/UObjectGlobals.h
+++ b/Engine/Source/Runtime/CoreUObject/Public/UObject/UObjectGlobals.h
@@ -2127,6 +2127,12 @@ struct COREUOBJECT_API FCoreUObjectDelegates
 /** Allows release builds to override not verifying GC assumptions. Useful for profiling as it's hitchy. */
 extern COREUOBJECT_API bool GShouldVerifyGCAssumptions;
 
+/** If non-zero, the engine will create Garbage Collector clusters to speed up Garbage Collection */
+extern COREUOBJECT_API int32 GCreateGCClusters;
+
+/** If non-zero, the engine will attempt to create clusters from asset files */
+extern COREUOBJECT_API int32 GAssetClustreringEnabled;
+
 /** A struct used as stub for deleted ones. */
 COREUOBJECT_API UScriptStruct* GetFallbackStruct();
 enum class EConstructDynamicType : uint8
diff --git a/Engine/Source/Runtime/Engine/Classes/Engine/CoreSettings.h b/Engine/Source/Runtime/Engine/Classes/Engine/CoreSettings.h
--- a/Engine/Source/Runtime/Engine/Classes/Engine/CoreSettings.h
+++ b/Engine/Source/Runtime/Engine/Classes/Engine/CoreSettings.h
@@ -190,6 +190,11 @@ protected:
 		ToolTip = "If true, the engine will attempt to create clusters of objects for better garbage collection performance."))
 	uint32 CreateGCClusters : 1;
 
+	UPROPERTY(EditAnywhere, config, Category = Optimization, meta = (
+		ConsoleVariable = "gc.AssetClustreringEnabled", DisplayName = "Asset Clustering Enabled",
+		ToolTip = "Whether to allow asset files to create actor clusters for GC."))
+	uint32 AssetClusteringEnabled : 1;
+
 	UPROPERTY(EditAnywhere, config, Category = Optimization, meta = (
 		ConsoleVariable = "gc.ActorClusteringEnabled", DisplayName = "Actor Clustering Enabled",
 		ToolTip = "Whether to allow levels to create actor clusters for GC."))
diff --git a/Engine/Source/Runtime/Engine/Private/CoreSettings.cpp b/Engine/Source/Runtime/Engine/Private/CoreSettings.cpp
--- a/Engine/Source/Runtime/Engine/Private/CoreSettings.cpp
+++ b/Engine/Source/Runtime/Engine/Private/CoreSettings.cpp
@@ -178,7 +178,8 @@ UGarbageCollectionSettings::UGarbageCollectionSettings()
 	MaxObjectsInGame = 2 * 1024 * 1024;	
 	CreateGCClusters = true;	
 	MinGCClusterSize = 5;
-	ActorClusteringEnabled = true;
+	AssetClusteringEnabled = true;
+	ActorClusteringEnabled = true;	
 	BlueprintClusteringEnabled = false;
 	UseDisregardForGCOnDedicatedServers = false;
 }
diff --git a/Engine/Source/Runtime/Engine/Private/Level.cpp b/Engine/Source/Runtime/Engine/Private/Level.cpp
--- a/Engine/Source/Runtime/Engine/Private/Level.cpp
+++ b/Engine/Source/Runtime/Engine/Private/Level.cpp
@@ -627,7 +627,7 @@ void ULevel::CreateCluster()
 	// Also, we don't want the level to reference the actors that are clusters because that would
 	// make things work even slower (references to clustered objects are expensive). That's why
 	// we keep a separate array for referencing unclustered actors (ActorsForGC).
-	if (FPlatformProperties::RequiresCookedData() && GActorClusteringEnabled && !bActorClusterCreated)
+	if (FPlatformProperties::RequiresCookedData() && GCreateGCClusters && GActorClusteringEnabled && !bActorClusterCreated)
 	{
 		TArray<AActor*> ClusterActors;
 
