Re-FORCEINLINEning functions that were made virtual in FAsyncLoadingThread

diff --git a/Engine/Source/Runtime/CoreUObject/Private/Serialization/AsyncLoading.cpp b/Engine/Source/Runtime/CoreUObject/Private/Serialization/AsyncLoading.cpp
--- a/Engine/Source/Runtime/CoreUObject/Private/Serialization/AsyncLoading.cpp
+++ b/Engine/Source/Runtime/CoreUObject/Private/Serialization/AsyncLoading.cpp
@@ -319,7 +319,7 @@ static FORCEINLINE bool IsTimeLimitExceeded(double InTickStartTime, bool bUseTim
 }
 FORCEINLINE bool FAsyncPackage::IsTimeLimitExceeded()
 {
-	return AsyncLoadingThread.IsAsyncLoadingSuspended() || ::IsTimeLimitExceeded(TickStartTime, bUseTimeLimit, TimeLimit, LastTypeOfWorkPerformed, LastObjectWorkWasPerformedOn);
+	return AsyncLoadingThread.IsAsyncLoadingSuspendedInternal() || ::IsTimeLimitExceeded(TickStartTime, bUseTimeLimit, TimeLimit, LastTypeOfWorkPerformed, LastObjectWorkWasPerformedOn);
 }
 
 DEFINE_LOG_CATEGORY_STATIC(LogAsyncArchive, Display, All);
@@ -4005,7 +4005,7 @@ EAsyncPackageState::Type FAsyncLoadingThread::ProcessAsyncLoading(int32& OutPack
 				}
 			}
 
-			if (IsAsyncLoadingSuspended())
+			if (IsAsyncLoadingSuspendedInternal())
 			{
 				return EAsyncPackageState::TimeOut;
 			}
@@ -4315,7 +4315,7 @@ EAsyncPackageState::Type FAsyncLoadingThread::ProcessLoadedPackages(bool bUseTim
 
 		
 	bDidSomething = LoadedPackagesToProcess.Num() > 0;
-	for (int32 PackageIndex = 0; PackageIndex < LoadedPackagesToProcess.Num() && !IsAsyncLoadingSuspended(); ++PackageIndex)
+	for (int32 PackageIndex = 0; PackageIndex < LoadedPackagesToProcess.Num() && !IsAsyncLoadingSuspendedInternal(); ++PackageIndex)
 	{
 		FAsyncPackage* Package = LoadedPackagesToProcess[PackageIndex];
 		if (Package->GetDependencyRefCount() == 0)
@@ -4487,7 +4487,7 @@ EAsyncPackageState::Type FAsyncLoadingThread::TickAsyncLoading(bool bUseTimeLimi
 	check(IsInGameThread());
 	check(!IsGarbageCollecting());
 
-	const bool bLoadingSuspended = IsAsyncLoadingSuspended();
+	const bool bLoadingSuspended = IsAsyncLoadingSuspendedInternal();
 	EAsyncPackageState::Type Result = bLoadingSuspended ? EAsyncPackageState::PendingImports : EAsyncPackageState::Complete;
 
 	if (!bLoadingSuspended)
@@ -5568,7 +5568,7 @@ EAsyncPackageState::Type FAsyncPackage::TickAsyncPackage(bool InbUseTimeLimit, b
 		}
 	} while (!IsTimeLimitExceeded() && LoadingState == EAsyncPackageState::TimeOut);
 
-	check(bUseTimeLimit || LoadingState != EAsyncPackageState::TimeOut || AsyncLoadingThread.IsAsyncLoadingSuspended() || IsGarbageCollectionWaiting());
+	check(bUseTimeLimit || LoadingState != EAsyncPackageState::TimeOut || AsyncLoadingThread.IsAsyncLoadingSuspendedInternal() || IsGarbageCollectionWaiting());
 
 	if (LinkerRoot && LoadingState == EAsyncPackageState::Complete)
 	{
@@ -6413,7 +6413,7 @@ EAsyncPackageState::Type FAsyncPackage::PostLoadDeferredObjects(double InTickSta
 	STAT(double PostLoadStartTime = FPlatformTime::Seconds());
 
 	while (DeferredPostLoadIndex < DeferredPostLoadObjects.Num() && 
-		!AsyncLoadingThread.IsAsyncLoadingSuspended() &&
+		!AsyncLoadingThread.IsAsyncLoadingSuspendedInternal() &&
 		!::IsTimeLimitExceeded(InTickStartTime, bInUseTimeLimit, InOutTimeLimit, LastTypeOfWorkPerformed, LastObjectWorkWasPerformedOn))
 	{
 		UObject* Object = DeferredPostLoadObjects[DeferredPostLoadIndex++];
@@ -6484,7 +6484,7 @@ EAsyncPackageState::Type FAsyncPackage::PostLoadDeferredObjects(double InTickSta
 		TArray<UObject*> CDODefaultSubobjects;
 		// Clear async loading flags (we still want RF_Async, but EInternalObjectFlags::AsyncLoading can be cleared)
 		while (DeferredFinalizeIndex < DeferredFinalizeObjects.Num() &&
-			(DeferredPostLoadIndex % 100 != 0 || (!AsyncLoadingThread.IsAsyncLoadingSuspended() && !::IsTimeLimitExceeded(InTickStartTime, bInUseTimeLimit, InOutTimeLimit, LastTypeOfWorkPerformed, LastObjectWorkWasPerformedOn))))
+			(DeferredPostLoadIndex % 100 != 0 || (!AsyncLoadingThread.IsAsyncLoadingSuspendedInternal() && !::IsTimeLimitExceeded(InTickStartTime, bInUseTimeLimit, InOutTimeLimit, LastTypeOfWorkPerformed, LastObjectWorkWasPerformedOn))))
 		{
 			UObject* Object = DeferredFinalizeObjects[DeferredFinalizeIndex++];
 			if (Object)
@@ -6583,7 +6583,7 @@ EAsyncPackageState::Type FAsyncPackage::CreateClusters(double InTickStartTime, b
 	LastTypeOfWorkPerformed = TEXT("CreateClusters");
 
 	while (DeferredClusterIndex < DeferredClusterObjects.Num() &&
-		(!AsyncLoadingThread.IsAsyncLoadingSuspended() && !::IsTimeLimitExceeded(InTickStartTime, bInUseTimeLimit, InOutTimeLimit, LastTypeOfWorkPerformed, LastObjectWorkWasPerformedOn)))
+		(!AsyncLoadingThread.IsAsyncLoadingSuspendedInternal() && !::IsTimeLimitExceeded(InTickStartTime, bInUseTimeLimit, InOutTimeLimit, LastTypeOfWorkPerformed, LastObjectWorkWasPerformedOn)))
 	{
 		UObject* ClusterRootObject = DeferredClusterObjects[DeferredClusterIndex++];
 		LastObjectWorkWasPerformedOn = ClusterRootObject;
@@ -6940,7 +6940,7 @@ void FAsyncLoadingThread::FlushLoading(int32 PackageID)
  	if (IsAsyncLoading())
 	{
 		// Flushing async loading while loading is suspend will result in infinite stall
-		UE_CLOG(IsAsyncLoadingSuspended(), LogStreaming, Fatal, TEXT("Cannot Flush Async Loading while async loading is suspended (%d)"), GetAsyncLoadingSuspendedCount());
+		UE_CLOG(IsAsyncLoadingSuspendedInternal(), LogStreaming, Fatal, TEXT("Cannot Flush Async Loading while async loading is suspended (%d)"), GetAsyncLoadingSuspendedCount());
 
 		SCOPE_CYCLE_COUNTER(STAT_FAsyncPackage_FlushAsyncLoadingGameThread);
 
@@ -7018,7 +7018,7 @@ EAsyncPackageState::Type FAsyncLoadingThread::ProcessLoadingUntilComplete(TFunct
 	SCOPE_CYCLE_COUNTER(STAT_FAsyncPackage_FlushAsyncLoadingGameThread);
 
 	// Flushing async loading while loading is suspend will result in infinite stall
-	UE_CLOG(IsAsyncLoadingSuspended(), LogStreaming, Fatal, TEXT("Cannot Flush Async Loading while async loading is suspended (%d)"), GetAsyncLoadingSuspendedCount());
+	UE_CLOG(IsAsyncLoadingSuspendedInternal(), LogStreaming, Fatal, TEXT("Cannot Flush Async Loading while async loading is suspended (%d)"), GetAsyncLoadingSuspendedCount());
 
 	if (TimeLimit <= 0.0f)
 	{
diff --git a/Engine/Source/Runtime/CoreUObject/Private/Serialization/AsyncLoadingThread.h b/Engine/Source/Runtime/CoreUObject/Private/Serialization/AsyncLoadingThread.h
--- a/Engine/Source/Runtime/CoreUObject/Private/Serialization/AsyncLoadingThread.h
+++ b/Engine/Source/Runtime/CoreUObject/Private/Serialization/AsyncLoadingThread.h
@@ -147,7 +147,7 @@ struct FMaxPackageSummarySize
 /**
  * Async loading thread. Preloads/serializes packages on async loading thread. Postloads objects on the game thread.
  */
-class FAsyncLoadingThread : public FRunnable, public IAsyncPackageLoader
+class FAsyncLoadingThread final : public FRunnable, public IAsyncPackageLoader
 {
 	/** Structure that holds the async loading thread ini settings */
 	struct FAsyncLoadingThreadSettings
@@ -388,10 +388,14 @@ public:
 	}
 
 	/** Returns true if async loading is suspended */
-	FORCEINLINE bool IsAsyncLoadingSuspended() override
+	FORCEINLINE bool IsAsyncLoadingSuspendedInternal() const
 	{
-		FPlatformMisc::MemoryBarrier();
-		return IsLoadingSuspended.GetValue() != 0;
+		return !!IsLoadingSuspended.GetValue();
+	}
+
+	virtual bool IsAsyncLoadingSuspended() override
+	{
+		return IsAsyncLoadingSuspendedInternal();
 	}
 
 	FORCEINLINE int32 GetAsyncLoadingSuspendedCount()
@@ -403,7 +407,6 @@ public:
 	/** Returns the number of async packages that are currently being processed */
 	FORCEINLINE int32 GetNumAsyncPackages() override
 	{
-		FPlatformMisc::MemoryBarrier();
 		return ExistingAsyncPackagesCounter.GetValue();
 	}
 
