Inlined IsGarbageCollectionWaiting so that it doesn't appear in profiles

diff --git a/Engine/Source/Runtime/CoreUObject/Private/Serialization/AsyncLoading.cpp b/Engine/Source/Runtime/CoreUObject/Private/Serialization/AsyncLoading.cpp
--- a/Engine/Source/Runtime/CoreUObject/Private/Serialization/AsyncLoading.cpp
+++ b/Engine/Source/Runtime/CoreUObject/Private/Serialization/AsyncLoading.cpp
@@ -41,7 +41,7 @@
 #include "HAL/LowLevelMemTracker.h"
 #include "HAL/LowLevelMemStats.h"
 #include "ProfilingDebugging/CsvProfiler.h"
-#include "UObject/GCScopeLock.h"
+#include "UObject/GarbageCollectionInternal.h"
 #include "ProfilingDebugging/MiscTrace.h"
 #include "Serialization/LoadTimeTracePrivate.h"
 
@@ -106,9 +106,6 @@ static FName StaticGetNativeClassName(UClass* InClass)
 /** Returns true if we're inside a FGCScopeLock */
 bool IsGarbageCollectionLocked();
 
-/** Returns true if GC wants to run */
-bool IsGarbageCollectionWaiting();
-
 /** Global request ID counter */
 static FThreadSafeCounter GPackageRequestID;
 
diff --git a/Engine/Source/Runtime/CoreUObject/Private/Serialization/AsyncLoading2.cpp b/Engine/Source/Runtime/CoreUObject/Private/Serialization/AsyncLoading2.cpp
--- a/Engine/Source/Runtime/CoreUObject/Private/Serialization/AsyncLoading2.cpp
+++ b/Engine/Source/Runtime/CoreUObject/Private/Serialization/AsyncLoading2.cpp
@@ -36,7 +36,7 @@
 #include "Blueprint/BlueprintSupport.h"
 #include "HAL/LowLevelMemTracker.h"
 #include "ProfilingDebugging/CsvProfiler.h"
-#include "UObject/GCScopeLock.h"
+#include "UObject/GarbageCollectionInternal.h"
 #include "ProfilingDebugging/MiscTrace.h"
 #include "Serialization/LoadTimeTracePrivate.h"
 #include "ProfilingDebugging/CountersTrace.h"
diff --git a/Engine/Source/Runtime/CoreUObject/Private/UObject/GCScopeLock.h b/Engine/Source/Runtime/CoreUObject/Private/UObject/GCScopeLock.h
--- a/Engine/Source/Runtime/CoreUObject/Private/UObject/GCScopeLock.h
+++ b/Engine/Source/Runtime/CoreUObject/Private/UObject/GCScopeLock.h
@@ -179,8 +179,8 @@ public:
 	}
 
 	/** True if GC wants to run on the game thread but is maybe blocked by some other thread */
-	bool IsGCWaiting() const
+	FORCEINLINE bool IsGCWaiting() const
 	{
-		return GCWantsToRunCounter.GetValue() > 0;
+		return !!GCWantsToRunCounter.GetValue();
 	}
 };
diff --git a/Engine/Source/Runtime/CoreUObject/Private/UObject/GarbageCollection.cpp b/Engine/Source/Runtime/CoreUObject/Private/UObject/GarbageCollection.cpp
--- a/Engine/Source/Runtime/CoreUObject/Private/UObject/GarbageCollection.cpp
+++ b/Engine/Source/Runtime/CoreUObject/Private/UObject/GarbageCollection.cpp
@@ -141,11 +141,6 @@ bool IsGarbageCollectionLocked()
 	return FGCCSyncObject::Get().IsAsyncLocked();
 }
 
-bool IsGarbageCollectionWaiting()
-{
-	return FGCCSyncObject::Get().IsGCWaiting();
-}
-
 // Minimum number of objects to spawn a GC sub-task for
 static int32 GMinDesiredObjectsPerSubTask = 128;
 static FAutoConsoleVariableRef CVarMinDesiredObjectsPerSubTask(
diff --git a/Engine/Source/Runtime/CoreUObject/Private/UObject/GarbageCollectionInternal.h b/Engine/Source/Runtime/CoreUObject/Private/UObject/GarbageCollectionInternal.h
new file mode 100644
--- /dev/null
+++ b/Engine/Source/Runtime/CoreUObject/Private/UObject/GarbageCollectionInternal.h
@@ -0,0 +1,20 @@
+// Copyright 1998-2019 Epic Games, Inc. All Rights Reserved.
+
+/*=============================================================================
+	GarbageCollectionInternal.h: Unreal realtime garbage collection internal helpers
+=============================================================================*/
+
+#pragma once
+
+#include "CoreMinimal.h"
+#include "Stats/Stats.h"
+#include "UObject/GarbageCollection.h"
+#include "UObject/GCScopeLock.h"
+
+extern FGCCSyncObject* GGCSingleton;
+
+/** Returns true if GC wants to run */
+FORCEINLINE bool IsGarbageCollectionWaiting()
+{
+	return GGCSingleton->IsGCWaiting();
+}
\ No newline at end of file
diff --git a/Engine/Source/Runtime/CoreUObject/Public/Serialization/AsyncPackageLoader.h b/Engine/Source/Runtime/CoreUObject/Public/Serialization/AsyncPackageLoader.h
--- a/Engine/Source/Runtime/CoreUObject/Public/Serialization/AsyncPackageLoader.h
+++ b/Engine/Source/Runtime/CoreUObject/Public/Serialization/AsyncPackageLoader.h
@@ -14,9 +14,6 @@ extern const FName PrestreamPackageClassNameLoad;
 /** Returns true if we're inside a FGCScopeLock */
 extern bool IsGarbageCollectionLocked();
 
-/** Returns true if GC wants to run */
-extern bool IsGarbageCollectionWaiting();
-
 bool IsFullyLoadedObj(UObject* Obj);
 
 bool IsNativeCodePackage(UPackage* Package);
