Fixed VS2019 compilation issue (int32 -> bool implicit conversion)

diff --git a/Engine/Source/Runtime/CoreUObject/Private/UObject/GarbageCollection.cpp b/Engine/Source/Runtime/CoreUObject/Private/UObject/GarbageCollection.cpp
--- a/Engine/Source/Runtime/CoreUObject/Private/UObject/GarbageCollection.cpp
+++ b/Engine/Source/Runtime/CoreUObject/Private/UObject/GarbageCollection.cpp
@@ -1895,7 +1895,7 @@ void CollectGarbageInternal(EObjectFlags KeepFlags, bool bPerformFullPurge)
 		{
 			const double StartTime = FPlatformTime::Seconds();
 			FRealtimeGC TagUsedRealtimeGC;
-			TagUsedRealtimeGC.PerformReachabilityAnalysis(KeepFlags, bForceSingleThreadedGC, GCreateGCClusters);
+			TagUsedRealtimeGC.PerformReachabilityAnalysis(KeepFlags, bForceSingleThreadedGC, !!GCreateGCClusters);
 			UE_LOG(LogGarbage, Log, TEXT("%f ms for GC"), (FPlatformTime::Seconds() - StartTime) * 1000);
 		}
 
