FORCEINLINEd IsCollectingGarbage to fix performance hotspot

diff --git a/Engine/Source/Runtime/CoreUObject/Private/UObject/GCScopeLock.h b/Engine/Source/Runtime/CoreUObject/Private/UObject/GCScopeLock.h
--- a/Engine/Source/Runtime/CoreUObject/Private/UObject/GCScopeLock.h
+++ b/Engine/Source/Runtime/CoreUObject/Private/UObject/GCScopeLock.h
@@ -13,34 +13,6 @@
 #include "HAL/ThreadSafeBool.h"
 #include "HAL/Event.h"
 
-/** Locks all UObject hash tables when performing GC */
-class FGCScopeLock
-{
-	/** Previous value of the GetGarbageCollectingFlag() */
-	bool bPreviousGabageCollectingFlagValue;
-public:
-
-	static FThreadSafeBool& GetGarbageCollectingFlag();
-
-	/** 
-	 * We're storing the value of GetGarbageCollectingFlag in the constructor, it's safe as only 
-	 * one thread is ever going to be setting it and calling this code - the game thread.
-	 **/
-	FORCEINLINE FGCScopeLock()
-		: bPreviousGabageCollectingFlagValue(GetGarbageCollectingFlag())
-	{		
-		void LockUObjectHashTables();
-		LockUObjectHashTables();		
-		GetGarbageCollectingFlag() = true;
-	}
-	FORCEINLINE ~FGCScopeLock()
-	{		
-		GetGarbageCollectingFlag() = bPreviousGabageCollectingFlagValue;
-		void UnlockUObjectHashTables();
-		UnlockUObjectHashTables();		
-	}
-};
-
 
 /**
 * Garbage Collection synchronization objects
diff --git a/Engine/Source/Runtime/CoreUObject/Private/UObject/GarbageCollection.cpp b/Engine/Source/Runtime/CoreUObject/Private/UObject/GarbageCollection.cpp
--- a/Engine/Source/Runtime/CoreUObject/Private/UObject/GarbageCollection.cpp
+++ b/Engine/Source/Runtime/CoreUObject/Private/UObject/GarbageCollection.cpp
@@ -71,11 +71,33 @@ static int32 GUnrechableObjectIndex = 0;
 /** Helpful constant for determining how many token slots we need to store a pointer **/
 static const uint32 GNumTokensPerPointer = sizeof(void*) / sizeof(uint32); //-V514
 
-FThreadSafeBool& FGCScopeLock::GetGarbageCollectingFlag()
+FThreadSafeBool GIsGarbageCollecting(false);
+
+/** Locks all UObject hash tables when performing GC */
+class FGCScopeLock
 {
-	static FThreadSafeBool IsGarbageCollecting(false);
-	return IsGarbageCollecting;
-}
+	/** Previous value of the GetGarbageCollectingFlag() */
+	bool bPreviousGabageCollectingFlagValue;
+public:
+
+	/**
+	 * We're storing the value of GetGarbageCollectingFlag in the constructor, it's safe as only
+	 * one thread is ever going to be setting it and calling this code - the game thread.
+	 **/
+	FORCEINLINE FGCScopeLock()
+		: bPreviousGabageCollectingFlagValue(GIsGarbageCollecting)
+	{
+		void LockUObjectHashTables();
+		LockUObjectHashTables();
+		GIsGarbageCollecting = true;
+	}
+	FORCEINLINE ~FGCScopeLock()
+	{
+		GIsGarbageCollecting = bPreviousGabageCollectingFlagValue;
+		void UnlockUObjectHashTables();
+		UnlockUObjectHashTables();
+	}
+};
 
 FGCCSyncObject::FGCCSyncObject()
 {
@@ -131,11 +153,6 @@ FGCScopeGuard::~FGCScopeGuard()
 	FGCCSyncObject::Get().UnlockAsync();
 }
 
-bool IsGarbageCollecting()
-{
-	return FGCScopeLock::GetGarbageCollectingFlag();
-}
-
 bool IsGarbageCollectionLocked()
 {
 	return FGCCSyncObject::Get().IsAsyncLocked();
diff --git a/Engine/Source/Runtime/CoreUObject/Public/UObject/GarbageCollection.h b/Engine/Source/Runtime/CoreUObject/Public/UObject/GarbageCollection.h
--- a/Engine/Source/Runtime/CoreUObject/Public/UObject/GarbageCollection.h
+++ b/Engine/Source/Runtime/CoreUObject/Public/UObject/GarbageCollection.h
@@ -10,6 +10,7 @@
 #include "Stats/Stats.h"
 #include "UObject/UObjectGlobals.h"
 #include "Serialization/ArchiveUObject.h"
+#include "HAL/ThreadSafeBool.h"
 
 /** Context sensitive keep flags for garbage collection */
 #define GARBAGE_COLLECTION_KEEPFLAGS	(GIsEditor ? RF_Standalone : RF_NoFlags)
@@ -513,3 +514,14 @@ public:
 	virtual void PerformReachabilityAnalysisOnObjects(FGCArrayStruct* ArrayStruct, bool bForceSingleThreaded, bool bWithClusters) = 0;
 
 };
+
+/** True if Garbage Collection is running. Use IsGarbageCollecting() functio n instead of using this variable directly */
+extern COREUOBJECT_API FThreadSafeBool GIsGarbageCollecting;
+
+/**
+* Whether we are inside garbage collection
+*/
+FORCEINLINE bool IsGarbageCollecting()
+{
+	return GIsGarbageCollecting;
+}
\ No newline at end of file
diff --git a/Engine/Source/Runtime/CoreUObject/Public/UObject/UObjectGlobals.h b/Engine/Source/Runtime/CoreUObject/Public/UObject/UObjectGlobals.h
--- a/Engine/Source/Runtime/CoreUObject/Public/UObject/UObjectGlobals.h
+++ b/Engine/Source/Runtime/CoreUObject/Public/UObject/UObjectGlobals.h
@@ -423,11 +423,6 @@ COREUOBJECT_API bool IsEventDrivenLoaderEnabled();
  */
 COREUOBJECT_API float GetAsyncLoadPercentage( const FName& PackageName );
 
-/**
-* Whether we are inside garbage collection
-*/
-COREUOBJECT_API bool IsGarbageCollecting();
-
 /**
 * Whether we are running on the Garbage Collector Thread
 */
