Added a new method to the BulkData runtime API, ::GetCopyAsBuffer which acts similar to ::GetCopy but will return the copy wrapped in a new object FBulkDataBuffer which adds typesafety and will 'own' the memory returned reducing the scope for memory leaks.

- In the case where the BulkData is single use or the parameter bDiscardInternalCopy is true, then this path will not allocate any memory or need to memcpy from the original data both of which are common in code paths that call ::GetCopy directly.
- FBulkDataBuffer offiers a TArrayView as an interface for accessing the memory it owns.

#rnx

diff --git a/Engine/Source/Runtime/CoreUObject/Public/Serialization/BulkData.h b/Engine/Source/Runtime/CoreUObject/Public/Serialization/BulkData.h
--- a/Engine/Source/Runtime/CoreUObject/Public/Serialization/BulkData.h
+++ b/Engine/Source/Runtime/CoreUObject/Public/Serialization/BulkData.h
@@ -571,6 +571,18 @@ public:
 	 */
 	void GetCopy( void** Dest, bool bDiscardInternalCopy = true );
 
+	/**
+	 * Returns a copy encapsulated by a FBulkDataBuffer.
+	 *
+	 * @param RequestedElementCount If set to greater than 0, the returned FBulkDataBuffer will be limited to
+	 * this number of elements. This will give an error if larger than the actual number of elements in the BulkData object.
+	 * @param bDiscardInternalCopy If true then the BulkData object will free it's internal buffer once called.
+	 *
+	 * @return A FBulkDataBuffer that owns a copy of the BulkData, this might be a subset of the data depending on the value of RequestedSize.
+	 */
+	template<typename ElementType>
+	FBulkDataBuffer<ElementType> GetCopyAsBuffer(int64 RequestedElementCount, bool bDiscardInternalCopy);
+
 	/**
 	 * Locks the bulk data and returns a pointer to it.
 	 *
@@ -845,7 +857,20 @@ protected:
 #endif // WITH_EDITOR
 };
 
+template<typename ElementType>
+FBulkDataBuffer<ElementType> FUntypedBulkData::GetCopyAsBuffer(int64 RequestedElementCount, bool bDiscardInternalCopy)
+{
+	const int64 MaxElementCount = GetElementCount();
+
+	check(RequestedElementCount <= MaxElementCount);
+
+	ElementType* Ptr = nullptr;
+	GetCopy((void**)& Ptr, bDiscardInternalCopy);
 
+	const int64 BufferSize = (RequestedElementCount > 0 ? RequestedElementCount : MaxElementCount);
+
+	return FBulkDataBuffer<ElementType>(Ptr, BufferSize);
+}
 /*-----------------------------------------------------------------------------
 	uint8 version of bulk data.
 -----------------------------------------------------------------------------*/
@@ -867,6 +892,17 @@ struct COREUOBJECT_API FByteBulkDataOld : public FUntypedBulkData
 	 * @param ElementIndex	Element index to serialize
 	 */
 	virtual void SerializeElement( FArchive& Ar, void* Data, int64 ElementIndex );
+
+	/**
+	 * Returns a copy encapsulated by a FBulkDataBuffer.
+	 *
+	 * @param RequestedElementCount If set to greater than 0, the returned FBulkDataBuffer will be limited to
+	 * this number of elements. This will give an error if larger than the actual number of elements in the BulkData object.
+	 * @param bDiscardInternalCopy If true then the BulkData object will free it's internal buffer once called.
+	 *
+	 * @return A FBulkDataBuffer that owns a copy of the BulkData, this might be a subset of the data depending on the value of RequestedSize.
+	 */
+	FBulkDataBuffer<uint8> GetCopyAsBuffer(int64 RequestedElementCount, bool bDiscardInternalCopy) { return FUntypedBulkData::GetCopyAsBuffer<uint8>(RequestedElementCount, bDiscardInternalCopy); }
 };
 
 /*-----------------------------------------------------------------------------
@@ -890,6 +926,17 @@ struct COREUOBJECT_API FWordBulkDataOld : public FUntypedBulkData
 	 * @param ElementIndex	Element index to serialize
 	 */
 	virtual void SerializeElement( FArchive& Ar, void* Data, int64 ElementIndex );
+
+	/**
+	 * Returns a copy encapsulated by a FBulkDataBuffer.
+	 *
+	 * @param RequestedElementCount If set to greater than 0, the returned FBulkDataBuffer will be limited to
+	 * this number of elements. This will give an error if larger than the actual number of elements in the BulkData object.
+	 * @param bDiscardInternalCopy If true then the BulkData object will free it's internal buffer once called.
+	 *
+	 * @return A FBulkDataBuffer that owns a copy of the BulkData, this might be a subset of the data depending on the value of RequestedSize.
+	 */
+	FBulkDataBuffer<uint16> GetCopyAsBuffer(int64 RequestedElementCount, bool bDiscardInternalCopy) { return FUntypedBulkData::GetCopyAsBuffer<uint16>(RequestedElementCount, bDiscardInternalCopy); }
 };
 
 /*-----------------------------------------------------------------------------
@@ -913,6 +960,17 @@ struct COREUOBJECT_API FIntBulkDataOld : public FUntypedBulkData
 	 * @param ElementIndex	Element index to serialize
 	 */
 	virtual void SerializeElement( FArchive& Ar, void* Data, int64 ElementIndex );
+
+	/**
+	 * Returns a copy encapsulated by a FBulkDataBuffer.
+	 *
+	 * @param RequestedElementCount If set to greater than 0, the returned FBulkDataBuffer will be limited to
+	 * this number of elements. This will give an error if larger than the actual number of elements in the BulkData object.
+	 * @param bDiscardInternalCopy If true then the BulkData object will free it's internal buffer once called.
+	 *
+	 * @return A FBulkDataBuffer that owns a copy of the BulkData, this might be a subset of the data depending on the value of RequestedSize.
+	 */
+	FBulkDataBuffer<int32> GetCopyAsBuffer(int64 RequestedElementCount, bool bDiscardInternalCopy) { return FUntypedBulkData::GetCopyAsBuffer<int32>(RequestedElementCount, bDiscardInternalCopy); }
 };
 
 /*-----------------------------------------------------------------------------
@@ -936,6 +994,17 @@ struct COREUOBJECT_API FFloatBulkDataOld : public FUntypedBulkData
 	 * @param ElementIndex	Element index to serialize
 	 */
 	virtual void SerializeElement( FArchive& Ar, void* Data, int64 ElementIndex );
+
+	/**
+	 * Returns a copy encapsulated by a FBulkDataBuffer.
+	 *
+	 * @param RequestedElementCount If set to greater than 0, the returned FBulkDataBuffer will be limited to
+	 * this number of elements. This will give an error if larger than the actual number of elements in the BulkData object.
+	 * @param bDiscardInternalCopy If true then the BulkData object will free it's internal buffer once called.
+	 *
+	 * @return A FBulkDataBuffer that owns a copy of the BulkData, this might be a subset of the data depending on the value of RequestedSize.
+	 */
+	FBulkDataBuffer<float> GetCopyAsBuffer(int64 RequestedElementCount, bool bDiscardInternalCopy) { return FUntypedBulkData::GetCopyAsBuffer<float>(RequestedElementCount, bDiscardInternalCopy); }
 };
 
 // Switch between the old and new types based on USE_NEW_BULKDATA
diff --git a/Engine/Source/Runtime/CoreUObject/Public/Serialization/BulkData2.h b/Engine/Source/Runtime/CoreUObject/Public/Serialization/BulkData2.h
--- a/Engine/Source/Runtime/CoreUObject/Public/Serialization/BulkData2.h
+++ b/Engine/Source/Runtime/CoreUObject/Public/Serialization/BulkData2.h
@@ -3,6 +3,7 @@
 #pragma once
 
 #include "CoreMinimal.h"
+#include "BulkDataBuffer.h"
 #include "Async/AsyncFileHandle.h"
 #include "IO/IoDispatcher.h"
 
@@ -74,6 +75,13 @@ public:
 	bool IsLocked() const;
 
 	void* Realloc(int64 InElementCount);
+
+	/**
+	 * Retrieves a copy of the bulk data.
+	 *
+	 * @param Dest [in/out] Pointer to pointer going to hold copy, can point to nullptr in which case memory is allocated
+	 * @param bDiscardInternalCopy Whether to discard/ free the potentially internally allocated copy of the data
+	 */
 	void GetCopy(void** Dest, bool bDiscardInternalCopy);
 
 	int64 GetBulkDataSize() const;
@@ -168,9 +176,24 @@ public:
 		FBulkDataBase::Serialize(Ar, Owner, Index, bAttemptFileMapping, sizeof(ElementType));
 	}
 
+	/**
+	 * Returns the number of elements held by the BulkData object.
+	 *
+	 * @return Number of elements.
+	 */
+	int64 GetElementCount() const 
+	{ 
+		return GetBulkDataSize() / GetElementSize(); 
+	}
+
+	/**
+	 * Returns size in bytes of single element.
+	 *
+	 * @return The size of the element.
+	 */
 	int32 GetElementSize() const
-	{
-		return sizeof(ElementType);
+	{ 
+		return sizeof(ElementType); 
 	}
 
 	ElementType* Lock(uint32 LockFlags)
@@ -187,6 +210,29 @@ public:
 	{
 		return (ElementType*)FBulkDataBase::Realloc(InElementCount);
 	}
+
+	/**
+	 * Returns a copy encapsulated by a FBulkDataBuffer.
+	 *
+	 * @param RequestedElementCount If set to greater than 0, the returned FBulkDataBuffer will be limited to
+	 * this number of elements. This will give an error if larger than the actual number of elements in the BulkData object.
+	 * @param bDiscardInternalCopy If true then the BulkData object will free it's internal buffer once called.
+	 *
+	 * @return A FBulkDataBuffer that owns a copy of the BulkData, this might be a subset of the data depending on the value of RequestedSize.
+	 */
+	FBulkDataBuffer<ElementType> GetCopyAsBuffer(int64 RequestedElementCount, bool bDiscardInternalCopy)
+	{
+		const int64 MaxElementCount = GetElementCount();
+
+		check(RequestedElementCount <= MaxElementCount);
+
+		ElementType* Ptr = nullptr;
+		GetCopy((void**)& Ptr, bDiscardInternalCopy);
+
+		const int64 BufferSize = (RequestedElementCount > 0 ? RequestedElementCount : MaxElementCount);
+
+		return FBulkDataBuffer<ElementType>(Ptr, BufferSize);
+	}
 };
 
 // Commonly used types
diff --git a/Engine/Source/Runtime/CoreUObject/Public/Serialization/BulkDataBuffer.h b/Engine/Source/Runtime/CoreUObject/Public/Serialization/BulkDataBuffer.h
new file mode 100644
--- /dev/null
+++ b/Engine/Source/Runtime/CoreUObject/Public/Serialization/BulkDataBuffer.h
@@ -0,0 +1,149 @@
+// Copyright 1998-2019 Epic Games, Inc. All Rights Reserved.
+
+#pragma once
+
+#include "CoreMinimal.h"
+
+/**
+ * Used to manage a raw data buffer provided by BulkData while providing access to it
+ * via a TArrayView.
+ *
+ * This object assumes that it owns the buffer that it has been given and will free the
+ * memory when the object is destroyed.
+ */
+template<typename TYPE>
+class FBulkDataBuffer
+{
+public:
+	using ViewType = TArrayView<TYPE>;
+
+	/** Constructor. */
+	FBulkDataBuffer() = default;
+
+	/**
+	 * Copy constructor which will create it's own memory buffer
+	 * and then copy from the source object rather than share access
+	 * to the same buffer.
+	 *
+	 * @param Other The source array to copy.
+	 */
+	FBulkDataBuffer(const FBulkDataBuffer& Other)
+	{
+		*this = Other;
+	}
+
+	/**
+	 * Move constructor.
+	 *
+	 * @param Other FBulkDataBuffer to move from.
+	 */
+	FBulkDataBuffer(FBulkDataBuffer&& Other)
+	{
+		View = Other.View;
+		Other.View = ViewType();
+	}
+
+	/** Constructor.
+	* @param InBuffer Pointer to memory to take ownership of, must've have been allocated via FMemory::Malloc/Realloc.
+	* @param InNumberOfElements The number of elements in the buffer.
+	*/
+	FBulkDataBuffer(TYPE* InBuffer, uint64 InNumberOfElements)
+		: View(InBuffer, InNumberOfElements)
+	{
+		check(InNumberOfElements <= TNumericLimits<uint32>::Max());
+	}
+
+	/** Destructor. */
+	~FBulkDataBuffer()
+	{
+		FreeBuffer();
+	}
+
+	/**
+	 * Assignment operator which will create it's own memory buffer
+	 * and then copy from the source object rather than share access
+	 * to the same buffer.
+	 *
+	 * @param Other The source FBulkDataBuffer to assign from.
+	 */
+	FBulkDataBuffer& operator =(const FBulkDataBuffer& Other)
+	{
+		FreeBuffer();
+
+		if (this != &Other)
+		{
+			const int32 BufferSize = Other.View.Num();
+
+			TYPE* BufferCopy = (TYPE*)FMemory::Malloc(BufferSize);
+			FMemory::Memcpy(BufferCopy, Other.View.GetData(), BufferSize);
+
+			View = ViewType(BufferCopy, BufferSize);
+		}
+
+		return *this;
+	}
+
+	/**
+	 * Move assignment operator.
+	 *
+	 * @param Other FBulkDataBuffer to assign and move from.
+	 */
+	FBulkDataBuffer& operator =(FBulkDataBuffer&& Other)
+	{
+		if (this != &Other)
+		{
+			FreeBuffer();
+
+			View = Other.View;
+			Other.View = ViewType();
+		}
+
+		return *this;
+	}
+
+	/**
+	 * Frees the internal buffer and sets the internal TArrayView to an empty state
+	 *
+	 */
+	void Empty()
+	{
+		FreeBuffer();
+
+		View = ViewType();
+	}
+
+	/**
+	 * Frees any existing buffer and takes ownership of the buffer provided instead.
+	 *
+	 * @param InBuffer Pointer to memory to take ownership of, must've have been allocated via FMemory::Malloc/Realloc.
+	*  @param InNumberOfElements The number of elements in the buffer.
+	 * buffers actual size if desired.
+	 */
+	void Reset(TYPE* InBuffer, uint64 InNumberOfElements)
+	{
+		FreeBuffer();
+
+		View = ViewType(InBuffer, InNumberOfElements);
+	}
+
+	/**
+	* Allows access to the data buffer owned by the object in the form of a TArrayView
+	* @return A valid TArrayView
+	*/
+	const ViewType& GetView() const 
+	{ 
+		return View; 
+	}
+
+private:
+
+	void FreeBuffer()
+	{
+		if (View.GetData() != nullptr)
+		{
+			FMemory::Free(View.GetData());
+		}
+	}
+
+	ViewType View;
+};
