AsyncLoading2:
Individual export create and serialize nodes are now merged into larger export bundle blocks
Replaced IoDispatcher queue with a ReadWithCallback function

diff --git a/Engine/Source/Runtime/Core/Private/IO/IoDispatcher.cpp b/Engine/Source/Runtime/Core/Private/IO/IoDispatcher.cpp
--- a/Engine/Source/Runtime/Core/Private/IO/IoDispatcher.cpp
+++ b/Engine/Source/Runtime/Core/Private/IO/IoDispatcher.cpp
@@ -523,6 +523,21 @@ public:
 		BatchAllocator.Destroy(Batch);
 	}
 
+	void ReadWithCallback(const FIoChunkId& Chunk, const FIoReadOptions& Options, TFunction<void(TIoStatusOr<FIoBuffer>)>&& Callback)
+	{
+		FIoRequestImpl* Request = AllocRequest(Chunk, Options);
+		Request->Result = IoStore->Resolve(Request->ChunkId, Request->Options);
+		if (Request->Result.IsOk())
+		{
+			Callback(Request->Result.ConsumeValueOrDie());
+		}
+		else
+		{
+			Callback(TIoStatusOr<FIoBuffer>(Request->Result.Status()));
+		}
+		FreeRequest(Request);
+	}
+
 	TIoStatusOr<uint64> GetSizeForChunk(const FIoChunkId& ChunkId) const
 	{
 		return IoStore->GetSizeForChunk(ChunkId);
@@ -649,6 +664,12 @@ FIoDispatcher::FreeBatch(FIoBatch Batch)
 	Impl->FreeBatch(Batch.Impl);
 }
 
+void
+FIoDispatcher::ReadWithCallback(const FIoChunkId& Chunk, const FIoReadOptions& Options, TFunction<void(TIoStatusOr<FIoBuffer>)>&& Callback)
+{
+	Impl->ReadWithCallback(Chunk, Options, MoveTemp(Callback));
+}
+
 TIoStatusOr<FIoChunkId>
 FIoDispatcher::OpenFileChunk(FStringView Filename)
 {
@@ -758,19 +779,13 @@ FIoRequest::Status() const
 	return Impl->Result.Status();
 }
 
-FIoBuffer	
-FIoRequest::GetChunk()
-{
-	return Impl->Result.ValueOrDie();
-}
-
 const FIoChunkId&
 FIoRequest::GetChunkId() const
 {
 	return Impl->ChunkId;
 }
 
-const TIoStatusOr<FIoBuffer>&
+TIoStatusOr<FIoBuffer>
 FIoRequest::GetResult() const
 {
 	return Impl->Result;
@@ -782,285 +797,4 @@ FIoRequest::GetUserData() const
 	return Impl->UserData;
 }
 
-//////////////////////////////////////////////////////////////////////////
-
-class FIoQueueImpl
-	: private FRunnable
-{
-public:
-	FIoQueueImpl(FIoDispatcherImpl& InDispatcher, FIoQueue::FBatchReadyCallback&& InBatchReadyCallback)
-		: Dispatcher(InDispatcher)
-		, BatchReadyCallback(Forward<FIoQueue::FBatchReadyCallback>(InBatchReadyCallback))
-	{
-		WakeUpEvent = FGenericPlatformProcess::GetSynchEventFromPool(true);
-		Thread = FRunnableThread::Create(this, TEXT("IoQueueThread"), 0, TPri_Normal);
-	}
-
-	~FIoQueueImpl()
-	{
-		Stop();
-		Thread->Kill(true);
-		Thread = nullptr;
-		FPlatformProcess::ReturnSynchEventToPool(WakeUpEvent);
-	}
-
-	void Enqueue(const FIoChunkId& ChunkId, FIoReadOptions ReadOptions, uint64 UserData, bool bDeferBatch)
-	{
-		FIoRequestImpl* Request = Dispatcher.AllocRequest(ChunkId, ReadOptions, UserData);
-		check(Request->NextRequest == nullptr);
-
-		{
-			// TODO: CAS
-			FScopeLock _(&QueuedLock);
-
-			Request->NextRequest	= FirstQueuedRequest;
-			FirstQueuedRequest		= Request;
-		}
-
-		if (!bDeferBatch)
-		{
-			IssueBatch();
-		}
-	
-		{
-			FScopeLock _(&NumPendingLock);
-
-			if (0 == NumPending++)
-			{
-				WakeUpEvent->Trigger();
-			}
-		}
-	}
-
-	bool Dequeue(FIoChunkId& ChunkId, TIoStatusOr<FIoBuffer>& Result, uint64& UserData)
-	{
-		FIoRequestImpl* CompletedRequest = nullptr;
-		{
-			// TODO: CAS
-			FScopeLock _(&CompletedLock);
-			
-			CompletedRequest		= FirstCompletedRequest;
-			FirstCompletedRequest	= CompletedRequest ? CompletedRequest->NextRequest : FirstCompletedRequest;
-		}
-
-		if (CompletedRequest)
-		{
-			ChunkId		= CompletedRequest->ChunkId;
-			Result		= CompletedRequest->Result;
-			UserData	= CompletedRequest->UserData;
-
-			Dispatcher.FreeRequest(CompletedRequest);
-
-			{
-				FScopeLock _(&NumPendingLock);
-
-				if (0 == --NumPending)
-				{
-					check(NumPending >= 0);
-					WakeUpEvent->Reset();
-				}
-			}
-
-			return true;
-		}
-
-		return false;
-	}
-
-	void IssueBatch()
-	{
-		FIoRequestImpl* QueuedRequests = nullptr;
-		{
-			// TODO: CAS
-			FScopeLock _(&QueuedLock);
-
-			QueuedRequests 		= FirstQueuedRequest;
-			FirstQueuedRequest	= nullptr;
-		}
-
-		if (QueuedRequests)
-		{
-			FIoBatchImpl* NewBatch = Dispatcher.AllocBatch(QueuedRequests);
-			{
-				// TODO: CAS
-				FScopeLock _(&PendingLock);
-
-				NewBatch->NextBatch	= FirstPendingBatch;
-				FirstPendingBatch	= NewBatch;
-			}
-		}
-	}
-
-	bool IsEmpty() const
-	{
-		FScopeLock _(&NumPendingLock);
-		return NumPending == 0;
-	}
-
-private:
-	
-	uint32 Run() override
-	{
-		bIsRunning = true;
-
-		FBatchQueue IssuedBatches;
-
-		while (bIsRunning)
-		{
-			if (!bIsRunning)
-			{
-				break;
-			}
-
-			// Dispatch pending batch
-			{
-				FIoBatchImpl* PendingBatch = nullptr;
-				{
-					// TODO: CAS
-					FScopeLock _(&PendingLock);
-
-					PendingBatch		= FirstPendingBatch;
-					FirstPendingBatch	= PendingBatch ? PendingBatch->NextBatch : nullptr;
-				}
-
-				if (PendingBatch)
-				{
-					PendingBatch->NextBatch = nullptr;
-					
-					Dispatcher.IssueBatch(PendingBatch);
-					Enqueue(PendingBatch, IssuedBatches);
-				}
-			}
-
-			// Process issued batches
-			if (FIoBatchImpl* IssuedBatch = Peek(IssuedBatches))
-			{
-				if (Dispatcher.IsBatchReady(IssuedBatch))
-				{
-					Dequeue(IssuedBatches);
-
-					FIoRequestImpl* Request = IssuedBatch->FirstRequest;
-					while (Request)
-					{
-						FIoRequestImpl* CompletedRequest	= Request;
-						Request								= Request->NextRequest;
-						
-						{
-							// TODO: CAS
-							FScopeLock _(&CompletedLock);
-
-							CompletedRequest->NextRequest	= FirstCompletedRequest;
-							FirstCompletedRequest			= CompletedRequest;
-						}
-					}
-
-					BatchReadyCallback();
-
-					IssuedBatch->FirstRequest = nullptr;
-					Dispatcher.FreeBatch(IssuedBatch);
-				}
-			}
-			else
-			{
-				WakeUpEvent->Wait();
-			}
-		}
-
-		return 0;
-	}
-
-	void Stop() override
-	{
-		if (bIsRunning)
-		{
-			bIsRunning = false;
-			WakeUpEvent->Trigger();
-		}
-	}
-
-	struct FBatchQueue
-	{
-		FIoBatchImpl* Head = nullptr;
-		FIoBatchImpl* Tail = nullptr;
-	};
-
-	void Enqueue(FIoBatchImpl* Batch, FBatchQueue& Batches)
-	{
-		if (Batches.Tail == nullptr)
-		{
-			Batches.Head = Batches.Tail	= Batch;
-		}
-		else
-		{
-			Batches.Tail->NextBatch	= Batch;
-			Batches.Tail			= Batch;
-		}
-	}
-
-	FIoBatchImpl* Dequeue(FBatchQueue& Batches)
-	{
-		if (FIoBatchImpl* Batch = Batches.Head)
-		{
-			Batches.Head		= Batch->NextBatch;
-			Batches.Tail		= Batches.Head == nullptr ? nullptr : Batches.Tail;
-			Batch->NextBatch	= nullptr;
-
-			return Batch;
-		}
-		
-		return nullptr;
-	}
-
-	FIoBatchImpl* Peek(FBatchQueue& Queue)
-	{
-		return Queue.Head;
-	}
-
-	FIoDispatcherImpl&				Dispatcher;
-	FIoQueue::FBatchReadyCallback	BatchReadyCallback;
-	FRunnableThread*				Thread					= nullptr;
-	FEvent*							WakeUpEvent				= nullptr;
-	bool							bIsRunning				= false;
-	FIoRequestImpl*					FirstQueuedRequest		= nullptr;
-	FIoBatchImpl*					FirstPendingBatch		= nullptr;
-	FIoRequestImpl*					FirstCompletedRequest	= nullptr;
-	int32							NumPending				= 0;
-	FCriticalSection				QueuedLock;
-	FCriticalSection				PendingLock;
-	FCriticalSection				CompletedLock;
-	mutable FCriticalSection		NumPendingLock;
-};
-
-//////////////////////////////////////////////////////////////////////////
-FIoQueue::FIoQueue(FIoDispatcher& IoDispatcher, FBatchReadyCallback BatchReadyCallback)
-: Impl(new FIoQueueImpl(*IoDispatcher.Impl, MoveTemp(BatchReadyCallback)))
-{ 
-}
-
-FIoQueue::~FIoQueue() = default;
-
-void
-FIoQueue::Enqueue(const FIoChunkId& ChunkId, FIoReadOptions ReadOptions, uint64 UserData, bool bDeferBatch)
-{
-	Impl->Enqueue(ChunkId, ReadOptions, UserData, bDeferBatch);
-}
-
-bool
-FIoQueue::Dequeue(FIoChunkId& ChunkId, TIoStatusOr<FIoBuffer>& Result, uint64& UserData)
-{
-	return Impl->Dequeue(ChunkId, Result, UserData);
-}
-
-void
-FIoQueue::IssueBatch()
-{
-	return Impl->IssueBatch();
-}
-
-bool
-FIoQueue::IsEmpty() const
-{
-	return Impl->IsEmpty();
-}
-
 #endif // PLATFORM_IMPLEMENTS_IO
diff --git a/Engine/Source/Runtime/Core/Public/IO/IoDispatcher.h b/Engine/Source/Runtime/Core/Public/IO/IoDispatcher.h
--- a/Engine/Source/Runtime/Core/Public/IO/IoDispatcher.h
+++ b/Engine/Source/Runtime/Core/Public/IO/IoDispatcher.h
@@ -25,7 +25,6 @@ class FIoDispatcher;
 class FIoStoreReader;
 class FIoStoreWriter;
 
-class FIoQueueImpl;
 class FIoRequestImpl;
 class FIoBatchImpl;
 class FIoDispatcherImpl;
@@ -637,9 +636,8 @@ public:
 
 	CORE_API bool							IsOk() const;
 	CORE_API FIoStatus						Status() const;
-	CORE_API FIoBuffer						GetChunk();
 	CORE_API const FIoChunkId&				GetChunkId() const;
-	CORE_API const TIoStatusOr<FIoBuffer>&	GetResult() const;
+	CORE_API TIoStatusOr<FIoBuffer>			GetResult() const;
 	CORE_API uint64							GetUserData() const;
 
 private:
@@ -695,6 +693,8 @@ public:
 	CORE_API FIoBatch				NewBatch();
 	CORE_API void					FreeBatch(FIoBatch Batch);
 
+	CORE_API void					ReadWithCallback(const FIoChunkId& Chunk, const FIoReadOptions& Options, TFunction<void(TIoStatusOr<FIoBuffer>)>&& Callback);
+
 	// Polling methods
 	CORE_API TIoStatusOr<uint64>	GetSizeForChunk(const FIoChunkId& ChunkId) const;
 
@@ -787,23 +787,3 @@ private:
 
 //////////////////////////////////////////////////////////////////////////
 
-class FIoQueue
-{
-public:
-	using FBatchReadyCallback = TFunction<void()>;
-
-	CORE_API		FIoQueue(FIoDispatcher& IoDispatcher, FBatchReadyCallback BatchReadyCallback);
-	CORE_API		~FIoQueue();
-
-					FIoQueue(const FIoQueue&) = delete;
-					FIoQueue(FIoQueue&&) = delete;
-					FIoQueue& operator=(const FIoQueue&) = delete;
-					FIoQueue& operator=(FIoQueue&&) = delete;
-
-	CORE_API void	Enqueue(const FIoChunkId& ChunkId, FIoReadOptions ReadOptions, uint64 UserData, bool bDeferBatch = false);
-	CORE_API bool	Dequeue(FIoChunkId& ChunkId, TIoStatusOr<FIoBuffer>& Result, uint64& UserData);
-	CORE_API void	IssueBatch();
-	CORE_API bool	IsEmpty() const;
-private:
-	TUniquePtr<FIoQueueImpl> Impl;
-};
diff --git a/Engine/Source/Runtime/CoreUObject/Private/Serialization/AsyncLoading2.cpp b/Engine/Source/Runtime/CoreUObject/Private/Serialization/AsyncLoading2.cpp
--- a/Engine/Source/Runtime/CoreUObject/Private/Serialization/AsyncLoading2.cpp
+++ b/Engine/Source/Runtime/CoreUObject/Private/Serialization/AsyncLoading2.cpp
@@ -125,6 +125,7 @@ struct FGlobalPackageId
 
 struct FExportMapEntry
 {
+	uint64 SerialSize;
 	int32 ObjectName[2];
 	FPackageIndex OuterIndex;
 	FPackageIndex ClassIndex;
@@ -134,6 +135,23 @@ struct FExportMapEntry
 	EObjectFlags ObjectFlags;
 };
 
+struct FExportBundleHeader
+{
+	uint32 FirstEntryIndex;
+	uint32 EntryCount;
+};
+
+struct FExportBundleEntry
+{
+	enum EExportBundleEntryType
+	{
+		ExportBundleEntryType_Create,
+		ExportBundleEntryType_Serialize
+	};
+	uint32 CommandType;
+	uint32 LocalExportIndex;
+};
+
 struct FGlobalImport
 {
 	FName ObjectName;
@@ -151,6 +169,7 @@ struct FPackageStoreEntry
 	int32 SlimportCount;
 	int32 SlimportOffset;
 	int32 ExportCount;
+	int32 ExportBundleCount;
 	int32 ScriptArcsOffset;
 	int32 ScriptArcsCount;
 };
@@ -159,6 +178,7 @@ struct FPackageSummary
 {
 	uint32 PackageFlags;
 	int32 ExportMapOffset;
+	int32 ExportBundlesOffset;
 	int32 GraphDataOffset;
 	int32 GraphDataSize;
 	int32 BulkDataStartOffset;
@@ -427,6 +447,11 @@ public:
 		return StoreEntries[GlobalPackageId.Id].ExportCount;
 	}
 
+	inline int32 GetPackageExportBundleCount(FGlobalPackageId GlobalPackageId) const
+	{
+		return StoreEntries[GlobalPackageId.Id].ExportBundleCount;
+	}
+
 	inline FString GetPackageFileName(FGlobalPackageId GlobalPackageId)
 	{
 		return StoreEntries[GlobalPackageId.Id].FileName.ToString();
@@ -630,7 +655,8 @@ enum class EChunkType : uint8
 	None,
 	PackageSummary,
 	ExportData,
-	BulkData
+	ExportBundleData,
+	BulkData,
 };
 
 FIoChunkId CreateChunkId(uint32 NameIndex, uint32 NameNumber, uint16 ChunkIndex, EChunkType ChunkType)
@@ -699,9 +725,9 @@ enum EEventLoadNode2 : uint8
 	Package_Delete,
 	Package_NumPhases,
 
-	Export_Create = 0,
-	Export_Serialize,
-	Export_NumPhases,
+	ExportBundle_StartIo = 0,
+	ExportBundle_Process,
+	ExportBundle_NumPhases,
 };
 
 class FEventLoadGraphAllocator;
@@ -730,7 +756,7 @@ public:
 
 private:
 	void ProcessDependencies(FAsyncLoadingThreadState2& ThreadState);
-	void Fire(FAsyncLoadingThreadState2& ThreadState);
+	void Fire();
 
 	union
 	{
@@ -827,23 +853,23 @@ struct FAsyncLoadEventSpec
 struct FAsyncLoadingThreadState2
 	: public FTlsAutoCleanup
 {
-	static FAsyncLoadingThreadState2* Create(FAsyncLoadEventGraphAllocator& GraphAllocator)
+	static FAsyncLoadingThreadState2* Create(FAsyncLoadEventGraphAllocator& GraphAllocator, FIoDispatcher& IoDispatcher)
 	{
 		check(TlsSlot != 0);
 		check(!FPlatformTLS::GetTlsValue(TlsSlot));
-		FAsyncLoadingThreadState2* State = new FAsyncLoadingThreadState2(GraphAllocator);
+		FAsyncLoadingThreadState2* State = new FAsyncLoadingThreadState2(GraphAllocator, IoDispatcher);
 		State->Register();
 		FPlatformTLS::SetTlsValue(TlsSlot, State);
 		return State;
 	}
 
-	static FAsyncLoadingThreadState2& Get()
+	static FAsyncLoadingThreadState2* Get()
 	{
 		check(TlsSlot != 0);
-		return *static_cast<FAsyncLoadingThreadState2*>(FPlatformTLS::GetTlsValue(TlsSlot));
+		return static_cast<FAsyncLoadingThreadState2*>(FPlatformTLS::GetTlsValue(TlsSlot));
 	}
 
-	FAsyncLoadingThreadState2(FAsyncLoadEventGraphAllocator& InGraphAllocator)
+	FAsyncLoadingThreadState2(FAsyncLoadEventGraphAllocator& InGraphAllocator, FIoDispatcher& InIoDispatcher)
 		: GraphAllocator(InGraphAllocator)
 	{
 
@@ -955,7 +981,7 @@ struct FProcessExternalIoTask
 
 			const int32 RequestIndex = IoRequest.GetUserData();	
 			const FExternalIoRequest& Request = ExternalIoRequests[RequestIndex];
-			Request.Callback(IoRequest.GetChunkId(), IoRequest.GetChunk());
+			Request.Callback(IoRequest.GetChunkId(), IoRequest.GetResult().ValueOrDie());
 
 			return true;
 		});
@@ -997,7 +1023,7 @@ struct FAsyncPackage2
 		check(RefCount > 0);
 		if (--RefCount == 0)
 		{
-			GetNode(EEventLoadNode2::Package_Delete)->ReleaseBarrier();
+			GetPackageNode(EEventLoadNode2::Package_Delete)->ReleaseBarrier();
 		}
 	}
 
@@ -1163,6 +1189,7 @@ struct FAsyncPackage2
 	EAsyncPackageState::Type CreateClusters();
 
 	void ImportPackagesRecursive();
+	void StartLoading();
 
 private:
 
@@ -1246,12 +1273,12 @@ private:
 	TSet<UPackage*> ImportedPackages;
 
 	FEventLoadNode2* PackageNodes = nullptr;
-	FEventLoadNode2* ExportNodes = nullptr;
-	uint32 ExportNodeCount = 0;
+	FEventLoadNode2* ExportBundleNodes = nullptr;
+	uint32 ExportBundleNodeCount = 0;
 
 	uint64 PackageSummarySize = 0;
 	FIoBuffer PackageSummaryIoBuffer;
-	TArray<FIoBuffer> ExportIoBuffers;
+	TArray<FIoBuffer> ExportBundleIoBuffers;
 
 	// FZenLinkerLoad
 	TArray<FExternalReadCallback> ExternalReadDependencies;
@@ -1260,6 +1287,10 @@ private:
 	const FExportMapEntry* ExportMap = nullptr;
 	TArray<UObject*> Exports;
 	FPackageImportStore ImportStore;
+
+	int32 ExportBundleCount = 0;
+	const FExportBundleHeader* ExportBundles = nullptr;
+	const FExportBundleEntry* ExportBundleEntries = nullptr;
 public:
 
 	FAsyncLoadingThread2Impl& GetAsyncLoadingThread()
@@ -1283,16 +1314,14 @@ public:
 
 	static EAsyncPackageState::Type Event_CreateLinker(FAsyncPackage2* Package, int32);
 	static EAsyncPackageState::Type Event_ProcessSummary(FAsyncPackage2* Package, int32);
-	static EAsyncPackageState::Type Event_CreateExport(FAsyncPackage2* Package, int32 LocalExportIndex);
-	static EAsyncPackageState::Type Event_SerializeExport(FAsyncPackage2* Package, int32 LocalExportIndex);
+	static EAsyncPackageState::Type Event_StartExportBundleIo(FAsyncPackage2* Package, int32 ExportBundleIndex);
+	static EAsyncPackageState::Type Event_ProcessExportBundle(FAsyncPackage2* Package, int32 ExportBundleIndex);
 	static EAsyncPackageState::Type Event_ExportsDone(FAsyncPackage2* Package, int32);
 	static EAsyncPackageState::Type Event_Tick(FAsyncPackage2* Package, int32);
 	static EAsyncPackageState::Type Event_Delete(FAsyncPackage2* Package, int32);
 
-	void ProcessIoRequest(const FIoChunkId& ChunkId, TIoStatusOr<FIoBuffer>& Result);
-
 	void EventDrivenCreateExport(int32 LocalExportIndex);
-	void EventDrivenSerializeExport(int32 LocalExportIndex);
+	void EventDrivenSerializeExport(int32 LocalExportIndex, const uint8* ExportData, uint64 ExportDataSize);
 
 	UObject* EventDrivenIndexToObject(FPackageIndex Index, bool bCheckSerialized);
 	template<class T>
@@ -1306,8 +1335,8 @@ public:
 		return CastChecked<T>(Result);
 	}
 
-	FEventLoadNode2* GetNode(EEventLoadNode2 Phase, FPackageIndex ImportOrExportIndex = FPackageIndex());
-	
+	FEventLoadNode2* GetPackageNode(EEventLoadNode2 Phase);
+	FEventLoadNode2* GetExportBundleNode(EEventLoadNode2 Phase, uint32 ExportBundleIndex);
 	FEventLoadNode2* GetNode(int32 NodeIndex);
 
 	/** [EDL] End Event driven loader specific stuff */
@@ -1398,10 +1427,11 @@ struct FScopedAsyncPackageEvent2
 class FAsyncLoadingThreadWorker : private FRunnable
 {
 public:
-	FAsyncLoadingThreadWorker(FAsyncLoadEventGraphAllocator& InGraphAllocator, FAsyncLoadEventQueue2& InEventQueue, FZenaphore& InZenaphore, TAtomic<int32>& InActiveWorkersCount)
+	FAsyncLoadingThreadWorker(FAsyncLoadEventGraphAllocator& InGraphAllocator, FAsyncLoadEventQueue2& InEventQueue, FIoDispatcher& InIoDispatcher, FZenaphore& InZenaphore, TAtomic<int32>& InActiveWorkersCount)
 		: Zenaphore(InZenaphore)
 		, EventQueue(InEventQueue)
 		, GraphAllocator(InGraphAllocator)
+		, IoDispatcher(InIoDispatcher)
 		, ActiveWorkersCount(InActiveWorkersCount)
 	{
 	}
@@ -1439,6 +1469,7 @@ private:
 	FZenaphore& Zenaphore;
 	FAsyncLoadEventQueue2& EventQueue;
 	FAsyncLoadEventGraphAllocator& GraphAllocator;
+	FIoDispatcher& IoDispatcher;
 	TAtomic<int32>& ActiveWorkersCount;
 	FRunnableThread* Thread = nullptr;
 	TAtomic<bool> bStopRequested { false };
@@ -1517,7 +1548,6 @@ private:
 	/** I/O Dispatcher */
 	FGlobalNameMap GlobalNameMap;
 	FIoDispatcher& IoDispatcher;
-	TUniquePtr<FIoQueue> IoQueue;
 
 	FPackageStore PackageStore;
 public:
@@ -1541,14 +1571,10 @@ public:
 	FAsyncLoadEventGraphAllocator GraphAllocator;
 	FAsyncLoadEventQueue2 EventQueue;
 	FAsyncLoadEventQueue2 AsyncEventQueue;
-	FAsyncLoadEventQueue2 CreateExportsEventQueue;
-	FAsyncLoadEventQueue2 SerializeExportsEventQueue;
+	FAsyncLoadEventQueue2 ProcessExportBundlesEventQueue;
 	TArray<FAsyncLoadEventQueue2*> AltEventQueues;
 	TArray<FAsyncLoadEventSpec> EventSpecs;
 
-	/** [EDL] Queues CreateLinker event */
-	void QueueEvent_CreateLinker(FAsyncPackage2* Pkg);
-
 	/** True if multithreaded async loading is currently being used. */
 	bool IsMultithreaded()
 	{
@@ -1789,9 +1815,6 @@ private:
 
 	/** Number of times we re-entered the async loading tick, mostly used by singlethreaded ticking. Debug purposes only. */
 	int32 AsyncLoadingTickCounter;
-
-	/** Enqueue I/O request */
-	void EnqueueIoRequest(FAsyncPackage2* Package, const FIoChunkId& ChunkId);
 };
 
 /**
@@ -1867,11 +1890,6 @@ void FAsyncLoadingThread2Impl::InitializeLoading()
 
 		IoDispatcher.Mount(IoStoreReader.Release());
 
-		IoQueue = MakeUnique<FIoQueue>(IoDispatcher, TFunction<void()>([this]()
-		{
-			AltZenaphore.NotifyOne();
-		}));
-
 #if USE_NEW_BULKDATA
 		FBulkDataBase::SetIODispatcher(&IoDispatcher);
 #endif
@@ -1929,10 +1947,6 @@ FAsyncPackage2* FAsyncLoadingThread2Impl::FindOrInsertPackage(FAsyncPackageDesc*
 			Package->AddCompletionCallback(MoveTemp(InDesc->PackageLoadedDelegate), bInternalCallback);
 		}
 	}
-	if (bInserted)
-	{
-		QueueEvent_CreateLinker(Package);
-	}
 	return Package;
 }
 
@@ -1955,6 +1969,10 @@ bool FAsyncLoadingThread2Impl::CreateAsyncPackagesFromQueue()
 			TRACE_CPUPROFILER_EVENT_SCOPE(ImportPackages);
 			Package->ImportPackagesRecursive();
 		}
+		if (bInserted)
+		{
+			Package->StartLoading();
+		}
 		delete PackageRequest;
 	}
 
@@ -2035,11 +2053,11 @@ void FEventLoadNode2::ReleaseBarrier()
 	check(BarrierCount > 0);
 	if (--BarrierCount == 0)
 	{
-		Fire(FAsyncLoadingThreadState2::Get());
+		Fire();
 	}
 }
 
-void FEventLoadNode2::Fire(FAsyncLoadingThreadState2& ThreadState)
+void FEventLoadNode2::Fire()
 {
 	//TRACE_CPUPROFILER_EVENT_SCOPE(Fire);
 
@@ -2047,9 +2065,10 @@ void FEventLoadNode2::Fire(FAsyncLoadingThreadState2& ThreadState)
 	bFired.Store(1);
 #endif
 
-	if (Spec->bExecuteImmediately)
+	FAsyncLoadingThreadState2* ThreadState = FAsyncLoadingThreadState2::Get();
+	if (Spec->bExecuteImmediately && ThreadState)
 	{
-		Execute(ThreadState);
+		Execute(*ThreadState);
 	}
 	else
 	{
@@ -2107,7 +2126,7 @@ void FEventLoadNode2::ProcessDependencies(FAsyncLoadingThreadState2& ThreadState
 		ThreadState.bShouldFireNodes = false;
 		while (ThreadState.NodesToFire.Num())
 		{
-			ThreadState.NodesToFire.Pop(false)->Fire(ThreadState);
+			ThreadState.NodesToFire.Pop(false)->Fire();
 		}
 		ThreadState.bShouldFireNodes = true;
 	}
@@ -2209,11 +2228,11 @@ uint32 FAsyncLoadingThreadWorker::Run()
 		FPlatformProcess::SetThreadAffinityMask(FPlatformAffinity::GetAsyncLoadingThreadMask());
 	}
 
-	FAsyncLoadingThreadState2::Create(GraphAllocator);
+	FAsyncLoadingThreadState2::Create(GraphAllocator, IoDispatcher);
 
 	FZenaphoreWaiter Waiter(Zenaphore, TEXT("WaitForEvents"));
 
-	FAsyncLoadingThreadState2& ThreadState = FAsyncLoadingThreadState2::Get();
+	FAsyncLoadingThreadState2& ThreadState = *FAsyncLoadingThreadState2::Get();
 
 	bool bSuspended = false;
 	while (!bStopRequested)
@@ -2259,13 +2278,6 @@ uint32 FAsyncLoadingThreadWorker::Run()
 	return 0;
 }
 
-void FAsyncLoadingThread2Impl::QueueEvent_CreateLinker(FAsyncPackage2* Package)
-{
-	check(Package);
-	TRACE_LOADTIME_BEGIN_LOAD_ASYNC_PACKAGE(Package);
-	Package->GetNode(EEventLoadNode2::Package_CreateLinker)->ReleaseBarrier();
-}
-
 FUObjectSerializeContext* FAsyncPackage2::GetSerializeContext()
 {
 	return FUObjectThreadContext::Get().GetSerializeContext();
@@ -2321,9 +2333,10 @@ void FAsyncPackage2::StartImportPackages()
 		TRACE_CPUPROFILER_EVENT_SCOPE(SetupSerializedArcs);
 		const FGlobalNameMap& GlobalNameMap = AsyncLoadingThread.GlobalNameMap;
 
-		const FPackageSummary* PackageSummary = reinterpret_cast<const FPackageSummary*>(PackageSummaryIoBuffer.Data());
+		const uint8* PackageSummaryData = PackageSummaryIoBuffer.Data();
+		const FPackageSummary* PackageSummary = reinterpret_cast<const FPackageSummary*>(PackageSummaryData);
 		uint64 ZenHeaderDataSize = PackageSummary->GraphDataSize;
-		const uint8* ZenHeaderData = PackageSummaryIoBuffer.Data() + PackageSummary->GraphDataOffset;
+		const uint8* ZenHeaderData = PackageSummaryData + PackageSummary->GraphDataOffset;
 		FSimpleArchive ZenHeaderArchive(ZenHeaderData, ZenHeaderDataSize);
 		int32 InternalArcCount;
 		ZenHeaderArchive << InternalArcCount;
@@ -2411,7 +2424,6 @@ void FAsyncPackage2::ImportPackagesRecursive()
 	}
 	bHasImportedPackagesRecursive = true;
 
-#if 1
 	FMemStack& MemStack = FMemStack::Get();
 	FMemMark Mark(MemStack);
 	TArrayView<FPackageIndex> LocalImportedPackages(new(MemStack)FPackageIndex[ImportStore.Count], ImportStore.Count);
@@ -2454,83 +2466,25 @@ void FAsyncPackage2::ImportPackagesRecursive()
 		if (bInserted)
 		{
 			ImportedAsyncPackage->ImportPackagesRecursive();
+			ImportedAsyncPackage->StartLoading();
 		}
 	}
-#else
-	bool bNeedToLoadPackage = false;
-	bool bNeedToCheckImportObjectsInPackage = true;
-	bool bHasAddedNeedToLoadPackage = false;
-	FPackageIndex LastImportedPackageIndex;
-
-	// slimports are sorted, first comes a new package, then N additional objects
-	for (int32 LocalImportIndex = 0; LocalImportIndex < ImportCount; ++LocalImportIndex)
-	{
-		const int32 GlobalImportIndex = ImportMap[LocalImportIndex];
-		const FPackageIndex ImportedPackageIndex = GlobalImportPackages[GlobalImportIndex];
-
-		const bool bIsANewImportPackage = ImportedPackageIndex != LastImportedPackageIndex;
-		if (bIsANewImportPackage)
-		{
-			LastImportedPackageIndex = ImportedPackageIndex;
-			bHasAddedNeedToLoadPackage = false;
+}
 
-			UPackage* ImportedPackage = (UPackage*)FindExistingSlimport(ImportedPackageIndex.ToImport());
-			if (ImportedPackage)
-			{
-				bNeedToLoadPackage = false;
-				if (ImportedPackage->HasAnyPackageFlags(PKG_CompiledIn))
-				{
-					bNeedToCheckImportObjectsInPackage = false;
-				}
-				else
-				{
-					AddObjectReference(ImportedPackage);
-					bNeedToCheckImportObjectsInPackage = true;
-				}
-			}
-			else
-			{
-				bNeedToLoadPackage = true;
-				bNeedToCheckImportObjectsInPackage = false;
-			}
-		}
-		else if (bNeedToCheckImportObjectsInPackage)
-		{
-			UObject* ImportedObject = FindExistingSlimport(GlobalImportIndex);
-			if (ImportedObject)
-			{
-				AddObjectReference(ImportedObject);
-				bNeedToLoadPackage |= !IsFullyLoadedObj(ImportedObject);
-			}
-			else
-			{
-				bNeedToLoadPackage = true;
-			}
-		}
+void FAsyncPackage2::StartLoading()
+{
+	TRACE_LOADTIME_BEGIN_LOAD_ASYNC_PACKAGE(this);
 
-		if (bNeedToLoadPackage && !bHasAddedNeedToLoadPackage)
+	FIoReadOptions ReadOptions;
+	AsyncLoadingThread.IoDispatcher.ReadWithCallback(CreateChunkId(PackageChunkId, 0, EChunkType::PackageSummary),
+		ReadOptions,
+		[this](TIoStatusOr<FIoBuffer> Result)
 		{
-			bHasAddedNeedToLoadPackage = true;
-
-			const FName ImportedPackageName = GlobalImportNames[ImportedPackageIndex.ToImport()];
-			FAsyncPackageDesc Info(INDEX_NONE, ImportedPackageName);
-			Info.Priority = Desc.Priority;
-			bool bInserted;
-			FAsyncPackage2* ImportedAsyncPackage = AsyncLoadingThread.FindOrInsertPackage(&Info, bInserted);
-			TRACE_LOADTIME_ASYNC_PACKAGE_IMPORT_DEPENDENCY(this, ImportedAsyncPackage);
-			ImportedAsyncPackage->AddRef();
-			ImportedAsyncPackages.Add(ImportedAsyncPackage);
-			if (bInserted)
-			{
-				ImportedAsyncPackage->ImportPackagesRecursive();
-			}
-
-			FEventLoadNode2* SetupExportsNode = GetNode(EEventLoadNode2::Package_SetupExports);
-			FEventLoadNode2* ImportedPackageLoadSummaryNode = ImportedAsyncPackage->GetNode(EEventLoadNode2::Package_ProcessSummary);
-			SetupExportsNode->DependsOn(ImportedPackageLoadSummaryNode);
-		}
-	}
-#endif
+			PackageSummaryIoBuffer = Result.ConsumeValueOrDie();
+			PackageSummarySize = PackageSummaryIoBuffer.DataSize();
+			GetPackageNode(EEventLoadNode2::Package_ProcessSummary)->ReleaseBarrier();
+		});
+	GetPackageNode(EEventLoadNode2::Package_CreateLinker)->ReleaseBarrier();
 }
 
 void FAsyncPackage2::SetupScriptArcs()
@@ -2597,31 +2551,63 @@ void FAsyncPackage2::SetupScriptArcs()
 	}
 }
 
-EAsyncPackageState::Type FAsyncPackage2::Event_CreateExport(FAsyncPackage2* Package, int32 LocalExportIndex)
+EAsyncPackageState::Type FAsyncPackage2::Event_StartExportBundleIo(FAsyncPackage2* Package, int32 ExportBundleIndex)
 {
-	TRACE_CPUPROFILER_EVENT_SCOPE(Event_CreateExport);
+	TRACE_CPUPROFILER_EVENT_SCOPE(Event_StartExportBundleIo);
 
-	check(Package->AsyncPackageLoadingState == EAsyncPackageLoadingState2::ProcessNewImportsAndExports);
-
-	FScopedAsyncPackageEvent2 Scope(Package);
-
-	Package->EventDrivenCreateExport(LocalExportIndex);
+	FIoReadOptions ReadOptions;
+	Package->AsyncLoadingThread.IoDispatcher.ReadWithCallback(CreateChunkId(Package->PackageChunkId, ExportBundleIndex, EChunkType::ExportBundleData),
+		ReadOptions,
+		[Package, ExportBundleIndex](TIoStatusOr<FIoBuffer> Result)
+		{
+			Package->ExportBundleIoBuffers[ExportBundleIndex] = Result.ConsumeValueOrDie();
+			Package->GetExportBundleNode(EEventLoadNode2::ExportBundle_Process, ExportBundleIndex)->ReleaseBarrier();
+		});
 	return EAsyncPackageState::Complete;
 }
 
-EAsyncPackageState::Type FAsyncPackage2::Event_SerializeExport(FAsyncPackage2* Package, int32 LocalExportIndex)
+EAsyncPackageState::Type FAsyncPackage2::Event_ProcessExportBundle(FAsyncPackage2* Package, int32 ExportBundleIndex)
 {
-	TRACE_CPUPROFILER_EVENT_SCOPE(Event_SerializeExport);
+	TRACE_CPUPROFILER_EVENT_SCOPE(Event_ProcessExportBundle);
+
+	check(Package->AsyncPackageLoadingState == EAsyncPackageLoadingState2::ProcessNewImportsAndExports);
+
+	FIoBuffer& ExportBundleIoBuffer = Package->ExportBundleIoBuffers[ExportBundleIndex];
+	const uint8* ExportData = ExportBundleIoBuffer.Data();
 
 	FScopedAsyncPackageEvent2 Scope(Package);
 
-	UObject* Object = Package->Exports[LocalExportIndex];;
-	check(Object);
-	if (Object->HasAnyFlags(RF_NeedLoad))
+	check(ExportBundleIndex < Package->ExportBundleCount);
+	const FExportBundleHeader* ExportBundle = Package->ExportBundles + ExportBundleIndex;
+	
+	const FExportBundleEntry* BundleEntry = Package->ExportBundleEntries + ExportBundle->FirstEntryIndex;
+	const FExportBundleEntry* BundleEntryEnd = BundleEntry + ExportBundle->EntryCount;
+	check(BundleEntry <= BundleEntryEnd);
+	while (BundleEntry < BundleEntryEnd)
 	{
-		Package->EventDrivenSerializeExport(LocalExportIndex);
+		if (BundleEntry->CommandType == FExportBundleEntry::ExportBundleEntryType_Create)
+		{
+			Package->EventDrivenCreateExport(BundleEntry->LocalExportIndex);
+		}
+		else
+		{
+			check(BundleEntry->CommandType == FExportBundleEntry::ExportBundleEntryType_Serialize);
+			uint64 ExportSerialSize = Package->ExportMap[BundleEntry->LocalExportIndex].SerialSize;
+			check(ExportData + ExportSerialSize <= ExportBundleIoBuffer.Data() + ExportBundleIoBuffer.DataSize());
+			UObject* Object = Package->Exports[BundleEntry->LocalExportIndex];
+			check(Object);
+			if (Object->HasAnyFlags(RF_NeedLoad))
+			{
+				Package->EventDrivenSerializeExport(BundleEntry->LocalExportIndex, ExportData, ExportSerialSize);
+			}
+			check(!Object->HasAnyFlags(RF_NeedLoad));
+
+			ExportData += ExportSerialSize;
+		}
+		++BundleEntry;
 	}
-	check(!Object->HasAnyFlags(RF_NeedLoad));
+	ExportBundleIoBuffer = FIoBuffer();
+
 	return EAsyncPackageState::Complete;
 }
 
@@ -2643,7 +2629,7 @@ UObject* FAsyncPackage2::EventDrivenIndexToObject(FPackageIndex Index, bool bChe
 	}
 	if (bCheckSerialized && !IsFullyLoadedObj(Result))
 	{
-		FEventLoadNode2* MyDependentNode = GetNode(EEventLoadNode2::Export_Serialize, Index);
+		/*FEventLoadNode2* MyDependentNode = GetExportNode(EEventLoadNode2::Export_Serialize, Index.ToExport());
 		if (!Result)
 		{
 			UE_LOG(LogStreaming, Error, TEXT("Missing Dependency, request for %s but it hasn't been created yet."), *Linker->GetPathName(Index));
@@ -2655,7 +2641,8 @@ UObject* FAsyncPackage2::EventDrivenIndexToObject(FPackageIndex Index, bool bChe
 		else
 		{
 			UE_LOG(LogStreaming, Fatal, TEXT("Missing Dependency, request for %s but it was still has RF_NeedLoad."), *Linker->GetPathName(Index));
-		}
+		}*/
+		UE_LOG(LogStreaming, Fatal, TEXT("Missing Dependency"), *Linker->GetPathName(Index));
 	}
 	if (Result)
 	{
@@ -2946,19 +2933,9 @@ void FAsyncPackage2::EventDrivenCreateExport(int32 LocalExportIndex)
 			check(StoreObject == Object);
 		}
 	}
-
-	if (bIsCompleteyLoaded) 
-	{
-		GetNode(EEventLoadNode2::Export_Serialize, FPackageIndex::FromExport(LocalExportIndex))->ReleaseBarrier();
-	}
-	else
-	{
-		TRACE_CPUPROFILER_EVENT_SCOPE(StartIO);
-		AsyncLoadingThread.EnqueueIoRequest(this, CreateChunkId(PackageChunkId, LocalExportIndex, EChunkType::ExportData));
-	}
 }
 
-void FAsyncPackage2::EventDrivenSerializeExport(int32 LocalExportIndex)
+void FAsyncPackage2::EventDrivenSerializeExport(int32 LocalExportIndex, const uint8* ExportData, uint64 ExportDataSize)
 {
 	//SCOPED_LOADTIMER(Package_PreLoadObjects);
 
@@ -3016,7 +2993,7 @@ void FAsyncPackage2::EventDrivenSerializeExport(int32 LocalExportIndex)
 		const FPackageFileSummary& Summary = Linker->Summary;
 		const FCustomVersionContainer& SummaryVersions = Summary.GetCustomVersionContainer();
 
-		FSimpleExportArchive Ar(ExportIoBuffers[LocalExportIndex].Data(), ExportIoBuffers[LocalExportIndex].DataSize());
+		FSimpleExportArchive Ar(ExportData, ExportDataSize);
 		Ar.SetUE4Ver(Summary.GetFileVersionUE4());
 		Ar.SetLicenseeUE4Ver(Summary.GetFileVersionLicenseeUE4());
 		Ar.SetEngineVer(Summary.SavedByEngineVersion);
@@ -3037,7 +3014,7 @@ void FAsyncPackage2::EventDrivenSerializeExport(int32 LocalExportIndex)
 
 		Object->ClearFlags(RF_NeedLoad);
 
-		TRACE_LOADTIME_SERIALIZE_EXPORT_SCOPE(Object, ExportIoBuffers[LocalExportIndex].DataSize());
+		TRACE_LOADTIME_SERIALIZE_EXPORT_SCOPE(Object, ExportDataSize);
 
 		FUObjectSerializeContext* LoadContext = GetSerializeContext();
 		UObject* PrevSerializedObject = LoadContext->SerializedObject;
@@ -3147,7 +3124,7 @@ EAsyncPackageState::Type FAsyncPackage2::Event_Tick(FAsyncPackage2* Package, int
 		{
 			LoadingState = Package->FinishObjects();
 		}
-	} while (!FAsyncLoadingThreadState2::Get().IsTimeLimitExceeded() && LoadingState == EAsyncPackageState::TimeOut);
+	} while (!FAsyncLoadingThreadState2::Get()->IsTimeLimitExceeded() && LoadingState == EAsyncPackageState::TimeOut);
 
 	if (Package->LinkerRoot && LoadingState == EAsyncPackageState::Complete)
 	{
@@ -3183,48 +3160,16 @@ EAsyncPackageState::Type FAsyncPackage2::Event_Delete(FAsyncPackage2* Package, i
 	return EAsyncPackageState::Complete;
 }
 
-void FAsyncPackage2::ProcessIoRequest(const FIoChunkId& ChunkId, TIoStatusOr<FIoBuffer>& Result)
-{
-	UE_CLOG(!Result.IsOk(), LogStreaming, Warning, TEXT("I/O Error: '%s', package: '%s'"), *Result.Status().ToString(), *Desc.NameToLoad.ToString()); 
-	
-	const FIoBuffer& IoBuffer = Result.ValueOrDie();
 
-	switch (GetChunkType(ChunkId))
-	{
-	case EChunkType::PackageSummary:
-	{
-		PackageSummarySize = IoBuffer.DataSize();
-		PackageSummaryIoBuffer = IoBuffer;
-		GetNode(EEventLoadNode2::Package_ProcessSummary)->ReleaseBarrier();
-		break;
-	}
-	case EChunkType::ExportData:
-	{
-		const int32 LocalExportIndex = GetChunkIndex(ChunkId);
-		check(ExportIoBuffers.Num() > LocalExportIndex);
-		ExportIoBuffers[LocalExportIndex] = IoBuffer;
-		GetNode(EEventLoadNode2::Export_Serialize, FPackageIndex::FromExport(LocalExportIndex))->ReleaseBarrier();
-		break;
-	}
-	}
+FEventLoadNode2* FAsyncPackage2::GetPackageNode(EEventLoadNode2 Phase)
+{
+	return PackageNodes + Phase;
 }
 
-FEventLoadNode2* FAsyncPackage2::GetNode(EEventLoadNode2 Phase, FPackageIndex ImportOrExportIndex)
+FEventLoadNode2* FAsyncPackage2::GetExportBundleNode(EEventLoadNode2 Phase, uint32 ExportBundleIndex)
 {
-	if (ImportOrExportIndex.IsNull())
-	{
-		return PackageNodes + Phase;
-	}
-	else if (ImportOrExportIndex.IsImport())
-	{
-		check(false);
-		return 0;
-	}
-	else
-	{
-		int32 ExportNodeIndex = ImportOrExportIndex.ToExport() * EEventLoadNode2::Export_NumPhases + Phase;
-		return ExportNodes + ExportNodeIndex;
-	}
+	uint32 ExportBundleNodeIndex = ExportBundleIndex * EEventLoadNode2::ExportBundle_NumPhases + Phase;
+	return ExportBundleNodes + ExportBundleNodeIndex;
 }
 
 FEventLoadNode2* FAsyncPackage2::GetNode(int32 NodeIndex)
@@ -3252,7 +3197,7 @@ EAsyncPackageState::Type FAsyncLoadingThread2Impl::ProcessAsyncLoadingFromGameTh
 	FAsyncLoadingTickScope2 InAsyncLoadingTick(*this);
 	uint32 LoopIterations = 0;
 
-	FAsyncLoadingThreadState2& ThreadState = FAsyncLoadingThreadState2::Get();
+	FAsyncLoadingThreadState2& ThreadState = *FAsyncLoadingThreadState2::Get();
 
 	while (true)
 	{
@@ -3283,23 +3228,6 @@ EAsyncPackageState::Type FAsyncLoadingThread2Impl::ProcessAsyncLoadingFromGameTh
 				break;
 			}
 
-			bool bProcessedResourceCompletionEvents = false;
-			FIoChunkId ChunkId;
-			TIoStatusOr<FIoBuffer> Result;
-			uint64 UserData;
-
-			while(IoQueue->Dequeue(ChunkId, Result, UserData))
-			{
-				FAsyncPackage2* Package = reinterpret_cast<FAsyncPackage2*>(UserData);
-				Package->ProcessIoRequest(ChunkId, Result);
-				bProcessedResourceCompletionEvents = true;
-			}
-			if (bProcessedResourceCompletionEvents)
-			{
-				OutPackagesProcessed++;
-				break;
-			}
-
 			bool bPopped = false;
 			for (FAsyncLoadEventQueue2* Queue : AltEventQueues)
 			{
@@ -3315,14 +3243,7 @@ EAsyncPackageState::Type FAsyncLoadingThread2Impl::ProcessAsyncLoadingFromGameTh
 				break;
 			}
 
-			if (!IoQueue->IsEmpty())
-			{
-				FPlatformProcess::Sleep(0.001f);
-			}
-			else
-			{
-				return EAsyncPackageState::Complete;
-			}
+			return EAsyncPackageState::Complete;
 		} while (false);
 	}
 	check(false);
@@ -3406,7 +3327,7 @@ EAsyncPackageState::Type FAsyncLoadingThread2Impl::ProcessLoadedPackagesFromGame
 		// The async loading thread might have queued some render thread tasks (we don't have a render thread yet, so these are actually sent to the game thread)
 		// We need to process them now before we do any postloads.
 		FTaskGraphInterface::Get().ProcessThreadUntilIdle(ENamedThreads::GameThread);
-		if (FAsyncLoadingThreadState2::Get().IsTimeLimitExceeded())
+		if (FAsyncLoadingThreadState2::Get()->IsTimeLimitExceeded())
 		{
 			return EAsyncPackageState::TimeOut;
 		}
@@ -3573,7 +3494,7 @@ EAsyncPackageState::Type FAsyncLoadingThread2Impl::TickAsyncLoadingFromGameThrea
 
 	if (!bLoadingSuspended)
 	{
-		FAsyncLoadingThreadState2::Get().SetTimeLimit(bUseTimeLimit, TimeLimit);
+		FAsyncLoadingThreadState2::Get()->SetTimeLimit(bUseTimeLimit, TimeLimit);
 
 		// First make sure there's no objects pending to be unhashed. This is important in uncooked builds since we don't 
 		// detach linkers immediately there and we may end up in getting unreachable objects from Linkers in CreateImports
@@ -3651,21 +3572,23 @@ FAsyncLoadingThread2Impl::FAsyncLoadingThread2Impl(IEDLBootNotificationManager&
 	FLoadTimeProfilerTracePrivate::Init();
 #endif
 
-	AltEventQueues.Add(&SerializeExportsEventQueue);
+	AltEventQueues.Add(&ProcessExportBundlesEventQueue);
 	AltEventQueues.Add(&AsyncEventQueue);
 	AltEventQueues.Add(&EventQueue);
-	AltEventQueues.Add(&CreateExportsEventQueue);
+	for (FAsyncLoadEventQueue2* Queue : AltEventQueues)
+	{
+		Queue->SetZenaphore(&AltZenaphore);
+	}
 
-	EventSpecs.AddDefaulted(EEventLoadNode2::Package_NumPhases + EEventLoadNode2::Export_NumPhases);
+	EventSpecs.AddDefaulted(EEventLoadNode2::Package_NumPhases + EEventLoadNode2::ExportBundle_NumPhases);
 	EventSpecs[EEventLoadNode2::Package_CreateLinker] = { &FAsyncPackage2::Event_CreateLinker, &AsyncEventQueue, false };
 	EventSpecs[EEventLoadNode2::Package_ProcessSummary] = { &FAsyncPackage2::Event_ProcessSummary, &AsyncEventQueue, false };
 	EventSpecs[EEventLoadNode2::Package_ExportsSerialized] = { &FAsyncPackage2::Event_ExportsDone, &AsyncEventQueue, true };
 	EventSpecs[EEventLoadNode2::Package_Tick] = { &FAsyncPackage2::Event_Tick, &EventQueue, false };
 	EventSpecs[EEventLoadNode2::Package_Delete] = { &FAsyncPackage2::Event_Delete, &AsyncEventQueue, false };
 
-
-	EventSpecs[EEventLoadNode2::Package_NumPhases + EEventLoadNode2::Export_Create] = { &FAsyncPackage2::Event_CreateExport, &CreateExportsEventQueue, false };
-	EventSpecs[EEventLoadNode2::Package_NumPhases + EEventLoadNode2::Export_Serialize] = { &FAsyncPackage2::Event_SerializeExport, &SerializeExportsEventQueue, false };
+	EventSpecs[EEventLoadNode2::Package_NumPhases + EEventLoadNode2::ExportBundle_StartIo] = { &FAsyncPackage2::Event_StartExportBundleIo, &AsyncEventQueue, false };
+	EventSpecs[EEventLoadNode2::Package_NumPhases + EEventLoadNode2::ExportBundle_Process] = { &FAsyncPackage2::Event_ProcessExportBundle, &ProcessExportBundlesEventQueue, false };
 
 	CancelLoadingEvent = FPlatformProcess::GetSynchEventFromPool();
 	ThreadSuspendedEvent = FPlatformProcess::GetSynchEventFromPool();
@@ -3673,7 +3596,7 @@ FAsyncLoadingThread2Impl::FAsyncLoadingThread2Impl(IEDLBootNotificationManager&
 	AsyncLoadingTickCounter = 0;
 
 	FAsyncLoadingThreadState2::TlsSlot = FPlatformTLS::AllocTlsSlot();
-	FAsyncLoadingThreadState2::Create(GraphAllocator);
+	FAsyncLoadingThreadState2::Create(GraphAllocator, IoDispatcher);
 }
 
 FAsyncLoadingThread2Impl::~FAsyncLoadingThread2Impl()
@@ -3717,41 +3640,25 @@ void FAsyncLoadingThread2Impl::StartThread()
 		UE_LOG(LogStreaming, Log, TEXT("Starting Async Loading Thread."));
 		bThreadStarted = true;
 		FPlatformMisc::MemoryBarrier();
-		Thread = FRunnableThread::Create(this, TEXT("FAsyncLoadingThread"), 0, TPri_Normal);
-		if (Thread)
-		{
-			TRACE_SET_THREAD_GROUP(Thread->GetThreadID(), "AsyncLoading");
-		}
 		
-		int32 WorkerCount = 3;
+		int32 WorkerCount = 0;
 		FParse::Value(FCommandLine::Get(), TEXT("-zenworkercount="), WorkerCount);
 		
 		if (WorkerCount > 0)
 		{
-			for (FAsyncLoadEventQueue2* Queue : AltEventQueues)
-			{
-				Queue->SetZenaphore(&AltZenaphore);
-			}
-
 			WorkerZenaphores.AddDefaulted(FMath::Max(3, WorkerCount));
 			Workers.Reserve(WorkerCount);
 			for (int32 WorkerIndex = 0; WorkerIndex < WorkerCount; ++WorkerIndex)
 			{
 				if (WorkerIndex == 0)
 				{
-					Workers.Emplace(GraphAllocator, SerializeExportsEventQueue, WorkerZenaphores[0], ActiveWorkersCount);
-					SerializeExportsEventQueue.SetZenaphore(&WorkerZenaphores[0]);
-					AltEventQueues.Remove(&SerializeExportsEventQueue);
-				}
-				else if (WorkerIndex == 1)
-				{
-					Workers.Emplace(GraphAllocator, CreateExportsEventQueue, WorkerZenaphores[1], ActiveWorkersCount);
-					CreateExportsEventQueue.SetZenaphore(&WorkerZenaphores[1]);
-					AltEventQueues.Remove(&CreateExportsEventQueue);
+					Workers.Emplace(GraphAllocator, ProcessExportBundlesEventQueue, IoDispatcher, WorkerZenaphores[0], ActiveWorkersCount);
+					ProcessExportBundlesEventQueue.SetZenaphore(&WorkerZenaphores[0]);
+					AltEventQueues.Remove(&ProcessExportBundlesEventQueue);
 				}
 				else
 				{
-					Workers.Emplace(GraphAllocator, AsyncEventQueue, WorkerZenaphores[2], ActiveWorkersCount);
+					Workers.Emplace(GraphAllocator, AsyncEventQueue, IoDispatcher, WorkerZenaphores[2], ActiveWorkersCount);
 					AsyncEventQueue.SetZenaphore(&WorkerZenaphores[2]);
 					AltEventQueues.Remove(&AsyncEventQueue);
 				}
@@ -3761,6 +3668,12 @@ void FAsyncLoadingThread2Impl::StartThread()
 
 		FCoreUObjectDelegates::GetPreGarbageCollectDelegate().AddRaw(this, &FAsyncLoadingThread2Impl::OnPreGarbageCollect);
 		FCoreUObjectDelegates::GetPostGarbageCollect().AddRaw(this, &FAsyncLoadingThread2Impl::OnPostGarbageCollect);
+
+		Thread = FRunnableThread::Create(this, TEXT("FAsyncLoadingThread"), 0, TPri_Normal);
+		if (Thread)
+		{
+			TRACE_SET_THREAD_GROUP(Thread->GetThreadID(), "AsyncLoading");
+		}
 	}
 }
 
@@ -3807,7 +3720,7 @@ uint32 FAsyncLoadingThread2Impl::Run()
 
 	AsyncLoadingThreadID = FPlatformTLS::GetCurrentThreadId();
 
-	FAsyncLoadingThreadState2::Create(GraphAllocator);
+	FAsyncLoadingThreadState2::Create(GraphAllocator, IoDispatcher);
 
 	TRACE_LOADTIME_START_ASYNC_LOADING();
 
@@ -3816,7 +3729,7 @@ uint32 FAsyncLoadingThread2Impl::Run()
 		FPlatformProcess::SetThreadAffinityMask(FPlatformAffinity::GetAsyncLoadingThreadMask());
 	}
 
-	FAsyncLoadingThreadState2& ThreadState = FAsyncLoadingThreadState2::Get();
+	FAsyncLoadingThreadState2& ThreadState = *FAsyncLoadingThreadState2::Get();
 
 	FZenaphoreWaiter Waiter(AltZenaphore, TEXT("WaitForEvents"));
 	bool bIsSuspended = false;
@@ -3845,24 +3758,6 @@ uint32 FAsyncLoadingThread2Impl::Run()
 				do
 				{
 					bDidSomething = false;
-					FIoChunkId ChunkId;
-					TIoStatusOr<FIoBuffer> Result;
-					uint64 UserData;
-
-					for (;;)
-					{
-						if (IoQueue->Dequeue(ChunkId, Result, UserData))
-						{
-							TRACE_CPUPROFILER_EVENT_SCOPE(ProcessResourceCompletionEvent);
-							FAsyncPackage2* Package = reinterpret_cast<FAsyncPackage2*>(UserData);
-							Package->ProcessIoRequest(ChunkId, Result);
-							bDidSomething = true;
-						}
-						else
-						{
-							break;
-						}
-					}
 
 					if (QueuedPackagesCounter)
 					{
@@ -3926,7 +3821,7 @@ EAsyncPackageState::Type FAsyncLoadingThread2Impl::TickAsyncThreadFromGameThread
 		{
 			EDLBootNotificationManager.FireCompletedCompiledInImports();
 		}
-		if (IsGarbageCollectionWaiting() || FAsyncLoadingThreadState2::Get().IsTimeLimitExceeded())
+		if (IsGarbageCollectionWaiting() || FAsyncLoadingThreadState2::Get()->IsTimeLimitExceeded())
 		{
 			Result = EAsyncPackageState::TimeOut;
 		}
@@ -4158,11 +4053,6 @@ void FAsyncLoadingThread2Impl::FireCompletedCompiledInImport(void* AsyncPackage,
 	static_cast<FAsyncPackage2*>(AsyncPackage)->GetNode(ExportNodeIndex)->ReleaseBarrier();
 }
 
-void FAsyncLoadingThread2Impl::EnqueueIoRequest(FAsyncPackage2* Package, const FIoChunkId& ChunkId)
-{
-	IoQueue->Enqueue(ChunkId, FIoReadOptions(), reinterpret_cast<uint64>(Package));
-}
-
 /*-----------------------------------------------------------------------------
 	FAsyncPackage implementation.
 -----------------------------------------------------------------------------*/
@@ -4203,13 +4093,11 @@ FAsyncPackage2::FAsyncPackage2(const FAsyncPackageDesc& InDesc, int32 InSerialNu
 	check(NameIndex != nullptr);
 	PackageChunkId = CreateChunkId(*NameIndex, Desc.NameToLoad.GetNumber(), 0, EChunkType::None);
 
-	{
-		ExportCount = AsyncLoadingThread.PackageStore.GetPackageExportCount(GlobalPackageId);
-		Exports.AddZeroed(ExportCount);
-		ExportIoBuffers.SetNum(ExportCount);
-		OwnedObjects.Reserve(ExportCount + 1); // ExportCount + UPackage
-	}
-	
+	ExportBundleCount = AsyncLoadingThread.PackageStore.GetPackageExportBundleCount(GlobalPackageId);
+	ExportCount = AsyncLoadingThread.PackageStore.GetPackageExportCount(GlobalPackageId);
+	Exports.AddZeroed(ExportCount);
+
+	ExportBundleIoBuffers.SetNum(ExportBundleCount);
 	CreateNodes(EventSpecs);
 }
 
@@ -4217,9 +4105,9 @@ void FAsyncPackage2::CreateNodes(const FAsyncLoadEventSpec* EventSpecs)
 {
 	{
 		TRACE_CPUPROFILER_EVENT_SCOPE(CreateNodes);
-		ExportNodeCount = ExportCount * EEventLoadNode2::Export_NumPhases;
+		ExportBundleNodeCount = ExportBundleCount * EEventLoadNode2::ExportBundle_NumPhases;
 
-		PackageNodes = GraphAllocator.AllocNodes(EEventLoadNode2::Package_NumPhases + ExportNodeCount);
+		PackageNodes = GraphAllocator.AllocNodes(EEventLoadNode2::Package_NumPhases + ExportBundleNodeCount);
 		for (int32 Phase = 0; Phase < EEventLoadNode2::Package_NumPhases; ++Phase)
 		{
 			new (PackageNodes + Phase) FEventLoadNode2(EventSpecs + Phase, this, -1);
@@ -4229,6 +4117,7 @@ void FAsyncPackage2::CreateNodes(const FAsyncLoadEventSpec* EventSpecs)
 		CreateLinkerNode->AddBarrier();
 
 		FEventLoadNode2* ProcessSummaryNode = PackageNodes + EEventLoadNode2::Package_ProcessSummary;
+		ProcessSummaryNode->DependsOn(CreateLinkerNode);
 		ProcessSummaryNode->AddBarrier();
 
 		FEventLoadNode2* ExportsSerializedNode = PackageNodes + EEventLoadNode2::Package_ExportsSerialized;
@@ -4240,18 +4129,17 @@ void FAsyncPackage2::CreateNodes(const FAsyncLoadEventSpec* EventSpecs)
 		DeleteNode->AddBarrier();
 		DeleteNode->DependsOn(TickNode);
 
-		// Add nodes for exports.
-		ExportNodes = PackageNodes + EEventLoadNode2::Package_NumPhases;
-		for (int32 LocalExportIndex = 0; LocalExportIndex < ExportCount; ++LocalExportIndex)
+		ExportBundleNodes = PackageNodes + EEventLoadNode2::Package_NumPhases;
+		for (int32 ExportBundleIndex = 0; ExportBundleIndex < ExportBundleCount; ++ExportBundleIndex)
 		{
-			uint32 NodeIndex = EEventLoadNode2::Export_NumPhases * LocalExportIndex;
-			FEventLoadNode2* CreateExportNode = ExportNodes + NodeIndex + EEventLoadNode2::Export_Create;
-			new (CreateExportNode) FEventLoadNode2(EventSpecs + EEventLoadNode2::Package_NumPhases + EEventLoadNode2::Export_Create, this, LocalExportIndex);
-			FEventLoadNode2* SerializeExportNode = ExportNodes + NodeIndex + EEventLoadNode2::Export_Serialize;
-			new (SerializeExportNode) FEventLoadNode2(EventSpecs + EEventLoadNode2::Package_NumPhases + EEventLoadNode2::Export_Serialize, this, LocalExportIndex);
-			CreateExportNode->DependsOn(ProcessSummaryNode);
-			ExportsSerializedNode->DependsOn(SerializeExportNode);
-			SerializeExportNode->AddBarrier();
+			uint32 NodeIndex = EEventLoadNode2::ExportBundle_NumPhases * ExportBundleIndex;
+			FEventLoadNode2* StartIoNode = ExportBundleNodes + NodeIndex + EEventLoadNode2::ExportBundle_StartIo;
+			new (StartIoNode) FEventLoadNode2(EventSpecs + EEventLoadNode2::Package_NumPhases + EEventLoadNode2::ExportBundle_StartIo, this, ExportBundleIndex);
+			FEventLoadNode2* ProcessNode = ExportBundleNodes + NodeIndex + EEventLoadNode2::ExportBundle_Process;
+			new (ProcessNode) FEventLoadNode2(EventSpecs + EEventLoadNode2::Package_NumPhases + EEventLoadNode2::ExportBundle_Process, this, ExportBundleIndex);
+			ProcessNode->AddBarrier();
+			StartIoNode->DependsOn(ProcessSummaryNode);
+			ExportsSerializedNode->DependsOn(ProcessNode);
 		}
 	}
 }
@@ -4262,7 +4150,7 @@ FAsyncPackage2::~FAsyncPackage2()
 
 	check(RefCount == 0);
 
-	FAsyncLoadingThreadState2::Get().DeferredFreeNodes.Add(MakeTuple(PackageNodes, EEventLoadNode2::Package_NumPhases + ExportNodeCount));
+	FAsyncLoadingThreadState2::Get()->DeferredFreeNodes.Add(MakeTuple(PackageNodes, EEventLoadNode2::Package_NumPhases + ExportBundleNodeCount));
 
 	TRACE_LOADTIME_DESTROY_ASYNC_PACKAGE(this);
 
@@ -4448,11 +4336,6 @@ EAsyncPackageState::Type FAsyncPackage2::CreateLinker()
 	}
 	Package->LinkerLoad = Linker;
 
-	{
-		TRACE_CPUPROFILER_EVENT_SCOPE(PackageSummaryStartIO);
-		AsyncLoadingThread.EnqueueIoRequest(this, CreateChunkId(PackageChunkId, 0, EChunkType::PackageSummary));
-	}
-
 	UE_LOG(LogStreaming, Verbose, TEXT("FAsyncPackage::CreateLinker for %s finished."), *Desc.NameToLoad.ToString());
 	return EAsyncPackageState::Complete;
 }
@@ -4467,8 +4350,14 @@ EAsyncPackageState::Type FAsyncPackage2::FinishLinker()
 	TRACE_CPUPROFILER_EVENT_SCOPE(FinishLinker);
 	LLM_SCOPE(ELLMTag::AsyncLoading);
 
-	const FPackageSummary* PackageSummary = reinterpret_cast<const FPackageSummary*>(PackageSummaryIoBuffer.Data());
-	ExportMap = reinterpret_cast<const FExportMapEntry*>(PackageSummaryIoBuffer.Data() + PackageSummary->ExportMapOffset);
+	const uint8* PackageSummaryData = PackageSummaryIoBuffer.Data();
+	const FPackageSummary* PackageSummary = reinterpret_cast<const FPackageSummary*>(PackageSummaryData);
+	
+	ExportMap = reinterpret_cast<const FExportMapEntry*>(PackageSummaryData + PackageSummary->ExportMapOffset);
+
+	ExportBundles = reinterpret_cast<const FExportBundleHeader*>(PackageSummaryData + PackageSummary->ExportBundlesOffset);
+	ExportBundleEntries = reinterpret_cast<const FExportBundleEntry*>(ExportBundles + ExportBundleCount);
+
 
 	TRACE_LOADTIME_PACKAGE_SUMMARY(this, PackageSummarySize, ImportStore.Count, ExportCount);
 
@@ -4511,7 +4400,7 @@ EAsyncPackageState::Type FAsyncPackage2::PostLoadObjects()
 	const bool bIsMultithreaded = AsyncLoadingThread.IsMultithreaded();
 
 	// PostLoad objects.
-	while (PostLoadIndex < ExportCount && !FAsyncLoadingThreadState2::Get().IsTimeLimitExceeded())
+	while (PostLoadIndex < ExportCount && !FAsyncLoadingThreadState2::Get()->IsTimeLimitExceeded())
 	{
 		UObject* Object = Exports[PostLoadIndex++];
 
@@ -4556,7 +4445,7 @@ EAsyncPackageState::Type FAsyncPackage2::PostLoadDeferredObjects()
 
 	while (DeferredPostLoadIndex < ExportCount && 
 		!AsyncLoadingThread.IsAsyncLoadingSuspended() &&
-		!FAsyncLoadingThreadState2::Get().IsTimeLimitExceeded())
+		!FAsyncLoadingThreadState2::Get()->IsTimeLimitExceeded())
 	{
 		UObject* Object = Exports[DeferredPostLoadIndex++];
 
@@ -4582,7 +4471,7 @@ EAsyncPackageState::Type FAsyncPackage2::PostLoadDeferredObjects()
 		// Clear async loading flags (we still want RF_Async, but EInternalObjectFlags::AsyncLoading can be cleared)
 		while (DeferredFinalizeIndex < ExportCount &&
 			!AsyncLoadingThread.IsAsyncLoadingSuspended() &&
-			!FAsyncLoadingThreadState2::Get().IsTimeLimitExceeded())
+			!FAsyncLoadingThreadState2::Get()->IsTimeLimitExceeded())
 		{
 			UObject* Object = Exports[DeferredFinalizeIndex++];
 
@@ -4655,7 +4544,7 @@ EAsyncPackageState::Type FAsyncPackage2::CreateClusters()
 {
 	while (DeferredClusterIndex < DeferredClusterObjects.Num() &&
 			!AsyncLoadingThread.IsAsyncLoadingSuspended() &&
-			!FAsyncLoadingThreadState2::Get().IsTimeLimitExceeded())
+			!FAsyncLoadingThreadState2::Get()->IsTimeLimitExceeded())
 	{
 		UObject* ClusterRootObject = DeferredClusterObjects[DeferredClusterIndex++];
 		ClusterRootObject->CreateCluster();
