The original BulkData code can now load optional (and duplicate nonoptional) data that has the BULKDATA_CookedForIoDispatcher flag set.

#rnx

diff --git a/Engine/Source/Runtime/CoreUObject/Private/Serialization/BulkData.cpp b/Engine/Source/Runtime/CoreUObject/Private/Serialization/BulkData.cpp
--- a/Engine/Source/Runtime/CoreUObject/Private/Serialization/BulkData.cpp
+++ b/Engine/Source/Runtime/CoreUObject/Private/Serialization/BulkData.cpp
@@ -1097,6 +1097,12 @@ void FUntypedBulkData::Serialize( FArchive& Ar, UObject* Owner, int32 Idx, bool
 
 						Ar << BulkDataOffsetInFile;
 						BulkDataOffsetInFile += OffsetInFileFixup;
+
+						if ((BulkDataFlags & BULKDATA_CookedForIoDispatcher) != 0)
+						{
+							FIoChunkId DummyChunkID;
+							Ar << DummyChunkID;
+						}
 					}
 					else
 					{
@@ -1118,6 +1124,12 @@ void FUntypedBulkData::Serialize( FArchive& Ar, UObject* Owner, int32 Idx, bool
 						int64 DummyBulkDataOffsetInFile;
 						Ar << DummyBulkDataOffsetInFile;
 
+						if ((BulkDataFlags & BULKDATA_CookedForIoDispatcher) != 0)
+						{
+							FIoChunkId DummyChunkID;
+							Ar << DummyChunkID;
+						}
+
 						Filename = FPaths::ChangeExtension(Filename, TEXT(".ubulk"));
 					}
 				}
