AsyncLoading2: Optimize post load dependencies

diff --git a/Engine/Source/Runtime/CoreUObject/Private/Serialization/AsyncLoading2.cpp b/Engine/Source/Runtime/CoreUObject/Private/Serialization/AsyncLoading2.cpp
--- a/Engine/Source/Runtime/CoreUObject/Private/Serialization/AsyncLoading2.cpp
+++ b/Engine/Source/Runtime/CoreUObject/Private/Serialization/AsyncLoading2.cpp
@@ -721,6 +721,7 @@ enum EEventLoadNode2 : uint8
 	Package_CreateLinker,
 	Package_ProcessSummary,
 	Package_ExportsSerialized,
+	Package_StartPostLoad,
 	Package_Tick,
 	Package_Delete,
 	Package_NumPhases,
@@ -1225,8 +1226,6 @@ private:
 	UPackage*				LinkerRoot;
 	/** Call backs called when we finished loading this package											*/
 	TArray<FCompletionCallback>	CompletionCallbacks;
-	/** Current index into linkers import table used to spread creation over several frames				*/
-	int32							FinishExternalReadDependenciesIndex;
 	/** Current index into ObjLoaded array used to spread routing PostLoad over several frames			*/
 	int32							PostLoadIndex;
 	/** Current index into DeferredPostLoadObjects array used to spread routing PostLoad over several frames			*/
@@ -1317,6 +1316,7 @@ public:
 	static EAsyncPackageState::Type Event_StartExportBundleIo(FAsyncPackage2* Package, int32 ExportBundleIndex);
 	static EAsyncPackageState::Type Event_ProcessExportBundle(FAsyncPackage2* Package, int32 ExportBundleIndex);
 	static EAsyncPackageState::Type Event_ExportsDone(FAsyncPackage2* Package, int32);
+	static EAsyncPackageState::Type Event_StartPostLoad(FAsyncPackage2* Package, int32);
 	static EAsyncPackageState::Type Event_Tick(FAsyncPackage2* Package, int32);
 	static EAsyncPackageState::Type Event_Delete(FAsyncPackage2* Package, int32);
 
@@ -3072,6 +3072,14 @@ EAsyncPackageState::Type FAsyncPackage2::Event_ExportsDone(FAsyncPackage2* Packa
 					Package->GetNode(Package_Tick));
 	}
 
+	Package->GetNode(EEventLoadNode2::Package_StartPostLoad)->ReleaseBarrier();
+	return EAsyncPackageState::Complete;
+}
+
+EAsyncPackageState::Type FAsyncPackage2::Event_StartPostLoad(FAsyncPackage2* Package, int32)
+{
+	TRACE_CPUPROFILER_EVENT_SCOPE(Event_StartPostLoad);
+	Package->GetNode(EEventLoadNode2::Package_Tick)->ReleaseBarrier();
 	return EAsyncPackageState::Complete;
 }
 
@@ -3100,7 +3108,9 @@ EAsyncPackageState::Type FAsyncPackage2::Event_Tick(FAsyncPackage2* Package, int
 		// Begin async loading, simulates BeginLoad
 		Package->BeginAsyncLoad();
 
-		if (LoadingState == EAsyncPackageState::Complete && !Package->bLoadHasFailed)
+		if (LoadingState == EAsyncPackageState::Complete &&
+			Package->ExternalReadDependencies.Num() > 0 &&
+			!Package->bLoadHasFailed)
 		{
 			SCOPED_LOADTIMER(Package_ExternalReadDependencies);
 			LoadingState = Package->FinishExternalReadDependencies();
@@ -3584,6 +3594,7 @@ FAsyncLoadingThread2Impl::FAsyncLoadingThread2Impl(IEDLBootNotificationManager&
 	EventSpecs[EEventLoadNode2::Package_CreateLinker] = { &FAsyncPackage2::Event_CreateLinker, &AsyncEventQueue, false };
 	EventSpecs[EEventLoadNode2::Package_ProcessSummary] = { &FAsyncPackage2::Event_ProcessSummary, &AsyncEventQueue, false };
 	EventSpecs[EEventLoadNode2::Package_ExportsSerialized] = { &FAsyncPackage2::Event_ExportsDone, &AsyncEventQueue, true };
+	EventSpecs[EEventLoadNode2::Package_StartPostLoad] = { &FAsyncPackage2::Event_StartPostLoad, &AsyncEventQueue, false };
 	EventSpecs[EEventLoadNode2::Package_Tick] = { &FAsyncPackage2::Event_Tick, &EventQueue, false };
 	EventSpecs[EEventLoadNode2::Package_Delete] = { &FAsyncPackage2::Event_Delete, &AsyncEventQueue, false };
 
@@ -3885,8 +3896,7 @@ void FAsyncLoadingThread2Impl::ResumeLoading()
 float FAsyncLoadingThread2Impl::GetAsyncLoadPercentage(const FName& PackageName)
 {
 	float LoadPercentage = -1.0f;
-	FScopeLock LockAsyncPackages(&AsyncPackagesCritical);
-	FAsyncPackage2* Package = AsyncPackageNameLookup.FindRef(PackageName);
+	FAsyncPackage2* Package = FindAsyncPackage(PackageName);
 	if (Package)
 	{
 		LoadPercentage = Package->GetLoadPercentage();
@@ -4064,7 +4074,6 @@ FAsyncPackage2::FAsyncPackage2(const FAsyncPackageDesc& InDesc, int32 InSerialNu
 : Desc(InDesc)
 , Linker(nullptr)
 , LinkerRoot(nullptr)
-, FinishExternalReadDependenciesIndex(0)
 , PostLoadIndex(0)
 , DeferredPostLoadIndex(0)
 , DeferredFinalizeIndex(0)
@@ -4121,9 +4130,11 @@ void FAsyncPackage2::CreateNodes(const FAsyncLoadEventSpec* EventSpecs)
 		ProcessSummaryNode->AddBarrier();
 
 		FEventLoadNode2* ExportsSerializedNode = PackageNodes + EEventLoadNode2::Package_ExportsSerialized;
+		FEventLoadNode2* StartPostLoadNode = PackageNodes + EEventLoadNode2::Package_StartPostLoad;
 		FEventLoadNode2* TickNode = PackageNodes + EEventLoadNode2::Package_Tick;
 
-		TickNode->DependsOn(ExportsSerializedNode);
+		StartPostLoadNode->AddBarrier();
+		TickNode->AddBarrier();
 
 		FEventLoadNode2* DeleteNode = PackageNodes + EEventLoadNode2::Package_Delete;
 		DeleteNode->AddBarrier();
@@ -4597,7 +4608,6 @@ EAsyncPackageState::Type FAsyncPackage2::FinishObjects()
 	// Simulate what EndLoad does.
 	// FLinkerManager::Get().DissociateImportsAndForcedExports(); //@todo: this should be avoidable
 	PostLoadIndex = 0;
-	FinishExternalReadDependenciesIndex = 0;
 
 	// Keep the linkers to close until we finish loading and it's safe to close them too
 	LoadContext->MoveDelayedLinkerClosePackages(DelayedLinkerClosePackages);
@@ -4671,7 +4681,6 @@ void FAsyncPackage2::Cancel()
 		}
 		ResetLoader();
 	}
-	FinishExternalReadDependenciesIndex = 0;
 }
 
 void FAsyncPackage2::AddCompletionCallback(TUniquePtr<FLoadPackageAsyncDelegate>&& Callback, bool bInternal)
