AsyncLoading2: Remove implicit arcs.
Merge CreateLinker into ProcessSummary node.
Cleanup functions for create export and serialize export.

diff --git a/Engine/Source/Runtime/CoreUObject/Private/Serialization/AsyncLoading2.cpp b/Engine/Source/Runtime/CoreUObject/Private/Serialization/AsyncLoading2.cpp
--- a/Engine/Source/Runtime/CoreUObject/Private/Serialization/AsyncLoading2.cpp
+++ b/Engine/Source/Runtime/CoreUObject/Private/Serialization/AsyncLoading2.cpp
@@ -718,7 +718,7 @@ enum class EAsyncPackageLoadingState2 : uint8
 
 enum EEventLoadNode2 : uint8
 {
-	Package_CreateLinker,
+	Package_CreateLinker, // TODO: remove
 	Package_ProcessSummary,
 	Package_ExportsSerialized,
 	Package_StartPostLoad,
@@ -742,6 +742,7 @@ public:
 	FEventLoadNode2(const FAsyncLoadEventSpec* InSpec, FAsyncPackage2* InPackage, int32 InImportOrExportIndex);
 	void DependsOn(FEventLoadNode2* Other);
 	void AddBarrier();
+	void AddBarrier(int32 Count);
 	void ReleaseBarrier();
 	void Execute(FAsyncLoadingThreadState2& ThreadState);
 
@@ -1311,7 +1312,6 @@ public:
 
 	bool bAllExportsSerialized;
 
-	static EAsyncPackageState::Type Event_CreateLinker(FAsyncPackage2* Package, int32);
 	static EAsyncPackageState::Type Event_ProcessSummary(FAsyncPackage2* Package, int32);
 	static EAsyncPackageState::Type Event_StartExportBundleIo(FAsyncPackage2* Package, int32 ExportBundleIndex);
 	static EAsyncPackageState::Type Event_ProcessExportBundle(FAsyncPackage2* Package, int32 ExportBundleIndex);
@@ -1356,7 +1356,7 @@ public:
 private:
 	void CreateNodes(const FAsyncLoadEventSpec* EventSpecs);
 	void SetupScriptArcs();
-	void StartImportPackages();
+	void SetupSerializedArcs();
 
 	/**
 	 * Begin async loading process. Simulates parts of BeginLoad.
@@ -2048,6 +2048,15 @@ void FEventLoadNode2::AddBarrier()
 	++BarrierCount;
 }
 
+void FEventLoadNode2::AddBarrier(int32 Count)
+{
+#if UE_BUILD_DEVELOPMENT
+	check(!bDone);
+	check(!bFired);
+#endif
+	BarrierCount += Count;
+}
+
 void FEventLoadNode2::ReleaseBarrier()
 {
 	check(BarrierCount > 0);
@@ -2283,35 +2292,14 @@ FUObjectSerializeContext* FAsyncPackage2::GetSerializeContext()
 	return FUObjectThreadContext::Get().GetSerializeContext();
 }
 
-EAsyncPackageState::Type FAsyncPackage2::Event_CreateLinker(FAsyncPackage2* Package, int32)
-{
-	TRACE_CPUPROFILER_EVENT_SCOPE(Event_CreateLinker);
-	// Keep track of time when we start loading.
-	if (Package->LoadStartTime == 0.0)
-	{
-		double Now = FPlatformTime::Seconds();
-		Package->LoadStartTime = Now;
-	}
-	FScopedAsyncPackageEvent2 Scope(Package);
-	//SCOPED_LOADTIMER(Package_CreateLinker);
-	check(!Package->Linker);
-	Package->CreateLinker();
-	check(Package->Linker);
-	check(Package->AsyncPackageLoadingState == EAsyncPackageLoadingState2::NewPackage);
-	Package->AsyncPackageLoadingState = EAsyncPackageLoadingState2::WaitingForSummary;
-	return EAsyncPackageState::Complete;
-}
-
 EAsyncPackageState::Type FAsyncPackage2::Event_ProcessSummary(FAsyncPackage2* Package, int32)
 {
 	TRACE_CPUPROFILER_EVENT_SCOPE(Event_ProcessSummary);
 	FScopedAsyncPackageEvent2 Scope(Package);
-	EAsyncPackageState::Type Result = Package->FinishLinker();
-	check(Result == EAsyncPackageState::Complete);
-	check(Package->Linker/* && Linker->HasFinishedInitialization()*/);
-
-	Package->StartImportPackages();
+	Package->CreateLinker();
+	Package->FinishLinker();
 
+	Package->SetupSerializedArcs();
 	if (GIsInitialLoad)
 	{
 		Package->SetupScriptArcs();
@@ -2320,55 +2308,56 @@ EAsyncPackageState::Type FAsyncPackage2::Event_ProcessSummary(FAsyncPackage2* Pa
 	check(Package->AsyncPackageLoadingState == EAsyncPackageLoadingState2::WaitingForSummary);
 	Package->AsyncPackageLoadingState = EAsyncPackageLoadingState2::ProcessNewImportsAndExports;
 
+	for (int32 ExportBundleIndex = 0; ExportBundleIndex < Package->ExportBundleCount; ++ExportBundleIndex)
+	{
+		Package->GetExportBundleNode(EEventLoadNode2::ExportBundle_StartIo, ExportBundleIndex)->ReleaseBarrier();
+	}
 	return EAsyncPackageState::Complete;
 }
 
-void FAsyncPackage2::StartImportPackages()
+void FAsyncPackage2::SetupSerializedArcs()
 {
-	TRACE_CPUPROFILER_EVENT_SCOPE(StartImportPackages);
+	TRACE_CPUPROFILER_EVENT_SCOPE(SetupSerializedArcs);
 
 	LLM_SCOPE(ELLMTag::AsyncLoading);
 
-	{
-		TRACE_CPUPROFILER_EVENT_SCOPE(SetupSerializedArcs);
-		const FGlobalNameMap& GlobalNameMap = AsyncLoadingThread.GlobalNameMap;
+	const FGlobalNameMap& GlobalNameMap = AsyncLoadingThread.GlobalNameMap;
 
-		const uint8* PackageSummaryData = PackageSummaryIoBuffer.Data();
-		const FPackageSummary* PackageSummary = reinterpret_cast<const FPackageSummary*>(PackageSummaryData);
-		uint64 ZenHeaderDataSize = PackageSummary->GraphDataSize;
-		const uint8* ZenHeaderData = PackageSummaryData + PackageSummary->GraphDataOffset;
-		FSimpleArchive ZenHeaderArchive(ZenHeaderData, ZenHeaderDataSize);
-		int32 InternalArcCount;
-		ZenHeaderArchive << InternalArcCount;
-		for (int32 InternalArcIndex = 0; InternalArcIndex < InternalArcCount; ++InternalArcIndex)
+	const uint8* PackageSummaryData = PackageSummaryIoBuffer.Data();
+	const FPackageSummary* PackageSummary = reinterpret_cast<const FPackageSummary*>(PackageSummaryData);
+	uint64 ZenHeaderDataSize = PackageSummary->GraphDataSize;
+	const uint8* ZenHeaderData = PackageSummaryData + PackageSummary->GraphDataOffset;
+	FSimpleArchive ZenHeaderArchive(ZenHeaderData, ZenHeaderDataSize);
+	int32 InternalArcCount;
+	ZenHeaderArchive << InternalArcCount;
+	for (int32 InternalArcIndex = 0; InternalArcIndex < InternalArcCount; ++InternalArcIndex)
+	{
+		int32 FromNodeIndex;
+		int32 ToNodeIndex;
+		ZenHeaderArchive << FromNodeIndex;
+		ZenHeaderArchive << ToNodeIndex;
+		PackageNodes[ToNodeIndex].DependsOn(PackageNodes + FromNodeIndex);
+	}
+	int32 ImportedPackagesCount;
+	ZenHeaderArchive << ImportedPackagesCount;
+	for (int32 ImportedPackageIndex = 0; ImportedPackageIndex < ImportedPackagesCount; ++ImportedPackageIndex)
+	{
+		int32 ImportedPackageNameIndex;
+		int32 ImportedPackageNameNumber;
+		ZenHeaderArchive << ImportedPackageNameIndex << ImportedPackageNameNumber;
+		FName ImportedPackageName = GlobalNameMap.GetName(ImportedPackageNameIndex, ImportedPackageNameNumber);
+		FAsyncPackage2* ImportedPackage = AsyncLoadingThread.FindAsyncPackage(ImportedPackageName);
+		int32 ExternalArcCount;
+		ZenHeaderArchive << ExternalArcCount;
+		for (int32 ExternalArcIndex = 0; ExternalArcIndex < ExternalArcCount; ++ExternalArcIndex)
 		{
 			int32 FromNodeIndex;
 			int32 ToNodeIndex;
 			ZenHeaderArchive << FromNodeIndex;
 			ZenHeaderArchive << ToNodeIndex;
-			PackageNodes[ToNodeIndex].DependsOn(PackageNodes + FromNodeIndex);
-		}
-		int32 ImportedPackagesCount;
-		ZenHeaderArchive << ImportedPackagesCount;
-		for (int32 ImportedPackageIndex = 0; ImportedPackageIndex < ImportedPackagesCount; ++ImportedPackageIndex)
-		{
-			int32 ImportedPackageNameIndex;
-			int32 ImportedPackageNameNumber;
-			ZenHeaderArchive << ImportedPackageNameIndex << ImportedPackageNameNumber;
-			FName ImportedPackageName = GlobalNameMap.GetName(ImportedPackageNameIndex, ImportedPackageNameNumber);
-			FAsyncPackage2* ImportedPackage = AsyncLoadingThread.FindAsyncPackage(ImportedPackageName);
-			int32 ExternalArcCount;
-			ZenHeaderArchive << ExternalArcCount;
-			for (int32 ExternalArcIndex = 0; ExternalArcIndex < ExternalArcCount; ++ExternalArcIndex)
+			if (ImportedPackage)
 			{
-				int32 FromNodeIndex;
-				int32 ToNodeIndex;
-				ZenHeaderArchive << FromNodeIndex;
-				ZenHeaderArchive << ToNodeIndex;
-				if (ImportedPackage)
-				{
-					PackageNodes[ToNodeIndex].DependsOn(ImportedPackage->PackageNodes + FromNodeIndex);
-				}
+				PackageNodes[ToNodeIndex].DependsOn(ImportedPackage->PackageNodes + FromNodeIndex);
 			}
 		}
 	}
@@ -2474,6 +2463,10 @@ void FAsyncPackage2::ImportPackagesRecursive()
 void FAsyncPackage2::StartLoading()
 {
 	TRACE_LOADTIME_BEGIN_LOAD_ASYNC_PACKAGE(this);
+	check(AsyncPackageLoadingState == EAsyncPackageLoadingState2::NewPackage);
+	AsyncPackageLoadingState = EAsyncPackageLoadingState2::WaitingForSummary;
+
+	LoadStartTime = FPlatformTime::Seconds();;
 
 	FIoReadOptions ReadOptions;
 	AsyncLoadingThread.IoDispatcher.ReadWithCallback(CreateChunkId(PackageChunkId, 0, EChunkType::PackageSummary),
@@ -2484,7 +2477,6 @@ void FAsyncPackage2::StartLoading()
 			PackageSummarySize = PackageSummaryIoBuffer.DataSize();
 			GetPackageNode(EEventLoadNode2::Package_ProcessSummary)->ReleaseBarrier();
 		});
-	GetPackageNode(EEventLoadNode2::Package_CreateLinker)->ReleaseBarrier();
 }
 
 void FAsyncPackage2::SetupScriptArcs()
@@ -2608,6 +2600,7 @@ EAsyncPackageState::Type FAsyncPackage2::Event_ProcessExportBundle(FAsyncPackage
 	}
 	ExportBundleIoBuffer = FIoBuffer();
 
+	Package->GetNode(Package_ExportsSerialized)->ReleaseBarrier();
 	return EAsyncPackageState::Complete;
 }
 
@@ -2627,6 +2620,7 @@ UObject* FAsyncPackage2::EventDrivenIndexToObject(FPackageIndex Index, bool bChe
 		Result = ImportStore.FindExistingSlimport(Index.ToImport());
 		check(Result);
 	}
+#if DO_CHECK
 	if (bCheckSerialized && !IsFullyLoadedObj(Result))
 	{
 		/*FEventLoadNode2* MyDependentNode = GetExportNode(EEventLoadNode2::Export_Serialize, Index.ToExport());
@@ -2648,15 +2642,18 @@ UObject* FAsyncPackage2::EventDrivenIndexToObject(FPackageIndex Index, bool bChe
 	{
 		UE_CLOG(Result->HasAnyInternalFlags(EInternalObjectFlags::Unreachable), LogStreaming, Fatal, TEXT("Returning an object  (%s) from EventDrivenIndexToObject that is unreachable."), *Result->GetFullName());
 	}
+#endif
 	return Result;
 }
 
 
 void FAsyncPackage2::EventDrivenCreateExport(int32 LocalExportIndex)
 {
-	//SCOPED_LOADTIMER(Package_CreateExports);
+	TRACE_CPUPROFILER_EVENT_SCOPE(CreateExport);
+
 	const FExportMapEntry& Export = ExportMap[LocalExportIndex];
 	UObject*& Object = Exports[LocalExportIndex];
+	check(!Object);
 
 	FName ObjectName;
 	{
@@ -2673,253 +2670,135 @@ void FAsyncPackage2::EventDrivenCreateExport(int32 LocalExportIndex)
 	// 	CastEventDrivenIndexToObject<UClass>(Export.ClassIndex, false), ELLMTagSet::AssetClasses);
 
 	bool bIsCompleteyLoaded = false;
+	UClass* LoadClass = Export.ClassIndex.IsNull() ? UClass::StaticClass() : CastEventDrivenIndexToObject<UClass>(Export.ClassIndex, true);
+	UObject* ThisParent = Export.OuterIndex.IsNull() ? LinkerRoot : EventDrivenIndexToObject(Export.OuterIndex, false);
+
+	checkf(LoadClass, TEXT("Could not find class object for %s in %s"), *ObjectName.ToString(), *Desc.NameToLoad.ToString());
+	checkf(ThisParent, TEXT("Could not find outer object for %s in %s"), *ObjectName.ToString(), *Desc.NameToLoad.ToString());
+	check(!dynamic_cast<UObjectRedirector*>(ThisParent));
 
-	// Check whether we already loaded the object and if not whether the context flags allow loading it.
-	//check(!Export.Object || Export.Object->HasAnyFlags(RF_ClassDefaultObject)); // we should not have this yet, unless it is a CDO
-	check(!Object); // we should not have this yet
+	// Try to find existing object first as we cannot in-place replace objects, could have been created by other export in this package
 	{
-		FUObjectSerializeContext* LoadContext = GetSerializeContext();
+		TRACE_CPUPROFILER_EVENT_SCOPE(FindExport);
+		Object = StaticFindObjectFastInternal(NULL, ThisParent, ObjectName, true);
+	}
 
+	// Object is found in memory.
+	if (Object)
+	{
+		// If this object was allocated but never loaded (components created by a constructor, CDOs, etc) make sure it gets loaded
+		// Do this for all subobjects created in the native constructor.
+		const EObjectFlags ObjectFlags = Object->GetFlags();
+		bIsCompleteyLoaded = !!(ObjectFlags & RF_LoadCompleted);
+		if (!bIsCompleteyLoaded)
 		{
-			// check(Export.ObjectName != NAME_None || !(Export.ObjectFlags&RF_Public));
-			check(LoadContext->HasStartedLoading());
-#if 0
-			if (Export.DynamicType == FObjectExport::EDynamicType::DynamicType)
+			UE_LOG(LogStreaming, VeryVerbose, TEXT("Note2: %s was constructed during load and is an export and so needs loading."), *Object->GetFullName());
+			ensure(!(ObjectFlags & (RF_NeedLoad | RF_WasLoaded))); // If export exist but is not completed, it is expected to have been created from a native constructor and not from EventDrivenCreateExport, but who knows...?
+			if (ObjectFlags & RF_ClassDefaultObject)
 			{
-//native blueprint 
-
-				Export.Object = ConstructDynamicType(*Linker->GetExportPathName(LocalExportIndex), EConstructDynamicType::OnlyAllocateClassObject);
-				check(Export.Object);
-				UDynamicClass* DC = Cast<UDynamicClass>(Export.Object);
-				UObject* DCD = DC ? DC->GetDefaultObject(false) : nullptr;
-				if (GIsInitialLoad || GUObjectArray.IsOpenForDisregardForGC())
-				{
-					Export.Object->AddToRoot();
-					if (DCD)
-					{
-						DCD->AddToRoot();
-					}
-				}
-				// TODO: Replace with DCD->SetInternalFlags(EObjectInternalFlag::Async)?
-				// But then it also need to be cleared in ClearOwnedObjects()
-				// if (DCD)
-				// {
-				// 	AddObjectReference(DCD);
-				// }
-				UE_LOG(LogStreaming, Verbose, TEXT("EventDrivenCreateExport: Created dynamic class %s"), *Export.Object->GetFullName());
-				if (Export.Object)
-				{
-					Export.Object->SetLinker(Linker, LocalExportIndex);
-				}
+				// never call PostLoadSubobjects on class default objects, this matches the behavior of the old linker where
+				// StaticAllocateObject prevents setting of RF_NeedPostLoad and RF_NeedPostLoadSubobjects, but FLinkerLoad::Preload
+				// assigns RF_NeedPostLoad for blueprint CDOs:
+				Object->SetFlags(RF_NeedLoad | RF_NeedPostLoad | RF_WasLoaded);
 			}
-			else if (Export.DynamicType == FObjectExport::EDynamicType::ClassDefaultObject)
+			else
 			{
-				UClass* LoadClass = nullptr;
-				if (!Export.ClassIndex.IsNull())
-				{
-					LoadClass = CastEventDrivenIndexToObject<UClass>(Export.ClassIndex, true);
-				}
-				if (!LoadClass)
-				{
-					UE_LOG(LogStreaming, Error, TEXT("Could not find class %s to create %s"), *Linker->ImpExp(Export.ClassIndex).ObjectName.ToString(), *Export.ObjectName.ToString());
-					return;
-				}
-				Export.Object = LoadClass->GetDefaultObject(true);
-				if (Export.Object)
-				{
-					Export.Object->SetLinker(Linker, LocalExportIndex);
-				}
+				Object->SetFlags(RF_NeedLoad | RF_NeedPostLoad | RF_NeedPostLoadSubobjects | RF_WasLoaded);
 			}
-			else
-#endif
+			AddOwnedObjectWithAsyncFlag(Object, /*bForceAdd*/ false);
+		}
+		else
+		{
+			AddOwnedObjectWithAsyncFlag(Object, /*bForceAdd*/ true);
+		}
+	}
+	else
+	{
+		// Find the Archetype object for the one we are loading.
+		check(!Export.TemplateIndex.IsNull());
+		UObject* Template = EventDrivenIndexToObject(Export.TemplateIndex, true);
+		checkf(Template, TEXT("Could not find template for %s in %s"), *ObjectName.ToString(), *Desc.NameToLoad.ToString());
+		// we also need to ensure that the template has set up any instances
+		Template->ConditionalPostLoadSubobjects();
+
+		check(!GVerifyObjectReferencesOnly); // not supported with the event driven loader
+		// Create the export object, marking it with the appropriate flags to
+		// indicate that the object's data still needs to be loaded.
+		EObjectFlags ObjectLoadFlags = Export.ObjectFlags;
+		ObjectLoadFlags = EObjectFlags(ObjectLoadFlags | RF_NeedLoad | RF_NeedPostLoad | RF_NeedPostLoadSubobjects | RF_WasLoaded);
+
+		// If we are about to create a CDO, we need to ensure that all parent sub-objects are loaded
+		// to get default value initialization to work.
+#if DO_CHECK
+		if ((ObjectLoadFlags & RF_ClassDefaultObject) != 0)
+		{
+			UClass* SuperClass = LoadClass->GetSuperClass();
+			UObject* SuperCDO = SuperClass ? SuperClass->GetDefaultObject() : nullptr;
+			check(!SuperCDO || Template == SuperCDO); // the template for a CDO is the CDO of the super
+			if (SuperClass && !SuperClass->IsNative())
 			{
-				UClass* LoadClass = nullptr;
-				if (Export.ClassIndex.IsNull())
+				check(SuperCDO);
+				if (SuperClass->HasAnyFlags(RF_NeedLoad))
 				{
-					LoadClass = UClass::StaticClass();
-				}
-				else
-				{
-					LoadClass = CastEventDrivenIndexToObject<UClass>(Export.ClassIndex, true);
-				}
-				if (!LoadClass)
-				{
-					UE_LOG(LogStreaming, Error, TEXT("Could not find class %s to create %s"), *Linker->ImpExp(Export.ClassIndex).ObjectName.ToString(), *ObjectName.ToString());
+					UE_LOG(LogStreaming, Fatal, TEXT("Super %s had RF_NeedLoad while creating %s"), *SuperClass->GetFullName(), *ObjectName.ToString());
 					return;
 				}
-				UObject* ThisParent = nullptr;
-				if (!Export.OuterIndex.IsNull())
-				{
-					ThisParent = EventDrivenIndexToObject(Export.OuterIndex, false);
-				}
-				else
-				{
-					check(LinkerRoot);
-					ThisParent = LinkerRoot;
-				}
-				check(!dynamic_cast<UObjectRedirector*>(ThisParent));
-				if (!ThisParent)
+				if (SuperCDO->HasAnyFlags(RF_NeedLoad))
 				{
-					UE_LOG(LogStreaming, Error, TEXT("Could not find outer %s to create %s"), *Linker->ImpExp(Export.OuterIndex).ObjectName.ToString(), *ObjectName.ToString());
+					UE_LOG(LogStreaming, Fatal, TEXT("Super CDO %s had RF_NeedLoad while creating %s"), *SuperCDO->GetFullName(), *ObjectName.ToString());
 					return;
 				}
+				TArray<UObject*> SuperSubObjects;
+				GetObjectsWithOuter(SuperCDO, SuperSubObjects, /*bIncludeNestedObjects=*/ false, /*ExclusionFlags=*/ RF_NoFlags, /*InternalExclusionFlags=*/ EInternalObjectFlags::Native);
 
-				// Try to find existing object first in case we're a forced export to be able to reconcile. Also do it for the
-				// case of async loading as we cannot in-place replace objects.
-
-				UObject* ActualObjectWithTheName = StaticFindObjectFastInternal(NULL, ThisParent, ObjectName, true);
-
-				// Always attempt to find object in memory first
-				if (ActualObjectWithTheName && (ActualObjectWithTheName->GetClass() == LoadClass))
-				{
-					Object = ActualObjectWithTheName;
-				}
-
-				// Object is found in memory.
-				if (Object)
-				{
-					//Object->SetLinker(Linker, LocalExportIndex);
-
-					// If this object was allocated but never loaded (components created by a constructor, CDOs, etc) make sure it gets loaded
-					// Do this for all subobjects created in the native constructor.
-					const EObjectFlags ObjectFlags = Object->GetFlags();
-					bIsCompleteyLoaded = !!(ObjectFlags & RF_LoadCompleted);
-					if (!bIsCompleteyLoaded)
-					{
-						UE_LOG(LogStreaming, VeryVerbose, TEXT("Note2: %s was constructed during load and is an export and so needs loading."), *Object->GetFullName());
-						ensure(!(ObjectFlags & (RF_NeedLoad | RF_WasLoaded))); // If export exist but is not completed, it is expected to have been created from a native constructor and not from EventDrivenCreateExport, but who knows...?
-						if (ObjectFlags & RF_ClassDefaultObject)
-						{
-							// never call PostLoadSubobjects on class default objects, this matches the behavior of the old linker where
-							// StaticAllocateObject prevents setting of RF_NeedPostLoad and RF_NeedPostLoadSubobjects, but FLinkerLoad::Preload
-							// assigns RF_NeedPostLoad for blueprint CDOs:
-							Object->SetFlags(RF_NeedLoad | RF_NeedPostLoad | RF_WasLoaded);
-						}
-						else
-						{
-							Object->SetFlags(RF_NeedLoad | RF_NeedPostLoad | RF_NeedPostLoadSubobjects | RF_WasLoaded);
-						}
-						AddOwnedObjectWithAsyncFlag(Object, /*bForceAdd*/ false);
-					}
-					else
-					{
-						AddOwnedObjectWithAsyncFlag(Object, /*bForceAdd*/ true);
-					}
-				}
-				else
+				for (UObject* SubObject : SuperSubObjects)
 				{
-					if (ActualObjectWithTheName && !ActualObjectWithTheName->GetClass()->IsChildOf(LoadClass))
+					if (SubObject->HasAnyFlags(RF_NeedLoad))
 					{
-						UE_LOG(LogLinker, Error, TEXT("Failed import: class '%s' name '%s' outer '%s'. There is another object (of '%s' class) at the path."),
-							*LoadClass->GetName(), *ObjectName.ToString(), *ThisParent->GetName(), *ActualObjectWithTheName->GetClass()->GetName());
+						UE_LOG(LogStreaming, Fatal, TEXT("Super CDO subobject %s had RF_NeedLoad while creating %s"), *SubObject->GetFullName(), *ObjectName.ToString());
 						return;
 					}
-
-					// Find the Archetype object for the one we are loading.
-					check(!Export.TemplateIndex.IsNull());
-					UObject* Template = EventDrivenIndexToObject(Export.TemplateIndex, true);
-					if (!Template)
-					{
-						UE_LOG(LogStreaming, Error, TEXT("Cannot construct %s in %s because we could not find its template %s"), *ObjectName.ToString(), *Linker->GetArchiveName(), *Linker->GetImportPathName(Export.TemplateIndex));
-						return;
-					}
-					// we also need to ensure that the template has set up any instances
-					Template->ConditionalPostLoadSubobjects();
-
-
-					check(!GVerifyObjectReferencesOnly); // not supported with the event driven loader
-					// Create the export object, marking it with the appropriate flags to
-					// indicate that the object's data still needs to be loaded.
-					EObjectFlags ObjectLoadFlags = Export.ObjectFlags;
-					ObjectLoadFlags = EObjectFlags(ObjectLoadFlags | RF_NeedLoad | RF_NeedPostLoad | RF_NeedPostLoadSubobjects | RF_WasLoaded);
-
-					// If we are about to create a CDO, we need to ensure that all parent sub-objects are loaded
-					// to get default value initialization to work.
-		#if DO_CHECK
-					if ((ObjectLoadFlags & RF_ClassDefaultObject) != 0)
-					{
-						UClass* SuperClass = LoadClass->GetSuperClass();
-						UObject* SuperCDO = SuperClass ? SuperClass->GetDefaultObject() : nullptr;
-						check(!SuperCDO || Template == SuperCDO); // the template for a CDO is the CDO of the super
-						if (SuperClass && !SuperClass->IsNative())
-						{
-							check(SuperCDO);
-							if (SuperClass->HasAnyFlags(RF_NeedLoad))
-							{
-								UE_LOG(LogStreaming, Fatal, TEXT("Super %s had RF_NeedLoad while creating %s"), *SuperClass->GetFullName(), *ObjectName.ToString());
-								return;
-							}
-							if (SuperCDO->HasAnyFlags(RF_NeedLoad))
-							{
-								UE_LOG(LogStreaming, Fatal, TEXT("Super CDO %s had RF_NeedLoad while creating %s"), *SuperCDO->GetFullName(), *ObjectName.ToString());
-								return;
-							}
-							TArray<UObject*> SuperSubObjects;
-							GetObjectsWithOuter(SuperCDO, SuperSubObjects, /*bIncludeNestedObjects=*/ false, /*ExclusionFlags=*/ RF_NoFlags, /*InternalExclusionFlags=*/ EInternalObjectFlags::Native);
-
-							for (UObject* SubObject : SuperSubObjects)
-							{
-								if (SubObject->HasAnyFlags(RF_NeedLoad))
-								{
-									UE_LOG(LogStreaming, Fatal, TEXT("Super CDO subobject %s had RF_NeedLoad while creating %s"), *SubObject->GetFullName(), *ObjectName.ToString());
-									return;
-								}
-							}
-						}
-						else
-						{
-							check(Template->IsA(LoadClass));
-						}
-					}
-		#endif
-					if (LoadClass->HasAnyFlags(RF_NeedLoad))
-					{
-						UE_LOG(LogStreaming, Fatal, TEXT("LoadClass %s had RF_NeedLoad while creating %s"), *LoadClass->GetFullName(), *ObjectName.ToString());
-						return;
-					}
-					{
-						UObject* LoadCDO = LoadClass->GetDefaultObject();
-						if (LoadCDO->HasAnyFlags(RF_NeedLoad))
-						{
-							UE_LOG(LogStreaming, Fatal, TEXT("Class CDO %s had RF_NeedLoad while creating %s"), *LoadCDO->GetFullName(), *ObjectName.ToString());
-							return;
-						}
-					}
-					if (Template->HasAnyFlags(RF_NeedLoad))
-					{
-						UE_LOG(LogStreaming, Fatal, TEXT("Template %s had RF_NeedLoad while creating %s"), *Template->GetFullName(), *ObjectName.ToString());
-						return;
-					}
-
-					{
-						TRACE_CPUPROFILER_EVENT_SCOPE(ConstructObject);
-						Object = StaticConstructObject_Internal
-							(
-							 LoadClass,
-							 ThisParent,
-							 ObjectName,
-							 ObjectLoadFlags,
-							 EInternalObjectFlags::None,
-							 Template,
-							 false,
-							 nullptr,
-							 true
-							);
-					}
-
-					if (GIsInitialLoad || GUObjectArray.IsOpenForDisregardForGC())
-					{
-						Object->AddToRoot();
-					}
-
-					AddOwnedObjectWithAsyncFlag(Object, /*bForceAdd*/ false);
-					//Object->SetLinker(Linker, LocalExportIndex);
-					check(Object->GetClass() == LoadClass);
-					check(Object->GetFName() == ObjectName);
 				}
 			}
+			else
+			{
+				check(Template->IsA(LoadClass));
+			}
+		}
+#endif
+		checkf(!LoadClass->HasAnyFlags(RF_NeedLoad),
+			TEXT("LoadClass %s had RF_NeedLoad while creating %s"), *LoadClass->GetFullName(), *ObjectName.ToString());
+		checkf(!(LoadClass->GetDefaultObject() && LoadClass->GetDefaultObject()->HasAnyFlags(RF_NeedLoad)), 
+			TEXT("Class CDO %s had RF_NeedLoad while creating %s"), *LoadClass->GetDefaultObject()->GetFullName(), *ObjectName.ToString());
+		checkf(!Template->HasAnyFlags(RF_NeedLoad),
+			TEXT("Template %s had RF_NeedLoad while creating %s"), *Template->GetFullName(), *ObjectName.ToString());
+
+		{
+			TRACE_CPUPROFILER_EVENT_SCOPE(ConstructObject);
+			Object = StaticConstructObject_Internal
+				(
+				 LoadClass,
+				 ThisParent,
+				 ObjectName,
+				 ObjectLoadFlags,
+				 EInternalObjectFlags::None,
+				 Template,
+				 false,
+				 nullptr,
+				 true
+				);
+		}
+
+		if (GIsInitialLoad || GUObjectArray.IsOpenForDisregardForGC())
+		{
+			Object->AddToRoot();
 		}
+
+		AddOwnedObjectWithAsyncFlag(Object, /*bForceAdd*/ false);
+		check(Object->GetClass() == LoadClass);
+		check(Object->GetFName() == ObjectName);
 	}
+
 	check(Object);
 	if (Export.GlobalImportIndex != -1)
 	{
@@ -2937,119 +2816,88 @@ void FAsyncPackage2::EventDrivenCreateExport(int32 LocalExportIndex)
 
 void FAsyncPackage2::EventDrivenSerializeExport(int32 LocalExportIndex, const uint8* ExportData, uint64 ExportDataSize)
 {
-	//SCOPED_LOADTIMER(Package_PreLoadObjects);
+	TRACE_CPUPROFILER_EVENT_SCOPE(SerializeExport);
 
 	const FExportMapEntry& Export = ExportMap[LocalExportIndex];
 	UObject* Object = Exports[LocalExportIndex];
+	check(Object);
 
 	LLM_SCOPE(ELLMTag::UObject);
 	LLM_SCOPED_TAG_WITH_OBJECT_IN_SET(GetLinkerRoot(), ELLMTagSet::Assets);
 	// LLM_SCOPED_TAG_WITH_OBJECT_IN_SET((Export.DynamicType == FObjectExport::EDynamicType::DynamicType) ? UDynamicClass::StaticClass() :
 	// 	CastEventDrivenIndexToObject<UClass>(Export.ClassIndex, false), ELLMTagSet::AssetClasses);
 
-#if 0
-	if (Object && Linker->bDynamicClassLinker)
-	{
-		//native blueprint 
-		UDynamicClass* UD = Cast<UDynamicClass>(Object);
-		if (UD)
-		{
-			check(Export.DynamicType == FObjectExport::EDynamicType::DynamicType);
-			UObject* LocObj = ConstructDynamicType(*Linker->GetExportPathName(LocalExportIndex), EConstructDynamicType::CallZConstructor);
-			check(UD == LocObj);
-		}
-		Object->ClearFlags(RF_NeedLoad);
-	}
-	else if (Object && Object->HasAnyFlags(RF_NeedLoad))
-#endif
+	// If this is a struct, make sure that its parent struct is completely loaded
+	if (UStruct* Struct = dynamic_cast<UStruct*>(Object))
 	{
-		//check(Object->GetLinker() == Linker);
-		//check(Object->GetLinkerIndex() == LocalExportIndex);
-
-		// If this is a struct, make sure that its parent struct is completely loaded
-		if (UStruct* Struct = dynamic_cast<UStruct*>(Object))
+		if (!Export.SuperIndex.IsNull())
 		{
-			UStruct* SuperStruct = nullptr;
-			if (!Export.SuperIndex.IsNull())
+			UStruct* SuperStruct = CastEventDrivenIndexToObject<UStruct>(Export.SuperIndex, true);
+			checkf(SuperStruct, TEXT("Could not find SuperStruct for %s"), *Object->GetFullName());
+			Struct->SetSuperStruct(SuperStruct);
+			if (UClass* ClassObject = dynamic_cast<UClass*>(Object))
 			{
-				SuperStruct = CastEventDrivenIndexToObject<UStruct>(Export.SuperIndex, true);
-				if (!SuperStruct)
-				{
-					// see FLinkerLoad::CreateExport, there may be some more we can do here
-					// UE_LOG(LogStreaming, Fatal, TEXT("Could not find SuperStruct %s to create %s"), *Linker->ImpExp(Export.SuperIndex).ObjectName.ToString(), *Export.ObjectName.ToString());
-					return;
-				}
-			}
-			if (SuperStruct)
-			{
-				Struct->SetSuperStruct(SuperStruct);
-				if (UClass* ClassObject = dynamic_cast<UClass*>(Object))
-				{
-					ClassObject->Bind();
-				}
+				ClassObject->Bind();
 			}
 		}
-		
-		const FPackageFileSummary& Summary = Linker->Summary;
-		const FCustomVersionContainer& SummaryVersions = Summary.GetCustomVersionContainer();
+	}
 
-		FSimpleExportArchive Ar(ExportData, ExportDataSize);
-		Ar.SetUE4Ver(Summary.GetFileVersionUE4());
-		Ar.SetLicenseeUE4Ver(Summary.GetFileVersionLicenseeUE4());
-		Ar.SetEngineVer(Summary.SavedByEngineVersion);
-		Ar.SetIsLoading(true);
-		Ar.SetIsPersistent(true);
-		if (Summary.PackageFlags & PKG_FilterEditorOnly)
-		{
-			Ar.SetFilterEditorOnly(true);
-		}
-		Ar.ArAllowLazyLoading = true;
+	FSimpleExportArchive Ar(ExportData, ExportDataSize);
+	Ar.SetUE4Ver(LinkerRoot->LinkerPackageVersion);
+	Ar.SetLicenseeUE4Ver(LinkerRoot->LinkerLicenseeVersion);
+	// Ar.SetEngineVer(Summary.SavedByEngineVersion); // very old versioning scheme
+	// Ar.SetCustomVersions(LinkerRoot->LinkerCustomVersion); // only if not cooking with -unversioned
+	Ar.SetIsLoading(true);
+	Ar.SetIsPersistent(true);
+	if (LinkerRoot->GetPackageFlags() & PKG_FilterEditorOnly)
+	{
+		Ar.SetFilterEditorOnly(true);
+	}
+	Ar.ArAllowLazyLoading = true;
 
-		// FSimpleExportArchive special fields
-		Ar.NameMap = &AsyncLoadingThread.GlobalNameMap.GetNameEntries();
-		Ar.ImportStore = &ImportStore;
-		Ar.Exports = &Exports;
-		Ar.ExternalReadDependencies = &ExternalReadDependencies;
-		Ar.ExternalIoRequests = &ExternalIoRequests;
+	// FSimpleExportArchive special fields
+	Ar.NameMap = &AsyncLoadingThread.GlobalNameMap.GetNameEntries();
+	Ar.ImportStore = &ImportStore;
+	Ar.Exports = &Exports;
+	Ar.ExternalReadDependencies = &ExternalReadDependencies;
+	Ar.ExternalIoRequests = &ExternalIoRequests;
 
-		Object->ClearFlags(RF_NeedLoad);
+	Object->ClearFlags(RF_NeedLoad);
 
-		TRACE_LOADTIME_SERIALIZE_EXPORT_SCOPE(Object, ExportDataSize);
+	TRACE_LOADTIME_SERIALIZE_EXPORT_SCOPE(Object, ExportDataSize);
 
-		FUObjectSerializeContext* LoadContext = GetSerializeContext();
-		UObject* PrevSerializedObject = LoadContext->SerializedObject;
-		LoadContext->SerializedObject = Object;
-		Linker->bForceSimpleIndexToObject = true;
+	FUObjectSerializeContext* LoadContext = GetSerializeContext();
+	UObject* PrevSerializedObject = LoadContext->SerializedObject;
+	LoadContext->SerializedObject = Object;
 
-		// Find the Archetype object for the one we are loading. This is piped to GetArchetypeFromLoader
-		check(!Export.TemplateIndex.IsNull());
-		UObject* Template = EventDrivenIndexToObject(Export.TemplateIndex, true);
-		check(Template);
+	// Find the Archetype object for the one we are loading. This is piped to GetArchetypeFromLoader
+	check(!Export.TemplateIndex.IsNull());
+	UObject* Template = EventDrivenIndexToObject(Export.TemplateIndex, true);
+	check(Template);
 
-		Ar.TemplateForGetArchetypeFromLoader = Template;
+	Ar.TemplateForGetArchetypeFromLoader = Template;
 
-		if (Object->HasAnyFlags(RF_ClassDefaultObject))
-		{
-			TRACE_CPUPROFILER_EVENT_SCOPE(SerializeDefaultObject);
-			Object->GetClass()->SerializeDefaultObject(Object, Ar);
-		}
-		else
-		{
-			TRACE_CPUPROFILER_EVENT_SCOPE(SerializeObject);
-			Object->Serialize(Ar);
-		}
+	if (Object->HasAnyFlags(RF_ClassDefaultObject))
+	{
+		TRACE_CPUPROFILER_EVENT_SCOPE(SerializeDefaultObject);
+		Object->GetClass()->SerializeDefaultObject(Object, Ar);
+	}
+	else
+	{
+		TRACE_CPUPROFILER_EVENT_SCOPE(SerializeObject);
+		Object->Serialize(Ar);
+	}
 
-		Object->SetFlags(RF_LoadCompleted);
-		LoadContext->SerializedObject = PrevSerializedObject;
+	Object->SetFlags(RF_LoadCompleted);
+	LoadContext->SerializedObject = PrevSerializedObject;
 
 #if DO_CHECK
-		if (Object->HasAnyFlags(RF_ClassDefaultObject) && Object->GetClass()->HasAnyClassFlags(CLASS_CompiledFromBlueprint))
-		{
-			check(Object->HasAllFlags(RF_NeedPostLoad | RF_WasLoaded));
-			//Object->SetFlags(RF_NeedPostLoad | RF_WasLoaded);
-		}
-#endif
+	if (Object->HasAnyFlags(RF_ClassDefaultObject) && Object->GetClass()->HasAnyClassFlags(CLASS_CompiledFromBlueprint))
+	{
+		check(Object->HasAllFlags(RF_NeedPostLoad | RF_WasLoaded));
+		//Object->SetFlags(RF_NeedPostLoad | RF_WasLoaded);
 	}
+#endif
 
 	// push stats so that we don't overflow number of tags per thread during blocking loading
 	LLM_PUSH_STATS_FOR_ASSET_TAGS();
@@ -3591,7 +3439,6 @@ FAsyncLoadingThread2Impl::FAsyncLoadingThread2Impl(IEDLBootNotificationManager&
 	}
 
 	EventSpecs.AddDefaulted(EEventLoadNode2::Package_NumPhases + EEventLoadNode2::ExportBundle_NumPhases);
-	EventSpecs[EEventLoadNode2::Package_CreateLinker] = { &FAsyncPackage2::Event_CreateLinker, &AsyncEventQueue, false };
 	EventSpecs[EEventLoadNode2::Package_ProcessSummary] = { &FAsyncPackage2::Event_ProcessSummary, &AsyncEventQueue, false };
 	EventSpecs[EEventLoadNode2::Package_ExportsSerialized] = { &FAsyncPackage2::Event_ExportsDone, &AsyncEventQueue, true };
 	EventSpecs[EEventLoadNode2::Package_StartPostLoad] = { &FAsyncPackage2::Event_StartPostLoad, &AsyncEventQueue, false };
@@ -4122,11 +3969,8 @@ void FAsyncPackage2::CreateNodes(const FAsyncLoadEventSpec* EventSpecs)
 			new (PackageNodes + Phase) FEventLoadNode2(EventSpecs + Phase, this, -1);
 		}
 
-		FEventLoadNode2* CreateLinkerNode = PackageNodes + EEventLoadNode2::Package_CreateLinker;
-		CreateLinkerNode->AddBarrier();
 
 		FEventLoadNode2* ProcessSummaryNode = PackageNodes + EEventLoadNode2::Package_ProcessSummary;
-		ProcessSummaryNode->DependsOn(CreateLinkerNode);
 		ProcessSummaryNode->AddBarrier();
 
 		FEventLoadNode2* ExportsSerializedNode = PackageNodes + EEventLoadNode2::Package_ExportsSerialized;
@@ -4138,7 +3982,6 @@ void FAsyncPackage2::CreateNodes(const FAsyncLoadEventSpec* EventSpecs)
 
 		FEventLoadNode2* DeleteNode = PackageNodes + EEventLoadNode2::Package_Delete;
 		DeleteNode->AddBarrier();
-		DeleteNode->DependsOn(TickNode);
 
 		ExportBundleNodes = PackageNodes + EEventLoadNode2::Package_NumPhases;
 		for (int32 ExportBundleIndex = 0; ExportBundleIndex < ExportBundleCount; ++ExportBundleIndex)
@@ -4149,9 +3992,9 @@ void FAsyncPackage2::CreateNodes(const FAsyncLoadEventSpec* EventSpecs)
 			FEventLoadNode2* ProcessNode = ExportBundleNodes + NodeIndex + EEventLoadNode2::ExportBundle_Process;
 			new (ProcessNode) FEventLoadNode2(EventSpecs + EEventLoadNode2::Package_NumPhases + EEventLoadNode2::ExportBundle_Process, this, ExportBundleIndex);
 			ProcessNode->AddBarrier();
-			StartIoNode->DependsOn(ProcessSummaryNode);
-			ExportsSerializedNode->DependsOn(ProcessNode);
+			StartIoNode->AddBarrier();
 		}
+		ExportsSerializedNode->AddBarrier(ExportBundleCount);
 	}
 }
 
@@ -4312,6 +4155,7 @@ void FAsyncPackage2::EndAsyncLoad()
  */
 EAsyncPackageState::Type FAsyncPackage2::CreateLinker()
 {
+	TRACE_CPUPROFILER_EVENT_SCOPE(CreateLinker);
 	check(Linker == nullptr);
 
 	// Try to find existing package or create it if not already present.
