The fallback path for new Bulkdata runtime now uses much less memory (it is similar to the memory footprint when using the old BulkData, FBulkDataStreamingToken and FCompactByteBulkData), it also now works with the DuplicateNonOptional path.

- We now only store each Filename once, and reference count it for each Bulkdata object that wants to use it rather than duplicating the strings to save memory. This is done via FStringTable which acts as storage for the strings, accessed via the PackageName.
- The FileToken is now an index in a TSparseArray that references it's PackageName and the BulkDataOffset, in turn the PackageName can be used to look up the filename.
- Reduced FileToken to a int32 (since this is the type that TSparseArray returns as it's index) this also means we need to change the InvalidToken to INDEX_NONE rather than 0 which is a valid index for TSparseArray.
- Since we store the original filename for the Package, we are not returning the filename to actually use by the BulkData object so places that want that need to call FBulkDataBase::ConvertFilenameFromFlags using the Filename as the input.
- When serializing DuplicateNonOptional data we need to fix up the fallback data if the optional data can be found, this fixes the DuplicateNonOptional path when using the new BulkData api.
-- We also need to use the flags we just loaded from disk to check if the DuplicateNonOptional has a FIoChunkId to read, not the flags currently owned by the object.

Future work:
 - Could consider removing the strings entirely and only storing the PackageName and Offset. We could use code similar to UStaticMesh::GetMipDataFilename to generate the filenames when ever we need them (wrapped up properly and assume that we don't need to test for file existance as this should've been done when the package was loaded) but would need to profile the perf impact of that vs memory savings.

#rnx

diff --git a/Engine/Source/Runtime/CoreUObject/Private/Serialization/BulkData2.cpp b/Engine/Source/Runtime/CoreUObject/Private/Serialization/BulkData2.cpp
--- a/Engine/Source/Runtime/CoreUObject/Private/Serialization/BulkData2.cpp
+++ b/Engine/Source/Runtime/CoreUObject/Private/Serialization/BulkData2.cpp
@@ -34,23 +34,99 @@ namespace FileTokenSystem
 {
 	struct Data
 	{
-		FString Filename;
 		int64 BulkDataOffsetInFile;
+		FString PackageHeaderFilename;
 	};
 
-	TMap<FBulkDataBase::FileToken, Data> TokenDataMap;
-	uint64 TokenCounter = 0;
+	// Internal to the FileTokenSystem namespace
+	namespace
+	{
+		struct InternalData
+		{
+			FName PackageName;
+			int64 BulkDataOffsetInFile;			
+		};
+
+		struct StringData
+		{
+			FString Filename;
+			uint16 RefCount;
+		};
+
+		/**
+		* Provides a ref counted PackageName->Filename look up table.
+		*/
+		class FStringTable
+		{
+		public:
+			void Add(const FName& PackageName, const FString& Filename)
+			{
+				if (StringData* ExistingEntry = Table.Find(PackageName))
+				{
+					ExistingEntry->RefCount++;
+				}
+				else
+				{
+					StringData& NewEntry = Table.Emplace(PackageName);
+					NewEntry.Filename = Filename;
+					NewEntry.RefCount = 1;
+				}
+			}
+
+			bool Remove(const FName& PackageName)
+			{
+				if (StringData* ExistingEntry = Table.Find(PackageName))
+				{
+					if (--ExistingEntry->RefCount == 0)
+					{
+						Table.Remove(PackageName);
+						return true;
+					}
+				}
+
+				return false;
+			}
+
+			void IncRef(const FName& PackageName)
+			{
+				if (StringData* ExistingEntry = Table.Find(PackageName))
+				{
+					ExistingEntry->RefCount++;
+				}
+			}
+
+			const FString& Resolve(const FName& PackageName)
+			{
+				return Table.Find(PackageName)->Filename;
+			}
+
+			int32 Num() const
+			{
+				return Table.Num();
+			}
+
+		private:
+			TMap<FName, StringData> Table;
+		};
+	}
+
+	FStringTable StringTable;
+	TSparseArray<InternalData> TokenData;
+
 	FRWLock TokenLock;
 
-	FBulkDataBase::FileToken RegisterFileToken( const FString& Filename, uint64 BulkDataSize, uint64 BulkDataOffsetInFile )
+	FBulkDataBase::FileToken RegisterFileToken( const FName& PackageName, const FString& Filename, uint64 BulkDataOffsetInFile )
 	{
 		FWriteScopeLock LockForScope(TokenLock);
 
-		Data& Data = TokenDataMap.Add(++TokenCounter);
-		Data.Filename = Filename;
+		StringTable.Add(PackageName, Filename);
+
+		InternalData Data;
+		Data.PackageName = PackageName;
 		Data.BulkDataOffsetInFile = BulkDataOffsetInFile;
 
-		return TokenCounter;
+		FBulkDataBase::FileToken FileToken  = TokenData.Add(Data);
+		return FileToken;
 	}
 
 	void UnregisterFileToken(FBulkDataBase::FileToken ID)
@@ -58,7 +134,11 @@ namespace FileTokenSystem
 		if (ID != FBulkDataBase::InvalidToken)
 		{
 			FWriteScopeLock LockForScope(TokenLock);
-			TokenDataMap.Remove(ID);
+
+			StringTable.Remove(TokenData[ID].PackageName);
+			TokenData.RemoveAt(ID);
+
+			check(StringTable.Num() <= TokenData.Num());
 		}
 	}
 
@@ -68,12 +148,17 @@ namespace FileTokenSystem
 		{
 			FWriteScopeLock LockForScope(TokenLock);
 
-			Data* Data = TokenDataMap.Find(ID);
-			check(Data);
+			FSparseArrayAllocationInfo AllocInfo = TokenData.AddUninitialized();
+			
+			const InternalData& OriginalData = TokenData[ID];
+			InternalData& NewData = TokenData[AllocInfo.Index];
 
-			TokenDataMap.Add(++TokenCounter, *Data);
+			NewData.PackageName = OriginalData.PackageName;
+			NewData.BulkDataOffsetInFile = OriginalData.BulkDataOffsetInFile;
 
-			return TokenCounter;		
+			StringTable.IncRef(NewData.PackageName);
+
+			return AllocInfo.Index;
 		}
 		else
 		{
@@ -89,10 +174,13 @@ namespace FileTokenSystem
 		}
 
 		FReadScopeLock LockForScope(TokenLock);
-		Data* Data = TokenDataMap.Find(ID);
-		check(Data);
+		const InternalData& DataSrc = TokenData[ID];
+
+		Data Output;
+		Output.BulkDataOffsetInFile = DataSrc.BulkDataOffsetInFile;
+		Output.PackageHeaderFilename = StringTable.Resolve(DataSrc.PackageName);
 
-		return *Data;
+		return Output;
 	}
 
 	FString GetFilename(FBulkDataBase::FileToken ID)
@@ -103,10 +191,7 @@ namespace FileTokenSystem
 		}
 
 		FReadScopeLock LockForScope(TokenLock);
-		Data* Data = TokenDataMap.Find(ID);
-		check(Data);
-
-		return Data->Filename;
+		return StringTable.Resolve(TokenData[ID].PackageName);
 	}
 
 	uint64 GetBulkDataOffset(FBulkDataBase::FileToken ID)
@@ -117,10 +202,7 @@ namespace FileTokenSystem
 		}
 
 		FReadScopeLock LockForScope(TokenLock);
-		Data* Data = TokenDataMap.Find(ID);
-		check(Data);
-
-		return Data != nullptr ? Data->BulkDataOffsetInFile : 0;
+		return TokenData[ID].BulkDataOffsetInFile;
 	}
 }
 
@@ -496,8 +578,10 @@ void FBulkDataBase::Serialize(FArchive& Ar, UObject* Owner, int32 /*Index*/, boo
 		}
 
 		FString Filename;
+		FName PackageName;
 		FLinkerLoad* Linker = nullptr;
-		if (!bUseZenLoader)
+
+		if (bUseZenLoader == false)
 		{
 			// Assuming that Owner/Package/Linker are all valid, the old BulkData system would
 			// generally fail if any of these were nullptr but had plenty of inconsistent checks
@@ -515,9 +599,10 @@ void FBulkDataBase::Serialize(FArchive& Ar, UObject* Owner, int32 /*Index*/, boo
 				BulkDataOffsetInFile += Linker->Summary.BulkDataStartOffset;
 			}
 
-			Filename = ConvertFilenameFromFlags(Linker->Filename);
+			Filename = Linker->Filename;
+			PackageName = Package->FileName;
 		}
-		
+	
 		FArchive* CacheableArchive = Ar.GetCacheableArchive();
 		if (Ar.IsAllowingLazyLoading() && CacheableArchive != nullptr)
 		{
@@ -531,11 +616,18 @@ void FBulkDataBase::Serialize(FArchive& Ar, UObject* Owner, int32 /*Index*/, boo
 			{
 				if (IsDuplicateNonOptional())
 				{
-					FString OptionalFilename = FPaths::ChangeExtension(Filename, OptionalExt);
 					if (IFileManager::Get().FileExists(*Filename))
 					{
 						SerializeDuplicateData(Ar, Linker, BulkDataFlags, BulkDataSizeOnDisk, BulkDataOffsetInFile);
-						Filename = OptionalFilename;
+
+						// Update the fallback data again if needed
+						if (!bUseZenLoader || ChunkID == FIoChunkId::InvalidChunkId)
+						{
+							// Remove the flag if we are not going to use it
+							BulkDataFlags &= ~BULKDATA_CookedForIoDispatcher;
+							Fallback.Token = InvalidToken;
+							Fallback.BulkDataSize = BulkDataSize;
+						}
 					}
 					else
 					{
@@ -544,16 +636,14 @@ void FBulkDataBase::Serialize(FArchive& Ar, UObject* Owner, int32 /*Index*/, boo
 						uint32 DummyValue32;
 						int64 DummyValue64;
 						SerializeDuplicateData(Ar, Linker, DummyValue32, DummyValue64, DummyValue64);
-
-						check(Filename.EndsWith(DefaultExt));
-					}			
-				}			
+					}
+				}
 			}
 
 			if (bUseZenLoader == false)
 			{
-				Fallback.Token = FileTokenSystem::RegisterFileToken(Filename, BulkDataSize, BulkDataOffsetInFile);
-			}		
+				Fallback.Token = FileTokenSystem::RegisterFileToken( PackageName, Filename, BulkDataOffsetInFile);
+			}
 		}
 		else
 		{
@@ -774,12 +864,13 @@ IBulkDataIORequest* FBulkDataBase::CreateStreamingRequest(int64 OffsetInBulkData
 		const int64 BulkDataSize = GetBulkDataSize();
 		FileTokenSystem::Data FileData = FileTokenSystem::GetFileData(Fallback.Token);
 
-		check(FileData.Filename.IsEmpty() == false);
+		check(FileData.PackageHeaderFilename.IsEmpty() == false);
+		const FString Filename = ConvertFilenameFromFlags(FileData.PackageHeaderFilename);
 
-		UE_CLOG(IsStoredCompressedOnDisk(), LogSerialization, Fatal, TEXT("Package level compression is no longer supported (%s)."), *FileData.Filename);
-		UE_CLOG(BulkDataSize <= 0, LogSerialization, Error, TEXT("(%s) has invalid bulk data size."), *FileData.Filename);
+		UE_CLOG(IsStoredCompressedOnDisk(), LogSerialization, Fatal, TEXT("Package level compression is no longer supported (%s)."), *Filename);
+		UE_CLOG(BulkDataSize <= 0, LogSerialization, Error, TEXT("(%s) has invalid bulk data size."), *Filename);
 
-		IAsyncReadFileHandle* IORequestHandle = FPlatformFileManager::Get().GetPlatformFile().OpenAsyncRead(*FileData.Filename);
+		IAsyncReadFileHandle* IORequestHandle = FPlatformFileManager::Get().GetPlatformFile().OpenAsyncRead(*Filename);
 		check(IORequestHandle); // this generally cannot fail because it is async
 
 		if (IORequestHandle == nullptr)
@@ -891,7 +982,8 @@ FString FBulkDataBase::GetFilename() const
 {
 	if (!IsUsingIODispatcher())
 	{
-		return FileTokenSystem::GetFilename(Fallback.Token);
+		FString Filename = FileTokenSystem::GetFilename(Fallback.Token);
+		return ConvertFilenameFromFlags(Filename);
 	}
 	else
 	{
@@ -910,6 +1002,8 @@ void FBulkDataBase::LoadDataDirectly(void** DstBuffer)
 		const int64 BulkDataSize = GetBulkDataSize();
 		FileTokenSystem::Data FileData = FileTokenSystem::GetFileData(Fallback.Token);
 
+		const FString Filename = ConvertFilenameFromFlags(FileData.PackageHeaderFilename);
+
 		// In the older code path if the bulkdata was not in a separate file we could just serialize from the package instead, but the new system
 		// does not keep a reference to the pack so we cannot do that.
 		// TODO: Maybe add a check for this condition?
@@ -919,10 +1013,10 @@ void FBulkDataBase::LoadDataDirectly(void** DstBuffer)
 								// ::Serialize is too slow!)
 
 		// If the data is inlined then we already loaded is during ::Serialize, this warning should help track cases where data is being discarded then re-requested.
-		UE_CLOG(IsInlined(), LogSerialization, Warning, TEXT("Reloading inlined bulk data directly from disk, consider not discarding it in the first place. Filename: '%s'."), *FileData.Filename);
+		UE_CLOG(IsInlined(), LogSerialization, Warning, TEXT("Reloading inlined bulk data directly from disk, consider not discarding it in the first place. Filename: '%s'."), *Filename);
 
-		FArchive* Ar = IFileManager::Get().CreateFileReader(*FileData.Filename, FILEREAD_Silent);
-		checkf(Ar != nullptr, TEXT("Failed to open the file to load bulk data from. Filename: '%s'."), *FileData.Filename);
+		FArchive* Ar = IFileManager::Get().CreateFileReader(*Filename, FILEREAD_Silent);
+		checkf(Ar != nullptr, TEXT("Failed to open the file to load bulk data from. Filename: '%s'."), *Filename);
 
 		// Seek to the beginning of the bulk data in the file.
 		Ar->Seek(FileData.BulkDataOffsetInFile);
@@ -960,7 +1054,7 @@ void FBulkDataBase::SerializeDuplicateData(FArchive& Ar, FLinkerLoad* Linker, ui
 
 	Ar << OutBulkDataOffsetInFile;
 
-	if ((BulkDataFlags & BULKDATA_CookedForIoDispatcher) != 0)
+	if ((OutBulkDataFlags & BULKDATA_CookedForIoDispatcher) != 0)
 	{
 		FIoChunkId DummyChunkID;
 		Ar << DummyChunkID;
diff --git a/Engine/Source/Runtime/CoreUObject/Public/Serialization/BulkData2.h b/Engine/Source/Runtime/CoreUObject/Public/Serialization/BulkData2.h
--- a/Engine/Source/Runtime/CoreUObject/Public/Serialization/BulkData2.h
+++ b/Engine/Source/Runtime/CoreUObject/Public/Serialization/BulkData2.h
@@ -49,8 +49,8 @@ public:
 	static void SetIODispatcher(FIoDispatcher* InIoDispatcher) { IoDispatcher = InIoDispatcher; }
 
 public:
-	using FileToken = uint64;
-	static constexpr FileToken InvalidToken = 0;
+	using FileToken = int32;
+	static constexpr FileToken InvalidToken = INDEX_NONE;
 
 	FBulkDataBase(const FBulkDataBase& Other) { *this = Other; }
 	FBulkDataBase(FBulkDataBase&& Other);
