Added a warning to FBulkDataBase::GetCopy to highlight cases where it would be faster to not allocate memory up front.

diff --git a/Engine/Source/Runtime/CoreUObject/Private/Serialization/BulkData2.cpp b/Engine/Source/Runtime/CoreUObject/Private/Serialization/BulkData2.cpp
--- a/Engine/Source/Runtime/CoreUObject/Private/Serialization/BulkData2.cpp
+++ b/Engine/Source/Runtime/CoreUObject/Private/Serialization/BulkData2.cpp
@@ -723,6 +723,7 @@ void FBulkDataBase::GetCopy(void** DstBuffer, bool bDiscardInternalCopy)
 
 			if (bDiscardInternalCopy && (CanLoadFromDisk() || IsSingleUse()))
 			{
+				UE_LOG(LogSerialization, Warning, TEXT("FBulkDataBase::GetCopy both copied and discarded it's data, passing in an empty pointer would avoid an extra allocate and memcpy!"));
 				FreeData();
 			}
 		}
