When writing out duplicate bulkdata make sure we save the position in the stream of the duplicate FIoChunkId correctly.

#rnx

diff --git a/Engine/Source/Runtime/CoreUObject/Private/Serialization/BulkData.cpp b/Engine/Source/Runtime/CoreUObject/Private/Serialization/BulkData.cpp
--- a/Engine/Source/Runtime/CoreUObject/Private/Serialization/BulkData.cpp
+++ b/Engine/Source/Runtime/CoreUObject/Private/Serialization/BulkData.cpp
@@ -1340,7 +1340,7 @@ void FUntypedBulkData::Serialize( FArchive& Ar, UObject* Owner, int32 Idx, bool
 						// is cooked data, see if we can remove the if check
 						if (Ar.IsCooking())
 						{
-							SavedBulkDataChunkIdPos = Ar.Tell();
+							SavedDupeBulkDataChunkIdPos = Ar.Tell();
 							
 							FIoChunkId InvalidChunkId = FIoChunkId::InvalidChunkId;
 							Ar << InvalidChunkId;
