AsyncLoading2: Cleanup some more linker stuff until it can be completely removed

diff --git a/Engine/Source/Runtime/CoreUObject/Private/Serialization/AsyncLoading2.cpp b/Engine/Source/Runtime/CoreUObject/Private/Serialization/AsyncLoading2.cpp
--- a/Engine/Source/Runtime/CoreUObject/Private/Serialization/AsyncLoading2.cpp
+++ b/Engine/Source/Runtime/CoreUObject/Private/Serialization/AsyncLoading2.cpp
@@ -1010,8 +1010,7 @@ private:
 };
 
 /**
-* Structure containing intermediate data required for async loading of all imports and exports of a
-* FLinkerLoad.
+* Structure containing intermediate data required for async loading of all exports of a package.
 */
 
 struct FAsyncPackage2
@@ -1064,21 +1063,6 @@ struct FAsyncPackage2
 	 */
 	double GetLoadStartTime() const;
 
-	/**
-	 * Emulates ResetLoaders for the package's Linker objects, hence deleting it.
-	 */
-	void ResetLoader();
-
-	/**
-	* Disassociates linker from this package
-	*/
-	void DetachLinker();
-
-	/**
-	* Flushes linker cache for all objects loaded with this package
-	*/
-	void FlushObjectLinkerCache();
-
 	/**
 	 * Returns the name of the package to load.
 	 */
@@ -1254,8 +1238,6 @@ private:
 	double						LoadStartTime;
 	/** Estimated load percentage.																		*/
 	float						LoadPercentage;
-	/** Packages that were loaded synchronously while async loading this package or packages added by verify import */
-	TArray<FLinkerLoad*> DelayedLinkerClosePackages;
 	/** Objects to create GC clusters from */
 	TArray<UObject*> DeferredClusterObjects;
 
@@ -1360,9 +1342,6 @@ public:
 	*/
 	EAsyncPackageState::Type PostLoadDeferredObjects();
 
-	/** Close any linkers that have been open as a result of synchronous load during async loading */
-	void CloseDelayedLinkers();
-
 private:
 	void CreateNodes(const FAsyncLoadEventSpec* EventSpecs);
 	void SetupScriptArcs();
@@ -2520,9 +2499,7 @@ void FAsyncPackage2::SetupScriptArcs()
 		check(ImportPackage);
 
 		// do initial loading stuff for compiled in packages
-		FLinkerLoad* ImportLinker = ImportPackage->LinkerLoad;
-		const bool bDynamicImport = ImportLinker && ImportLinker->bDynamicClassLinker;
-		if (!ImportLinker && ImportPackage->HasAnyPackageFlags(PKG_CompiledIn) && !bDynamicImport)
+		if (ImportPackage->HasAnyPackageFlags(PKG_CompiledIn))
 		{
 			FPackageIndex OuterMostIndex = FPackageIndex::FromImport(GlobalImportIndex);
 			FPackageIndex OuterMostNonPackageIndex = OuterMostIndex;
@@ -3222,15 +3199,6 @@ EAsyncPackageState::Type FAsyncLoadingThread2Impl::ProcessLoadedPackagesFromGame
 			// this is to ensure we can re-enter FlushAsyncLoading from any of the callbacks
 			LoadedPackagesToProcess.RemoveAt(PackageIndex--);
 
-			// Emulates ResetLoaders on the package linker's linkerroot.
-			if (!Package->IsBeingProcessedRecursively())
-			{
-				Package->ResetLoader();
-			}
-
-			// Close linkers opened by synchronous loads during async loading
-			Package->CloseDelayedLinkers();
-
 			// Incremented on the Async Thread, now decrement as we're done with this package				
 			const int32 NewExistingAsyncPackagesCounterValue = ExistingAsyncPackagesCounter.Decrement();
 
@@ -3241,7 +3209,10 @@ EAsyncPackageState::Type FAsyncLoadingThread2Impl::ProcessLoadedPackagesFromGame
 			// Call external callbacks
 			const bool bInternalCallbacks = false;
 			const EAsyncLoadingResult::Type LoadingResult = Package->HasLoadFailed() ? EAsyncLoadingResult::Failed : EAsyncLoadingResult::Succeeded;
-			Package->CallCompletionCallbacks(bInternalCallbacks, LoadingResult);
+			{
+				TRACE_CPUPROFILER_EVENT_SCOPE(PackageCompletionCallbacks);
+				Package->CallCompletionCallbacks(bInternalCallbacks, LoadingResult);
+			}
 #if WITH_EDITOR
 			// In the editor we need to find any assets and add them to list for later callback
 			Package->GetLoadedAssets(LoadedAssets);
@@ -4020,7 +3991,6 @@ FAsyncPackage2::~FAsyncPackage2()
 	TRACE_LOADTIME_DESTROY_ASYNC_PACKAGE(this);
 
 	MarkRequestIDsAsComplete();
-	DetachLinker();
 	SerialNumber = 0; // the weak pointer will always fail now
 	
 	ensure(OwnedObjects.Num() == 0);
@@ -4080,38 +4050,6 @@ double FAsyncPackage2::GetLoadStartTime() const
 	return LoadStartTime;
 }
 
-/**
- * Emulates ResetLoaders for the package's Linker objects, hence deleting it. 
- */
-void FAsyncPackage2::ResetLoader()
-{
-	LLM_SCOPE(ELLMTag::AsyncLoading);
-
-	// Reset loader.
-	if (Linker)
-	{
-		// Flush cache and queue for delete
-		Linker->FlushCache();
-		Linker->Detach();
-		// FLinkerManager::Get().RemoveLinker(Linker);
-		Linker = nullptr;
-	}
-}
-
-void FAsyncPackage2::DetachLinker()
-{	
-	if (Linker)
-	{
-		Linker->FlushCache();
-		checkf(bLoadHasFinished || bLoadHasFailed, TEXT("FAsyncPackage::DetachLinker called before load finished on package \"%s\""), *this->GetPackageName().ToString());
-		Linker = nullptr;
-	}
-}
-
-void FAsyncPackage2::FlushObjectLinkerCache()
-{
-	}
-
 #if WITH_EDITOR 
 void FAsyncPackage2::GetLoadedAssets(TArray<FWeakObjectPtr>& AssetList)
 {
@@ -4458,7 +4396,6 @@ EAsyncPackageState::Type FAsyncPackage2::FinishObjects()
 				LinkerRoot->MarkPendingKill();
 				LinkerRoot->Rename(*MakeUniqueObjectName(GetTransientPackage(), UPackage::StaticClass()).ToString(), nullptr, REN_DontCreateRedirectors | REN_DoNotDirty | REN_ForceNoResetLoaders | REN_NonTransactional);
 			}
-			DetachLinker();
 		}
 
 		LoadingResult = EAsyncLoadingResult::Failed;
@@ -4468,17 +4405,6 @@ EAsyncPackageState::Type FAsyncPackage2::FinishObjects()
 	// FLinkerManager::Get().DissociateImportsAndForcedExports(); //@todo: this should be avoidable
 	PostLoadIndex = 0;
 
-	// Keep the linkers to close until we finish loading and it's safe to close them too
-	LoadContext->MoveDelayedLinkerClosePackages(DelayedLinkerClosePackages);
-
-	if (Linker)
-	{
-		// Flush linker cache now to reduce peak memory usage (5.5-10x)
-		// We shouldn't need it anyway at this point and even if something attempts to read in PostLoad, 
-		// we're just going to re-cache then.
-		Linker->FlushCache();
-	}
-
 	const bool bInternalCallbacks = true;
 	CallCompletionCallbacks(bInternalCallbacks, LoadingResult);
 
@@ -4493,10 +4419,6 @@ EAsyncPackageState::Type FAsyncPackage2::FinishObjects()
 	return EAsyncPackageState::Complete;
 }
 
-void FAsyncPackage2::CloseDelayedLinkers()
-{
-}
-
 void FAsyncPackage2::CallCompletionCallbacks(bool bInternal, EAsyncLoadingResult::Type LoadingResult)
 {
 	checkSlow(bInternal || !IsInAsyncLoadingThread());
@@ -4528,17 +4450,12 @@ void FAsyncPackage2::Cancel()
 
 	if (LinkerRoot)
 	{
-		if (Linker)
-		{
-			Linker->FlushCache();
-		}
 		if (bCreatedLinkerRoot)
 		{
 			LinkerRoot->ClearFlags(RF_WasLoaded);
 			LinkerRoot->bHasBeenFullyLoaded = false;
 			LinkerRoot->Rename(*MakeUniqueObjectName(GetTransientPackage(), UPackage::StaticClass()).ToString(), nullptr, REN_DontCreateRedirectors | REN_DoNotDirty | REN_ForceNoResetLoaders | REN_NonTransactional);
 		}
-		ResetLoader();
 	}
 }
 
