Added unversioned property serialization - disabled by default

Loads ~2x faster and is ~6x more compact compared to tagged property serialization.

Most important differences to tagged properties are:
- Assumes code and data schema matches
- Only used in unversioned (cooked) packages
- Doesn't store field types nor converts values when field types change
- Does not handle name or GUID redirects
- Does not store memory-wise zero properties
- Based on property indices instead of property names / tags
- Metadata put in a compact header instead of interleaving tags / metadata with data

diff --git a/Engine/Config/BaseEngine.ini b/Engine/Config/BaseEngine.ini
--- a/Engine/Config/BaseEngine.ini
+++ b/Engine/Config/BaseEngine.ini
@@ -843,6 +843,8 @@ ZeroEngineVersionWarning=True
 UseStrictEngineVersioning=True
 CanStripEditorOnlyExportsAndImports=True
 CanSkipEditorReferencedPackagesWhenCooking=False
+CanUseUnversionedPropertySerialization=False
+TestUnversionedPropertySerializationWhenCooking=False
 DetailedCallstacksInNonMonolithicBuilds=True
 UseSeperateBulkDataFiles=False
 HangDuration=0.0
diff --git a/Engine/Source/Runtime/Core/Private/Serialization/Archive.cpp b/Engine/Source/Runtime/Core/Private/Serialization/Archive.cpp
--- a/Engine/Source/Runtime/Core/Private/Serialization/Archive.cpp
+++ b/Engine/Source/Runtime/Core/Private/Serialization/Archive.cpp
@@ -132,6 +132,7 @@ void FArchiveState::Reset()
 	ArIsTransacting						= false;
 	ArIsTextFormat						= false;
 	ArWantBinaryPropertySerialization	= false;
+	ArUseUnversionedPropertySerialization = false;
 	ArForceUnicode						= false;
 	ArIsPersistent						= false;
 	ArIsError							= false;
@@ -189,6 +190,7 @@ void FArchiveState::CopyTrivialFArchiveStatusMembers(const FArchiveState& Archiv
 	ArIsTransacting                      = ArchiveToCopy.ArIsTransacting;
 	ArIsTextFormat                       = ArchiveToCopy.ArIsTextFormat;
 	ArWantBinaryPropertySerialization    = ArchiveToCopy.ArWantBinaryPropertySerialization;
+	ArUseUnversionedPropertySerialization = ArchiveToCopy.ArUseUnversionedPropertySerialization;
 	ArForceUnicode                       = ArchiveToCopy.ArForceUnicode;
 	ArIsPersistent                       = ArchiveToCopy.ArIsPersistent;
 	ArIsError                            = ArchiveToCopy.ArIsError;
@@ -1132,6 +1134,11 @@ void FArchiveState::SetWantBinaryPropertySerialization(bool bInWantBinaryPropert
 	ArWantBinaryPropertySerialization = bInWantBinaryPropertySerialization;
 }
 
+void FArchiveState::SetUseUnversionedPropertySerialization(bool bInUseUnversioned)
+{
+	ArUseUnversionedPropertySerialization = bInUseUnversioned;
+}
+
 void FArchiveState::SetForceUnicode(bool bInForceUnicode)
 {
 	ArForceUnicode = bInForceUnicode;
diff --git a/Engine/Source/Runtime/Core/Public/Serialization/Archive.h b/Engine/Source/Runtime/Core/Public/Serialization/Archive.h
--- a/Engine/Source/Runtime/Core/Public/Serialization/Archive.h
+++ b/Engine/Source/Runtime/Core/Public/Serialization/Archive.h
@@ -219,6 +219,11 @@ public:
 		return ArWantBinaryPropertySerialization;
 	}
 
+	FORCEINLINE bool UseUnversionedPropertySerialization() const
+	{
+		return ArUseUnversionedPropertySerialization;
+	}
+
 	FORCEINLINE bool IsForcingUnicode() const
 	{
 		return ArForceUnicode;
@@ -610,6 +615,9 @@ protected:
 	/** Whether this archive wants properties to be serialized in binary form instead of tagged. */
 	uint8 ArWantBinaryPropertySerialization : 1;
 
+	/** Whether tagged property serialization is replaced by faster unversioned serialization. This assumes writer and reader share the same property definitions. */
+	uint8 ArUseUnversionedPropertySerialization : 1;
+
 	/** Whether this archive wants to always save strings in unicode format */
 	uint8 ArForceUnicode : 1;
 
@@ -724,6 +732,9 @@ public:
 	 */
 	virtual void SetWantBinaryPropertySerialization(bool bInWantBinaryPropertySerialization);
 
+	/** Sets whether tagged property serialization should be replaced by faster unversioned serialization. This assumes writer and reader share the same property definitions. */
+	virtual void SetUseUnversionedPropertySerialization(bool bInUseUnversioned);
+
 	/**
 	 * Sets whether this archive wants to force saving as Unicode.
 	 * This is needed when we need to make sure ANSI strings are saved as Unicode.
@@ -1592,6 +1603,7 @@ public:
 	using FArchiveState::IsTransacting;
 	using FArchiveState::IsTextFormat;
 	using FArchiveState::WantBinaryPropertySerialization;
+	using FArchiveState::UseUnversionedPropertySerialization;
 	using FArchiveState::IsForcingUnicode;
 	using FArchiveState::IsPersistent;
 	using FArchiveState::IsError;
@@ -1751,6 +1763,7 @@ private:
 	using FArchiveState::ArIsTransacting;
 	using FArchiveState::ArIsTextFormat;
 	using FArchiveState::ArWantBinaryPropertySerialization;
+	using FArchiveState::ArUseUnversionedPropertySerialization;
 	using FArchiveState::ArForceUnicode;
 	using FArchiveState::ArIsPersistent;
 
@@ -1786,6 +1799,7 @@ public:
 	using FArchiveState::SetIsTransacting;
 	using FArchiveState::SetIsTextFormat;
 	using FArchiveState::SetWantBinaryPropertySerialization;
+	using FArchiveState::SetUseUnversionedPropertySerialization;
 	using FArchiveState::SetForceUnicode;
 	using FArchiveState::SetIsPersistent;
 	using FArchiveState::SetUE4Ver;
diff --git a/Engine/Source/Runtime/CoreUObject/Private/Serialization/AsyncLoading2.cpp b/Engine/Source/Runtime/CoreUObject/Private/Serialization/AsyncLoading2.cpp
--- a/Engine/Source/Runtime/CoreUObject/Private/Serialization/AsyncLoading2.cpp
+++ b/Engine/Source/Runtime/CoreUObject/Private/Serialization/AsyncLoading2.cpp
@@ -41,6 +41,7 @@
 #include "Serialization/LoadTimeTracePrivate.h"
 #include "ProfilingDebugging/CountersTrace.h"
 #include "Serialization/AsyncPackage.h"
+#include "Serialization/UnversionedPropertySerialization.h"
 #include "Serialization/Zenaphore.h"
 #include "IO/IoDispatcher.h"
 #include "UObject/GCObject.h"
@@ -2849,6 +2850,7 @@ void FAsyncPackage2::EventDrivenSerializeExport(int32 LocalExportIndex, const ui
 	Ar.Exports = &Exports;
 	Ar.ExternalReadDependencies = &ExternalReadDependencies;
 	Ar.ExternalIoRequests = &ExternalIoRequests;
+	Ar.SetUseUnversionedPropertySerialization(Summary.bUnversioned && CanUseUnversionedPropertySerialization());
 
 	Object->ClearFlags(RF_NeedLoad);
 
diff --git a/Engine/Source/Runtime/CoreUObject/Private/Serialization/UnversionedPropertySerialization.cpp b/Engine/Source/Runtime/CoreUObject/Private/Serialization/UnversionedPropertySerialization.cpp
new file mode 100644
--- /dev/null
+++ b/Engine/Source/Runtime/CoreUObject/Private/Serialization/UnversionedPropertySerialization.cpp
@@ -0,0 +1,804 @@
+// Copyright 1998-2019 Epic Games, Inc. All Rights Reserved.
+
+#include "Serialization/UnversionedPropertySerialization.h"
+#include "Serialization/UnversionedPropertySerializationTest.h"
+#include "UObject/UnrealType.h"
+
+// Caches a property array per UStruct to avoid link-walking and touching UProperty data.
+//
+// As a reference point, this uses ~6MB of memory in an internal project and makes all
+// unversioned property loading except non-numeric SerializeItem() calls about 2x faster.
+#ifndef CACHE_UNVERSIONED_PROPERTY_SCHEMA
+// x64 platforms tend to have more memory...
+#	define CACHE_UNVERSIONED_PROPERTY_SCHEMA PLATFORM_CPU_X86_FAMILY && PLATFORM_64BITS
+#endif
+
+// Helper to pass around appropriate default value types depending on CACHE_UNVERSIONED_PROPERTY_SCHEMA
+struct FDefaultStruct
+{
+	const uint8* Data;
+#if CACHE_UNVERSIONED_PROPERTY_SCHEMA
+	uint32 StructSize;
+
+	FDefaultStruct(const uint8* InData, UStruct* InStruct) : Data(InData), StructSize(InData ? InStruct->GetPropertiesSize() : 0)  {}
+#else
+	UStruct* Struct;
+
+	FDefaultStruct(const uint8* InData, UStruct* InStruct) : Data(InData), Struct(InStruct) {}
+#endif
+};
+
+// Serializes a UProperty at a specific static array index
+//
+// Extracts and caches relevant UProperty state when using
+// CACHE_UNVERSIONED_PROPERTY_SCHEMA to improve data locality.
+// Otherwise, extracts only the needed state on demand.
+class FUnversionedPropertySerializer
+{
+public:
+	FUnversionedPropertySerializer(UProperty* InProperty, int32 InArrayIndex)
+		: Property(InProperty)
+#if CACHE_UNVERSIONED_PROPERTY_SCHEMA
+		, Offset(Property->GetOffset_ForInternal() + Property->ElementSize * InArrayIndex)
+		, bSerializeAsInteger(CanSerializeAsInteger(Property))
+		, IntType(GetIntType(Property->GetMinAlignment()))
+		, FastZeroIntNum(CanSerializeAsZero(InProperty, IntType) ? GetIntNum(InProperty, IntType) : uint8(0))
+#else
+		, ArrayIndex(InArrayIndex)
+#endif
+	{}
+
+	UProperty* GetProperty() const
+	{
+		return Property;
+	}
+
+	void* GetValue(uint8* Data) const
+	{
+#if CACHE_UNVERSIONED_PROPERTY_SCHEMA
+		return Data + Offset;
+#else
+		return Property->ContainerPtrToValuePtr<uint8>(Data, ArrayIndex);
+#endif
+	}
+
+	const void* GetValue(const uint8* Data) const
+	{
+		return GetValue(const_cast<uint8*>(Data));
+	}
+
+	const void* GetDefaultValue(FDefaultStruct Defaults) const
+	{
+#if CACHE_UNVERSIONED_PROPERTY_SCHEMA
+		return Offset < Defaults.StructSize ? Defaults.Data + Offset : nullptr;
+#else
+		return Property->ContainerPtrToValuePtrForDefaults<uint8>(Defaults.Struct, Defaults.Data);
+#endif
+	}
+
+	FORCEINLINE void Serialize(FStructuredArchive::FSlot Slot, uint8* Data, FDefaultStruct Defaults) const
+	{
+#if !CACHE_UNVERSIONED_PROPERTY_SCHEMA
+		bool bSerializeAsInteger = CanSerializeAsInteger(Property);
+#endif
+
+		if (bSerializeAsInteger)
+		{
+#if !CACHE_UNVERSIONED_PROPERTY_SCHEMA
+			EIntegerType IntType = GetIntType(Property->GetMinAlignment());
+#endif
+			SerializeAsInteger(Slot, GetValue(Data), IntType);
+		}
+		else
+		{
+			// Note that each bitfield stores a redundant byte -- the zero mask could tell us if the bit is 0 or 1.
+			//
+			// Removing this redundant byte and the FBoolProperty::SerializeItem() calls would add some complexity,
+			// but only reduce data size by ~0.5% in the data set I used to guide trade-offs. It would not make much 
+			// difference for loading performance either. --jtorp
+			Property->SerializeItem(Slot, GetValue(Data), GetDefaultValue(Defaults));
+		}
+	}
+
+	FORCEINLINE void LoadZero(uint8* Data) const
+	{
+#if !CACHE_UNVERSIONED_PROPERTY_SCHEMA
+		EIntegerType IntType = GetIntType(Property->GetMinAlignment());
+		uint32 FastZeroIntNum = GetIntNum(Property, IntType);
+#endif
+
+		if (FastZeroIntNum == 1)
+		{
+			MemZeroInt(GetValue(Data), IntType);
+		}
+		else
+		{
+			MemZeroInts(GetValue(Data), FastZeroIntNum, IntType);
+		}
+	}
+
+	bool ShouldSaveAsZero(const uint8* Data) const
+	{
+#if !CACHE_UNVERSIONED_PROPERTY_SCHEMA
+		EIntegerType IntType = GetIntType(Property->GetMinAlignment());
+		uint32 FastZeroIntNum = CanSerializeAsZero(Property, IntType) ? GetIntNum(Property, IntType) : uint8(0);
+#endif
+
+		if (FastZeroIntNum == 1)
+		{
+			return IsIntZero(GetValue(Data), IntType);
+		}
+		else
+		{
+			return FastZeroIntNum > 0 && IsIntRangeZero(GetValue(Data), FastZeroIntNum, IntType);
+		}
+	}
+
+	bool IsDefault(const uint8* Data, FDefaultStruct Defaults, uint32 PortFlags) const
+	{
+		return Property->Identical(GetValue(Data), GetDefaultValue(Defaults), PortFlags);
+	}
+
+private:
+	enum class EIntegerType : uint8 { Uint8, Uint16, Uint32, Uint64 };
+
+	static uint32 GetIntNum(const UProperty* Property, EIntegerType IntType)
+	{
+		return Property->ElementSize / GetSizeOf(IntType);
+	}
+
+	static bool CanSerializeAsZero(const UProperty* Property, EIntegerType IntType)
+	{
+		static constexpr uint32 MaxZeroComparisons = 16;
+
+		EClassCastFlags CastFlags = Property->GetClass()->ClassCastFlags;
+
+		if ((CastFlags & (CASTCLASS_UStructProperty | CASTCLASS_UBoolProperty)) == 0)
+		{
+			checkf(GetIntNum(Property, IntType) < MaxZeroComparisons, TEXT("Unexpectedly large property type encountered %s"), *Property->GetName());
+
+			return true;
+		}
+		else if ((CastFlags & CASTCLASS_UBoolProperty) != 0)
+		{
+			return static_cast<const UBoolProperty*>(Property)->IsNativeBool();
+		}
+		else
+		{
+			bool bIsAtomic = !!(static_cast<const UStructProperty*>(Property)->Struct->StructFlags & STRUCT_Atomic);
+			return bIsAtomic && GetIntNum(Property, IntType) < MaxZeroComparisons;
+		}
+	}
+
+	static bool CanSerializeAsInteger(UProperty* Property)
+	{
+		EClassCastFlags CastFlags = Property->GetClass()->ClassCastFlags;
+
+		if ((CastFlags & CASTCLASS_UBoolProperty) != 0)
+		{
+			return static_cast<const UBoolProperty*>(Property)->IsNativeBool();
+		}
+		
+		return (CastFlags & (CASTCLASS_UNumericProperty | CASTCLASS_UEnumProperty)) != 0;
+	}
+
+	template<typename T>
+	static void SerializeAs(FStructuredArchive::FSlot Slot, void* Value)
+	{
+		Slot << *reinterpret_cast<T*>(Value);
+	}
+
+	static void SerializeAsInteger(FStructuredArchive::FSlot Slot, void* Value, EIntegerType IntType)
+	{
+		static constexpr void (*Functions[4])(FStructuredArchive::FSlot, void*) = { &SerializeAs<uint8>, &SerializeAs<uint16>, &SerializeAs<uint32>, &SerializeAs<uint64> };
+		(*Functions[static_cast<uint32>(IntType)])(Slot, Value);
+	}
+
+	template<typename T>
+	static bool IsZero(const void* Value)
+	{
+		return *reinterpret_cast<const T*>(Value) == 0;
+	}
+
+	static bool IsIntZero(const void* Value, EIntegerType IntType)
+	{
+		static constexpr bool (*Functions[4])(const void*) = { &IsZero<uint8>, &IsZero<uint16>, &IsZero<uint32>, &IsZero<uint64> };
+		return (*Functions[static_cast<uint32>(IntType)])(Value);
+	}
+
+	template<typename T>
+	static bool IsRangeZero(const void* Value, uint32 Num)
+	{
+		T Out = 0;
+		for (uint32 Idx = 0; Idx < Num; ++Idx)
+		{
+			Out |= reinterpret_cast<const T*>(Value)[Idx];
+		}
+		return Out == 0;
+	}
+
+	static bool IsIntRangeZero(const void* Value, uint32 Num, EIntegerType IntType)
+	{
+		static constexpr bool (*Functions[4])(const void*, uint32 Num) = { &IsRangeZero<uint8>, &IsRangeZero<uint16>, &IsRangeZero<uint32>, &IsRangeZero<uint64> };
+		return (*Functions[static_cast<uint32>(IntType)])(Value, Num);	
+	}
+
+	template<typename T>
+	static void MemZero(void* Value)
+	{
+		*reinterpret_cast<T*>(Value) = 0;
+	}
+
+	static void MemZeroInt(void* Value, EIntegerType IntType)
+	{
+		static constexpr void (*Functions[4])(void*) = { &MemZero<uint8>, &MemZero<uint16>, &MemZero<uint32>, &MemZero<uint64> };
+		(*Functions[static_cast<uint32>(IntType)])(Value);
+	}
+
+	template<typename T>
+	static void MemZeroRange(void* Value, uint32 Num)
+	{
+		for (uint32 Idx = 0; Idx < Num; ++Idx)
+		{
+			reinterpret_cast<T*>(Value)[Idx] = 0;
+		}
+	}
+	
+	static void MemZeroInts(void* Value, uint32 Num, EIntegerType IntType)
+	{
+		static constexpr void (*Functions[4])(void*, uint32) = { &MemZeroRange<uint8>, &MemZeroRange<uint16>, &MemZeroRange<uint32>, &MemZeroRange<uint64> };
+		(*Functions[static_cast<uint32>(IntType)])(Value, Num);
+	}
+
+	static uint32 GetSizeOf(EIntegerType Type)
+	{
+		return 1u << static_cast<uint32>(Type);
+	}
+
+	// @param Bytes must be 1, 2, 4 or 8
+	static constexpr uint32 Log2For1248(uint32 Bytes)
+	{
+		return (Bytes >> 1) - (Bytes >> 3);
+	}
+
+	static EIntegerType GetIntType(uint32 Alignment)
+	{
+		if (Alignment >= 8)
+		{
+			return EIntegerType::Uint64;
+		}
+		
+		return static_cast<EIntegerType>(Log2For1248(Alignment));
+	}
+
+	UProperty* Property;
+#if CACHE_UNVERSIONED_PROPERTY_SCHEMA
+	uint32 Offset;
+	bool bSerializeAsInteger;
+	EIntegerType IntType;
+	uint8 FastZeroIntNum;
+#else
+	uint32 ArrayIndex;
+#endif
+};
+
+#if CACHE_UNVERSIONED_PROPERTY_SCHEMA
+
+// Serialization is based on indices into this property array
+struct FUnversionedStructSchema
+{
+	uint32 Num;
+	FUnversionedPropertySerializer Serializers[0];
+
+	static FUnversionedStructSchema* Create(const UStruct* Struct)
+	{
+		TArray<FUnversionedPropertySerializer, TInlineAllocator<256>> Serializers;
+		for (UProperty* Property = Struct->PropertyLink; Property; Property = Property->PropertyLinkNext)
+		{
+			if (!Property->IsEditorOnly())
+			{
+				for (int32 ArrayIdx = 0, ArrayDim = Property->ArrayDim; ArrayIdx < ArrayDim; ++ArrayIdx)
+				{
+					Serializers.Add(FUnversionedPropertySerializer(Property, ArrayIdx));
+				}
+			}
+		}
+
+		uint32 Bytes = sizeof(FUnversionedStructSchema) + Serializers.Num() * sizeof(FUnversionedPropertySerializer);
+		FUnversionedStructSchema* Schema = reinterpret_cast<FUnversionedStructSchema*>(FMemory::Malloc(Bytes, alignof(FUnversionedPropertySerializer)));
+		Schema->Num = Serializers.Num();
+		FMemory::Memcpy(Schema->Serializers, Serializers.GetData(), Serializers.Num() * sizeof(FUnversionedPropertySerializer));
+
+		return Schema;
+	}
+};
+
+const FUnversionedStructSchema& GetOrCreateUnversionedSchema(const UStruct* Struct)
+{
+	if (const FUnversionedStructSchema* ExistingSchema = Struct->UnversionedSchema)
+	{
+		return *ExistingSchema;
+	}
+	
+	FUnversionedStructSchema* CreatedSchema = FUnversionedStructSchema::Create(Struct);
+
+	void** CachedSchemaPtr = reinterpret_cast<void**>(const_cast<FUnversionedStructSchema**>(&Struct->UnversionedSchema));
+	if (const FUnversionedStructSchema* ExistingSchema = reinterpret_cast<const FUnversionedStructSchema*>(FPlatformAtomics::InterlockedCompareExchangePointer(CachedSchemaPtr, CreatedSchema, nullptr)))
+	{
+		delete CreatedSchema;
+		return *ExistingSchema;
+	}
+
+	return *CreatedSchema;
+}
+
+using FUnversionedSchemaIterator = const FUnversionedPropertySerializer*;
+
+#else
+
+// Link-walking iterator
+struct FUnversionedSchemaIterator
+{
+	UProperty* Property = nullptr;
+	uint32 ArrayIndex = 0;
+
+	FUnversionedSchemaIterator() {}
+
+	explicit FUnversionedSchemaIterator(UProperty* FirstProperty)
+		: Property(SkipEditorOnlyProperties(FirstProperty))
+	{}
+
+	void operator++()
+	{
+		check(Property);
+
+		if (ArrayIndex + 1 == Property->ArrayDim)
+		{
+			Property = SkipEditorOnlyProperties(Property->PropertyLinkNext);
+			ArrayIndex = 0;
+		}
+		else
+		{
+			++ArrayIndex;
+		}
+	}
+
+	void operator+=(uint32 Num)
+	{
+		while (Num--)
+		{
+			this->operator++();
+		}
+	}
+
+	FUnversionedPropertySerializer operator*() const
+	{
+		return FUnversionedPropertySerializer(Property, ArrayIndex);
+	}
+
+	bool operator!=(FUnversionedSchemaIterator Rhs) const
+	{
+		return (Property != Rhs.Property) | (ArrayIndex != Rhs.ArrayIndex);
+	}
+
+	static UProperty* SkipEditorOnlyProperties(UProperty* Property)
+	{
+#if WITH_EDITORONLY_DATA
+		while (Property && Property->IsEditorOnly())
+		{
+			Property = Property->PropertyLinkNext;
+		}
+#endif
+		return Property;
+	}
+};
+
+#endif // CACHE_UNVERSIONED_PROPERTY_SCHEMA
+
+struct FUnversionedSchemaRange
+{
+	explicit FUnversionedSchemaRange(const UStruct* Struct)
+	{
+#if CACHE_UNVERSIONED_PROPERTY_SCHEMA
+		const FUnversionedStructSchema& Schema = GetOrCreateUnversionedSchema(Struct);
+		Begin = Schema.Serializers;
+		End = Schema.Serializers + Schema.Num;	
+#else
+		Begin = FUnversionedSchemaIterator(Struct->PropertyLink);
+		// End is default-initialized
+#endif
+	}
+	
+	FUnversionedSchemaIterator begin() const { return Begin; }
+	FUnversionedSchemaIterator end() const { return End; }
+
+	FUnversionedSchemaIterator Begin, End;
+};
+
+// List of serialized property indices and which of them are non-zero.
+//
+// Serialized as a stream of 16-bit skip-x keep-y fragments and a zero bitmask.
+class FUnversionedHeader
+{
+public:
+	void Save(FStructuredArchive::FStream Stream) const
+	{
+		for (FFragment Fragment : Fragments)
+		{
+			uint16 Packed = Fragment.Pack();
+			Stream.EnterElement() << Packed;
+		}
+
+		if (ZeroMask.Num() > 0)
+		{
+			SaveZeroMaskData(Stream, ZeroMask.Num(), ZeroMask.GetData());
+		}
+	}
+
+	void Load(FStructuredArchive::FStream Stream)
+	{
+		uint32 ZeroMaskSize = 0;
+		do
+		{
+			uint16 Packed;
+			Stream.EnterElement() << Packed;
+			FFragment Fragment = FFragment::Unpack(Packed);
+
+			Fragments.Add(Fragment);
+
+			ZeroMaskSize += Fragment.ValueNum * Fragment.bHasAnyZeroes;
+			bHasNonZeroValues |= (Fragment.ValueNum > 0) & (!Fragment.bHasAnyZeroes);
+		}
+		while (!Fragments.Last().bIsLast);
+
+		if (ZeroMaskSize > 0)
+		{
+			ZeroMask.Init(false, ZeroMaskSize);
+			LoadZeroMaskData(Stream, ZeroMaskSize, ZeroMask.GetData());
+			bHasNonZeroValues = bHasNonZeroValues || ZeroMask.Find(false) != INDEX_NONE;
+		}
+	}
+	
+	bool HasValues() const
+	{
+		return bHasNonZeroValues | (ZeroMask.Num() > 0);
+	}
+
+	bool HasNonZeroValues() const
+	{
+		return bHasNonZeroValues;
+	}
+
+protected:
+	struct FFragment
+	{
+		static constexpr uint32 SkipMax = 127;
+		static constexpr uint32 ValueMax = 127;
+
+		uint8 SkipNum = 0; // Number of properties to skip before values
+		bool bHasAnyZeroes = false;
+		uint8 ValueNum = 0;  // Number of subsequent property values stored
+		bool bIsLast = 0; // Is this the last fragment of the header?
+	
+		static constexpr uint32 SkipNumMask = 0x007fu;
+		static constexpr uint32 HasZeroMask = 0x0080u;
+		static constexpr uint32 ValueNumShift = 9u;
+		static constexpr uint32 IsLastMask  = 0x0100u;
+
+		uint16 Pack() const
+		{
+			return SkipNum | (bHasAnyZeroes ? HasZeroMask : 0) | (uint16)ValueNum << ValueNumShift | (bIsLast ? IsLastMask : 0);
+		}
+
+		static FFragment Unpack(uint16 Int)
+		{
+			FFragment Fragment;
+			Fragment.SkipNum = static_cast<uint8>(Int & SkipNumMask);
+			Fragment.bHasAnyZeroes = (Int & HasZeroMask) != 0;
+			Fragment.ValueNum = static_cast<uint8>(Int >> ValueNumShift);
+			Fragment.bIsLast = (Int & IsLastMask) != 0;
+			return Fragment;
+		}
+	};
+
+	using FZeroMask = TBitArray<TInlineAllocator<8>>;
+
+	alignas (PLATFORM_CACHE_LINE_SIZE) TArray<FFragment, TInlineAllocator<32>> Fragments;
+	bool bHasNonZeroValues = false;
+	FZeroMask ZeroMask;
+
+	void SaveZeroMaskData(FStructuredArchive::FStream Stream, uint32 NumBits, const uint32* Data) const
+	{
+		check(NumBits > 0);
+
+		uint32 LastWordMask = ~0u >> (32 - (NumBits % 32));
+		if (NumBits <= 8)
+		{
+			uint8 Word = static_cast<uint8>(*Data & LastWordMask);
+			checkf(Word != 0, TEXT("Zero mask shouldn't be saved when no bits are set"));
+			Stream << Word;
+		}
+		else if (NumBits <= 16)
+		{
+			uint16 Word = static_cast<uint16>(*Data & LastWordMask);
+			checkf(Word != 0, TEXT("Zero mask shouldn't be saved when no bits are set"));
+			Stream << Word;
+		}
+		else
+		{
+			uint32 NumWords = FMath::DivideAndRoundUp(NumBits, 32u);
+
+			for (uint32 WordIdx = 0; WordIdx < NumWords - 1; ++WordIdx)
+			{
+				uint32 Word = Data[WordIdx];
+				Stream << Word;
+			}
+
+			uint32 LastWord = Data[NumWords - 1] & LastWordMask;
+			Stream << LastWord;
+		}
+	}
+
+	void LoadZeroMaskData(FStructuredArchive::FStream Stream, uint32 NumBits, uint32* Data)
+	{
+		if (NumBits <= 8)
+		{
+			uint8 Int;
+			Stream << Int;
+			*Data = Int;
+		}
+		else if (NumBits <= 16)
+		{
+			uint16 Int;
+			Stream << Int;
+			*Data = Int;
+		}
+		else
+		{
+			for (uint32 Idx = 0, Num = FMath::DivideAndRoundUp(NumBits, 32u); Idx < Num; ++Idx)
+			{
+				Stream << Data[Idx];
+			}
+		}
+	}
+
+public:
+	class FIterator
+	{
+	public:
+		FORCEINLINE FIterator(const FUnversionedHeader& Header, const FUnversionedSchemaRange& Schema)
+			: SchemaIt(Schema.Begin)
+			, ZeroMask(Header.ZeroMask)
+			, FragmentIt(Header.Fragments.GetData())
+			, bDone(!Header.HasValues())
+#if DO_CHECK
+			, SchemaEnd(Schema.End)
+#endif
+		{
+			if (!bDone)
+			{
+				Skip();
+			}
+		}
+
+		void Next()
+		{
+			++SchemaIt;
+			--RemainingFragmentValues;
+			ZeroMaskIndex += FragmentIt->bHasAnyZeroes;
+
+			if (RemainingFragmentValues == 0)
+			{
+				if (FragmentIt->bIsLast)
+				{
+					bDone = true;
+				}
+				else
+				{
+					++FragmentIt;
+					Skip();
+				}
+			}
+		}
+
+		explicit operator bool() const
+		{
+			return !bDone;
+		}
+
+		FUnversionedPropertySerializer GetSerializer() const
+		{
+			check(SchemaIt != SchemaEnd);
+			return *SchemaIt;
+		}
+
+		bool IsNonZero() const
+		{
+			return !FragmentIt->bHasAnyZeroes || !ZeroMask[ZeroMaskIndex];
+		}
+
+	private:
+		FUnversionedSchemaIterator SchemaIt;
+		const FZeroMask& ZeroMask;
+		const FFragment* FragmentIt = nullptr;
+		bool bDone = false;
+		uint32 ZeroMaskIndex = 0;
+		uint32 RemainingFragmentValues = 0;
+#if DO_CHECK
+		FUnversionedSchemaIterator SchemaEnd;
+#endif
+
+		void Skip()
+		{
+			SchemaIt += FragmentIt->SkipNum;
+			
+			while (FragmentIt->ValueNum == 0)
+			{
+				check(!FragmentIt->bIsLast);
+				++FragmentIt;
+				SchemaIt += FragmentIt->SkipNum;
+			}
+
+			RemainingFragmentValues = FragmentIt->ValueNum;
+			check(SchemaIt + RemainingFragmentValues <= SchemaEnd);
+		}
+	};
+};
+
+class FUnversionedHeaderBuilder : public FUnversionedHeader
+{
+public:
+	FUnversionedHeaderBuilder()
+	{
+		Fragments.AddDefaulted();
+	}
+
+	void IncludeProperty(bool bIsZero)
+	{
+		if (Fragments.Last().ValueNum == FFragment::ValueMax)
+		{
+			TrimZeroMask(Fragments.Last());
+			Fragments.AddDefaulted();
+		}
+		
+		++Fragments.Last().ValueNum;
+		Fragments.Last().bHasAnyZeroes |= bIsZero;
+		ZeroMask.Add(bIsZero);
+		bHasNonZeroValues |= !bIsZero;
+	}
+
+	void ExcludeProperty()
+	{
+		if (Fragments.Last().ValueNum || Fragments.Last().SkipNum == FFragment::SkipMax)
+		{
+			TrimZeroMask(Fragments.Last());
+			Fragments.AddDefaulted();
+		}
+
+		++Fragments.Last().SkipNum;
+	}
+
+	void Finalize()
+	{
+		TrimZeroMask(Fragments.Last());
+
+		// Trim trailing skips
+		while (Fragments.Num() > 1 && Fragments.Last().ValueNum == 0)
+		{
+			checkf(!Fragments.Last().bHasAnyZeroes, TEXT("No values implies no zero-values"));			
+			Fragments.Pop(/* allow shrink */ false);
+		}
+
+		Fragments.Last().bIsLast = true;
+	}
+
+private:
+	void TrimZeroMask(const FFragment& Fragment)
+	{
+		if (!Fragment.bHasAnyZeroes)
+		{
+			ZeroMask.RemoveAt(ZeroMask.Num() - Fragment.ValueNum, Fragment.ValueNum);
+		}
+	}
+};
+
+bool CanUseUnversionedPropertySerialization()
+{
+	bool bTemp;
+	static bool bAllow = GConfig->GetBool(TEXT("Core.System"), TEXT("CanUseUnversionedPropertySerialization"), bTemp, GEngineIni) && bTemp;
+	return bAllow;
+}
+
+void DestroyUnversionedSchema(const UStruct* Struct)
+{
+#if CACHE_UNVERSIONED_PROPERTY_SCHEMA
+	delete Struct->UnversionedSchema;
+	Struct->UnversionedSchema = nullptr;
+#endif
+}
+
+void SerializeUnversionedProperties(const UStruct* Struct, FStructuredArchive::FSlot Slot, uint8* Data, UStruct* DefaultsStruct, uint8* DefaultsData)
+{
+	check(CanUseUnversionedPropertySerialization());
+
+	FArchive& UnderlyingArchive = Slot.GetUnderlyingArchive();
+	FStructuredArchive::FRecord StructRecord = Slot.EnterRecord();
+
+	if (UnderlyingArchive.IsLoading())
+	{
+		FUnversionedHeader Header;
+		Header.Load(StructRecord.EnterStream(SA_FIELD_NAME(TEXT("Header"))));
+
+		if (Header.HasValues())
+		{
+			FUnversionedSchemaRange Schema(Struct);
+
+			if (Header.HasNonZeroValues())
+			{
+				FDefaultStruct Defaults(DefaultsData, DefaultsStruct);
+
+				FStructuredArchive::FStream ValueStream = StructRecord.EnterStream(SA_FIELD_NAME(TEXT("Values")));
+				for (FUnversionedHeader::FIterator It(Header, Schema); It; It.Next())
+				{
+					if (It.IsNonZero())
+					{
+						It.GetSerializer().Serialize(ValueStream.EnterElement(), Data, Defaults);
+					}
+					else
+					{
+						It.GetSerializer().LoadZero(Data);
+					}
+				}
+			}
+			else
+			{
+				for (FUnversionedHeader::FIterator It(Header, Schema); It; It.Next())
+				{
+					check(!It.IsNonZero());
+					It.GetSerializer().LoadZero(Data);
+				}	
+			}
+		}
+
+	}
+	else
+	{
+		FUnversionedPropertyTestRunner TestRunner({Struct, Data, DefaultsStruct, DefaultsData});
+		FUnversionedPropertyTestCollector TestCollector;
+
+		// Decide which properties to save and build header based on schema property indices
+		const bool bDense = !UnderlyingArchive.DoDelta() || UnderlyingArchive.IsTransacting() || (!DefaultsData && !dynamic_cast<const UClass*>(Struct));
+		FDefaultStruct Defaults(DefaultsData, DefaultsStruct);
+
+		FUnversionedSchemaRange Schema(Struct);
+		FUnversionedHeaderBuilder Header;
+		for (FUnversionedPropertySerializer Serializer : Schema)
+		{
+			if (Serializer.GetProperty()->ShouldSerializeValue(UnderlyingArchive) &&
+				(bDense || !Serializer.IsDefault(Data, Defaults, UnderlyingArchive.GetPortFlags())))
+			{
+				Header.IncludeProperty(Serializer.ShouldSaveAsZero(Data));
+				TestCollector.RecordSavedProperty(Serializer.GetProperty());
+			}
+			else
+			{
+				Header.ExcludeProperty();
+			}
+		}
+		Header.Finalize();
+
+		// Save header and non-zero values 
+		Header.Save(StructRecord.EnterStream(SA_FIELD_NAME(TEXT("Header"))));
+		if (Header.HasNonZeroValues())
+		{
+			FStructuredArchive::FStream ValueStream = StructRecord.EnterStream(SA_FIELD_NAME(TEXT("Values")));
+			for (FUnversionedHeader::FIterator It(Header, Schema); It; It.Next())
+			{
+				if (It.IsNonZero())
+				{
+					It.GetSerializer().Serialize(ValueStream.EnterElement(), Data, Defaults);
+				}
+			}
+		}
+	}
+}
\ No newline at end of file
diff --git a/Engine/Source/Runtime/CoreUObject/Private/Serialization/UnversionedPropertySerialization.h b/Engine/Source/Runtime/CoreUObject/Private/Serialization/UnversionedPropertySerialization.h
new file mode 100644
--- /dev/null
+++ b/Engine/Source/Runtime/CoreUObject/Private/Serialization/UnversionedPropertySerialization.h
@@ -0,0 +1,15 @@
+// Copyright 1998-2019 Epic Games, Inc. All Rights Reserved.
+
+#pragma once
+
+#include "UObject/UnrealType.h"
+
+struct FUnversionedStructSchema;
+
+// Check if unversioned property serialization is configured to be used
+bool CanUseUnversionedPropertySerialization();
+
+// Serialize sparse unversioned properties for a particular struct
+void SerializeUnversionedProperties(const UStruct* Struct, FStructuredArchive::FSlot Slot, uint8* Data, UStruct* DefaultsStruct, uint8* Defaults);
+
+void DestroyUnversionedSchema(const UStruct* Struct);
\ No newline at end of file
diff --git a/Engine/Source/Runtime/CoreUObject/Private/Serialization/UnversionedPropertySerializationTest.cpp b/Engine/Source/Runtime/CoreUObject/Private/Serialization/UnversionedPropertySerializationTest.cpp
new file mode 100644
--- /dev/null
+++ b/Engine/Source/Runtime/CoreUObject/Private/Serialization/UnversionedPropertySerializationTest.cpp
@@ -0,0 +1,405 @@
+// Copyright 1998-2019 Epic Games, Inc. All Rights Reserved.
+
+#include "Serialization/UnversionedPropertySerializationTest.h"
+#include "ProfilingDebugging/CookStats.h"
+#include "Serialization/MemoryReader.h"
+#include "Serialization/MemoryWriter.h"
+#include "UObject/UnrealType.h"
+
+#if UE_ENABLE_UNVERSIONED_PROPERTY_TEST
+
+namespace PropertySerializationStats
+{
+	static TAtomic<uint64> Structs = 0;
+	static TAtomic<uint64> VersionedBytes = 0;
+	static TAtomic<uint64> UnversionedBytes = 0;
+	static TAtomic<uint64> UselessBytes = 0;
+
+#if ENABLE_COOK_STATS
+	static FCookStatsManager::FAutoRegisterCallback RegisterCookStats([](FCookStatsManager::AddStatFuncRef AddStat)
+	{
+		AddStat(TEXT("UnversionedProperties"), 
+				FCookStatsManager::CreateKeyValueArray(
+					TEXT("SavedStructs"), Structs.Load(),
+					TEXT("SavedMB"), static_cast<uint32>(UnversionedBytes.Load() >> 20),
+					TEXT("EquivalentTaggedMB"), static_cast<uint32>(VersionedBytes.Load() >> 20),
+					TEXT("CompressionRatio"), static_cast<float>(VersionedBytes.Load()) / UnversionedBytes.Load(),
+					TEXT("BitfieldWasteKB"), static_cast<uint32>(UselessBytes.Load()) >> 10));
+	});
+#endif
+}
+
+// Serializes a UStruct to memory using both unversioned and versioned tagged property serialization,
+// then creates two struct instances, loads the data back and compares that they are identical.
+struct FUnversionedPropertyTest : public FUnversionedPropertyTestInput
+{
+	explicit FUnversionedPropertyTest(const FUnversionedPropertyTestInput& Input) : FUnversionedPropertyTestInput(Input) {}
+
+	class FTestLinker : public FArchiveProxy
+	{
+	public:
+		using FArchiveProxy::FArchiveProxy;
+
+		virtual FArchive& operator<<(FName& Value) override
+		{
+			uint32 UnstableInt = Value.GetDisplayIndex().ToUnstableInt();
+			int32 Number = Value.GetNumber();
+			InnerArchive << UnstableInt << Number;
+
+			if (IsLoading())
+			{
+				Value = FName::CreateFromDisplayId(FNameEntryId::FromUnstableInt(UnstableInt), Number);
+			}
+
+			return *this;
+		}
+
+		virtual FArchive& operator<<(UObject*& Value) override
+		{
+			return InnerArchive << reinterpret_cast<UPTRINT&>(Value);
+		}
+
+		virtual FArchive& operator<<(FLazyObjectPtr& Value) override { return FArchiveUObject::SerializeLazyObjectPtr(*this, Value); }
+		virtual FArchive& operator<<(FSoftObjectPtr& Value) override { return FArchiveUObject::SerializeSoftObjectPtr(*this, Value); }
+		virtual FArchive& operator<<(FSoftObjectPath& Value) override { return FArchiveUObject::SerializeSoftObjectPath(*this, Value); }
+		virtual FArchive& operator<<(FWeakObjectPtr& Value) override { return FArchiveUObject::SerializeWeakObjectPtr(*this, Value); }
+	};
+
+	enum class EPath { Versioned, Unversioned };
+
+	static const TCHAR* ToString(EPath Path)
+	{
+		return Path == EPath::Unversioned ? TEXT("unversioned") : TEXT("versioned");
+	}
+
+	struct FSaveResult
+	{
+		TArray<uint8> Data;
+		TArray<UProperty*> Properties;
+		EPath Path;
+	};
+
+	static thread_local FSaveResult* TlsSaveResult;
+
+	FSaveResult Save(EPath Path)
+	{
+		FSaveResult Result;
+		Result.Path = Path;
+		
+		FMemoryWriter Writer(Result.Data);
+		Writer.SetUseUnversionedPropertySerialization(Path == EPath::Unversioned);
+		FTestLinker Linker(Writer);
+		FBinaryArchiveFormatter Formatter(Linker);
+		FStructuredArchive StructuredArchive(Formatter);
+		FStructuredArchive::FSlot Slot = StructuredArchive.Open();
+
+		TlsSaveResult = &Result;
+		Struct->SerializeTaggedProperties(Slot, OriginalInstance, DefaultsStruct, Defaults);
+		check(TlsSaveResult == nullptr);
+
+		return Result;
+	}
+
+	struct FTestInstance
+	{
+		explicit FTestInstance(const UStruct* InType)
+			: Type(InType)
+		{
+			Instance = FMemory::Malloc(Type->GetStructureSize(), Type->GetMinAlignment());
+			Type->InitializeStruct(Instance);
+		}
+
+		FTestInstance(FTestInstance&& Other)
+			: Type(Other.Type)
+			, Instance(Other.Instance)
+		{
+			Other.Instance = nullptr;
+		}
+
+		~FTestInstance()
+		{
+			if (Instance)
+			{
+				Type->DestroyStruct(Instance);
+				FMemory::Free(Instance);
+			}
+		}
+
+		const UStruct* Type;
+		void* Instance;
+	};
+
+	FTestInstance Load(const FSaveResult& Saved)
+	{
+		FMemoryReader Reader(Saved.Data);
+		Reader.SetUseUnversionedPropertySerialization(Saved.Path == EPath::Unversioned);
+		FTestLinker Linker(Reader);
+		FBinaryArchiveFormatter Formatter(Linker);
+		FStructuredArchive StructuredArchive(Formatter);
+		FStructuredArchive::FSlot Slot = StructuredArchive.Open();
+
+		TGuardValue<bool> Guard(GIsSavingPackage, false);
+		FTestInstance Result(Struct);
+		// Call UStruct::SerializeTaggedProperties() directly to bypass
+		// UUserDefinedStruct::SerializeTaggedProperties() for test loading,
+		// since that is what the test saving does.
+		Struct->UStruct::SerializeTaggedProperties(Slot, (uint8*)Result.Instance, DefaultsStruct, Defaults);
+
+		checkf(Reader.Tell() == Saved.Data.Num(), TEXT("Failed to consume all %s saved property data"), ToString(Saved.Path));
+
+		return Result;
+	}
+
+	static constexpr uint32 EqualsPortFlags = 0;
+
+	// UProperty::Identical() flavor suited to comparing loaded instances
+	static bool Equals(const UProperty* Property, const void* A, const void* B)
+	{
+		if (Property->GetPropertyFlags() & (CPF_EditorOnly | CPF_Transient))
+		{
+			return true;
+		}
+		else if (const UStructProperty* StructProperty = Cast<UStructProperty>(Property))
+		{
+			return Equals(StructProperty, A, B);
+		}
+		else if (const UArrayProperty* ArrayProperty = Cast<UArrayProperty>(Property))
+		{
+			return Equals(ArrayProperty, A, B);
+		}
+		else if (const USetProperty* SetProperty = Cast<USetProperty>(Property))
+		{
+			return Equals(SetProperty, A, B);
+		}
+		else if (const UMapProperty* MapProperty = Cast<UMapProperty>(Property))
+		{
+			return Equals(MapProperty, A, B);
+		}
+
+		return Property->Identical(A, B, EqualsPortFlags);
+	}
+
+	static bool Equals(const UArrayProperty* Property, const void* A, const void* B)
+	{
+		FScriptArrayHelper HelperA(Property, A);
+		FScriptArrayHelper HelperB(Property, B);
+
+		if (HelperA.Num() != HelperB.Num())
+		{
+			return false;
+		}
+
+		for (int32 Idx = 0, Num = HelperA.Num(); Idx < Num; ++Idx)
+		{
+			if (!Equals(Property->Inner, HelperA.GetRawPtr(Idx), HelperB.GetRawPtr(Idx)))
+			{
+				return false;
+			}
+		}
+
+		return true;
+	}
+	
+	static bool Equals(const USetProperty* Property, const void* A, const void* B)
+	{
+		FScriptSetHelper HelperA(Property, A);
+		FScriptSetHelper HelperB(Property, B);
+		const UProperty* ElemProperty   = HelperA.GetElementProperty();
+
+		if (HelperA.Num() != HelperB.Num())
+		{
+			return false;
+		}
+
+		for (int32 Num = HelperA.Num(), IndexA = 0; Num > 0; --Num)
+		{
+			while (!HelperA.IsValidIndex(IndexA))
+			{
+				++IndexA;
+			}
+			
+			const uint8* ElemA = HelperA.GetElementPtr(IndexA);
+			const uint8* ElemB = HelperB.FindElementPtrFromHash(ElemA);
+		
+			if (!ElemB || !Equals(ElemProperty, ElemA, ElemB))
+			{
+				return false;
+			} 
+		}
+
+		return true;
+	}
+	
+	static bool Equals(const UMapProperty* Property, const void* A, const void* B)
+	{
+		FScriptMapHelper HelperA(Property, A);
+		FScriptMapHelper HelperB(Property, B);
+		const UProperty* KeyProp   = HelperA.GetKeyProperty();
+		const UProperty* ValueProp = HelperA.GetValueProperty();
+		int32 ValueOffset = HelperA.MapLayout.ValueOffset;
+
+		if (HelperA.Num() != HelperB.Num())
+		{
+			return false;
+		}
+		
+		for (int32 Num = HelperA.Num(), IndexA = 0; Num > 0; --Num)
+		{
+			while (!HelperA.IsValidIndex(IndexA))
+			{
+				++IndexA;
+			}
+			
+			const uint8* PairA = HelperA.GetPairPtr(IndexA);
+			const uint8* PairB = HelperB.FindMapPairPtrFromHash(PairA);
+		
+			if (!PairB || !Equals(KeyProp, PairA, PairB))
+			{
+				return false;
+			} 
+
+			if (!Equals(ValueProp, PairA + ValueOffset, PairB + ValueOffset))
+			{
+				return false;
+			}
+		}
+
+		return true;
+	}
+
+	static bool Equals(const UStructProperty* Property, const void* A, const void* B)
+	{
+		UScriptStruct* Struct = Property->Struct;
+		if (Struct->StructFlags & STRUCT_IdenticalNative)
+		{
+			bool bResult = false;
+			if (Struct->GetCppStructOps()->Identical(A, B, EqualsPortFlags, bResult))
+			{
+				return bResult;
+			}
+		}
+
+		// Skip deprecated fields
+		for (TFieldIterator<UProperty> It(Struct, EFieldIteratorFlags::IncludeSuper, EFieldIteratorFlags::ExcludeDeprecated); It; ++It)
+		{
+			for (int32 Idx = 0, MaxIdx = It->ArrayDim; Idx < MaxIdx; ++Idx)
+			{
+				if (!Equals_InContainer(*It, A, B, Idx))
+				{
+					return false;
+				}
+			}
+		}
+		return true;
+	}
+
+	static bool Equals_InContainer(const UProperty* Property, const void* A, const void* B, uint32 Idx)
+	{
+		return Equals(Property, Property->ContainerPtrToValuePtr<void>(A, Idx), Property->ContainerPtrToValuePtr<void>(B, Idx));
+	}
+
+	static FString GetValueAsText(UProperty* Property, uint32 ArrayIdx, void* StructInstance)
+	{
+		FString Value;
+		Property->ExportText_InContainer(ArrayIdx, Value, StructInstance, nullptr, nullptr, 0);
+		return Value;
+	}
+
+
+	void CheckEqual(UProperty* Property, void* VersionedInstance, void* UnversionedInstance)
+	{
+		for (int32 Idx = 0, Num = Property->ArrayDim; Idx < Num; ++Idx)
+		{
+			if (!Equals_InContainer(Property, VersionedInstance, UnversionedInstance, Idx))
+			{
+				bool VersionedOk = Equals_InContainer(Property, VersionedInstance, OriginalInstance, Idx);
+				bool UnversionedOk = Equals_InContainer(Property, UnversionedInstance, OriginalInstance, Idx);
+				const TCHAR* OkPaths = VersionedOk ? (UnversionedOk ? TEXT("Both paths") : TEXT("Versioned path"))
+													: (UnversionedOk ? TEXT("Unversioned path") : TEXT("Neither path"));
+
+				// When debugging test failures, put a breakpoint inside this if statement
+				if (FPlatformMisc::IsDebuggerPresent())
+				{
+					// These strings might be too long to fit in the assert message.
+					// We could write traversal code that identifies which nested property differs
+					// and only generates a text representation for that value. 
+					FString VersionedValue = GetValueAsText(Property, Idx, VersionedInstance);
+					FString UnversionedValue = GetValueAsText(Property, Idx, UnversionedInstance);
+					FString OriginalValue = GetValueAsText(Property, Idx, OriginalInstance);
+
+					FSaveResult VersionedSaved2 = Save(EPath::Versioned);
+					FSaveResult UnversionedSaved2 = Save(EPath::Unversioned);
+
+					FTestInstance VersionedLoaded2 = Load(VersionedSaved2);
+					FTestInstance UnversionedLoaded2 = Load(UnversionedSaved2);
+				}		
+
+				checkf(false, TEXT("The %s %s.%s roundtripped differently in versioned vs unversioned tagged property serialization. "
+					"%s loaded an instance equal to the original."), *Property->GetClass()->GetName(), *Struct->GetName(), *Property->GetName(), OkPaths);
+			}
+		}		
+	}
+
+	static TArray<UProperty*> ExcludeEditorOnlyProperties(const TArray<UProperty*>& Properties)
+	{
+		TArray<UProperty*> Out;
+		Out.Reserve(Properties.Num());
+
+		for (UProperty* Property : Properties)
+		{
+			if (!Property->IsEditorOnly())
+			{
+				Out.Add(Property);
+			}
+		}
+
+		return Out;
+	}
+
+	void Run()
+	{
+		FSaveResult VersionedSaved = Save(EPath::Versioned);
+		FSaveResult UnversionedSaved = Save(EPath::Unversioned);	
+
+		check(ExcludeEditorOnlyProperties(VersionedSaved.Properties) == UnversionedSaved.Properties);
+
+		FTestInstance VersionedLoaded = Load(VersionedSaved);
+		FTestInstance UnversionedLoaded = Load(UnversionedSaved);
+
+		for (UProperty* Property : UnversionedSaved.Properties)
+		{
+			CheckEqual(Property, VersionedLoaded.Instance, UnversionedLoaded.Instance);
+			PropertySerializationStats::UselessBytes += Property->IsA<UBoolProperty>() && !Cast<UBoolProperty>(Property)->IsNativeBool();
+		}
+		
+		++PropertySerializationStats::Structs;
+		PropertySerializationStats::VersionedBytes += VersionedSaved.Data.Num();
+		PropertySerializationStats::UnversionedBytes += UnversionedSaved.Data.Num();
+	}
+};
+
+thread_local FUnversionedPropertyTest::FSaveResult* FUnversionedPropertyTest::TlsSaveResult;
+thread_local bool FUnversionedPropertyTestRunner::bTlsTesting;
+
+void RunUnversionedPropertyTest(const FUnversionedPropertyTestInput& Input)
+{
+	FUnversionedPropertyTest Test(Input);
+	Test.Run();
+}
+
+FUnversionedPropertyTestCollector::FUnversionedPropertyTestCollector()
+{
+	if (FUnversionedPropertyTest::FSaveResult* Result = FUnversionedPropertyTest::TlsSaveResult)
+	{
+		Out = &Result->Properties;
+
+		// Nested SerializeTaggedProperties() call should not record nested properties
+		FUnversionedPropertyTest::TlsSaveResult = nullptr;
+	}
+	else
+	{
+		Out = nullptr;
+	}
+}
+
+#endif
\ No newline at end of file
diff --git a/Engine/Source/Runtime/CoreUObject/Private/Serialization/UnversionedPropertySerializationTest.h b/Engine/Source/Runtime/CoreUObject/Private/Serialization/UnversionedPropertySerializationTest.h
new file mode 100644
--- /dev/null
+++ b/Engine/Source/Runtime/CoreUObject/Private/Serialization/UnversionedPropertySerializationTest.h
@@ -0,0 +1,78 @@
+// Copyright 1998-2019 Epic Games, Inc. All Rights Reserved.
+
+#pragma once
+
+#include "Misc/ConfigCacheIni.h"
+
+#define UE_ENABLE_UNVERSIONED_PROPERTY_TEST WITH_EDITORONLY_DATA
+
+struct FUnversionedPropertyTestInput
+{
+	const UStruct* Struct;
+	uint8* const OriginalInstance;
+	UStruct* const DefaultsStruct;
+	uint8* const Defaults;
+};
+
+#if UE_ENABLE_UNVERSIONED_PROPERTY_TEST
+
+void RunUnversionedPropertyTest(const FUnversionedPropertyTestInput& Input);
+
+// Avoids starting tests recursively when we test serializing nested structs or struct containers.
+// The test is both started from and recursively calls SerializeTaggedProperties().
+struct FUnversionedPropertyTestRunner
+{
+	static thread_local bool bTlsTesting;
+	bool bStartedTest = false;
+
+	explicit FUnversionedPropertyTestRunner(const FUnversionedPropertyTestInput& Input)
+	{
+		bool bTemp;
+		static bool bEnabled = GConfig->GetBool(TEXT("Core.System"), TEXT("TestUnversionedPropertySerializationWhenCooking"), bTemp, GEngineIni) && bTemp;
+			
+		if (bEnabled && !bTlsTesting)
+		{
+			bTlsTesting = true;
+			bStartedTest = true;
+			RunUnversionedPropertyTest(Input);
+		}
+	}
+
+	~FUnversionedPropertyTestRunner()
+	{
+		if (bStartedTest)
+		{
+			bTlsTesting = false;
+		}
+	}
+};
+
+// Intrusive test helper that records which properties were saved
+struct FUnversionedPropertyTestCollector
+{
+	TArray<UProperty*>* Out;
+
+	FUnversionedPropertyTestCollector();
+	
+	void RecordSavedProperty(UProperty* Property)
+	{
+		if (Out)
+		{
+			Out->Add(Property);
+		}
+	}
+};
+
+#else  // !UE_ENABLE_UNVERSIONED_PROPERTY_TEST
+
+struct FUnversionedPropertyTestRunner
+{	
+	explicit FUnversionedPropertyTestRunner(const FUnversionedPropertyTestInput& Input) {}
+};
+
+struct FUnversionedPropertyTestCollector
+{
+	void RecordSavedProperty(UProperty* Property) {}
+};
+
+#endif
\ No newline at end of file
diff --git a/Engine/Source/Runtime/CoreUObject/Private/UObject/Class.cpp b/Engine/Source/Runtime/CoreUObject/Private/UObject/Class.cpp
--- a/Engine/Source/Runtime/CoreUObject/Private/UObject/Class.cpp
+++ b/Engine/Source/Runtime/CoreUObject/Private/UObject/Class.cpp
@@ -44,6 +44,8 @@
 #include "UObject/GarbageCollection.h"
 #include "UObject/UObjectThreadContext.h"
 #include "Serialization/LoadTimeTracePrivate.h"
+#include "Serialization/UnversionedPropertySerialization.h"
+#include "Serialization/UnversionedPropertySerializationTest.h"
 
 // This flag enables some expensive class tree validation that is meant to catch mutations of 
 // the class tree outside of SetSuperStruct. It has been disabled because loading blueprints 
@@ -1110,6 +1112,18 @@ void UStruct::LoadTaggedPropertiesFromText(FStructuredArchive::FSlot Slot, uint8
 }
 
 void UStruct::SerializeTaggedProperties(FStructuredArchive::FSlot Slot, uint8* Data, UStruct* DefaultsStruct, uint8* Defaults, const UObject* BreakRecursionIfFullyLoad) const
+{
+	if (Slot.GetUnderlyingArchive().UseUnversionedPropertySerialization())
+	{
+		SerializeUnversionedProperties(this, Slot, Data, DefaultsStruct, Defaults);
+	}
+	else
+	{
+		SerializeVersionedTaggedProperties(Slot, Data, DefaultsStruct, Defaults, BreakRecursionIfFullyLoad);
+	}
+}
+
+void UStruct::SerializeVersionedTaggedProperties(FStructuredArchive::FSlot Slot, uint8* Data, UStruct* DefaultsStruct, uint8* Defaults, const UObject* BreakRecursionIfFullyLoad) const
 {
 	FArchive& UnderlyingArchive = Slot.GetUnderlyingArchive();
 	//SCOPED_LOADTIMER(SerializeTaggedPropertiesTime);
@@ -1325,6 +1339,8 @@ void UStruct::SerializeTaggedProperties(FStructuredArchive::FSlot Slot, uint8* D
 	}
 	else
 	{
+		FUnversionedPropertyTestCollector TestCollector;
+
 		FStructuredArchive::FRecord PropertiesRecord = Slot.EnterRecord();
 
 		check(UnderlyingArchive.IsSaving() || UnderlyingArchive.IsCountingMemory());
@@ -1380,6 +1396,8 @@ void UStruct::SerializeTaggedProperties(FStructuredArchive::FSlot Slot, uint8* D
 						FArchive::FScopeAddDebugData P(UnderlyingArchive, NAME_PropertySerialize);
 						FArchive::FScopeAddDebugData S(UnderlyingArchive, Property->GetFName());
 #endif
+						TestCollector.RecordSavedProperty(Property);
+
 						FPropertyTag Tag( UnderlyingArchive, Property, Idx, DataPtr, DefaultValue );
 						// If available use the property guid from BlueprintGeneratedClasses, provided we aren't cooking data.
 						if (bArePropertyGuidsAvailable && !UnderlyingArchive.IsCooking())
@@ -1440,6 +1458,7 @@ void UStruct::SerializeTaggedProperties(FStructuredArchive::FSlot Slot, uint8* D
 }
 void UStruct::FinishDestroy()
 {
+	DestroyUnversionedSchema(this);
 	Script.Empty();
 	Super::FinishDestroy();
 }
@@ -4249,6 +4268,8 @@ void UClass::PurgeClass(bool bRecompilingOnLoad)
 	FuncMap.Empty();
 	ClearFunctionMapsCaches();
 	PropertyLink = NULL;
+
+	DestroyUnversionedSchema(this);
 }
 
 UClass* UClass::FindCommonBase(UClass* InClassA, UClass* InClassB)
diff --git a/Engine/Source/Runtime/CoreUObject/Private/UObject/LinkerLoad.cpp b/Engine/Source/Runtime/CoreUObject/Private/UObject/LinkerLoad.cpp
--- a/Engine/Source/Runtime/CoreUObject/Private/UObject/LinkerLoad.cpp
+++ b/Engine/Source/Runtime/CoreUObject/Private/UObject/LinkerLoad.cpp
@@ -35,6 +35,7 @@
 #include "Serialization/Formatters/BinaryArchiveFormatter.h"
 #include "Serialization/Formatters/JsonArchiveInputFormatter.h"
 #include "Serialization/ArchiveUObjectFromStructuredArchive.h"
+#include "Serialization/UnversionedPropertySerialization.h"
 #include "Serialization/LoadTimeTracePrivate.h"
 #include "HAL/FileManager.h"
 #include "UObject/CoreRedirects.h"
@@ -1183,7 +1184,13 @@ FLinkerLoad::ELinkerStatus FLinkerLoad::SerializePackageFileSummary()
 			UE_LOG(LogLinker, Warning, TEXT("Asset '%s' has been saved with engine version newer than current and therefore can't be loaded. CurrEngineVersion: %s AssetEngineVersion: %s"), *Filename, *FEngineVersion::Current().ToString(), *Summary.CompatibleWithEngineVersion.ToString());
 			return LINKER_Failed;
 		}
-		else if (!FPlatformProperties::RequiresCookedData() && !Summary.SavedByEngineVersion.HasChangelist() && FEngineVersion::Current().HasChangelist())
+		
+		// Set desired property tag format
+		bool bUseUnversionedProperties = Summary.bUnversioned && CanUseUnversionedPropertySerialization();
+		SetUseUnversionedPropertySerialization(bUseUnversionedProperties);
+		Loader->SetUseUnversionedPropertySerialization(bUseUnversionedProperties);
+		
+		if (!FPlatformProperties::RequiresCookedData() && !Summary.SavedByEngineVersion.HasChangelist() && FEngineVersion::Current().HasChangelist())
 		{
 			// This warning can be disabled in ini with [Core.System] ZeroEngineVersionWarning=False
 			static struct FInitZeroEngineVersionWarning
diff --git a/Engine/Source/Runtime/CoreUObject/Private/UObject/LinkerSave.cpp b/Engine/Source/Runtime/CoreUObject/Private/UObject/LinkerSave.cpp
--- a/Engine/Source/Runtime/CoreUObject/Private/UObject/LinkerSave.cpp
+++ b/Engine/Source/Runtime/CoreUObject/Private/UObject/LinkerSave.cpp
@@ -2,7 +2,9 @@
 
 #include "UObject/LinkerSave.h"
 #include "HAL/FileManager.h"
+#include "Misc/ConfigCacheIni.h"
 #include "Serialization/LargeMemoryWriter.h"
+#include "Serialization/UnversionedPropertySerialization.h"
 #include "UObject/Package.h"
 #include "UObject/SavePackage.h"
 #include "UObject/Class.h"
@@ -63,6 +65,9 @@ FLinkerSave::FLinkerSave(FPackageHeaderSaver& InPackageHeaderSaver, UPackage* In
 		this->SetIsSaving(true);
 		this->SetIsPersistent(true);
 		ArForceByteSwapping		= bForceByteSwapping;
+		bool bUseUnversionedProperties = bInSaveUnversioned && CanUseUnversionedPropertySerialization();
+		SetUseUnversionedPropertySerialization(bUseUnversionedProperties);
+		Saver->SetUseUnversionedPropertySerialization(bUseUnversionedProperties);
 
 #if USE_STABLE_LOCALIZATION_KEYS
 		if (GIsEditor)
@@ -117,6 +122,9 @@ FLinkerSave::FLinkerSave(FPackageHeaderSaver& InPackageHeaderSaver, UPackage* In
 		this->SetIsSaving(true);
 		this->SetIsPersistent(true);
 		ArForceByteSwapping = bForceByteSwapping;
+		bool bUseUnversionedProperties = bInSaveUnversioned && CanUseUnversionedPropertySerialization();
+		SetUseUnversionedPropertySerialization(bUseUnversionedProperties);
+		Saver->SetUseUnversionedPropertySerialization(bUseUnversionedProperties);
 
 #if USE_STABLE_LOCALIZATION_KEYS
 		if (GIsEditor)
@@ -166,6 +174,9 @@ FLinkerSave::FLinkerSave(FPackageHeaderSaver& InPackageHeaderSaver, UPackage* In
 		this->SetIsSaving(true);
 		this->SetIsPersistent(true);
 		ArForceByteSwapping		= bForceByteSwapping;
+		bool bUseUnversionedProperties = bInSaveUnversioned && CanUseUnversionedPropertySerialization();
+		SetUseUnversionedPropertySerialization(bUseUnversionedProperties);
+		Saver->SetUseUnversionedPropertySerialization(bUseUnversionedProperties);
 
 #if USE_STABLE_LOCALIZATION_KEYS
 		if (GIsEditor)
diff --git a/Engine/Source/Runtime/CoreUObject/Private/UObject/PropertyArray.cpp b/Engine/Source/Runtime/CoreUObject/Private/UObject/PropertyArray.cpp
--- a/Engine/Source/Runtime/CoreUObject/Private/UObject/PropertyArray.cpp
+++ b/Engine/Source/Runtime/CoreUObject/Private/UObject/PropertyArray.cpp
@@ -62,12 +62,14 @@ void UArrayProperty::SerializeItem(FStructuredArchive::FSlot Slot, void* Value,
 {
 	checkSlow(Inner);
 	FArchive& UnderlyingArchive = Slot.GetUnderlyingArchive();
-	FPropertyTag InnerTag(UnderlyingArchive, Inner, 0, (uint8*)Value, (uint8*)Defaults);
-	
-	if (UnderlyingArchive.IsTextFormat() && InnerTag.Type == NAME_StructProperty)
+	TOptional<FPropertyTag> MaybeInnerTag;
+
+	if (UnderlyingArchive.IsTextFormat() && !UnderlyingArchive.UseUnversionedPropertySerialization()
+		&& Inner && Inner->IsA<UStructProperty>())
 	{
-		Slot << SA_ATTRIBUTE(TEXT("InnerStructName"), InnerTag.StructName);
-		Slot << SA_OPTIONAL_ATTRIBUTE(TEXT("InnerStructGuid"), InnerTag.StructGuid, FGuid());
+		MaybeInnerTag.Emplace(UnderlyingArchive, Inner, 0, (uint8*)Value, (uint8*)Defaults);	
+		Slot << SA_ATTRIBUTE(TEXT("InnerStructName"), MaybeInnerTag.GetValue().StructName);
+		Slot << SA_OPTIONAL_ATTRIBUTE(TEXT("InnerStructGuid"), MaybeInnerTag.GetValue().StructGuid, FGuid());
 	}
 	
 	// Ensure that the Inner itself has been loaded before calling SerializeItem() on it
@@ -101,12 +103,17 @@ void UArrayProperty::SerializeItem(FStructuredArchive::FSlot Slot, void* Value,
 	ArrayHelper.CountBytes( UnderlyingArchive );
 
 	// Serialize a PropertyTag for the inner property of this array, allows us to validate the inner struct to see if it has changed
-	if (UnderlyingArchive.UE4Ver() >= VER_UE4_INNER_ARRAY_TAG_INFO && InnerTag.Type == NAME_StructProperty)
+	if (!UnderlyingArchive.UseUnversionedPropertySerialization() && 
+		UnderlyingArchive.UE4Ver() >= VER_UE4_INNER_ARRAY_TAG_INFO &&
+		Inner && Inner->IsA<UStructProperty>())
 	{
-		if (!UnderlyingArchive.IsTextFormat())
+		if (!MaybeInnerTag)
 		{
-			UnderlyingArchive << InnerTag;
+			MaybeInnerTag.Emplace(UnderlyingArchive, Inner, 0, (uint8*)Value, (uint8*)Defaults);
+			UnderlyingArchive << MaybeInnerTag.GetValue();
 		}
+
+		FPropertyTag& InnerTag = MaybeInnerTag.GetValue();
 		
 		if (UnderlyingArchive.IsLoading())
 		{
@@ -226,8 +233,10 @@ void UArrayProperty::SerializeItem(FStructuredArchive::FSlot Slot, void* Value,
 		UnderlyingArchive.ArUseCustomPropertyList = bUsingCustomPropertyList;
 	}
 
-	if (UnderlyingArchive.UE4Ver() >= VER_UE4_INNER_ARRAY_TAG_INFO && UnderlyingArchive.IsSaving() && InnerTag.Type == NAME_StructProperty && !UnderlyingArchive.IsTextFormat())
+	if (MaybeInnerTag.IsSet() && UnderlyingArchive.IsSaving() && !UnderlyingArchive.IsTextFormat())
 	{
+		FPropertyTag& InnerTag = MaybeInnerTag.GetValue();
+
 		// set the tag's size
 		InnerTag.Size = UnderlyingArchive.Tell() - DataOffset;
 
diff --git a/Engine/Source/Runtime/CoreUObject/Private/UObject/PropertyTag.cpp b/Engine/Source/Runtime/CoreUObject/Private/UObject/PropertyTag.cpp
--- a/Engine/Source/Runtime/CoreUObject/Private/UObject/PropertyTag.cpp
+++ b/Engine/Source/Runtime/CoreUObject/Private/UObject/PropertyTag.cpp
@@ -18,6 +18,7 @@ FPropertyTag::FPropertyTag( FArchive& InSaveAr, UProperty* Property, int32 InInd
 	, Name      (Property->GetFName())
 	, ArrayIndex(InIndex)
 {
+	check(!InSaveAr.UseUnversionedPropertySerialization());
 	if (Property)
 	{
 		// Handle structs.
@@ -84,6 +85,7 @@ void operator<<(FStructuredArchive::FSlot Slot, FPropertyTag& Tag)
 
 	int32 Version = UnderlyingArchive.UE4Ver();
 
+	check(!UnderlyingArchive.UseUnversionedPropertySerialization());
 	checkf(!UnderlyingArchive.IsSaving() || Tag.Prop, TEXT("FPropertyTag must be constructed with a valid property when used for saving data!"));
 
 	if (!bIsTextFormat)
diff --git a/Engine/Source/Runtime/CoreUObject/Public/UObject/Class.h b/Engine/Source/Runtime/CoreUObject/Public/UObject/Class.h
--- a/Engine/Source/Runtime/CoreUObject/Public/UObject/Class.h
+++ b/Engine/Source/Runtime/CoreUObject/Public/UObject/Class.h
@@ -297,6 +297,9 @@ public:
 	/** Array of object references embedded in script code. Mirrored for easy access by realtime garbage collection code */
 	TArray<UObject*> ScriptObjectReferences;
 
+	/** Cached schema for optimized unversioned property serialization, owned by this. */
+	mutable const struct FUnversionedStructSchema* UnversionedSchema = nullptr;
+
 public:
 	// Constructors.
 	UStruct( EStaticConstructor, int32 InSize, int32 InAlignment, EObjectFlags InFlags );
@@ -384,7 +387,6 @@ public:
 	 */
 	virtual void DestroyStruct(void* Dest, int32 ArrayDim = 1) const;
 
-public:
 	/** Look up a property by an alternate name if it was not found in the first search, this is overridden for user structs */
 	virtual UProperty* CustomFindProperty(const FName InName) const { return nullptr; };
 
@@ -502,6 +504,8 @@ private:
 	friend class FStructBaseChain;
 	friend class FBlueprintCompileReinstancer;
 #endif
+
+	void SerializeVersionedTaggedProperties(FStructuredArchive::FSlot Slot, uint8* Data, UStruct* DefaultsStruct, uint8* Defaults, const UObject* BreakRecursionIfFullyLoad) const;
 };
 
 enum EStructFlags
