Fix AsyncLoading2 compile error

diff --git a/Engine/Source/Runtime/CoreUObject/Private/Serialization/AsyncLoading2.cpp b/Engine/Source/Runtime/CoreUObject/Private/Serialization/AsyncLoading2.cpp
--- a/Engine/Source/Runtime/CoreUObject/Private/Serialization/AsyncLoading2.cpp
+++ b/Engine/Source/Runtime/CoreUObject/Private/Serialization/AsyncLoading2.cpp
@@ -2850,7 +2850,7 @@ void FAsyncPackage2::EventDrivenSerializeExport(int32 LocalExportIndex, const ui
 	Ar.Exports = &Exports;
 	Ar.ExternalReadDependencies = &ExternalReadDependencies;
 	Ar.ExternalIoRequests = &ExternalIoRequests;
-	Ar.SetUseUnversionedPropertySerialization(Summary.bUnversioned && CanUseUnversionedPropertySerialization());
+	//Ar.SetUseUnversionedPropertySerialization(Summary.bUnversioned && CanUseUnversionedPropertySerialization());
 
 	Object->ClearFlags(RF_NeedLoad);
 
