AsyncLoading2: Fix linker assert

diff --git a/Engine/Source/Runtime/CoreUObject/Private/Serialization/AsyncLoading2.cpp b/Engine/Source/Runtime/CoreUObject/Private/Serialization/AsyncLoading2.cpp
--- a/Engine/Source/Runtime/CoreUObject/Private/Serialization/AsyncLoading2.cpp
+++ b/Engine/Source/Runtime/CoreUObject/Private/Serialization/AsyncLoading2.cpp
@@ -4021,6 +4021,7 @@ void FAsyncPackage2::ClearOwnedObjects()
 		Object->ClearInternalFlags(InternalFlagsToClear);
 	}
 	OwnedObjects.Empty();
+	LinkerRoot->LinkerLoad = nullptr; // TEMP 
 }
 
 void FAsyncPackage2::AddRequestID(int32 Id)
