Compilation fixes

diff --git a/Engine/Source/Runtime/CoreUObject/Private/Serialization/UnversionedPropertySerialization.cpp b/Engine/Source/Runtime/CoreUObject/Private/Serialization/UnversionedPropertySerialization.cpp
--- a/Engine/Source/Runtime/CoreUObject/Private/Serialization/UnversionedPropertySerialization.cpp
+++ b/Engine/Source/Runtime/CoreUObject/Private/Serialization/UnversionedPropertySerialization.cpp
@@ -10,7 +10,7 @@
 // unversioned property loading except non-numeric SerializeItem() calls about 2x faster.
 #ifndef CACHE_UNVERSIONED_PROPERTY_SCHEMA
 // x64 platforms tend to have more memory...
-#	define CACHE_UNVERSIONED_PROPERTY_SCHEMA PLATFORM_CPU_X86_FAMILY && PLATFORM_64BITS
+#	define CACHE_UNVERSIONED_PROPERTY_SCHEMA (PLATFORM_CPU_X86_FAMILY && PLATFORM_64BITS)
 #endif
 
 // Helper to pass around appropriate default value types depending on CACHE_UNVERSIONED_PROPERTY_SCHEMA
@@ -640,7 +640,6 @@ public:
 			}
 
 			RemainingFragmentValues = FragmentIt->ValueNum;
-			check(SchemaIt + RemainingFragmentValues <= SchemaEnd);
 		}
 	};
 };
diff --git a/Engine/Source/Runtime/CoreUObject/Private/Serialization/UnversionedPropertySerializationTest.cpp b/Engine/Source/Runtime/CoreUObject/Private/Serialization/UnversionedPropertySerializationTest.cpp
--- a/Engine/Source/Runtime/CoreUObject/Private/Serialization/UnversionedPropertySerializationTest.cpp
+++ b/Engine/Source/Runtime/CoreUObject/Private/Serialization/UnversionedPropertySerializationTest.cpp
@@ -10,10 +10,10 @@
 
 namespace PropertySerializationStats
 {
-	static TAtomic<uint64> Structs = 0;
-	static TAtomic<uint64> VersionedBytes = 0;
-	static TAtomic<uint64> UnversionedBytes = 0;
-	static TAtomic<uint64> UselessBytes = 0;
+	static TAtomic<uint64> Structs;
+	static TAtomic<uint64> VersionedBytes;
+	static TAtomic<uint64> UnversionedBytes;
+	static TAtomic<uint64> UselessBytes;
 
 #if ENABLE_COOK_STATS
 	static FCookStatsManager::FAutoRegisterCallback RegisterCookStats([](FCookStatsManager::AddStatFuncRef AddStat)
