Unversioned property optimizations

Note that ASSUME() currently doesn't map to __builtin_assume on clang but I plan to fix this

diff --git a/Engine/Source/Runtime/CoreUObject/Private/Serialization/UnversionedPropertySerialization.cpp b/Engine/Source/Runtime/CoreUObject/Private/Serialization/UnversionedPropertySerialization.cpp
--- a/Engine/Source/Runtime/CoreUObject/Private/Serialization/UnversionedPropertySerialization.cpp
+++ b/Engine/Source/Runtime/CoreUObject/Private/Serialization/UnversionedPropertySerialization.cpp
@@ -106,14 +106,13 @@ public:
 		EIntegerType IntType = GetIntType(Property->GetMinAlignment());
 		uint32 FastZeroIntNum = GetIntNum(Property, IntType);
 #endif
-
-		if (FastZeroIntNum == 1)
-		{
-			MemZeroInt(GetValue(Data), IntType);
-		}
-		else
+		switch (IntType)
 		{
-			MemZeroInts(GetValue(Data), FastZeroIntNum, IntType);
+			case EIntegerType::Uint8 : MemZeroRange<uint8 >(GetValue(Data), FastZeroIntNum); break;
+			case EIntegerType::Uint16: MemZeroRange<uint16>(GetValue(Data), FastZeroIntNum); break;
+			case EIntegerType::Uint32: MemZeroRange<uint32>(GetValue(Data), FastZeroIntNum); break;
+			case EIntegerType::Uint64: MemZeroRange<uint64>(GetValue(Data), FastZeroIntNum); break;
+			default: ASSUME(0);
 		}
 	}
 
@@ -124,6 +123,7 @@ public:
 		uint32 FastZeroIntNum = CanSerializeAsZero(Property, IntType) ? GetIntNum(Property, IntType) : uint8(0);
 #endif
 
+		// Can simplified and faster using a switch() statement like LoadZero()
 		if (FastZeroIntNum == 1)
 		{
 			return IsIntZero(GetValue(Data), IntType);
@@ -190,8 +190,14 @@ private:
 
 	static void SerializeAsInteger(FStructuredArchive::FSlot Slot, void* Value, EIntegerType IntType)
 	{
-		static constexpr void (*Functions[4])(FStructuredArchive::FSlot, void*) = { &SerializeAs<uint8>, &SerializeAs<uint16>, &SerializeAs<uint32>, &SerializeAs<uint64> };
-		(*Functions[static_cast<uint32>(IntType)])(Slot, Value);
+		 switch (IntType)
+		 {
+			 case  EIntegerType::Uint8:	Slot << *reinterpret_cast< uint8*>(Value); break;
+			 case EIntegerType::Uint16: Slot << *reinterpret_cast<uint16*>(Value); break;
+			 case EIntegerType::Uint32: Slot << *reinterpret_cast<uint32*>(Value); break;
+			 case EIntegerType::Uint64: Slot << *reinterpret_cast<uint64*>(Value); break;
+			 default: ASSUME(0);
+		 }
 	}
 
 	template<typename T>
@@ -224,31 +230,14 @@ private:
 	}
 
 	template<typename T>
-	static void MemZero(void* Value)
-	{
-		*reinterpret_cast<T*>(Value) = 0;
-	}
-
-	static void MemZeroInt(void* Value, EIntegerType IntType)
-	{
-		static constexpr void (*Functions[4])(void*) = { &MemZero<uint8>, &MemZero<uint16>, &MemZero<uint32>, &MemZero<uint64> };
-		(*Functions[static_cast<uint32>(IntType)])(Value);
-	}
-
-	template<typename T>
-	static void MemZeroRange(void* Value, uint32 Num)
+	FORCEINLINE static void MemZeroRange(void* Value, uint32 Num)
 	{
+		ASSUME(Num > 0);
 		for (uint32 Idx = 0; Idx < Num; ++Idx)
 		{
 			reinterpret_cast<T*>(Value)[Idx] = 0;
 		}
 	}
-	
-	static void MemZeroInts(void* Value, uint32 Num, EIntegerType IntType)
-	{
-		static constexpr void (*Functions[4])(void*, uint32) = { &MemZeroRange<uint8>, &MemZeroRange<uint16>, &MemZeroRange<uint32>, &MemZeroRange<uint64> };
-		(*Functions[static_cast<uint32>(IntType)])(Value, Num);
-	}
 
 	static uint32 GetSizeOf(EIntegerType Type)
 	{
@@ -437,25 +426,30 @@ public:
 
 	void Load(FStructuredArchive::FStream Stream)
 	{
-		uint32 ZeroMaskSize = 0;
+		FFragment Fragment;
+		uint32 ZeroMaskNum = 0;
+		uint32 UnmaskedNum = 0;
 		do
 		{
 			uint16 Packed;
 			Stream.EnterElement() << Packed;
-			FFragment Fragment = FFragment::Unpack(Packed);
+			Fragment = FFragment::Unpack(Packed);
 
 			Fragments.Add(Fragment);
 
-			ZeroMaskSize += Fragment.ValueNum * Fragment.bHasAnyZeroes;
-			bHasNonZeroValues |= (Fragment.ValueNum > 0) & (!Fragment.bHasAnyZeroes);
+			(Fragment.bHasAnyZeroes ? ZeroMaskNum : UnmaskedNum) += Fragment.ValueNum;
 		}
-		while (!Fragments.Last().bIsLast);
+		while (!Fragment.bIsLast);
 
-		if (ZeroMaskSize > 0)
+		if (ZeroMaskNum > 0)
+		{
+			ZeroMask.SetNumUninitialized(ZeroMaskNum);
+			LoadZeroMaskData(Stream, ZeroMaskNum, ZeroMask.GetData());
+			bHasNonZeroValues = UnmaskedNum > 0 || ZeroMask.Find(false) != INDEX_NONE;
+		}
+		else
 		{
-			ZeroMask.Init(false, ZeroMaskSize);
-			LoadZeroMaskData(Stream, ZeroMaskSize, ZeroMask.GetData());
-			bHasNonZeroValues = bHasNonZeroValues || ZeroMask.Find(false) != INDEX_NONE;
+			bHasNonZeroValues = UnmaskedNum > 0;
 		}
 	}
 	
