AsyncLoading2: Optimize cost of calling UObject::GetArchetype from UObject::ConditionalPostLoad.
Tested on loading of 70 000 packages on PC and 60 000 packages on another platform.
Cache archetype pointer in game thread from PostLoadDeferredObjects.
No contention in annotations reader/writer lock, but GetArchetype may miss the cache a few times when serialized objects are postloaded from another object.
Reduces contention in FUObjectHashTables (by not calling GetArchetypeFromRequiredInfoImpl) which improves create and serialize on the async loading thread.

Results on PC (Superluminal cpu sample times):
Total UObject::GetArchetype time down from ~210 ms to ~20 ms => ~190 ms faster
PostLoadDeferredObjects exclusive time up from ~75 ms to ~110 ms => ~35 ms slower
=> ~155 ms total win

Results on other platform (Unreal Insights cpu scope times):
Total UObject::GetArchetype time down from ~325 ms to ~60 ms => 265 ms faster
PostLoadDeferredObjects scope time down from from ~580 ms to ~350 ms => 230 ms faster
New CacheArchetype scope in PostLoadDeferredObjects for storing the archetype pointer in annotations. => ~60 ms
=> ~205 ms total win

Slower alternative:
Cache archetype pointer in async loading thread from either CreateExport, SerializeExport or ExportsDone .
GetArchetype never misses the cache for serialized objects, but adds contention to annotations reader/writer lock.

Todo:
Not using annotations would be much faster and would enable setting the archetype early from the async loading thread.

diff --git a/Engine/Source/Runtime/CoreUObject/Private/Serialization/AsyncLoading2.cpp b/Engine/Source/Runtime/CoreUObject/Private/Serialization/AsyncLoading2.cpp
--- a/Engine/Source/Runtime/CoreUObject/Private/Serialization/AsyncLoading2.cpp
+++ b/Engine/Source/Runtime/CoreUObject/Private/Serialization/AsyncLoading2.cpp
@@ -36,6 +36,7 @@
 #include "Blueprint/BlueprintSupport.h"
 #include "HAL/LowLevelMemTracker.h"
 #include "ProfilingDebugging/CsvProfiler.h"
+#include "UObject/UObjectArchetypeInternal.h"
 #include "UObject/GarbageCollectionInternal.h"
 #include "ProfilingDebugging/MiscTrace.h"
 #include "Serialization/LoadTimeTracePrivate.h"
@@ -4213,7 +4214,9 @@ EAsyncPackageState::Type FAsyncPackage2::PostLoadObjects()
 	// PostLoad objects.
 	while (PostLoadIndex < ExportCount && !FAsyncLoadingThreadState2::Get()->IsTimeLimitExceeded())
 	{
-		UObject* Object = Exports[PostLoadIndex++];
+		int32 LocalExportIndex = PostLoadIndex++;
+		const FExportMapEntry& Export = ExportMap[LocalExportIndex];
+		UObject* Object = Exports[LocalExportIndex];
 
 		check(Object);
 		check(!Object->HasAnyFlags(RF_NeedLoad));
@@ -4228,6 +4231,10 @@ EAsyncPackageState::Type FAsyncPackage2::PostLoadObjects()
 			ThreadContext.CurrentlyPostLoadedObjectByALT = Object;
 			{
 				TRACE_LOADTIME_POSTLOAD_EXPORT_SCOPE(Object);
+				// cache archetype before post load
+				// prevents GetArchetype from hitting the expensive GetArchetypeFromRequiredInfoImpl
+				UObject* Template = EventDrivenIndexToObject(Export.TemplateIndex, true);
+				CacheArchetypeForObject(Object, Template);
 				Object->ConditionalPostLoad();
 				Object->AtomicallyClearInternalFlags(EInternalObjectFlags::AsyncLoading);
 			}
@@ -4252,7 +4259,21 @@ EAsyncPackageState::Type FAsyncPackage2::PostLoadDeferredObjects()
 
 	FUObjectSerializeContext* LoadContext = GetSerializeContext();
 
-	STAT(double PostLoadStartTime = FPlatformTime::Seconds());
+	{
+		// cache archetype for all exports before post load
+		// prevents GetArchetype from hitting the expensive GetArchetypeFromRequiredInfoImpl
+		TRACE_CPUPROFILER_EVENT_SCOPE(CacheArchetype);
+		for (int32 LocalExportIndex = 0; LocalExportIndex < ExportCount; ++LocalExportIndex)
+		{
+			UObject* Object = Exports[LocalExportIndex];
+			if (Object->HasAnyFlags(RF_NeedPostLoad))
+			{
+				const FExportMapEntry& Export = ExportMap[LocalExportIndex];
+				UObject* Template = EventDrivenIndexToObject(Export.TemplateIndex, true);
+				CacheArchetypeForObject(Object, Template);
+			}
+		}
+	}
 
 	while (DeferredPostLoadIndex < ExportCount && 
 		!AsyncLoadingThread.IsAsyncLoadingSuspended() &&
diff --git a/Engine/Source/Runtime/CoreUObject/Private/UObject/UObjectArchetype.cpp b/Engine/Source/Runtime/CoreUObject/Private/UObject/UObjectArchetype.cpp
--- a/Engine/Source/Runtime/CoreUObject/Private/UObject/UObjectArchetype.cpp
+++ b/Engine/Source/Runtime/CoreUObject/Private/UObject/UObjectArchetype.cpp
@@ -4,6 +4,7 @@
 	UObjectArchetype.cpp: Unreal object archetype relationship management
 =============================================================================*/
 
+#include "UObject/UObjectArchetypeInternal.h"
 #include "CoreMinimal.h"
 #include "UObject/UObjectHash.h"
 #include "UObject/Object.h"
@@ -142,6 +143,18 @@ UObject* GetArchetypeFromRequiredInfoImpl(const UClass* Class, const UObject* Ou
 	return Result;
 }
 
+void CacheArchetypeForObject(UObject* Object, UObject* Archetype)
+{
+#if UE_CACHE_ARCHETYPE
+#if UE_VERIFY_CACHED_ARCHETYPE
+	bool bUseUpToDateClass = false;
+	UObject* VerifyArchetype = GetArchetypeFromRequiredInfoImpl(Object->GetClass(), Object->GetOuter(), Object->GetFName(), Object->GetFlags(), bUseUpToDateClass);
+	checkf(Archetype == VerifyArchetype, TEXT("Cached archetype mismatch, expected: %s, cached: %s"), *GetFullNameSafe(VerifyArchetype), *GetFullNameSafe(Archetype));
+#endif
+	ArchetypeAnnotation.AddAnnotation(Object, GUObjectArray.ObjectToIndex(Archetype));
+#endif
+}
+
 UObject* UObject::GetArchetypeFromRequiredInfo(const UClass* Class, const UObject* Outer, FName Name, EObjectFlags ObjectFlags)
 {
 	bool bUseUpToDateClass = false;
@@ -199,4 +212,4 @@ void CleanupCachedArchetypes()
 #if UE_CACHE_ARCHETYPE
 	ArchetypeAnnotation.RemoveAllAnnotations();
 #endif
-}
\ No newline at end of file
+}
diff --git a/Engine/Source/Runtime/CoreUObject/Private/UObject/UObjectArchetypeInternal.h b/Engine/Source/Runtime/CoreUObject/Private/UObject/UObjectArchetypeInternal.h
new file mode 100644
--- /dev/null
+++ b/Engine/Source/Runtime/CoreUObject/Private/UObject/UObjectArchetypeInternal.h
@@ -0,0 +1,6 @@
+// Copyright 1998-2019 Epic Games, Inc. All Rights Reserved.
+
+#pragma once
+
+class UObject;
+void CacheArchetypeForObject(UObject* Object, UObject* Archetype);
