AsyncLoading2:
Fixed IoStore pak file filter
More gracefully handle missing packages

diff --git a/Engine/Source/Programs/AutomationTool/Scripts/CopyBuildToStagingDirectory.Automation.cs b/Engine/Source/Programs/AutomationTool/Scripts/CopyBuildToStagingDirectory.Automation.cs
--- a/Engine/Source/Programs/AutomationTool/Scripts/CopyBuildToStagingDirectory.Automation.cs
+++ b/Engine/Source/Programs/AutomationTool/Scripts/CopyBuildToStagingDirectory.Automation.cs
@@ -1751,7 +1751,7 @@ public partial class Project : CommandUtils
 				Src.HasExtension(".ubulkmanifest") ||
 				Src.HasExtension(".uglimport") ||
 				Src.HasExtension(".unamemap") ||
-				Src.HasExtension(".uscriptarc") ||
+				Src.HasExtension(".uscriptarcs") ||
 				Src.HasExtension(".uslimport") ||
 				Src.HasExtension(".ustoretoc"))
 			{
diff --git a/Engine/Source/Runtime/CoreUObject/Private/Serialization/AsyncLoading2.cpp b/Engine/Source/Runtime/CoreUObject/Private/Serialization/AsyncLoading2.cpp
--- a/Engine/Source/Runtime/CoreUObject/Private/Serialization/AsyncLoading2.cpp
+++ b/Engine/Source/Runtime/CoreUObject/Private/Serialization/AsyncLoading2.cpp
@@ -1790,6 +1790,10 @@ private:
 	FAsyncPackage2* CreateAsyncPackage(const FAsyncPackageDesc& Desc)
 	{
 		FGlobalPackageId* GlobalPackageId = PackageStore.GetGlobalPackageId(Desc.NameToLoad);
+		if (!GlobalPackageId)
+		{
+			return nullptr;
+		}
 		check(GlobalPackageId);
 		return new FAsyncPackage2(Desc, AsyncPackageSerialNumber.Increment(), *this, EDLBootNotificationManager, GraphAllocator, EventSpecs.GetData(), *GlobalPackageId);
 	}
@@ -1923,6 +1927,10 @@ FAsyncPackage2* FAsyncLoadingThread2Impl::FindOrInsertPackage(FAsyncPackageDesc*
 		if (!Package)
 		{
 			Package = CreateAsyncPackage(*InDesc);
+			if (!Package)
+			{
+				return nullptr;
+			}
 			Package->AddRef();
 			ExistingAsyncPackagesCounter.Increment();
 			AsyncPackageNameLookup.Add(Package->GetPackageName(), Package);
@@ -2440,13 +2448,16 @@ void FAsyncPackage2::ImportPackagesRecursive()
 		Info.Priority = Desc.Priority;
 		bool bInserted;
 		FAsyncPackage2* ImportedAsyncPackage = AsyncLoadingThread.FindOrInsertPackage(&Info, bInserted);
-		TRACE_LOADTIME_ASYNC_PACKAGE_IMPORT_DEPENDENCY(this, ImportedAsyncPackage);
-		ImportedAsyncPackage->AddRef();
-		ImportedAsyncPackages.Add(ImportedAsyncPackage);
-		if (bInserted)
+		if (ImportedAsyncPackage)
 		{
-			ImportedAsyncPackage->ImportPackagesRecursive();
-			ImportedAsyncPackage->StartLoading();
+			TRACE_LOADTIME_ASYNC_PACKAGE_IMPORT_DEPENDENCY(this, ImportedAsyncPackage);
+			ImportedAsyncPackage->AddRef();
+			ImportedAsyncPackages.Add(ImportedAsyncPackage);
+			if (bInserted)
+			{
+				ImportedAsyncPackage->ImportPackagesRecursive();
+				ImportedAsyncPackage->StartLoading();
+			}
 		}
 	}
 }
