Fix unversioned property serialization

diff --git a/Engine/Source/Runtime/CoreUObject/Private/Serialization/UnversionedPropertySerialization.cpp b/Engine/Source/Runtime/CoreUObject/Private/Serialization/UnversionedPropertySerialization.cpp
--- a/Engine/Source/Runtime/CoreUObject/Private/Serialization/UnversionedPropertySerialization.cpp
+++ b/Engine/Source/Runtime/CoreUObject/Private/Serialization/UnversionedPropertySerialization.cpp
@@ -284,7 +284,7 @@ struct FUnversionedStructSchema
 		TArray<FUnversionedPropertySerializer, TInlineAllocator<256>> Serializers;
 		for (UProperty* Property = Struct->PropertyLink; Property; Property = Property->PropertyLinkNext)
 		{
-			if (!Property->IsEditorOnly())
+			if (!Property->IsEditorOnlyProperty())
 			{
 				for (int32 ArrayIdx = 0, ArrayDim = Property->ArrayDim; ArrayIdx < ArrayDim; ++ArrayIdx)
 				{
@@ -373,7 +373,7 @@ struct FUnversionedSchemaIterator
 	static UProperty* SkipEditorOnlyProperties(UProperty* Property)
 	{
 #if WITH_EDITORONLY_DATA
-		while (Property && Property->IsEditorOnly())
+		while (Property && Property->IsEditorOnlyProperty())
 		{
 			Property = Property->PropertyLinkNext;
 		}
diff --git a/Engine/Source/Runtime/CoreUObject/Private/Serialization/UnversionedPropertySerializationTest.cpp b/Engine/Source/Runtime/CoreUObject/Private/Serialization/UnversionedPropertySerializationTest.cpp
--- a/Engine/Source/Runtime/CoreUObject/Private/Serialization/UnversionedPropertySerializationTest.cpp
+++ b/Engine/Source/Runtime/CoreUObject/Private/Serialization/UnversionedPropertySerializationTest.cpp
@@ -347,7 +347,7 @@ struct FUnversionedPropertyTest : public FUnversionedPropertyTestInput
 
 		for (UProperty* Property : Properties)
 		{
-			if (!Property->IsEditorOnly())
+			if (!Property->IsEditorOnlyProperty())
 			{
 				Out.Add(Property);
 			}
