Fix static analysis warnings

diff --git a/Engine/Source/Runtime/CoreUObject/Private/Serialization/UnversionedPropertySerialization.cpp b/Engine/Source/Runtime/CoreUObject/Private/Serialization/UnversionedPropertySerialization.cpp
--- a/Engine/Source/Runtime/CoreUObject/Private/Serialization/UnversionedPropertySerialization.cpp
+++ b/Engine/Source/Runtime/CoreUObject/Private/Serialization/UnversionedPropertySerialization.cpp
@@ -401,7 +401,8 @@ struct FUnversionedSchemaRange
 	FUnversionedSchemaIterator begin() const { return Begin; }
 	FUnversionedSchemaIterator end() const { return End; }
 
-	FUnversionedSchemaIterator Begin, End;
+	FUnversionedSchemaIterator Begin;
+	FUnversionedSchemaIterator End;
 };
 
 // List of serialized property indices and which of them are non-zero.
@@ -505,7 +506,7 @@ protected:
 	{
 		check(NumBits > 0);
 
-		uint32 LastWordMask = ~0u >> (32 - (NumBits % 32));
+		uint32 LastWordMask = ~0u >> ((32u - NumBits) % 32u);
 		if (NumBits <= 8)
 		{
 			uint8 Word = static_cast<uint8>(*Data & LastWordMask);
