Unversioned property serialization test handles hash table key structs with editor only properties

diff --git a/Engine/Source/Runtime/CoreUObject/Private/Serialization/UnversionedPropertySerializationTest.cpp b/Engine/Source/Runtime/CoreUObject/Private/Serialization/UnversionedPropertySerializationTest.cpp
--- a/Engine/Source/Runtime/CoreUObject/Private/Serialization/UnversionedPropertySerializationTest.cpp
+++ b/Engine/Source/Runtime/CoreUObject/Private/Serialization/UnversionedPropertySerializationTest.cpp
@@ -199,42 +199,56 @@ struct FUnversionedPropertyTest : public FUnversionedPropertyTestInput
 
 		return true;
 	}
+
+	static const uint8* FindElementPtr(const FScriptSetHelper& Helper, const uint8* Element)
+	{
+		const UProperty* ElemProp = Helper.GetElementProperty();
+		int32 Index = Helper.Set->FindIndex(Element, Helper.SetLayout,
+											[ElemProp](const void* Element) { return ElemProp->GetValueTypeHash(Element); },
+											[ElemProp](const void* A, const void* B) { return Equals(ElemProp, A, B); });
+		return Index >= 0 ? Helper.GetElementPtr(Index) : nullptr;
+	}
 	
 	static bool Equals(const USetProperty* Property, const void* A, const void* B)
 	{
 		FScriptSetHelper HelperA(Property, A);
 		FScriptSetHelper HelperB(Property, B);
-		const UProperty* ElemProperty   = HelperA.GetElementProperty();
 
 		if (HelperA.Num() != HelperB.Num())
 		{
 			return false;
 		}
 
-		for (int32 Num = HelperA.Num(), IndexA = 0; Num > 0; --Num)
+		for (int32 Num = HelperA.Num(), IndexA = 0; IndexA < Num; ++IndexA)
 		{
-			while (!HelperA.IsValidIndex(IndexA))
+			if (HelperA.IsValidIndex(IndexA))
 			{
-				++IndexA;
-			}
-			
-			const uint8* ElemA = HelperA.GetElementPtr(IndexA);
-			const uint8* ElemB = HelperB.FindElementPtrFromHash(ElemA);
+				const uint8* ElemA = HelperA.GetElementPtr(IndexA);
+				const uint8* ElemB = FindElementPtr(HelperB, ElemA);
 		
-			if (!ElemB || !Equals(ElemProperty, ElemA, ElemB))
-			{
-				return false;
-			} 
+				if (!ElemB)
+				{
+					return false;
+				} 
+			}
 		}
 
 		return true;
 	}
+
+	static const uint8* FindPairPtr(const FScriptMapHelper& Helper, const uint8* Key)
+	{
+		const UProperty* KeyProp = Helper.GetKeyProperty();
+		int32 Index = Helper.Map->FindPairIndex(Key, Helper.MapLayout,
+												[KeyProp](const void* Key) { return KeyProp->GetValueTypeHash(Key); },
+												[KeyProp](const void* A, const void* B) { return Equals(KeyProp, A, B); });
+		return Index >= 0 ? Helper.GetPairPtr(Index) : nullptr;
+	}
 	
 	static bool Equals(const UMapProperty* Property, const void* A, const void* B)
 	{
 		FScriptMapHelper HelperA(Property, A);
 		FScriptMapHelper HelperB(Property, B);
-		const UProperty* KeyProp   = HelperA.GetKeyProperty();
 		const UProperty* ValueProp = HelperA.GetValueProperty();
 		int32 ValueOffset = HelperA.MapLayout.ValueOffset;
 
@@ -243,24 +257,22 @@ struct FUnversionedPropertyTest : public FUnversionedPropertyTestInput
 			return false;
 		}
 		
-		for (int32 Num = HelperA.Num(), IndexA = 0; Num > 0; --Num)
+		for (int32 Num = HelperA.Num(), IndexA = 0; IndexA < Num; ++IndexA)
 		{
-			while (!HelperA.IsValidIndex(IndexA))
+			if (HelperA.IsValidIndex(IndexA))
 			{
-				++IndexA;
-			}
-			
-			const uint8* PairA = HelperA.GetPairPtr(IndexA);
-			const uint8* PairB = HelperB.FindMapPairPtrFromHash(PairA);
+				const uint8* PairA = HelperA.GetPairPtr(IndexA);
+				const uint8* PairB = FindPairPtr(HelperB, PairA);
 		
-			if (!PairB || !Equals(KeyProp, PairA, PairB))
-			{
-				return false;
-			} 
+				if (!PairB)
+				{
+					return false;
+				} 
 
-			if (!Equals(ValueProp, PairA + ValueOffset, PairB + ValueOffset))
-			{
-				return false;
+				if (!Equals(ValueProp, PairA + ValueOffset, PairB + ValueOffset))
+				{
+					return false;
+				}
 			}
 		}
 
