Try to conform to IncludeTool code restrictions

diff --git a/Engine/Source/Runtime/CoreUObject/Private/Serialization/UnversionedPropertySerialization.cpp b/Engine/Source/Runtime/CoreUObject/Private/Serialization/UnversionedPropertySerialization.cpp
--- a/Engine/Source/Runtime/CoreUObject/Private/Serialization/UnversionedPropertySerialization.cpp
+++ b/Engine/Source/Runtime/CoreUObject/Private/Serialization/UnversionedPropertySerialization.cpp
@@ -325,15 +325,14 @@ using FUnversionedSchemaIterator = const FUnversionedPropertySerializer*;
 
 #else
 
-// Link-walking iterator
-struct FUnversionedSchemaIterator
+struct FLinkWalkingSchemaIterator
 {
 	UProperty* Property = nullptr;
 	uint32 ArrayIndex = 0;
 
-	FUnversionedSchemaIterator() {}
+	FLinkWalkingSchemaIterator() {}
 
-	explicit FUnversionedSchemaIterator(UProperty* FirstProperty)
+	explicit FLinkWalkingSchemaIterator(UProperty* FirstProperty)
 		: Property(SkipEditorOnlyProperties(FirstProperty))
 	{}
 
@@ -365,7 +364,7 @@ struct FUnversionedSchemaIterator
 		return FUnversionedPropertySerializer(Property, ArrayIndex);
 	}
 
-	bool operator!=(FUnversionedSchemaIterator Rhs) const
+	bool operator!=(FLinkWalkingSchemaIterator Rhs) const
 	{
 		return (Property != Rhs.Property) | (ArrayIndex != Rhs.ArrayIndex);
 	}
@@ -382,6 +381,8 @@ struct FUnversionedSchemaIterator
 	}
 };
 
+using FUnversionedSchemaIterator = FLinkWalkingSchemaIterator;
+
 #endif // CACHE_UNVERSIONED_PROPERTY_SCHEMA
 
 struct FUnversionedSchemaRange
