BulkData no longer generates FIoChunkIds at cook time so that there are no changes to the cooked data formats.

- We now generate FIoChunkIds entirely at runtime using the BulkDataOffset as the unique id (although this gets truncated and we lose the top byte it should be safe and unique enough)
- Added a specialised version of ::CreateIoChunkId named ::CreateBulkdataChunkId
- Renamed BULKDATA_CookedForIoDispatcher to BULKDATA_UsesIoDispatcher to reflect its new use.
- Since we no longer write out the FIoChunkId during the cook we don't need to write out dummy values or store the offsets so that code has been entirely removed.
- FSavePackageContext is now always passed around when cooking, it's the members that we check as optional.
- FPackageStoreBulkDataManifest no longer has a base interface
  - It also stores both the original offset (which is used in generating the ChunkId) and the adjusted offset which we will use for generating the actual mappings. If we used the adjusted offset for both we'd have to adjust the offset when using ZenLoader at runtime which is a (minor) waste of cycles.
  - It is always created during a cook.
  - The manifest file is written to <projectroot>\saved\BulkDataManifests\<platformname>.ubulkmanifest by moving it from the cooked directory we will reduce confusion as to what the file is and make sure that nobody ships it. At some point we can look into removing the need for this file.

#rnx

diff --git a/Engine/Source/Developer/IoStoreUtilities/Private/IoStoreUtilities.cpp b/Engine/Source/Developer/IoStoreUtilities/Private/IoStoreUtilities.cpp
--- a/Engine/Source/Developer/IoStoreUtilities/Private/IoStoreUtilities.cpp
+++ b/Engine/Source/Developer/IoStoreUtilities/Private/IoStoreUtilities.cpp
@@ -192,6 +192,15 @@ static FIoChunkId CreateChunkId(int32 GlobalPackageId, uint16 ChunkIndex, EIoChu
 	return ChunkId;
 }
 
+static FIoChunkId CreateChunkIdForBulkData(int32 GlobalPackageId, uint64 BulkdataOffset, EIoChunkType ChunkType, const TCHAR* DebugString)
+{
+	FIoChunkId ChunkId = CreateBulkdataChunkId(GlobalPackageId, BulkdataOffset, ChunkType);
+#if OUTPUT_CHUNKID_DIRECTORY
+	ChunkIdCsv.AddChunk(GlobalPackageId, 0, 0, (uint8)ChunkType, GetTypeHash(ChunkId), DebugString);
+#endif
+	return ChunkId;
+}
+
 struct FZenPackageSummary
 {
 	uint32 PackageFlags;
@@ -978,8 +987,8 @@ int32 CreateTarget(const FContainerTarget& Target)
 	FString CookedDir = Target.CookedDirectory;
 	FString OutputDir = Target.OutputDirectory;
 	FString RelativePrefixForLegacyFilename = TEXT("../../../");
-
-	FPackageStoreBulkDataManifest BulkDataManifest(CookedDir);
+	
+	FPackageStoreBulkDataManifest BulkDataManifest(Target.TargetPlatform->PlatformName());
 	const bool bWithBulkDataManifest = BulkDataManifest.Load();
 	if (bWithBulkDataManifest)
 	{
@@ -1658,7 +1667,7 @@ int32 CreateTarget(const FContainerTarget& Target)
 				TUniquePtr<FArchive> BulkAr(IFileManager::Get().CreateFileReader(*UBulkFileName));
 				if (BulkAr != nullptr)
 				{
-					const FIoChunkId BulkDataChunkId = CreateChunkId(Package.GlobalPackageId, 0, EIoChunkType::BulkData, *Package.FileName);
+					const FIoChunkId BulkDataChunkId = CreateChunkIdForBulkData(Package.GlobalPackageId, INDEX_NONE, EIoChunkType::BulkData, *Package.FileName);
 #if !SKIP_WRITE_CONTAINER
 					uint8* BulkBuffer = static_cast<uint8*>(FMemory::Malloc(BulkAr->TotalSize()));
 					BulkAr->Serialize(BulkBuffer, BulkAr->TotalSize());
@@ -1684,7 +1693,7 @@ int32 CreateTarget(const FContainerTarget& Target)
 						for (const FPackageStoreBulkDataManifest::PackageDesc::BulkDataDesc& BulkDataDesc : PackageDesc->GetDataArray())
 						{
 #if !SKIP_WRITE_CONTAINER
-							const FIoChunkId AccessChunkId = CreateChunkId(Package.GlobalPackageId, BulkDataDesc.Index, EIoChunkType::BulkData, *Package.FileName);
+							const FIoChunkId AccessChunkId = CreateChunkIdForBulkData(Package.GlobalPackageId, BulkDataDesc.ChunkId, EIoChunkType::BulkData, *Package.FileName);
 							const FIoStatus PartialResult = IoStoreWriter->MapPartialRange(BulkDataChunkId, BulkDataDesc.Offset, BulkDataDesc.Size, AccessChunkId);
 							if (!PartialResult.IsOk())
 							{
@@ -1696,6 +1705,10 @@ int32 CreateTarget(const FContainerTarget& Target)
 							++BulkPartialChunkCount;
 						}
 					}
+					else
+					{
+						UE_LOG(LogIoStore, Error, TEXT("Unable to find an entry in the bulkdata manifest for '%s' the file might be out of date!"), *Package.FileName);
+					}
 				}
 			}
 #endif
diff --git a/Engine/Source/Editor/UnrealEd/Private/CookOnTheFlyServer.cpp b/Engine/Source/Editor/UnrealEd/Private/CookOnTheFlyServer.cpp
--- a/Engine/Source/Editor/UnrealEd/Private/CookOnTheFlyServer.cpp
+++ b/Engine/Source/Editor/UnrealEd/Private/CookOnTheFlyServer.cpp
@@ -4031,7 +4031,7 @@ void UCookOnTheFlyServer::SaveCookedPackage(UPackage* Package, uint32 SaveFlags,
 						}
 						else
 						{
-							FSavePackageContext* const SavePackageContext = IsUsingPackageStore() ? SavePackageContexts[PlatformIndex] : nullptr;
+							FSavePackageContext* const SavePackageContext = SavePackageContexts[PlatformIndex];
 
 							Result = GEditor->Save(	Package, World, FlagsToCook, *PlatFilename, 
 													GError, nullptr, bSwap, false, SaveFlags, Target, 
@@ -6711,12 +6711,6 @@ void UCookOnTheFlyServer::InitializeSandbox()
 
 void UCookOnTheFlyServer::InitializePackageStore(const TArray<FName>& TargetPlatformNames)
 {
-	// TODO: should ideally support all cook modes
-	if (!IsUsingPackageStore())
-	{
-		return;
-	}
-
 	const FString RootPath = FPaths::RootDir();
 	const FString RootPathSandbox = ConvertToFullSandboxPath(*RootPath, true);
 
@@ -6724,33 +6718,32 @@ void UCookOnTheFlyServer::InitializePackageStore(const TArray<FName>& TargetPlat
 
 	for (const FName PlatformName : TargetPlatformNames)
 	{
-		const FString ResolvedRootPath = RootPathSandbox.Replace(TEXT("[Platform]"), *PlatformName.ToString());
+		const FString PlatformString = PlatformName.ToString();
+		const FString ResolvedRootPath = RootPathSandbox.Replace(TEXT("[Platform]"), *PlatformString);
 
 		// just leak all memory for now
-		FPackageStoreBulkDataManifest* BulkDataManifest	= new FPackageStoreBulkDataManifest(ResolvedRootPath);
-		FLooseFileWriter* LooseFileWriter				= new FLooseFileWriter();
-		FSavePackageContext* SavePackageContext			= new FSavePackageContext(*LooseFileWriter, *BulkDataManifest);
+		FPackageStoreBulkDataManifest* BulkDataManifest	= new FPackageStoreBulkDataManifest(PlatformString);
+		FLooseFileWriter* LooseFileWriter				= IsUsingPackageStore() ? new FLooseFileWriter() : nullptr;
 
+		FSavePackageContext* SavePackageContext			= new FSavePackageContext(LooseFileWriter, BulkDataManifest);
 		SavePackageContexts.Add(SavePackageContext);
 	}
 }
 
 void UCookOnTheFlyServer::FinalizePackageStore()
 {
-	if (!IsUsingPackageStore())
-	{
-		return;
-	}
-
 	SCOPE_TIMER(FinalizePackageStore);
 
-	UE_LOG(LogCook, Display, TEXT("Saving PackageStoreHeaders(s)..."));
+	UE_LOG(LogCook, Display, TEXT("Saving BulkData manifest(s)..."));
 	const TArray<ITargetPlatform*>& TargetPlatforms = GetCookingTargetPlatforms();
 	for (int32 PlatformIndex = 0; PlatformIndex < TargetPlatforms.Num(); ++PlatformIndex)
 	{
-		SavePackageContexts[PlatformIndex]->BulkDataManifest.Save();
+		if (SavePackageContexts[PlatformIndex] != nullptr && SavePackageContexts[PlatformIndex]->BulkDataManifest != nullptr)
+		{
+			SavePackageContexts[PlatformIndex]->BulkDataManifest->Save();
+		}
 	}
-	UE_LOG(LogCook, Display, TEXT("Done saving PackageStoreHeaders(s)"));
+	UE_LOG(LogCook, Display, TEXT("Done saving BulkData manifest(s)"));
 }
 
 void UCookOnTheFlyServer::InitializeTargetPlatforms()
diff --git a/Engine/Source/Runtime/Core/Public/IO/IoDispatcher.h b/Engine/Source/Runtime/Core/Public/IO/IoDispatcher.h
--- a/Engine/Source/Runtime/Core/Public/IO/IoDispatcher.h
+++ b/Engine/Source/Runtime/Core/Public/IO/IoDispatcher.h
@@ -600,6 +600,23 @@ static FIoChunkId CreateIoChunkId(uint32 GlobalPackageId, uint16 ChunkIndex, EIo
 	return ChunkId;
 }
 
+/**
+ * Creates a chunk identifier for BulkData
+ */
+static FIoChunkId CreateBulkdataChunkId(int32 GlobalPackageId, uint64 BulkDataChunkId, EIoChunkType ChunkType)
+{
+	uint8 Data[12] = { 0 };
+
+	*reinterpret_cast<int32*>(&Data[0]) = GlobalPackageId;
+	*reinterpret_cast<uint64*>(&Data[4]) = BulkDataChunkId; // Top byte will get overwritten!
+	*reinterpret_cast<uint8*>(&Data[11]) = static_cast<uint8>(ChunkType);
+
+	FIoChunkId ChunkId;
+	ChunkId.Set(Data, 12);
+
+	return ChunkId;
+}
+
 //////////////////////////////////////////////////////////////////////////
 
 class FIoReadOptions
diff --git a/Engine/Source/Runtime/CoreUObject/Private/Serialization/BulkData.cpp b/Engine/Source/Runtime/CoreUObject/Private/Serialization/BulkData.cpp
--- a/Engine/Source/Runtime/CoreUObject/Private/Serialization/BulkData.cpp
+++ b/Engine/Source/Runtime/CoreUObject/Private/Serialization/BulkData.cpp
@@ -986,13 +986,6 @@ void FUntypedBulkData::Serialize( FArchive& Ar, UObject* Owner, int32 Idx, bool
 				BulkDataFlags &= ~BULKDATA_BadDataVersion;		
 			}
 			
-			if ((BulkDataFlags & BULKDATA_CookedForIoDispatcher) != 0)
-			{
-				// Load the bulk data index that will be used to create I/O chunk id's
-				uint16 BulkDataIndex = ~uint16(0);
- 				Ar << BulkDataIndex;
-			}
-
 			// determine whether the payload is stored inline or at the end of the file
 			const bool bPayloadInline = !(BulkDataFlags&BULKDATA_PayloadAtEndOfFile);
 
@@ -1097,12 +1090,6 @@ void FUntypedBulkData::Serialize( FArchive& Ar, UObject* Owner, int32 Idx, bool
 
 						Ar << BulkDataOffsetInFile;
 						BulkDataOffsetInFile += OffsetInFileFixup;
-
-						if ((BulkDataFlags & BULKDATA_CookedForIoDispatcher) != 0)
-						{
-							uint16 BulkDataIndex = ~uint16(0);
-							Ar << BulkDataIndex;
-						}
 					}
 					else
 					{
@@ -1124,12 +1111,6 @@ void FUntypedBulkData::Serialize( FArchive& Ar, UObject* Owner, int32 Idx, bool
 						int64 DummyBulkDataOffsetInFile;
 						Ar << DummyBulkDataOffsetInFile;
 
-						if ((BulkDataFlags & BULKDATA_CookedForIoDispatcher) != 0)
-						{
-							uint16 BulkDataIndex = ~uint16(0);
-							Ar << BulkDataIndex;
-						}
-
 						Filename = FPaths::ChangeExtension(Filename, TEXT(".ubulk"));
 					}
 				}
@@ -1233,7 +1214,6 @@ void FUntypedBulkData::Serialize( FArchive& Ar, UObject* Owner, int32 Idx, bool
 			// Only serialize status information if wanted.
 			int64 SavedBulkDataSizeOnDiskPos	= INDEX_NONE;
 			int64 SavedBulkDataOffsetInFilePos	= INDEX_NONE;
-			int64 SavedBulkDataChunkIdPos		= INDEX_NONE;
 			
 			// Keep track of position we are going to serialize placeholder BulkDataSizeOnDisk.
 			SavedBulkDataSizeOnDiskPos = Ar.Tell();
@@ -1270,19 +1250,6 @@ void FUntypedBulkData::Serialize( FArchive& Ar, UObject* Owner, int32 Idx, bool
 				bStoreInline = true;
 			}
 
-			// Write the ChunkID that will be used by the IoDispatcher, but ONLY if we are producing cooked data
-			if(Ar.IsCooking() )
-			{
-				FArchive::FScopeSetDebugSerializationFlags S(Ar, DSF_IgnoreDiff);
-
-				SavedBulkDataChunkIdPos = Ar.Tell();
-
-				uint16 BulkDataIndex = ~uint16(0);
-				Ar << BulkDataIndex;
-
-				BulkDataFlags |= BULKDATA_CookedForIoDispatcher;
-			}
-
 			if (!bStoreInline)
 			{
 				// set the flag indicating where the payload is stored
@@ -1298,7 +1265,6 @@ void FUntypedBulkData::Serialize( FArchive& Ar, UObject* Owner, int32 Idx, bool
 				BulkStore.BulkDataOffsetInFilePos = SavedBulkDataOffsetInFilePos;
 				BulkStore.BulkDataSizeOnDiskPos = SavedBulkDataSizeOnDiskPos;
 				BulkStore.BulkDataFlagsPos = SavedBulkDataFlagsPos;
-				BulkStore.BulkDataChunkIdPos = SavedBulkDataChunkIdPos;
 				BulkStore.BulkDataFlags = BulkDataFlags;
 				BulkStore.BulkData = this;
 
@@ -1308,7 +1274,6 @@ void FUntypedBulkData::Serialize( FArchive& Ar, UObject* Owner, int32 Idx, bool
 					int64 SavedDupeBulkDataFlagsPos = INDEX_NONE;
 					int64 SavedDupeBulkDataSizeOnDiskPos = INDEX_NONE;
 					int64 SavedDupeBulkDataOffsetInFilePos = INDEX_NONE;
-					int64 SavedDupeBulkDataChunkIdPos = INDEX_NONE;
 
 					int32 SavedDupeBulkDataFlags = (BulkDataFlags & ~BULKDATA_DuplicateNonOptionalPayload) | BULKDATA_OptionalPayload;
 					{
@@ -1335,18 +1300,6 @@ void FUntypedBulkData::Serialize( FArchive& Ar, UObject* Owner, int32 Idx, bool
 						BulkDataOffsetInFile = INDEX_NONE;
 						// And serialize the placeholder which is going to be overwritten later.
 						Ar << BulkDataOffsetInFile;
-
-						// TODO: Fairly sure we can only be in the BULKDATA_DuplicateNonOptionalPayload branch if this
-						// is cooked data, see if we can remove the if check
-						if (Ar.IsCooking())
-						{
-							SavedDupeBulkDataChunkIdPos = Ar.Tell();
-							
-							uint16 BulkDataIndex = ~uint16(0);
-							Ar << BulkDataIndex;
-
-							check((SavedDupeBulkDataFlags& BULKDATA_CookedForIoDispatcher) != 0); // Validate that the BULKDATA_CookedForIoDispatcher flag was already set
-						}
 					}
 
 					// add duplicate bulkdata with different flag
@@ -1357,8 +1310,6 @@ void FUntypedBulkData::Serialize( FArchive& Ar, UObject* Owner, int32 Idx, bool
 					DupeBulkStore.BulkDataSizeOnDiskPos = SavedDupeBulkDataSizeOnDiskPos;
 					DupeBulkStore.BulkDataFlagsPos = SavedDupeBulkDataFlagsPos;
 					DupeBulkStore.BulkDataFlags = SavedDupeBulkDataFlags;
-					DupeBulkStore.BulkDataChunkIdPos = SavedDupeBulkDataChunkIdPos;
-
 					DupeBulkStore.BulkData = this;
 				}
 				
diff --git a/Engine/Source/Runtime/CoreUObject/Private/Serialization/BulkData2.cpp b/Engine/Source/Runtime/CoreUObject/Private/Serialization/BulkData2.cpp
--- a/Engine/Source/Runtime/CoreUObject/Private/Serialization/BulkData2.cpp
+++ b/Engine/Source/Runtime/CoreUObject/Private/Serialization/BulkData2.cpp
@@ -566,43 +566,33 @@ void FBulkDataBase::Serialize(FArchive& Ar, UObject* Owner, int32 /*Index*/, boo
 			Ar << DummyValue;
 		}
 
-		if ((BulkDataFlags & BULKDATA_CookedForIoDispatcher) != 0)
-		{
-			uint16 BulkDataIndex;
-			Ar << BulkDataIndex;
+		check((BulkDataFlags & BULKDATA_UsesIoDispatcher) == 0);
 
-			if (BulkDataIndex != InvalidBulkDataIndex)
-			{
-				UPackage* Package = Owner->GetOutermost();
-				ChunkID = CreateIoChunkId(Package->GetGlobalPackageId(), BulkDataIndex, EIoChunkType::BulkData);
-			}
-			else
-			{
-				ChunkID = FIoChunkId::InvalidChunkId;
-			}
-		}
+		// Assuming that Owner/Package/Linker are all valid, the old BulkData system would
+		// generally fail if any of these were nullptr but had plenty of inconsistent checks
+		// scattered throughout.
+		check(Owner != nullptr);
+		const UPackage* Package = Owner->GetOutermost();
+		check(Package != nullptr);
+
+		if (!IsInlined() && bUseZenLoader)
+		{		
+			ChunkID = CreateBulkdataChunkId(Package->GetGlobalPackageId(), BulkDataOffsetInFile, EIoChunkType::BulkData);
 
-		if (!bUseZenLoader || ChunkID == FIoChunkId::InvalidChunkId)
+			BulkDataFlags |= BULKDATA_UsesIoDispatcher;
+		}
+		else
 		{
-			// Remove the flag if we are not going to use it
-			BulkDataFlags &= ~BULKDATA_CookedForIoDispatcher;
 			Fallback.Token = InvalidToken;
 			Fallback.BulkDataSize = BulkDataSize;
 		}
 
 		FString Filename;
 		FName PackageName;
-		FLinkerLoad* Linker = nullptr;
+		const FLinkerLoad* Linker = nullptr;
 
 		if (bUseZenLoader == false)
 		{
-			// Assuming that Owner/Package/Linker are all valid, the old BulkData system would
-			// generally fail if any of these were nullptr but had plenty of inconsistent checks
-			// scattered throughout.
-			check(Owner != nullptr);
-			UPackage* Package = Owner->GetOutermost();
-			check(Package != nullptr);
-
 			Linker = FLinkerLoad::FindExistingLinkerForPackage(Package);
 			check(Linker != nullptr);
 
@@ -637,7 +627,7 @@ void FBulkDataBase::Serialize(FArchive& Ar, UObject* Owner, int32 /*Index*/, boo
 						if (!bUseZenLoader || ChunkID == FIoChunkId::InvalidChunkId)
 						{
 							// Remove the flag if we are not going to use it
-							BulkDataFlags &= ~BULKDATA_CookedForIoDispatcher;
+							BulkDataFlags &= ~BULKDATA_UsesIoDispatcher;
 							Fallback.Token = InvalidToken;
 							Fallback.BulkDataSize = BulkDataSize;
 						}
@@ -854,7 +844,7 @@ bool FBulkDataBase::IsMemoryMapped() const
 
 bool FBulkDataBase::IsUsingIODispatcher() const
 {
-	return (BulkDataFlags & BULKDATA_CookedForIoDispatcher) != 0;
+	return (BulkDataFlags & BULKDATA_UsesIoDispatcher) != 0;
 }
 
 IBulkDataIORequest* FBulkDataBase::CreateStreamingRequest(EAsyncIOPriorityAndFlags Priority, FAsyncFileCallBack* CompleteCallback, uint8* UserSuppliedMemory) const
@@ -1050,7 +1040,7 @@ void FBulkDataBase::LoadDataDirectly(void** DstBuffer)
 	}
 }
 
-void FBulkDataBase::SerializeDuplicateData(FArchive& Ar, FLinkerLoad* Linker, uint32& OutBulkDataFlags, int64& OutBulkDataSizeOnDisk, int64& OutBulkDataOffsetInFile)
+void FBulkDataBase::SerializeDuplicateData(FArchive& Ar, const FLinkerLoad* Linker, uint32& OutBulkDataFlags, int64& OutBulkDataSizeOnDisk, int64& OutBulkDataOffsetInFile)
 {
 	Ar << OutBulkDataFlags;
 
@@ -1068,7 +1058,7 @@ void FBulkDataBase::SerializeDuplicateData(FArchive& Ar, FLinkerLoad* Linker, ui
 
 	Ar << OutBulkDataOffsetInFile;
 
-	if ((OutBulkDataFlags & BULKDATA_CookedForIoDispatcher) != 0)
+	if ((OutBulkDataFlags & BULKDATA_UsesIoDispatcher) != 0)
 	{
 		uint16 DummyBulkDataIndex = InvalidBulkDataIndex;
 		Ar << DummyBulkDataIndex;
diff --git a/Engine/Source/Runtime/CoreUObject/Private/UObject/SavePackage.cpp b/Engine/Source/Runtime/CoreUObject/Private/UObject/SavePackage.cpp
--- a/Engine/Source/Runtime/CoreUObject/Private/UObject/SavePackage.cpp
+++ b/Engine/Source/Runtime/CoreUObject/Private/UObject/SavePackage.cpp
@@ -469,35 +469,27 @@ private:
 	TSet<FNameEntryId> ReferencedNames;
 };
 
-void FPackageStoreBulkDataManifest::PackageDesc::AddData(uint16 InIndex, uint64 InOffset, uint64 InSize)
+void FPackageStoreBulkDataManifest::PackageDesc::AddData(uint64 InChunkId, uint64 InOffset, uint64 InSize)
 {
-	// The index is supposed to be unique
-	auto func = [=](const BulkDataDesc& Entry) { return Entry.Index == InIndex; };
+	// The ChunkId is supposed to be unique
+	auto func = [=](const BulkDataDesc& Entry) { return Entry.ChunkId == InChunkId; };
 	check(Data.FindByPredicate(func) == nullptr);
 
 	BulkDataDesc& Entry = Data.Emplace_GetRef();
 
-	Entry.Index = InIndex;
+	Entry.ChunkId = InChunkId;
 	Entry.Offset = InOffset;
 	Entry.Size = InSize;
 }
 
-FPackageStoreBulkDataManifest::FPackageStoreBulkDataManifest(const FString& InRootPath)
+FPackageStoreBulkDataManifest::FPackageStoreBulkDataManifest(const FString& PlatformName)
 {
-	RootPath = InRootPath;
-	FPaths::NormalizeFilename(RootPath);
-
-	if (!RootPath.EndsWith("/"))
-	{
-		RootPath.Append("/");
-	}
-	
-	Filename = RootPath + "megafile.ubulkmanifest";
+	Filename = FPaths::ProjectSavedDir() / TEXT("BulkDataManifests") / PlatformName + TEXT(".ubulkmanifest");
 }
 
 FArchive& operator<<(FArchive& Ar, FPackageStoreBulkDataManifest::PackageDesc::BulkDataDesc& Entry)
 {
-	Ar << Entry.Index;
+	Ar << Entry.ChunkId;
 	Ar << Entry.Offset;
 	Ar << Entry.Size; 
 
@@ -543,12 +535,12 @@ const FPackageStoreBulkDataManifest::PackageDesc* FPackageStoreBulkDataManifest:
 	return Data.Find(NormalizedFilename);
 }
 
-void FPackageStoreBulkDataManifest::AddFileAccess(const FString& PackageFilename, uint16 InIndex, uint64 InOffset, uint64 InSize)
+void FPackageStoreBulkDataManifest::AddFileAccess(const FString& PackageFilename, uint64 InChunkId, uint64 InOffset, uint64 InSize)
 {
 	const FString NormalizedFilename = FixFilename(PackageFilename);
 	
 	PackageDesc& Entry = GetOrCreateFileAccess(NormalizedFilename);
-	Entry.AddData(InIndex, InOffset, InSize);
+	Entry.AddData(InChunkId, InOffset, InSize);
 }
 
 FPackageStoreBulkDataManifest::PackageDesc& FPackageStoreBulkDataManifest::GetOrCreateFileAccess(const FString& PackageFilename)
@@ -5723,7 +5715,7 @@ FSavePackageResultStruct UPackage::Save(UPackage* InOuter, UObject* Base, EObjec
 
 							if (IsEventDrivenLoaderEnabledInCookedBuilds() && Linker->IsCooking())
 							{
-								if (SavePackageContext)
+								if (SavePackageContext != nullptr && SavePackageContext->PackageStoreWriter != nullptr)
 								{
 									FIoBuffer IoBuffer(FIoBuffer::AssumeOwnership, Writer->ReleaseOwnership(), DataSize);
 
@@ -5733,7 +5725,7 @@ FSavePackageResultStruct UPackage::Save(UPackage* InOuter, UObject* Base, EObjec
 									HeaderInfo.PackageName		= InOuter->GetFName();
 									HeaderInfo.LooseFilePath	= Filename;
 
-									SavePackageContext->PackageStoreWriter.WriteHeader(HeaderInfo, FIoBuffer(IoBuffer.Data(), HeaderSize, IoBuffer));
+									SavePackageContext->PackageStoreWriter->WriteHeader(HeaderInfo, FIoBuffer(IoBuffer.Data(), HeaderSize, IoBuffer));
 
 									FPackageStoreWriter::ExportsInfo ExportsInfo;
 									ExportsInfo.LooseFilePath	= Filename;
@@ -5749,7 +5741,7 @@ FSavePackageResultStruct UPackage::Save(UPackage* InOuter, UObject* Base, EObjec
 										ExportsInfo.Exports.Add(FIoBuffer(IoBuffer.Data() + Export.SerialOffset, Export.SerialSize, IoBuffer));
 									}
 
-									SavePackageContext->PackageStoreWriter.WriteExports(ExportsInfo, FIoBuffer(ExportsData, DataSize - HeaderSize, IoBuffer));
+									SavePackageContext->PackageStoreWriter->WriteExports(ExportsInfo, FIoBuffer(ExportsData, DataSize - HeaderSize, IoBuffer));
 								}
 								else
 								{
@@ -6289,23 +6281,11 @@ void SaveBulkData(FLinkerSave* Linker, const UPackage* InOuter, const TCHAR* Fil
 				*Linker << SizeOnDiskAsInt32;
 			}
 
-			if ((BulkDataStorageInfo.BulkDataFlags & BULKDATA_CookedForIoDispatcher) != 0 && SavePackageContext != nullptr)
+			if (SavePackageContext != nullptr && SavePackageContext->BulkDataManifest != nullptr)
 			{
-				check(BulkDataIndex < TNumericLimits<uint16>::Max());
-
-				SavePackageContext->BulkDataManifest.AddFileAccess
-				(
-					Filename, BulkDataIndex, BulkStartOffset, SizeOnDisk
-				);
-
-				check(BulkDataStorageInfo.BulkDataChunkIdPos != INDEX_NONE);
-
-				Linker->Seek(BulkDataStorageInfo.BulkDataChunkIdPos);
-				*Linker << BulkDataIndex;
-				
-				BulkDataIndex++;
+				SavePackageContext->BulkDataManifest->AddFileAccess(Filename, StoredBulkStartOffset, BulkStartOffset, SizeOnDisk);
 			}
-
+			
 			Linker->Seek(LinkerEndOffset);
 
 			// Restore BulkData flags to before serialization started
@@ -6321,7 +6301,7 @@ void SaveBulkData(FLinkerSave* Linker, const UPackage* InOuter, const TCHAR* Fil
 
 			const bool bWriteBulkToDisk = !bDiffing;
 
-			if (SavePackageContext && bWriteBulkToDisk)
+			if (SavePackageContext != nullptr && SavePackageContext->PackageStoreWriter != nullptr && bWriteBulkToDisk)
 			{
 				auto ToIoBuffer = [](FLargeMemoryWriter* Writer)
 				{
@@ -6334,13 +6314,13 @@ void SaveBulkData(FLinkerSave* Linker, const UPackage* InOuter, const TCHAR* Fil
 				BulkInfo.LooseFilePath = Filename;
 				BulkInfo.BulkdataType = FPackageStoreWriter::FBulkDataInfo::Standard;
 
-				SavePackageContext->PackageStoreWriter.WriteBulkdata(BulkInfo, ToIoBuffer(BulkArchive.Get()));
+				SavePackageContext->PackageStoreWriter->WriteBulkdata(BulkInfo, ToIoBuffer(BulkArchive.Get()));
 
 				BulkInfo.BulkdataType = FPackageStoreWriter::FBulkDataInfo::Optional;
-				SavePackageContext->PackageStoreWriter.WriteBulkdata(BulkInfo, ToIoBuffer(OptionalBulkArchive.Get()));
+				SavePackageContext->PackageStoreWriter->WriteBulkdata(BulkInfo, ToIoBuffer(OptionalBulkArchive.Get()));
 
 				BulkInfo.BulkdataType = FPackageStoreWriter::FBulkDataInfo::Mmap;
-				SavePackageContext->PackageStoreWriter.WriteBulkdata(BulkInfo, ToIoBuffer(MappedBulkArchive.Get()));
+				SavePackageContext->PackageStoreWriter->WriteBulkdata(BulkInfo, ToIoBuffer(MappedBulkArchive.Get()));
 			}
 			else
 			{
diff --git a/Engine/Source/Runtime/CoreUObject/Public/Serialization/BulkData.h b/Engine/Source/Runtime/CoreUObject/Public/Serialization/BulkData.h
--- a/Engine/Source/Runtime/CoreUObject/Public/Serialization/BulkData.h
+++ b/Engine/Source/Runtime/CoreUObject/Public/Serialization/BulkData.h
@@ -64,8 +64,8 @@ enum EBulkDataFlags
 	BULKDATA_DuplicateNonOptionalPayload		= 1 << 14,
 	/** Indicates that an old ID is present in the data, at some point when the DDCs are flushed we can remove this. */
 	BULKDATA_BadDataVersion						= 1 << 15,
-	/** Indicates that the data was cooked for use with IoDispatcher. */
-	BULKDATA_CookedForIoDispatcher				= 1 << 16,
+	/** Assigned at runtime to indicate that the BulkData should be using the IoDispatcher when loading, not filepaths. */
+	BULKDATA_UsesIoDispatcher					= 1 << 16,
 };
 
 /**
diff --git a/Engine/Source/Runtime/CoreUObject/Public/Serialization/BulkData2.h b/Engine/Source/Runtime/CoreUObject/Public/Serialization/BulkData2.h
--- a/Engine/Source/Runtime/CoreUObject/Public/Serialization/BulkData2.h
+++ b/Engine/Source/Runtime/CoreUObject/Public/Serialization/BulkData2.h
@@ -129,7 +129,7 @@ public:
 private:
 	void LoadDataDirectly(void** DstBuffer);
 
-	void SerializeDuplicateData(FArchive& Ar, FLinkerLoad* Linker, uint32& OutBulkDataFlags, int64& OutBulkDataSizeOnDisk, int64& OutBulkDataOffsetInFile);
+	void SerializeDuplicateData(FArchive& Ar, const FLinkerLoad* Linker, uint32& OutBulkDataFlags, int64& OutBulkDataSizeOnDisk, int64& OutBulkDataOffsetInFile);
 	void SerializeBulkData(FArchive& Ar, void* DstBuffer, int64 DataLength);
 
 	void AllocateData(SIZE_T SizeInBytes);
diff --git a/Engine/Source/Runtime/CoreUObject/Public/UObject/LinkerSave.h b/Engine/Source/Runtime/CoreUObject/Public/UObject/LinkerSave.h
--- a/Engine/Source/Runtime/CoreUObject/Public/UObject/LinkerSave.h
+++ b/Engine/Source/Runtime/CoreUObject/Public/UObject/LinkerSave.h
@@ -57,8 +57,6 @@ public:
 		int64 BulkDataSizeOnDiskPos;
 		/** Offset to the location where the bulk data flags are stored */
 		int64 BulkDataFlagsPos;
-		/** Offset to the location where the bulk data ChunkID is stored */
-		int64 BulkDataChunkIdPos;
 		/** Bulk data flags at the time of serialization */
 		uint32 BulkDataFlags;
 		/** The bulkdata */
diff --git a/Engine/Source/Runtime/CoreUObject/Public/UObject/SavePackage.h b/Engine/Source/Runtime/CoreUObject/Public/UObject/SavePackage.h
--- a/Engine/Source/Runtime/CoreUObject/Public/UObject/SavePackage.h
+++ b/Engine/Source/Runtime/CoreUObject/Public/UObject/SavePackage.h
@@ -12,23 +12,16 @@ class FArchive;
 class FLinkerLoad;
 class FLinkerSave;
 
-class IBulkDataManifest 
+class FPackageStoreBulkDataManifest
 {
 public:
-	virtual ~IBulkDataManifest() = default;
-
-	virtual void Save() = 0;
-	virtual void AddFileAccess(const FString& PackageFilename, uint16 InIndex, uint64 InOffset, uint64 InSize) = 0;
-};
-
-class FPackageStoreBulkDataManifest : public IBulkDataManifest
-{
-public:
-	COREUOBJECT_API FPackageStoreBulkDataManifest(const FString& InRootPath);
-	COREUOBJECT_API virtual ~FPackageStoreBulkDataManifest();
+	COREUOBJECT_API FPackageStoreBulkDataManifest(const FString& PlatformName);
+	COREUOBJECT_API ~FPackageStoreBulkDataManifest();
 
 	COREUOBJECT_API bool Load();
-	virtual void Save() override;
+	COREUOBJECT_API void Save() ;
+
+	void AddFileAccess(const FString& PackageFilename, uint64 InChunkId, uint64 InOffset, uint64 InSize);
 	
 	const FString& GetFilename() const { return Filename; }
 
@@ -37,12 +30,13 @@ public:
 	public:
 		struct BulkDataDesc
 		{
-			uint16 Index;
+			uint64 ChunkId;	// Note this is the Offset before the linker BulkDataStartOffset is
+							// applied, to make it easier to compute at runtime.
 			uint64 Offset;
 			uint64 Size;
 		};
 
-		void AddData(uint16 InIndex, uint64 InOffset, uint64 InSize);
+		void AddData(uint64 InChunkId, uint64 InOffset, uint64 InSize);
 		const TArray<BulkDataDesc>& GetDataArray() const { return Data; }
 	private:
 		friend FArchive& operator<<(FArchive& Ar, PackageDesc& Entry);
@@ -52,8 +46,6 @@ public:
 	COREUOBJECT_API const PackageDesc* Find(const FString& PackageName) const;
 
 private:
-	virtual void AddFileAccess(const FString& PackageFilename, uint16 InIndex, uint64 InOffset, uint64 InSize) override;
-
 	PackageDesc& GetOrCreateFileAccess(const FString& PackageFilename);
 
 	FString FixFilename(const FString& InFileName) const;
@@ -126,12 +118,12 @@ private:
 class FSavePackageContext
 {
 public:
-	FSavePackageContext(FPackageStoreWriter& InPackageStoreWriter, IBulkDataManifest& InBulkDataManifest)
+	FSavePackageContext(FPackageStoreWriter* InPackageStoreWriter, FPackageStoreBulkDataManifest* InBulkDataManifest)
 	: PackageStoreWriter(InPackageStoreWriter) 
 	, BulkDataManifest(InBulkDataManifest)
 	{
 	}
 
-	FPackageStoreWriter& PackageStoreWriter;
-	IBulkDataManifest& BulkDataManifest;
+	FPackageStoreWriter* PackageStoreWriter;
+	FPackageStoreBulkDataManifest* BulkDataManifest;
 };
