AsyncLoading2: Change expensive asserts from ensures to checks to be compiled out in test

diff --git a/Engine/Source/Runtime/CoreUObject/Private/Serialization/AsyncLoading2.cpp b/Engine/Source/Runtime/CoreUObject/Private/Serialization/AsyncLoading2.cpp
--- a/Engine/Source/Runtime/CoreUObject/Private/Serialization/AsyncLoading2.cpp
+++ b/Engine/Source/Runtime/CoreUObject/Private/Serialization/AsyncLoading2.cpp
@@ -1045,7 +1045,7 @@ struct FAsyncPackage2
 		}
 		else
 		{
-			ensure(!OwnedObjects.Contains(Object));
+			check(!OwnedObjects.Contains(Object));
 			OwnedObjects.Add(Object);
 		}
 	}
@@ -1054,10 +1054,10 @@ struct FAsyncPackage2
 	{
 		if (bForceAdd || !IsInAsyncLoadingThread())
 		{
-			ensure(!OwnedObjects.Contains(Object));
+			check(!OwnedObjects.Contains(Object));
 			OwnedObjects.Add(Object);
 		}
-		ensure(OwnedObjects.Contains(Object));
+		check(OwnedObjects.Contains(Object));
 	}
 
 	void AddOwnedObjectWithAsyncFlag(UObject* Object, bool bForceAdd)
@@ -1065,10 +1065,10 @@ struct FAsyncPackage2
 		AddOwnedObject(Object, bForceAdd);
 		if (bForceAdd || IsInGameThread())
 		{
-			ensure(!Object->HasAnyInternalFlags(EInternalObjectFlags::Async));
+			check(!Object->HasAnyInternalFlags(EInternalObjectFlags::Async));
 			Object->SetInternalFlags(EInternalObjectFlags::Async);
 		}
-		ensure(Object->HasAnyInternalFlags(EInternalObjectFlags::Async));
+		check(Object->HasAnyInternalFlags(EInternalObjectFlags::Async));
 	}
 
 	void ClearOwnedObjects();
@@ -2599,7 +2599,7 @@ void FAsyncPackage2::EventDrivenCreateExport(int32 LocalExportIndex)
 		if (!bIsCompleteyLoaded)
 		{
 			UE_LOG(LogStreaming, VeryVerbose, TEXT("Note2: %s was constructed during load and is an export and so needs loading."), *Object->GetFullName());
-			ensure(!(ObjectFlags & (RF_NeedLoad | RF_WasLoaded))); // If export exist but is not completed, it is expected to have been created from a native constructor and not from EventDrivenCreateExport, but who knows...?
+			check(!(ObjectFlags & (RF_NeedLoad | RF_WasLoaded))); // If export exist but is not completed, it is expected to have been created from a native constructor and not from EventDrivenCreateExport, but who knows...?
 			if (ObjectFlags & RF_ClassDefaultObject)
 			{
 				// never call PostLoadSubobjects on class default objects, this matches the behavior of the old linker where
@@ -3673,12 +3673,12 @@ static void VerifyLoadFlagsWhenFinishedLoading()
 			const bool bHasAnyLoadIntermediateFlags = !!(Flags & LoadIntermediateFlags);
 			const bool bWasLoaded = !!(Flags & RF_WasLoaded);
 			const bool bLoadCompleted = !!(Flags & RF_LoadCompleted);
-			ensure(!bHasAnyAsyncFlags);
-			ensure(!bHasAnyLoadIntermediateFlags);
+			check(!bHasAnyAsyncFlags);
+			check(!bHasAnyLoadIntermediateFlags);
 			if (bWasLoaded)
 			{
 				const bool bIsPackage = Obj->IsA(UPackage::StaticClass());
-				ensure(bIsPackage || bLoadCompleted);
+				check(bIsPackage || bLoadCompleted);
 			}
 		}
 	}
@@ -3909,7 +3909,7 @@ FAsyncPackage2::~FAsyncPackage2()
 	MarkRequestIDsAsComplete();
 	SerialNumber = 0; // the weak pointer will always fail now
 	
-	ensure(OwnedObjects.Num() == 0);
+	check(OwnedObjects.Num() == 0);
 }
 
 
@@ -3921,10 +3921,10 @@ void FAsyncPackage2::ClearOwnedObjects()
 		const EInternalObjectFlags InternalFlags = Object->GetInternalFlags();
 		EInternalObjectFlags InternalFlagsToClear = EInternalObjectFlags::None;
 
-		ensure(!(Flags & (RF_NeedPostLoad | RF_NeedPostLoadSubobjects)));
+		check(!(Flags & (RF_NeedPostLoad | RF_NeedPostLoadSubobjects)));
 		if (!!(InternalFlags & EInternalObjectFlags::AsyncLoading))
 		{
-			ensure(!(Flags & RF_WasLoaded));
+			check(!(Flags & RF_WasLoaded));
 			InternalFlagsToClear |= EInternalObjectFlags::AsyncLoading;
 		}
 
