[AsyncLoading2] Moved storetoc, glimports, slimports and script arcs into container file

diff --git a/Engine/Source/Developer/IoStoreUtilities/Private/IoStoreUtilities.cpp b/Engine/Source/Developer/IoStoreUtilities/Private/IoStoreUtilities.cpp
--- a/Engine/Source/Developer/IoStoreUtilities/Private/IoStoreUtilities.cpp
+++ b/Engine/Source/Developer/IoStoreUtilities/Private/IoStoreUtilities.cpp
@@ -18,6 +18,7 @@
 #include "Serialization/JsonReader.h"
 #include "Serialization/JsonSerializer.h"
 #include "Serialization/BufferWriter.h"
+#include "Serialization/LargeMemoryWriter.h"
 #include "UObject/PackageFileSummary.h"
 #include "UObject/ObjectResource.h"
 #include "UObject/Package.h"
@@ -96,23 +97,20 @@ public:
 		return NameIndices.Num();
 	}
 
-	int32 Save(const FString& Filename)
+	int32 Save(FArchive& Ar, const FString& CsvFilePath)
 	{
 		int32 TotalSize = 0;
 		{
-			TUniquePtr<FArchive> BinArchive(IFileManager::Get().CreateFileWriter(*Filename));
 			int32 NameCount = NameMap.Num();
-			*BinArchive << NameCount;
+			Ar << NameCount;
 			for (int32 I = 0; I < NameCount; ++I)
 			{
-				FName::GetEntry(NameMap[I])->Write(*BinArchive);
+				FName::GetEntry(NameMap[I])->Write(Ar);
 			}
-			TotalSize = BinArchive->TotalSize();
+			TotalSize = Ar.TotalSize();
 		}
 #if OUTPUT_NAMEMAP_CSV
 		{
-			FString CsvFilePath = Filename;
-			CsvFilePath.Append(TEXT(".csv"));
 			TUniquePtr<FArchive> CsvArchive(IFileManager::Get().CreateFileWriter(*CsvFilePath));
 			if (CsvArchive)
 			{
@@ -1028,10 +1026,10 @@ int32 CreateTarget(const FContainerTarget& Target)
 	uint64 ImportPreloadCount = 0;
 	uint64 ExportPreloadCount = 0;
 
-	TUniquePtr<FArchive> StoreTocArchive(IFileManager::Get().CreateFileWriter(*(OutputDir / TEXT("megafile.ustoretoc"))));
-	TUniquePtr<FArchive> GlimportArchive(IFileManager::Get().CreateFileWriter(*(OutputDir / TEXT("megafile.uglimport"))));
-	TUniquePtr<FArchive> SlimportArchive(IFileManager::Get().CreateFileWriter(*(OutputDir / TEXT("megafile.uslimport"))));
-	TUniquePtr<FArchive> ScriptArcsArchive(IFileManager::Get().CreateFileWriter(*(OutputDir / TEXT("megafile.uscriptarcs"))));
+	TUniquePtr<FLargeMemoryWriter> StoreTocArchive = MakeUnique<FLargeMemoryWriter>(0, true);
+	TUniquePtr<FLargeMemoryWriter> GlimportArchive = MakeUnique<FLargeMemoryWriter>(0, true);
+	TUniquePtr<FLargeMemoryWriter> SlimportArchive = MakeUnique<FLargeMemoryWriter>(0, true);
+	TUniquePtr<FLargeMemoryWriter> ScriptArcsArchive = MakeUnique<FLargeMemoryWriter>(0, true);
 
 	FIoStoreEnvironment IoStoreEnv;
 #if !SKIP_WRITE_CONTAINER
@@ -1716,12 +1714,55 @@ int32 CreateTarget(const FContainerTarget& Target)
 		}
 	}
 
-	UE_LOG(LogIoStore, Display, TEXT("Saving NameMap..."));
-	FString NameMapFilename = OutputDir / TEXT("Container.namemap");
-	uint64 GlobalNameMapSize = NameMapBuilder.Save(NameMapFilename);
+	auto AppendToContainer = [](FIoStoreWriter& IoWriter, FLargeMemoryWriter& Ar, const FIoChunkId& ChunkId) -> FIoStatus
+	{
+		FIoBuffer IoBuffer(FIoBuffer::Wrap, Ar.GetData(), Ar.TotalSize());
+		return IoWriter.Append(ChunkId, IoBuffer);
+	};
+
+	{
+		UE_LOG(LogIoStore, Display, TEXT("Saving global meta data to container file"));
+		FLargeMemoryWriter GlobalMetaArchive(0, true);
+
+		int32 StoreTocByteCount = StoreTocArchive->TotalSize();
+		GlobalMetaArchive << StoreTocByteCount;
+		GlobalMetaArchive.Serialize(StoreTocArchive->GetData(), StoreTocByteCount);
+
+		int32 SlimportsByteCount = SlimportArchive->TotalSize();
+		GlobalMetaArchive << SlimportsByteCount;
+		GlobalMetaArchive.Serialize(SlimportArchive->GetData(), SlimportsByteCount);
+
+		int32 GlimportsByteCount = GlimportArchive->TotalSize();
+		GlobalMetaArchive << GlimportsByteCount;
+		GlobalMetaArchive.Serialize(GlimportArchive->GetData(), GlimportsByteCount);
+
+		const FIoStatus Status = AppendToContainer(*IoStoreWriter, GlobalMetaArchive, CreateIoChunkId(0, 0, EIoChunkType::LoaderGlobalMeta));
+		if (!Status.IsOk())
+		{
+			UE_LOG(LogIoStore, Error, TEXT("Failed to save global meta data to container file"));
+		}
+	}
 
-	UE_LOG(LogIoStore, Display, TEXT("Saving StoreToc..."));
-	StoreTocArchive->Close();
+	{
+		UE_LOG(LogIoStore, Display, TEXT("Saving initial load meta data to container file"));
+		const FIoStatus Status = AppendToContainer(*IoStoreWriter, *ScriptArcsArchive, CreateIoChunkId(0, 0, EIoChunkType::LoaderInitialLoadMeta));
+		if (!Status.IsOk())
+		{
+			UE_LOG(LogIoStore, Error, TEXT("Failed to save initial load meta data to container file"));
+		}
+	}
+	
+	uint64 GlobalNameMapSize = 0;
+	{
+		UE_LOG(LogIoStore, Display, TEXT("Saving global name map to container file"));
+		FLargeMemoryWriter GlobalNameMapArchive(0, true);
+		GlobalNameMapSize = NameMapBuilder.Save(GlobalNameMapArchive, OutputDir / TEXT("Container.namemap.csv"));
+		const FIoStatus Status = AppendToContainer(*IoStoreWriter, GlobalNameMapArchive, CreateIoChunkId(0, 0, EIoChunkType::LoaderGlobalNameMap));
+		if (!Status.IsOk())
+		{
+			UE_LOG(LogIoStore, Error, TEXT("Failed to save global name map to container file"));
+		}
+	}
 
 	UE_LOG(LogIoStore, Display, TEXT("Calculating stats..."));
 	uint64 UExpSize = 0;
diff --git a/Engine/Source/Programs/AutomationTool/Scripts/CopyBuildToStagingDirectory.Automation.cs b/Engine/Source/Programs/AutomationTool/Scripts/CopyBuildToStagingDirectory.Automation.cs
--- a/Engine/Source/Programs/AutomationTool/Scripts/CopyBuildToStagingDirectory.Automation.cs
+++ b/Engine/Source/Programs/AutomationTool/Scripts/CopyBuildToStagingDirectory.Automation.cs
@@ -1745,15 +1745,7 @@ public partial class Project : CommandUtils
 			}
 
 			// Filter I/O store container files
-			if (Src.HasExtension(".namemap") ||
-				Src.HasExtension(".ucas") ||
-				Src.HasExtension(".utoc") ||
-				Src.HasExtension(".ubulkmanifest") ||
-				Src.HasExtension(".uglimport") ||
-				Src.HasExtension(".unamemap") ||
-				Src.HasExtension(".uscriptarcs") ||
-				Src.HasExtension(".uslimport") ||
-				Src.HasExtension(".ustoretoc"))
+			if (Src.HasExtension(".ucas") || Src.HasExtension(".utoc"))
 			{
 				LogInformation("Excluding {0}", Src);
 				continue;
diff --git a/Engine/Source/Runtime/Core/Public/IO/IoDispatcher.h b/Engine/Source/Runtime/Core/Public/IO/IoDispatcher.h
--- a/Engine/Source/Runtime/Core/Public/IO/IoDispatcher.h
+++ b/Engine/Source/Runtime/Core/Public/IO/IoDispatcher.h
@@ -581,6 +581,9 @@ enum class EIoChunkType : uint8
 	ExportData,
 	ExportBundleData,
 	BulkData,
+	LoaderGlobalMeta,
+	LoaderInitialLoadMeta,
+	LoaderGlobalNameMap
 };
 
 /**
diff --git a/Engine/Source/Runtime/CoreUObject/Private/Serialization/AsyncLoading2.cpp b/Engine/Source/Runtime/CoreUObject/Private/Serialization/AsyncLoading2.cpp
--- a/Engine/Source/Runtime/CoreUObject/Private/Serialization/AsyncLoading2.cpp
+++ b/Engine/Source/Runtime/CoreUObject/Private/Serialization/AsyncLoading2.cpp
@@ -48,6 +48,7 @@
 #include "UObject/GCObject.h"
 #include "UObject/ObjectRedirector.h"
 #include "Serialization/BulkData.h"
+#include "Serialization/LargeMemoryReader.h"
 
 #if UE_BUILD_DEVELOPMENT || UE_BUILD_DEBUG
 //PRAGMA_DISABLE_OPTIMIZATION
@@ -191,19 +192,33 @@ struct FPackageSummary
 class FGlobalNameMap
 {
 public:
-	void Load(const FString& FilePath)
+	void Load(FIoDispatcher& IoDispatcher)
 	{
-		TUniquePtr<FArchive> Ar(IFileManager::Get().CreateFileReader(*FilePath));
-		check(Ar);
+		FEvent* Event = FPlatformProcess::GetSynchEventFromPool();
+		FIoBuffer IoBuffer;
+
+		IoDispatcher.ReadWithCallback(
+			CreateIoChunkId(0, 0, EIoChunkType::LoaderGlobalNameMap),
+			FIoReadOptions(),
+			[Event, &IoBuffer](TIoStatusOr<FIoBuffer> Result)
+			{
+				IoBuffer = Result.ConsumeValueOrDie();
+				Event->Trigger();
+			});
+
+		Event->Wait();
+		FPlatformProcess::ReturnSynchEventToPool(Event);
+
+		FLargeMemoryReader Ar(IoBuffer.Data(), IoBuffer.DataSize());
 
 		int32 NameCount;
-		*Ar << NameCount;
+		Ar << NameCount;
 		NameEntries.Reserve(NameCount);
 		FNameEntrySerialized SerializedNameEntry(ENAME_LinkerConstructor);
 
 		for (int32 I = 0; I < NameCount; ++I)
 		{
-			*Ar << SerializedNameEntry;
+			Ar << SerializedNameEntry;
 			NameEntries.Emplace(FName(SerializedNameEntry).GetDisplayIndex());
 		}
 	}
@@ -258,20 +273,35 @@ class FPackageStore
 	int32 PackageCount = 0;
 
 public:
-	void Load(const FString& RootDir, FGlobalNameMap& GlobalNameMap)
+	void Load(FIoDispatcher& IoDispatcher, FGlobalNameMap& GlobalNameMap)
 	{
+		FEvent* Event = FPlatformProcess::GetSynchEventFromPool();
+		FIoBuffer IoBuffer;
+
+		IoDispatcher.ReadWithCallback(
+			CreateIoChunkId(0, 0, EIoChunkType::LoaderGlobalMeta),
+			FIoReadOptions(),
+			[Event, &IoBuffer](TIoStatusOr<FIoBuffer> Result)
+			{
+				IoBuffer = Result.ConsumeValueOrDie();
+				Event->Trigger();
+			});
+
+		Event->Wait();
+		
+		FLargeMemoryReader GlobalMetaAr(IoBuffer.Data(), IoBuffer.DataSize());
+		
 		{
 			TRACE_CPUPROFILER_EVENT_SCOPE(LoadPackageStoreToc);
 
-			TUniquePtr<FArchive> StoreTocArchive(IFileManager::Get().CreateFileReader(*(RootDir / TEXT("megafile.ustoretoc"))));
-			check(StoreTocArchive);
+			int32 PackageByteCount = 0;
+			GlobalMetaAr << PackageByteCount;
 
-			const int32 PackageByteCount = StoreTocArchive->TotalSize();
 			PackageCount = PackageByteCount / sizeof(FPackageStoreEntry);
 			StoreEntries = reinterpret_cast<FPackageStoreEntry*>(FMemory::Malloc(PackageByteCount));
 
 			// In-place loading
-			StoreTocArchive->Serialize(StoreEntries, PackageByteCount);
+			GlobalMetaAr.Serialize(StoreEntries, PackageByteCount);
 			PackageNameToGlobalPackageId.Reserve(PackageCount);
 
 			// FName fixup
@@ -297,25 +327,23 @@ public:
 		{
 			TRACE_CPUPROFILER_EVENT_SCOPE(LoadPackageStoreSlimports);
 
-			TUniquePtr<FArchive> SlimportArchive(IFileManager::Get().CreateFileReader(*(RootDir / TEXT("megafile.uslimport"))));
-			check(SlimportArchive);
+			int32 SlimportByteCount = 0;
+			GlobalMetaAr << SlimportByteCount;
 
-			const int32 SlimportByteCount = SlimportArchive->TotalSize();
 			SlimportCount = SlimportByteCount / sizeof(int32);
 			Slimports = reinterpret_cast<int32*>(FMemory::Malloc(SlimportByteCount));
-			SlimportArchive->Serialize(Slimports, SlimportByteCount);
+			GlobalMetaAr.Serialize(Slimports, SlimportByteCount);
 		}
 
 		{
 			TRACE_CPUPROFILER_EVENT_SCOPE(LoadPackageStoreGlimports);
 
-			TUniquePtr<FArchive> ImportArchive(IFileManager::Get().CreateFileReader(*(RootDir / TEXT("megafile.uglimport"))));
-			check(ImportArchive);
+			int32 ImportByteCount = 0;
+			GlobalMetaAr << ImportByteCount; 
 
-			const int32 ImportByteCount = ImportArchive->TotalSize();
 			ImportStore.Count = ImportByteCount / sizeof(FGlobalImport);
 			FGlobalImport* Imports = reinterpret_cast<FGlobalImport*>(FMemory::Malloc(ImportByteCount));
-			ImportArchive->Serialize(Imports, ImportByteCount);
+			GlobalMetaAr.Serialize(Imports, ImportByteCount);
 
 			TRACE_CPUPROFILER_EVENT_SCOPE(LoadPackageStoreGlimportsFixup);
 			{
@@ -338,18 +366,29 @@ public:
 			}
 		}
 
+		// Load initial loading meta data
+		IoDispatcher.ReadWithCallback(
+			CreateIoChunkId(0, 0, EIoChunkType::LoaderInitialLoadMeta),
+			FIoReadOptions(),
+			[Event, &IoBuffer](TIoStatusOr<FIoBuffer> Result)
+			{
+				IoBuffer = Result.ConsumeValueOrDie();
+				Event->Trigger();
+			});
+
+		Event->Wait();
+		FPlatformProcess::ReturnSynchEventToPool(Event);
+
 		{
 			TRACE_CPUPROFILER_EVENT_SCOPE(LoadPackageStoreScriptArcs);
 
-			TUniquePtr<FArchive> ScriptArcsArchive(IFileManager::Get().CreateFileReader(*(RootDir / TEXT("megafile.uscriptarcs"))));
-			check(ScriptArcsArchive);
+			FLargeMemoryReader InitialLoadAr(IoBuffer.Data(), IoBuffer.DataSize());
 
-			const int32 ScriptArcsByteCount = ScriptArcsArchive->TotalSize();
+			const int32 ScriptArcsByteCount = InitialLoadAr.TotalSize();
 			ScriptArcsCount = ScriptArcsByteCount / sizeof(int32);
 			ScriptArcs = reinterpret_cast<int32*>(FMemory::Malloc(ScriptArcsByteCount));
-			ScriptArcsArchive->Serialize(ScriptArcs, ScriptArcsByteCount);
+			InitialLoadAr.Serialize(ScriptArcs, ScriptArcsByteCount);
 		}
-
 	}
 
 	inline FGlobalImportStore& GetImportStore()
@@ -1684,19 +1723,6 @@ struct FAsyncLoadingTickScope2
 
 void FAsyncLoadingThread2Impl::InitializeLoading()
 {
-	{
-		TRACE_CPUPROFILER_EVENT_SCOPE(LoadGlobalNameMap);
-
-		FString GlobalNameMapFilePath = FPaths::ProjectDir() / TEXT("Container.namemap");
-		UE_LOG(LogStreaming, Log, TEXT("Loading global name map '%s'"), *GlobalNameMapFilePath);
-		GlobalNameMap.Load(GlobalNameMapFilePath);
-	}
-
-	{
-		TRACE_CPUPROFILER_EVENT_SCOPE(LoadPackageStore);
-		PackageStore.Load(FPaths::ProjectDir(), GlobalNameMap);
-	}
-
 	{
 		TRACE_CPUPROFILER_EVENT_SCOPE(InitIoDispatcher);
 
@@ -1712,6 +1738,16 @@ void FAsyncLoadingThread2Impl::InitializeLoading()
 #endif
 	}
 
+	{
+		TRACE_CPUPROFILER_EVENT_SCOPE(LoadGlobalNameMap);
+		GlobalNameMap.Load(IoDispatcher);
+	}
+
+	{
+		TRACE_CPUPROFILER_EVENT_SCOPE(LoadPackageStore);
+		PackageStore.Load(IoDispatcher, GlobalNameMap);
+	}
+
 	AsyncThreadReady.Increment();
 }
 
