Temporarily disable the test foruniquness when adding to the BulkData manifest.

OptionalBulkData will break this test, which only matters if the manifest is used. So for the immediate copy up it is safe to disable, then re-enable with the submit that will fix optional bulkdata.

#rnx

diff --git a/Engine/Source/Runtime/CoreUObject/Private/UObject/SavePackage.cpp b/Engine/Source/Runtime/CoreUObject/Private/UObject/SavePackage.cpp
--- a/Engine/Source/Runtime/CoreUObject/Private/UObject/SavePackage.cpp
+++ b/Engine/Source/Runtime/CoreUObject/Private/UObject/SavePackage.cpp
@@ -471,9 +471,11 @@ private:
 
 void FPackageStoreBulkDataManifest::PackageDesc::AddData(uint64 InChunkId, uint64 InOffset, uint64 InSize)
 {
+#if 0 // DIsable until optional data is working again
 	// The ChunkId is supposed to be unique
 	auto func = [=](const BulkDataDesc& Entry) { return Entry.ChunkId == InChunkId; };
 	check(Data.FindByPredicate(func) == nullptr);
+#endif 
 
 	BulkDataDesc& Entry = Data.Emplace_GetRef();
 
