GC will now run clustering code if clusters are enabled AND there's actual clusters allocated

diff --git a/Engine/Source/Runtime/CoreUObject/Private/UObject/GarbageCollection.cpp b/Engine/Source/Runtime/CoreUObject/Private/UObject/GarbageCollection.cpp
--- a/Engine/Source/Runtime/CoreUObject/Private/UObject/GarbageCollection.cpp
+++ b/Engine/Source/Runtime/CoreUObject/Private/UObject/GarbageCollection.cpp
@@ -1907,12 +1907,14 @@ void CollectGarbageInternal(EObjectFlags KeepFlags, bool bPerformFullPurge)
 		// or detailed per class gc stats are enabled (not thread safe)
 		// Temporarily forcing single-threaded GC in the editor until Modify() can be safely removed from HandleObjectReference.
 		const bool bForceSingleThreadedGC = ShouldForceSingleThreadedGC();
+		// Run with GC clustering code enabled only if clustering is enabled and there's actual allocated clusters
+		const bool bWithClusters = !!GCreateGCClusters && GUObjectClusters.GetNumAllocatedClusters();
 
 		// Perform reachability analysis.
 		{
 			const double StartTime = FPlatformTime::Seconds();
 			FRealtimeGC TagUsedRealtimeGC;
-			TagUsedRealtimeGC.PerformReachabilityAnalysis(KeepFlags, bForceSingleThreadedGC, !!GCreateGCClusters);
+			TagUsedRealtimeGC.PerformReachabilityAnalysis(KeepFlags, bForceSingleThreadedGC, bWithClusters);
 			UE_LOG(LogGarbage, Log, TEXT("%f ms for GC"), (FPlatformTime::Seconds() - StartTime) * 1000);
 		}
 
