Removed redundant UPackage::AddReferencedObjects function

diff --git a/Engine/Source/Runtime/CoreUObject/Private/UObject/Package.cpp b/Engine/Source/Runtime/CoreUObject/Private/UObject/Package.cpp
--- a/Engine/Source/Runtime/CoreUObject/Private/UObject/Package.cpp
+++ b/Engine/Source/Runtime/CoreUObject/Private/UObject/Package.cpp
@@ -111,19 +111,6 @@ void UPackage::Serialize( FArchive& Ar )
 	}
 }
 
-void UPackage::AddReferencedObjects(UObject* InThis, FReferenceCollector& Collector)
-{	
-	UPackage* This = CastChecked<UPackage>(InThis);
-#if WITH_EDITOR
-	if( GIsEditor )
-	{
-		// Required by the unified GC when running in the editor
-		Collector.AddReferencedObject(This->MetaData, This);
-	}
-#endif
-	Super::AddReferencedObjects(This, Collector);
-}
-
 #if WITH_EDITORONLY_DATA
 bool UPackage::IsOwned() const
 {
@@ -284,14 +271,12 @@ void UPackage::SetLoadedByEditorPropertiesOnly(bool bIsEditorOnly, bool bRecursi
 #if WITH_EDITORONLY_DATA
 IMPLEMENT_CORE_INTRINSIC_CLASS(UPackage, UObject,
 	{
-		Class->ClassAddReferencedObjects = &UPackage::AddReferencedObjects;
 		Class->EmitObjectReference(STRUCT_OFFSET(UPackage, MetaData), TEXT("MetaData"));
 	}
 );
 #else
 IMPLEMENT_CORE_INTRINSIC_CLASS(UPackage, UObject,
 	{
-		Class->ClassAddReferencedObjects = &UPackage::AddReferencedObjects;
 	}
 );
 #endif
diff --git a/Engine/Source/Runtime/CoreUObject/Public/UObject/Package.h b/Engine/Source/Runtime/CoreUObject/Public/UObject/Package.h
--- a/Engine/Source/Runtime/CoreUObject/Public/UObject/Package.h
+++ b/Engine/Source/Runtime/CoreUObject/Public/UObject/Package.h
@@ -234,8 +234,6 @@ public:
 	/** Serializer */
 	virtual void Serialize( FArchive& Ar ) override;
 
-	static void AddReferencedObjects(UObject* InThis, FReferenceCollector& Collector);
-
 	/** Packages are never assets */
 	virtual bool IsAsset() const override { return false; }
 
