GC optimization: Class reference will now be completely ignored by GC for instances of native classes and additionally the Outer reference will not be processed for UPackages

diff --git a/Engine/Source/Runtime/CoreUObject/Private/UObject/GarbageCollection.cpp b/Engine/Source/Runtime/CoreUObject/Private/UObject/GarbageCollection.cpp
--- a/Engine/Source/Runtime/CoreUObject/Private/UObject/GarbageCollection.cpp
+++ b/Engine/Source/Runtime/CoreUObject/Private/UObject/GarbageCollection.cpp
@@ -1032,7 +1032,14 @@ class FRealtimeGC : public FGarbageCollectionTracer
 	void PerformReachabilityAnalysisOnObjectsInternal(FGCArrayStruct* ArrayStruct)
 	{
 		FGCReferenceProcessor<bParallel, bWithClusters> ReferenceProcessor;
-		TFastReferenceCollector<bParallel, FGCReferenceProcessor<bParallel, bWithClusters>, FGCCollector<bParallel, bWithClusters>, FGCArrayPool> ReferenceCollector(ReferenceProcessor, FGCArrayPool::Get());
+		// NOTE: we want to run with automatic token stream generation off as it should be already generated at this point,
+		// BUT we want to be ignoring Noop tokens as they're only pointing either at null references or at objects that never get GC'd (native classes)
+		TFastReferenceCollector<bParallel, 
+			FGCReferenceProcessor<bParallel, bWithClusters>, 
+			FGCCollector<bParallel, bWithClusters>, 
+			FGCArrayPool, 
+			/* bAutoGenerateTokenStream = */ false, 
+			/* bIgnoreNoopTokens = */ true> ReferenceCollector(ReferenceProcessor, FGCArrayPool::Get());
 		ReferenceCollector.CollectReferences(*ArrayStruct);
 	}
 
@@ -2554,15 +2561,25 @@ void UClass::AssembleReferenceTokenStream(bool bForce)
 			UObjectBase::EmitBaseReferences(this);
 		}
 
-#if !WITH_EDITOR
-		// In no-editor builds UObject::ARO is empty, thus only classes
-		// which implement their own ARO function need to have the ARO token generated.
-		if (ClassAddReferencedObjects != &UObject::AddReferencedObjects)
-#endif
 		{
 			check(ClassAddReferencedObjects != NULL);
-			ReferenceTokenStream.ReplaceOrAddAddReferencedObjectsCall(ClassAddReferencedObjects);
+			const bool bKeepOuter = GetFName() != NAME_Package;
+			const bool bKeepClass = !(GetClassFlags() & CLASS_Native);
+
+			ClassAddReferencedObjectsType AddReferencedObjectsFn = nullptr;
+#if !WITH_EDITOR
+			// In no-editor builds UObject::ARO is empty, thus only classes
+			// which implement their own ARO function need to have the ARO token generated.
+			if (ClassAddReferencedObjects != &UObject::AddReferencedObjects)
+			{
+				AddReferencedObjectsFn = ClassAddReferencedObjects;
+			}
+#else
+			AddReferencedObjectsFn = ClassAddReferencedObjects;
+#endif
+			ReferenceTokenStream.Fixup(AddReferencedObjectsFn, bKeepOuter, bKeepClass);
 		}
+
 		if (ReferenceTokenStream.IsEmpty())
 		{
 			return;
@@ -2600,15 +2617,17 @@ void FGCReferenceTokenStream::PrependStream( const FGCReferenceTokenStream& Othe
 	Tokens = MoveTemp(TempTokens);
 }
 
-void FGCReferenceTokenStream::ReplaceOrAddAddReferencedObjectsCall(void (*AddReferencedObjectsPtr)(UObject*, class FReferenceCollector&))
+void FGCReferenceTokenStream::Fixup(void (*AddReferencedObjectsPtr)(UObject*, class FReferenceCollector&), bool bKeepOuterToken, bool bKeepClassToken)
 {
+	bool bReplacedARO = false;
+
 	// Try to find exiting ARO pointer and replace it (to avoid removing and readding tokens).
 	for (int32 TokenStreamIndex = 0; TokenStreamIndex < Tokens.Num(); ++TokenStreamIndex)
 	{
 		uint32 TokenIndex = (uint32)TokenStreamIndex;
-		const EGCReferenceType TokenType = (EGCReferenceType)AccessReferenceInfo(TokenIndex).Type;
+		FGCReferenceInfo Token = Tokens[TokenIndex];
 		// Read token type and skip additional data if present.
-		switch (TokenType)
+		switch (Token.Type)
 		{
 		case GCRT_ArrayStruct:
 			{
@@ -2637,9 +2656,14 @@ void FGCReferenceTokenStream::ReplaceOrAddAddReferencedObjectsCall(void (*AddRef
 		case GCRT_AddReferencedObjects:
 			{
 				// Store the pointer after the ARO token.
-				StorePointer(&Tokens[++TokenIndex], (const void*)AddReferencedObjectsPtr);
-				return;
+				if (AddReferencedObjectsPtr)
+				{
+					StorePointer(&Tokens[TokenIndex + 1], (const void*)AddReferencedObjectsPtr);					
+				}
+				bReplacedARO = true;
+				TokenIndex += GNumTokensPerPointer;
 			}
+			break;
 		case GCRT_AddTMapReferencedObjects:
 		case GCRT_AddTSetReferencedObjects:
 			{
@@ -2647,22 +2671,52 @@ void FGCReferenceTokenStream::ReplaceOrAddAddReferencedObjectsCall(void (*AddRef
 				TokenIndex += GNumTokensPerPointer;
 			}
 			break;
-		case GCRT_None:
-		case GCRT_Object:
+		case GCRT_Class:
+		case GCRT_NoopClass:
+			{
+				if (bKeepClassToken)
+				{
+					Token.Type = GCRT_Class;
+				}
+				else
+				{
+					Token.Type = GCRT_NoopClass;
+				}
+				Tokens[TokenIndex] = Token;
+			}
+			break;
 		case GCRT_PersistentObject:
+		case GCRT_NoopPersistentObject:
+			{
+				if (bKeepOuterToken)
+				{
+					Token.Type = GCRT_PersistentObject;
+				}
+				else
+				{
+					Token.Type = GCRT_NoopPersistentObject;
+				}
+				Tokens[TokenIndex] = Token;
+			}
+			break;
+		case GCRT_None:
+		case GCRT_Object:		
 		case GCRT_ArrayObject:
 		case GCRT_EndOfPointer:
-		case GCRT_EndOfStream:		
+		case GCRT_EndOfStream:			
 			break;
 		default:
-			UE_LOG(LogGarbage, Fatal, TEXT("Unknown token type (%u) when trying to add ARO token."), (uint32)TokenType);
+			UE_LOG(LogGarbage, Fatal, TEXT("Unknown token type (%u) when trying to add ARO token."), (uint32)Token.Type);
 			break;
 		};
 		TokenStreamIndex = (int32)TokenIndex;
 	}
 	// ARO is not in the token stream yet.
-	EmitReferenceInfo(FGCReferenceInfo(GCRT_AddReferencedObjects, 0));
-	EmitPointer((const void*)AddReferencedObjectsPtr);
+	if (!bReplacedARO && AddReferencedObjectsPtr)
+	{
+		EmitReferenceInfo(FGCReferenceInfo(GCRT_AddReferencedObjects, 0));
+		EmitPointer((const void*)AddReferencedObjectsPtr);
+	}
 }
 
 int32 FGCReferenceTokenStream::EmitReferenceInfo(FGCReferenceInfo ReferenceInfo)
diff --git a/Engine/Source/Runtime/CoreUObject/Private/UObject/UObjectBase.cpp b/Engine/Source/Runtime/CoreUObject/Private/UObject/UObjectBase.cpp
--- a/Engine/Source/Runtime/CoreUObject/Private/UObject/UObjectBase.cpp
+++ b/Engine/Source/Runtime/CoreUObject/Private/UObject/UObjectBase.cpp
@@ -378,7 +378,8 @@ void UObjectBase::EmitBaseReferences(UClass *RootClass)
 {
 	static const FName ClassPropertyName(TEXT("Class"));
 	static const FName OuterPropertyName(TEXT("Outer"));
-	RootClass->EmitObjectReference(STRUCT_OFFSET(UObjectBase, ClassPrivate), ClassPropertyName);
+
+	RootClass->EmitObjectReference(STRUCT_OFFSET(UObjectBase, ClassPrivate), ClassPropertyName, GCRT_Class);
 	RootClass->EmitObjectReference(STRUCT_OFFSET(UObjectBase, OuterPrivate), OuterPropertyName, GCRT_PersistentObject);
 }
 
diff --git a/Engine/Source/Runtime/CoreUObject/Public/UObject/FastReferenceCollector.h b/Engine/Source/Runtime/CoreUObject/Public/UObject/FastReferenceCollector.h
--- a/Engine/Source/Runtime/CoreUObject/Public/UObject/FastReferenceCollector.h
+++ b/Engine/Source/Runtime/CoreUObject/Public/UObject/FastReferenceCollector.h
@@ -262,7 +262,7 @@ private:
 		 void ReturnToPool(FGCArrayStruct* ArrayStruct);
 	 };
  */
-template <bool bParallel, typename ReferenceProcessorType, typename CollectorType, typename ArrayPoolType, bool bAutoGenerateTokenStream = false>
+template <bool bParallel, typename ReferenceProcessorType, typename CollectorType, typename ArrayPoolType, bool bAutoGenerateTokenStream = false, bool bIgnoreNoopTokens = false>
 class TFastReferenceCollector
 {
 private:
@@ -699,8 +699,9 @@ private:
 					switch(ReferenceInfo.Type)
 					{
 					case GCRT_Object:
+					case GCRT_Class:
 					{
-						// We're dealing with an object reference.
+						// We're dealing with an object reference (this code should be identical to GCRT_NoopClass if !bIgnoreNoopTokens)
 						UObject**	ObjectPtr = (UObject**)(StackEntryData + ReferenceInfo.Offset);
 						UObject*&	Object = *ObjectPtr;
 						TokenReturnCount = ReferenceInfo.ReturnCount;
@@ -747,7 +748,7 @@ private:
 					break;
 					case GCRT_PersistentObject:
 					{
-						// We're dealing with an object reference.
+						// We're dealing with an object reference (this code should be identical to GCRT_NoopPersistentObject if !bIgnoreNoopTokens)
 						UObject**	ObjectPtr = (UObject**)(StackEntryData + ReferenceInfo.Offset);
 						UObject*&	Object = *ObjectPtr;
 						TokenReturnCount = ReferenceInfo.ReturnCount;
@@ -805,6 +806,38 @@ private:
 						TokenReturnCount = ReferenceInfo.ReturnCount;
 					}
 					break;
+					case GCRT_NoopPersistentObject:
+					{
+						if (!bIgnoreNoopTokens)
+						{
+							// We're dealing with an object reference (this code should be identical to GCRT_PersistentObject)
+							UObject**	ObjectPtr = (UObject**)(StackEntryData + ReferenceInfo.Offset);
+							UObject*&	Object = *ObjectPtr;
+							TokenReturnCount = ReferenceInfo.ReturnCount;
+							ReferenceProcessor.HandleTokenStreamObjectReference(NewObjectsToSerialize, CurrentObject, Object, ReferenceTokenStreamIndex, false);
+						}
+						else
+						{
+							TokenReturnCount = ReferenceInfo.ReturnCount;
+						}
+					}
+					break;
+					case GCRT_NoopClass:
+					{
+						if (!bIgnoreNoopTokens)
+						{
+							// We're dealing with an object reference (this code should be identical to GCRT_Object and GCRT_Class)
+							UObject**	ObjectPtr = (UObject**)(StackEntryData + ReferenceInfo.Offset);
+							UObject*&	Object = *ObjectPtr;
+							TokenReturnCount = ReferenceInfo.ReturnCount;
+							ReferenceProcessor.HandleTokenStreamObjectReference(NewObjectsToSerialize, CurrentObject, Object, ReferenceTokenStreamIndex, true);
+						}
+						else
+						{
+							TokenReturnCount = ReferenceInfo.ReturnCount;
+						}
+					}
+					break;
 					case GCRT_EndOfStream:
 					{
 						// Break out of loop.
diff --git a/Engine/Source/Runtime/CoreUObject/Public/UObject/GarbageCollection.h b/Engine/Source/Runtime/CoreUObject/Public/UObject/GarbageCollection.h
--- a/Engine/Source/Runtime/CoreUObject/Public/UObject/GarbageCollection.h
+++ b/Engine/Source/Runtime/CoreUObject/Public/UObject/GarbageCollection.h
@@ -40,6 +40,7 @@ enum EGCReferenceType
 {
 	GCRT_None			= 0,
 	GCRT_Object,
+	GCRT_Class,
 	GCRT_PersistentObject,
 	GCRT_ArrayObject,
 	GCRT_ArrayStruct,
@@ -50,6 +51,8 @@ enum EGCReferenceType
 	GCRT_AddTSetReferencedObjects,
 	GCRT_EndOfPointer,
 	GCRT_EndOfStream,
+	GCRT_NoopPersistentObject,
+	GCRT_NoopClass,
 };
 
 /** 
@@ -322,9 +325,9 @@ struct FGCReferenceTokenStream
 	uint32 EmitReturn();
 
 	/**
-	 * Helper function to quickly replace or add ARO call.
+	 * Helper function to perform post parent token stream prepend fixup
 	 */
-	void ReplaceOrAddAddReferencedObjectsCall(void (*AddReferencedObjectsPtr)(UObject*, class FReferenceCollector&));
+	void Fixup(void (*AddReferencedObjectsPtr)(UObject*, class FReferenceCollector&), bool bKeepOuterToken, bool bKeepClassToken);
 
 	/**
 	 * Reads count and advances stream.
