Setup TLS allocation caches for AsyncLoading2

diff --git a/Engine/Source/Runtime/CoreUObject/Private/Serialization/AsyncLoading2.cpp b/Engine/Source/Runtime/CoreUObject/Private/Serialization/AsyncLoading2.cpp
--- a/Engine/Source/Runtime/CoreUObject/Private/Serialization/AsyncLoading2.cpp
+++ b/Engine/Source/Runtime/CoreUObject/Private/Serialization/AsyncLoading2.cpp
@@ -2089,10 +2089,8 @@ uint32 FAsyncLoadingThreadWorker::Run()
 {
 	LLM_SCOPE(ELLMTag::AsyncLoading);
 
-	if (!IsInGameThread())
-	{
-		FPlatformProcess::SetThreadAffinityMask(FPlatformAffinity::GetAsyncLoadingThreadMask());
-	}
+	FPlatformProcess::SetThreadAffinityMask(FPlatformAffinity::GetAsyncLoadingThreadMask());
+	FMemory::SetupTLSCachesOnCurrentThread();
 
 	FAsyncLoadingThreadState2::Create(GraphAllocator, IoDispatcher);
 
@@ -3426,10 +3424,8 @@ uint32 FAsyncLoadingThread2Impl::Run()
 
 	TRACE_LOADTIME_START_ASYNC_LOADING();
 
-	if (!IsInGameThread())
-	{
-		FPlatformProcess::SetThreadAffinityMask(FPlatformAffinity::GetAsyncLoadingThreadMask());
-	}
+	FPlatformProcess::SetThreadAffinityMask(FPlatformAffinity::GetAsyncLoadingThreadMask());
+	FMemory::SetupTLSCachesOnCurrentThread();
 
 	FAsyncLoadingThreadState2& ThreadState = *FAsyncLoadingThreadState2::Get();
 
