Move MD5 computation on writer thread during package save

diff --git a/Engine/Source/Editor/UnrealEd/Private/CookOnTheFlyServer.cpp b/Engine/Source/Editor/UnrealEd/Private/CookOnTheFlyServer.cpp
--- a/Engine/Source/Editor/UnrealEd/Private/CookOnTheFlyServer.cpp
+++ b/Engine/Source/Editor/UnrealEd/Private/CookOnTheFlyServer.cpp
@@ -2967,7 +2967,14 @@ void UCookOnTheFlyServer::SaveCookedPackages(
 							{
 								FAssetPackageData* PackageData = Generator->GetAssetPackageData(Package->GetFName());
 								PackageData->DiskSize = SavePackageResult.TotalFileSize;
-								PackageData->CookedHash = SavePackageResult.CookedHash;
+								SavePackageResult.CookedHash.Next([PackageData](const FMD5Hash& CookedHash)
+								{
+									// Store the cooked hash in the Asset Registry when it is done computing in another thread.
+									// NOTE: For this to work, we rely on:
+									// 1) UPackage::WaitForAsyncFileWrites to have been called before any use of the CookedHash - it is called in CookByTheBookFinished before the registry does any work with the registries
+									// 2) PackageData must continue to be a valid pointer - the asset registry allocates the FAssetPackageData individually and doesn't relocate or delete them until pruning, which happens after WaitForAsyncFileWrites
+									PackageData->CookedHash = CookedHash;
+								});
 							}
 						}
 
@@ -7995,7 +8002,14 @@ uint32 UCookOnTheFlyServer::FullLoadAndSave(uint32& CookedPackageCount)
 							{
 								FAssetPackageData* PackageData = Generator->GetAssetPackageData(Package->GetFName());
 								PackageData->DiskSize = SaveResult.TotalFileSize;
-								PackageData->CookedHash = SaveResult.CookedHash;
+								SaveResult.CookedHash.Next([PackageData](const FMD5Hash& CookedHash)
+								{
+									// Store the cooked hash in the Asset Registry when it is done computing in another thread.
+									// NOTE: For this to work, we rely on:
+									// 1) UPackage::WaitForAsyncFileWrites to have been called before any use of the CookedHash - it is called in CookByTheBookFinished before the registry does any work with the registries
+									// 2) PackageData must continue to be a valid pointer - the asset registry allocates the FAssetPackageData individually and doesn't relocate or delete them until pruning, which happens after WaitForAsyncFileWrites
+									PackageData->CookedHash = CookedHash;
+								});
 							}
 
 							FPlatformAtomics::InterlockedIncrement(&ParallelSavedPackages);
diff --git a/Engine/Source/Editor/UnrealEd/Private/EditorEngine.cpp b/Engine/Source/Editor/UnrealEd/Private/EditorEngine.cpp
--- a/Engine/Source/Editor/UnrealEd/Private/EditorEngine.cpp
+++ b/Engine/Source/Editor/UnrealEd/Private/EditorEngine.cpp
@@ -4440,7 +4440,7 @@ FSavePackageResultStruct UEditorEngine::Save( UPackage* InOuter, UObject* InBase
 	SlowTask.EnterProgressFrame(70);
 
 	UPackage::PreSavePackageEvent.Broadcast(InOuter);
-	const FSavePackageResultStruct Result = UPackage::Save(InOuter, Base, TopLevelFlags, Filename, Error, Conform, bForceByteSwapping, bWarnOfLongFilename, SaveFlags, TargetPlatform, FinalTimeStamp, bSlowTask, InOutDiffMap, SavePackageContext);
+	FSavePackageResultStruct Result = UPackage::Save(InOuter, Base, TopLevelFlags, Filename, Error, Conform, bForceByteSwapping, bWarnOfLongFilename, SaveFlags, TargetPlatform, FinalTimeStamp, bSlowTask, InOutDiffMap, SavePackageContext);
 
 	SlowTask.EnterProgressFrame(10);
 
diff --git a/Engine/Source/Runtime/CoreUObject/Private/UObject/AsyncWorkSequence.h b/Engine/Source/Runtime/CoreUObject/Private/UObject/AsyncWorkSequence.h
new file mode 100644
--- /dev/null
+++ b/Engine/Source/Runtime/CoreUObject/Private/UObject/AsyncWorkSequence.h
@@ -0,0 +1,55 @@
+// Copyright 1998-2019 Epic Games, Inc. All Rights Reserved.
+
+#pragma once
+#include "Async/Async.h"
+#include "Containers/Array.h"
+#include "Misc/AssertionMacros.h"
+#include "Templates/Function.h"
+
+template <typename StateType>
+class TAsyncWorkSequence
+{
+public:
+	typedef TUniqueFunction<void(StateType& state)> WorkFunctionType;
+
+	TAsyncWorkSequence()
+	{
+	}
+
+	void AddWork(WorkFunctionType&& Work)
+	{
+		check(Work);
+		check(!bFinalized);
+
+		WorkItems.Emplace(MoveTemp(Work));
+	}
+
+	template<typename CallableType>
+	auto Finalize(EAsyncExecution Execution, CallableType&& Callable)
+	{
+		check(!bFinalized);
+		bFinalized = true;
+
+		return Async(Execution, [WorkItems = MoveTemp(WorkItems), Callable = MoveTemp(Callable)] () mutable
+		{
+			StateType State;
+
+			for (WorkFunctionType& Work : WorkItems)
+			{
+				// Moving the Work capture to a local variable so that after the iteration, any
+				// heavyweight data in the capture storage is cleaned up instead of waiting
+				// to clean them all up at the end.
+				const WorkFunctionType WorkLocal = MoveTemp(Work);
+				WorkLocal(State);
+			}
+			WorkItems.Empty();
+
+			return Callable(State);
+		});
+	}
+
+private:
+	TArray<WorkFunctionType> WorkItems;
+	bool bFinalized = false;
+
+};
diff --git a/Engine/Source/Runtime/CoreUObject/Private/UObject/SavePackage.cpp b/Engine/Source/Runtime/CoreUObject/Private/UObject/SavePackage.cpp
--- a/Engine/Source/Runtime/CoreUObject/Private/UObject/SavePackage.cpp
+++ b/Engine/Source/Runtime/CoreUObject/Private/UObject/SavePackage.cpp
@@ -58,6 +58,7 @@
 #include "Serialization/Formatters/BinaryArchiveFormatter.h"
 #include "Serialization/Formatters/JsonArchiveOutputFormatter.h"
 #include "Serialization/ArchiveUObjectFromStructuredArchive.h"
+#include "UObject/AsyncWorkSequence.h"
 
 DEFINE_LOG_CATEGORY_STATIC(LogSavePackage, Log, All);
 
@@ -71,7 +72,7 @@ static FCriticalSection InitializeCoreClassesCritSec;
 static void SaveThumbnails(UPackage* InOuter, FLinkerSave* Linker, FStructuredArchive::FSlot Slot);
 static void SaveAssetRegistryData(UPackage* InOuter, FLinkerSave* Linker, FStructuredArchive::FSlot Slot);
 static void SaveBulkData(FLinkerSave* Linker, const UPackage* InOuter, const TCHAR* Filename, const ITargetPlatform* TargetPlatform,
-						 FSavePackageContext* SavePackageContext, const bool bTextFormat, const bool bDiffing, FMD5* CookedPackageHash, int64 TotalPackageSizeUncompressed);
+						 FSavePackageContext* SavePackageContext, const bool bTextFormat, const bool bDiffing, const bool bComputeHash, TAsyncWorkSequence<FMD5>& AsyncWriteAndHashSequence, int64 TotalPackageSizeUncompressed);
 static void SaveWorldLevelInfo(UPackage* InOuter, FLinkerSave* Linker, FStructuredArchive::FRecord Record);
 static EObjectMark GetExcludedObjectMarksForTargetPlatform(const class ITargetPlatform* TargetPlatform, const bool bIsCooking);
 
@@ -361,80 +362,57 @@ struct FLargeMemoryDelete
 
 typedef TUniquePtr<uint8, FLargeMemoryDelete> FLargeMemoryPtr;
 
-static void AsyncWriteFile(FLargeMemoryPtr Data, const int64 DataSize, const TCHAR* Filename)
+enum class EAsyncWriteOptions
 {
-	class FAsyncWriteWorker : public FNonAbandonableTask
-	{
-	public:
-		FString			Filename;
-		FLargeMemoryPtr Data;
-		const int64		DataSize;
+	None = 0,
+	WriteFileToDisk = 0x01,
+	ComputeHash = 0x02
+};
+ENUM_CLASS_FLAGS(EAsyncWriteOptions)
 
-		FAsyncWriteWorker(const TCHAR* InFilename, FLargeMemoryPtr InData, const int64 InDataSize)
-			: Filename(InFilename)
-			, Data(MoveTemp(InData))
-			, DataSize(InDataSize)
-		{
-			check(InDataSize);
-		}
-		
-		/** Write the file  */
-		void DoWork()
+static void AsyncWriteFile(TAsyncWorkSequence<FMD5>& AsyncWriteAndHashSequence, FLargeMemoryPtr Data, const int64 DataSize, const TCHAR* Filename, EAsyncWriteOptions Options)
+{
+	OutstandingAsyncWrites.Increment();
+	FString OutputFilename(Filename);
+	AsyncWriteAndHashSequence.AddWork([Data = MoveTemp(Data), DataSize, OutputFilename = MoveTemp(OutputFilename), Options](FMD5& State) mutable
+	{
+		if (EnumHasAnyFlags(Options, EAsyncWriteOptions::ComputeHash))
 		{
-			WriteToFile(Filename, Data.Get(), DataSize);
-			
-			OutstandingAsyncWrites.Decrement();
+			State.Update(Data.Get(), DataSize);
 		}
 
-		FORCEINLINE TStatId GetStatId() const
+		if (EnumHasAnyFlags(Options, EAsyncWriteOptions::WriteFileToDisk))
 		{
-			RETURN_QUICK_DECLARE_CYCLE_STAT(FAsyncWriteWorker, STATGROUP_ThreadPoolAsyncTasks);
+			WriteToFile(OutputFilename, Data.Get(), DataSize);
 		}
-	};
 
-	OutstandingAsyncWrites.Increment();
-	(new FAutoDeleteAsyncTask<FAsyncWriteWorker>(Filename, MoveTemp(Data), DataSize))->StartBackgroundTask();
+		OutstandingAsyncWrites.Decrement();
+	});
 }
 
-static void AsyncWriteFileWithSplitExports(FLargeMemoryPtr Data, const int64 DataSize, const int64 HeaderSize, const TCHAR* Filename)
+static void AsyncWriteFileWithSplitExports(TAsyncWorkSequence<FMD5>& AsyncWriteAndHashSequence, FLargeMemoryPtr Data, const int64 DataSize, const int64 HeaderSize, const TCHAR* Filename, EAsyncWriteOptions Options)
 {
-	class FAsyncWriteWorkerWithSplitExports : public FNonAbandonableTask
+	OutstandingAsyncWrites.Increment();
+	FString OutputFilename(Filename);
+	AsyncWriteAndHashSequence.AddWork([Data = MoveTemp(Data), DataSize, HeaderSize, OutputFilename = MoveTemp(OutputFilename), Options](FMD5& State) mutable
 	{
-	public:
-		FString			Filename;
-		FLargeMemoryPtr Data;
-		const int64		DataSize;
-		const int64		HeaderSize;
-
-		FAsyncWriteWorkerWithSplitExports(const TCHAR* InFilename, FLargeMemoryPtr InData, const int64 InDataSize, const int64 InHeaderSize)
-			: Filename(InFilename)
-			, Data(MoveTemp(InData))
-			, DataSize(InDataSize)
-			, HeaderSize(InHeaderSize)
+		if (EnumHasAnyFlags(Options, EAsyncWriteOptions::ComputeHash))
 		{
-			check(InDataSize);
+			State.Update(Data.Get(), DataSize);
 		}
 
-		void DoWork()
+		if (EnumHasAnyFlags(Options, EAsyncWriteOptions::WriteFileToDisk))
 		{
 			// Write .uasset file
-			WriteToFile(Filename, Data.Get(), HeaderSize);
+			WriteToFile(OutputFilename, Data.Get(), HeaderSize);
 
 			// Write .uexp file
-			const FString FilenameExports = FPaths::ChangeExtension(Filename, TEXT(".uexp"));
+			const FString FilenameExports = FPaths::ChangeExtension(OutputFilename, TEXT(".uexp"));
 			WriteToFile(FilenameExports, Data.Get() + HeaderSize, DataSize - HeaderSize);
-
-			OutstandingAsyncWrites.Decrement();
-		}
-
-		FORCEINLINE TStatId GetStatId() const
-		{
-			RETURN_QUICK_DECLARE_CYCLE_STAT(FAsyncWriteWorkerWithSplitExports, STATGROUP_ThreadPoolAsyncTasks);
 		}
-	};
 
-	OutstandingAsyncWrites.Increment();
-	(new FAutoDeleteAsyncTask<FAsyncWriteWorkerWithSplitExports>(Filename, MoveTemp(Data), DataSize, HeaderSize))->StartBackgroundTask();
+		OutstandingAsyncWrites.Decrement();
+	});
 }
 
 class FPackageNameMapSaver
@@ -3370,7 +3348,7 @@ void VerifyEDLCookInfo()
 	FEDLCookChecker::Verify();
 }
 
-void AddFileToHash(FString const &Filename, FMD5 &Hash)
+void AddFileToHash(FString const& Filename, FMD5& Hash)
 {
 	TArray<uint8> LocalScratch;
 	LocalScratch.SetNumUninitialized(1024 * 64);
@@ -3524,7 +3502,9 @@ FSavePackageResultStruct UPackage::Save(UPackage* InOuter, UObject* Base, EObjec
 
 		uint32 Time = 0; CLOCK_CYCLES(Time);
 		int64 TotalPackageSizeUncompressed = 0;
-		FMD5 CookedPackageHash;
+
+		TFuture<FMD5Hash> PackageMD5Destination;
+		TAsyncWorkSequence<FMD5> AsyncWriteAndHashSequence;
 
 		// Make sure package is fully loaded before saving. 
 		if (!Base && !InOuter->IsFullyLoaded())
@@ -5431,7 +5411,7 @@ FSavePackageResultStruct UPackage::Save(UPackage* InOuter, UObject* Base, EObjec
 				SlowTask.EnterProgressFrame(1, NSLOCTEXT("Core", "SerializingBulkData", "Serializing bulk data"));
 
 				SaveBulkData(Linker.Get(), InOuter, Filename, TargetPlatform, SavePackageContext, bTextFormat, bDiffing,
-							 bComputeHash ? &CookedPackageHash : nullptr, TotalPackageSizeUncompressed );
+							 bComputeHash, AsyncWriteAndHashSequence, TotalPackageSizeUncompressed );
 
 #if WITH_EDITOR
 				if (bIsCooking && AdditionalFilesFromExports.Num() > 0)
@@ -5442,16 +5422,20 @@ FSavePackageResultStruct UPackage::Save(UPackage* InOuter, UObject* Base, EObjec
 						const int64 Size = Writer.TotalSize();
 						TotalPackageSizeUncompressed += Size;
 
-						if (bComputeHash)
-						{
-							CookedPackageHash.Update(Writer.GetData(), Size);
-						}
-
-						if (bWriteFileToDisk)
+						if (bComputeHash || bWriteFileToDisk)
 						{
 							FLargeMemoryPtr DataPtr(Writer.ReleaseOwnership());
 
-							AsyncWriteFile(MoveTemp(DataPtr), Size, *Writer.GetArchiveName());
+							EAsyncWriteOptions WriteOptions(EAsyncWriteOptions::None);
+							if (bComputeHash)
+							{
+								WriteOptions |= EAsyncWriteOptions::ComputeHash;
+							}
+							if (bWriteFileToDisk)
+							{
+								WriteOptions |= EAsyncWriteOptions::WriteFileToDisk;
+							}
+							AsyncWriteFile(AsyncWriteAndHashSequence, MoveTemp(DataPtr), Size, *Writer.GetArchiveName(), WriteOptions);
 						}
 					}
 					AdditionalFilesFromExports.Empty();
@@ -5710,17 +5694,21 @@ FSavePackageResultStruct UPackage::Save(UPackage* InOuter, UObject* Base, EObjec
 
 							TotalPackageSizeUncompressed += DataSize;
 
-							if (bComputeHash)
-							{
-								CookedPackageHash.Update(Writer->GetData(), DataSize);
-							}
-
 							if (IsEventDrivenLoaderEnabledInCookedBuilds() && Linker->IsCooking())
 							{
 								if (SavePackageContext != nullptr && SavePackageContext->PackageStoreWriter != nullptr)
 								{
 									FIoBuffer IoBuffer(FIoBuffer::AssumeOwnership, Writer->ReleaseOwnership(), DataSize);
 
+									if (bComputeHash)
+									{
+										FIoBuffer InnerBuffer(IoBuffer.Data(), IoBuffer.DataSize(), IoBuffer);
+										AsyncWriteAndHashSequence.AddWork([InnerBuffer = MoveTemp(InnerBuffer)](FMD5& State)
+										{
+											State.Update(InnerBuffer.Data(), InnerBuffer.DataSize());
+										});
+									}
+
 									const int32 HeaderSize = Linker->Summary.TotalHeaderSize;
 
 									FPackageStoreWriter::HeaderInfo HeaderInfo;
@@ -5747,12 +5735,22 @@ FSavePackageResultStruct UPackage::Save(UPackage* InOuter, UObject* Base, EObjec
 								}
 								else
 								{
-									AsyncWriteFileWithSplitExports(FLargeMemoryPtr(Writer->ReleaseOwnership()), DataSize, Linker->Summary.TotalHeaderSize, *NewPathToSave);
+									EAsyncWriteOptions WriteOptions(EAsyncWriteOptions::WriteFileToDisk);
+									if (bComputeHash)
+									{
+										WriteOptions |= EAsyncWriteOptions::ComputeHash;
+									}
+									AsyncWriteFileWithSplitExports(AsyncWriteAndHashSequence, FLargeMemoryPtr(Writer->ReleaseOwnership()), DataSize, Linker->Summary.TotalHeaderSize, *NewPathToSave, WriteOptions);
 								}
 							}
 							else
 							{
-								AsyncWriteFile(FLargeMemoryPtr(Writer->ReleaseOwnership()), DataSize, *NewPathToSave);
+								EAsyncWriteOptions WriteOptions(EAsyncWriteOptions::WriteFileToDisk);
+								if (bComputeHash)
+								{
+									WriteOptions |= EAsyncWriteOptions::ComputeHash;
+								}
+								AsyncWriteFile(AsyncWriteAndHashSequence, FLargeMemoryPtr(Writer->ReleaseOwnership()), DataSize, *NewPathToSave, WriteOptions);
 							}
 						}
 						Linker->CloseAndDestroySaver();
@@ -5773,16 +5771,21 @@ FSavePackageResultStruct UPackage::Save(UPackage* InOuter, UObject* Base, EObjec
 						UE_LOG(LogSavePackage, Log,  TEXT("Moving '%s' to '%s'"), *TempFilename, *NewPath );
 						TotalPackageSizeUncompressed += IFileManager::Get().FileSize(*TempFilename);
 
-						if (SaveFlags & SAVE_ComputeHash)
-						{
-							AddFileToHash(TempFilename, CookedPackageHash);
-						}
-
 						Success = IFileManager::Get().Move( *NewPath, *TempFilename );
 						if (FinalTimeStamp != FDateTime::MinValue())
 						{
 							IFileManager::Get().SetTimeStamp(*NewPath, FinalTimeStamp);
 						}
+
+						if (bComputeHash)
+						{
+							OutstandingAsyncWrites.Increment();
+							AsyncWriteAndHashSequence.AddWork([NewPath](FMD5& State)
+							{
+								AddFileToHash(NewPath, State);
+								OutstandingAsyncWrites.Decrement();
+							});
+						}
 					}
 
 					if( Success == false )
@@ -5899,16 +5902,20 @@ FSavePackageResultStruct UPackage::Save(UPackage* InOuter, UObject* Base, EObjec
 
 		if (Success)
 		{
-			FMD5Hash OutputHash;
-			OutputHash.Set(CookedPackageHash);
+			auto HashCompletionFunc = [](FMD5& State)
+			{
+				FMD5Hash OutputHash;
+				OutputHash.Set(State);
+				return OutputHash;
+			};
 
 			if (bRequestStub)
 			{
-				return FSavePackageResultStruct(ESavePackageResult::GenerateStub, TotalPackageSizeUncompressed, OutputHash);
+				return FSavePackageResultStruct(ESavePackageResult::GenerateStub, TotalPackageSizeUncompressed, AsyncWriteAndHashSequence.Finalize(EAsyncExecution::TaskGraph, MoveTemp(HashCompletionFunc)));
 			}
 			else
 			{
-				return FSavePackageResultStruct(bDiffOnlyIdentical ? ESavePackageResult::Success : ESavePackageResult::DifferentContent, TotalPackageSizeUncompressed, OutputHash);
+				return FSavePackageResultStruct(bDiffOnlyIdentical ? ESavePackageResult::Success : ESavePackageResult::DifferentContent, TotalPackageSizeUncompressed, AsyncWriteAndHashSequence.Finalize(EAsyncExecution::TaskGraph, MoveTemp(HashCompletionFunc)));
 			}
 		}
 		else
@@ -6126,7 +6133,7 @@ static void SaveAssetRegistryData(UPackage* InOuter, FLinkerSave* Linker, FStruc
 }
 
 void SaveBulkData(FLinkerSave* Linker, const UPackage* InOuter, const TCHAR* Filename, const ITargetPlatform* TargetPlatform,
-				  FSavePackageContext* SavePackageContext, const bool bTextFormat, const bool bDiffing, FMD5* CookedPackageHash, int64 TotalPackageSizeUncompressed)
+				  FSavePackageContext* SavePackageContext, const bool bTextFormat, const bool bDiffing, const bool bComputeHash, TAsyncWorkSequence<FMD5>& AsyncWriteAndHashSequence, int64 TotalPackageSizeUncompressed)
 {
 	// Now we write all the bulkdata that is supposed to be at the end of the package
 	// and fix up the offset
@@ -6333,20 +6340,24 @@ void SaveBulkData(FLinkerSave* Linker, const UPackage* InOuter, const TCHAR* Fil
 				{
 					if (const int64 DataSize = Archive->TotalSize())
 					{
-						if (CookedPackageHash != nullptr)
-						{
-							CookedPackageHash->Update(Archive->GetData(), DataSize);
-						}
-
 						TotalPackageSizeUncompressed += DataSize;
 
-						if (bWriteBulkToDisk)
+						if (bComputeHash || bWriteBulkToDisk)
 						{
 							FLargeMemoryPtr DataPtr(Archive->ReleaseOwnership());
 
 							const FString ArchiveFilename = FPaths::ChangeExtension(Filename, BulkFileExtension);
 
-							AsyncWriteFile(MoveTemp(DataPtr), DataSize, *ArchiveFilename);
+							EAsyncWriteOptions WriteOptions(EAsyncWriteOptions::None);
+							if (bComputeHash)
+							{
+								WriteOptions |= EAsyncWriteOptions::ComputeHash;
+							}
+							if (bWriteBulkToDisk)
+							{
+								WriteOptions |= EAsyncWriteOptions::WriteFileToDisk;
+							}
+							AsyncWriteFile(AsyncWriteAndHashSequence, MoveTemp(DataPtr), DataSize, *ArchiveFilename, WriteOptions);
 						}
 					}
 				};
diff --git a/Engine/Source/Runtime/CoreUObject/Public/UObject/Package.h b/Engine/Source/Runtime/CoreUObject/Public/UObject/Package.h
--- a/Engine/Source/Runtime/CoreUObject/Public/UObject/Package.h
+++ b/Engine/Source/Runtime/CoreUObject/Public/UObject/Package.h
@@ -13,6 +13,7 @@
 #include "Serialization/CustomVersion.h"
 #include "Templates/UniquePtr.h"
 #include "Misc/SecureHash.h"
+#include "Async/Future.h"
 
 class Error;
 
@@ -58,13 +59,13 @@ struct FSavePackageResultStruct
 	int64 TotalFileSize;
 
 	/** MD5 hash of the cooked data */
-	FMD5Hash CookedHash;
+	TFuture<FMD5Hash> CookedHash;
 
 	/** Constructors, it will implicitly construct from the result enum */
 	FSavePackageResultStruct() : Result(ESavePackageResult::Error), TotalFileSize(0) {}
 	FSavePackageResultStruct(ESavePackageResult InResult) : Result(InResult), TotalFileSize(0) {}
 	FSavePackageResultStruct(ESavePackageResult InResult, int64 InTotalFileSize) : Result(InResult), TotalFileSize(InTotalFileSize) {}
-	FSavePackageResultStruct(ESavePackageResult InResult, int64 InTotalFileSize, FMD5Hash InHash) : Result(InResult), TotalFileSize(InTotalFileSize), CookedHash(InHash) {}
+	FSavePackageResultStruct(ESavePackageResult InResult, int64 InTotalFileSize, TFuture<FMD5Hash>&& InHash) : Result(InResult), TotalFileSize(InTotalFileSize), CookedHash(MoveTemp(InHash)) {}
 
 	bool operator==(const FSavePackageResultStruct& Other) const
 	{
