Add overloads for FUntypedBulkData2::Serialize so existing licensee code relying on the default parameters in the old BulkData api will not have compiler errors.

At a later point I will remove the parameters entirely and add a proper deprecation path.

#rnx

diff --git a/Engine/Source/Runtime/CoreUObject/Public/Serialization/BulkData2.h b/Engine/Source/Runtime/CoreUObject/Public/Serialization/BulkData2.h
--- a/Engine/Source/Runtime/CoreUObject/Public/Serialization/BulkData2.h
+++ b/Engine/Source/Runtime/CoreUObject/Public/Serialization/BulkData2.h
@@ -180,6 +180,21 @@ public:
 	{
 		FBulkDataBase::Serialize(Ar, Owner, Index, bAttemptFileMapping, sizeof(ElementType));
 	}
+	
+	// @TODO: The following two ::Serialize methods are a work around for the default parameters in the old 
+	// BulkData api that are not used anywhere and to avoid causing code compilation issues for licensee code.
+	// At some point in the future we should remove Index and bAttemptFileMapping from both the old and new 
+	// BulkData api implementations of ::Serialize and then use UE_DEPRECATE to update existing code properly.
+	FORCEINLINE void Serialize(FArchive& Ar, UObject* Owner)
+	{	
+		Serialize(Ar, Owner, INDEX_NONE, false);
+	}
+
+	// @TODO: See above
+	FORCEINLINE void Serialize(FArchive& Ar, UObject* Owner, int32 Index)
+	{
+		Serialize(Ar, Owner, Index, false);
+	}
 
 	/**
 	 * Returns the number of elements held by the BulkData object.
