AsyncLoading2: Global meta data cleanup.
- Move 19 MB global ImportMap (slimports) to per package summary - now only used in serialize
- Add 3.6 MB global ImportedPackages - use in ImportPackagesRecursive and for ref counts tracking for GC
- Reduce serialized global import data to FNames only from 10.6 MB to 3.5 MB (actually only used during initial load, and nice for debugging)
- Add 15.7 MB in-place loaded InitialLoadIoBuffer for initial load meta data
- Find all script imports at end of initial load and release InitialLoadIoBuffer
- Sort global imports by script first, then by full name (i.e. by package)
- Use global import indices in ExportMap
- In PreGarbageCollect, convert raw import object pointers to weak pointers to avoid having ever having to do StaticFindObject fo find them again

diff --git a/Engine/Source/Developer/IoStoreUtilities/Private/IoStoreUtilities.cpp b/Engine/Source/Developer/IoStoreUtilities/Private/IoStoreUtilities.cpp
--- a/Engine/Source/Developer/IoStoreUtilities/Private/IoStoreUtilities.cpp
+++ b/Engine/Source/Developer/IoStoreUtilities/Private/IoStoreUtilities.cpp
@@ -204,11 +204,13 @@ struct FZenPackageSummary
 {
 	uint32 PackageFlags;
 	int32 NameMapOffset;
+	int32 ImportMapOffset;
 	int32 ExportMapOffset;
 	int32 ExportBundlesOffset;
 	int32 GraphDataOffset;
 	int32 GraphDataSize;
 	int32 BulkDataStartOffset;
+	int32 GlobalImportIndex;
 	int32 Pad;
 };
 
@@ -220,7 +222,6 @@ enum EPreloadDependencyType
 
 enum EEventLoadNode2 : uint8
 {
-	Package_CreateLinker,
 	Package_ProcessSummary,
 	Package_ExportsSerialized,
 	Package_StartPostLoad,
@@ -377,11 +378,11 @@ struct FPackage
 	int32 NameCount = 0;
 	int32 ImportCount = 0;
 	int32 ImportOffset = 0;
-	int32 ScriptArcsCount = 0;
-	int32 ScriptArcsOffset = 0;
-	int32 SlimportCount = 0;
-	int32 SlimportOffset = -1;
+	int32 ImportedPackagesCount = 0;
+	int32 ImportedPackagesOffset = -1;
 	int32 ExportCount = 0;
+	int32 FirstGlobalImport;
+	int32 GlobalImportCount;
 	int32 ExportIndexOffset = -1;
 	int32 PreloadIndexOffset = -1;
 	int64 BulkDataStartOffset = -1;
@@ -390,13 +391,15 @@ struct FPackage
 	int64 SummarySize = 0;
 	int64 UGraphSize = 0;
 	int64 NameMapSize = 0;
+	int64 ImportMapSize = 0;
 	int64 ExportMapSize = 0;
 	int64 ExportBundlesSize = 0;
 
 	bool bHasCircularImportDependencies = false;
 	bool bHasCircularPostLoadDependencies = false;
 
-	TSet<FName> ImportedPackageNames;
+	TArray<FString> ImportedFullNames;
+
 	TArray<FPackage*> ImportedPackages;
 	TSet<FPackage*> AllReachablePackages;
 	TSet<FPackage*> ImportedPreloadPackages;
@@ -910,11 +913,30 @@ struct FImportData
 	int32 GlobalIndex = -1;
 	int32 OuterIndex = -1;
 	int32 OutermostIndex = -1;
+	int32 GlobalExportIndex = -1;
 	int32 RefCount = 0;
 	FName ObjectName;
 	bool bIsPackage = false;
 	bool bIsScript = false;
 	FString FullName;
+	FPackage* Package = nullptr;
+
+	bool operator<(const FImportData& Other) const
+	{
+		if (bIsScript != Other.bIsScript)
+		{
+			return bIsScript;
+		}
+		if (OutermostIndex != Other.OutermostIndex)
+		{
+			return OutermostIndex < Other.OutermostIndex;
+		}
+		if (bIsPackage != Other.bIsPackage)
+		{
+			return bIsPackage;
+		}
+		return FullName < Other.FullName;
+	}
 };
 
 struct FExportData
@@ -924,6 +946,10 @@ struct FExportData
 	FName ObjectName;
 	int32 SourceIndex = -1;
 	int32 GlobalImportIndex = -1;
+	FPackageIndex OuterIndex;
+	FPackageIndex ClassIndex;
+	FPackageIndex SuperIndex;
+	FPackageIndex TemplateIndex;
 	FString FullName;
 
 	FExportGraphNode* CreateNode = nullptr;
@@ -1069,7 +1095,6 @@ int32 CreateTarget(const FContainerTarget& Target)
 	TMap<FString, int32> GlobalExportsByFullName;
 	TArray<FString> TempFullNames;
 	TArray<FPackageIndex> PreloadDependencies;
-	uint64 UniqueImportPackages = 0;
 	uint64 UPackageImports = 0;
 	TArray<FPackageFileSummary> Summaries;
 	TArray<int32> ImportPreloadCounts;
@@ -1081,9 +1106,9 @@ int32 CreateTarget(const FContainerTarget& Target)
 	uint64 ExportPreloadCount = 0;
 
 	TUniquePtr<FLargeMemoryWriter> StoreTocArchive = MakeUnique<FLargeMemoryWriter>(0, true);
-	TUniquePtr<FLargeMemoryWriter> GlimportArchive = MakeUnique<FLargeMemoryWriter>(0, true);
-	TUniquePtr<FLargeMemoryWriter> SlimportArchive = MakeUnique<FLargeMemoryWriter>(0, true);
-	TUniquePtr<FLargeMemoryWriter> ScriptArcsArchive = MakeUnique<FLargeMemoryWriter>(0, true);
+	TUniquePtr<FLargeMemoryWriter> ImportedPackagesArchive = MakeUnique<FLargeMemoryWriter>(0, true);
+	TUniquePtr<FLargeMemoryWriter> GlobalImportNamesArchive = MakeUnique<FLargeMemoryWriter>(0, true);
+	TUniquePtr<FLargeMemoryWriter> InitialLoadArchive = MakeUnique<FLargeMemoryWriter>(0, true);
 
 	FIoStoreEnvironment IoStoreEnv;
 #if !SKIP_WRITE_CONTAINER
@@ -1200,22 +1225,10 @@ int32 CreateTarget(const FContainerTarget& Target)
 
 			UPackageImports += NumPackages;
 
-			Package.SlimportCount = Summary.ImportCount;
-			Package.SlimportOffset = SlimportArchive->Tell();
-			TempFullNames.Reset();
-			TempFullNames.SetNum(Summary.ImportCount, false);
+			Package.ImportedFullNames.SetNum(Summary.ImportCount);
 			for (int32 I = 0; I < Summary.ImportCount; ++I)
 			{
-				FindImport(GlobalImports, GlobalImportsByFullName, TempFullNames, Imports.GetData() + BaseIndex, I);
-
-				FImportData& ImportData = GlobalImports[*GlobalImportsByFullName.Find(TempFullNames[I])];
-				*SlimportArchive << ImportData.GlobalIndex;
-
-				if (ImportData.bIsPackage)
-				{
-					Package.ImportedPackageNames.Add(ImportData.ObjectName);
-				}
-				Package.Imports.Add(ImportData.GlobalIndex);
+				FindImport(GlobalImports, GlobalImportsByFullName, Package.ImportedFullNames, Imports.GetData() + BaseIndex, I);
 			}
 		}
 
@@ -1296,13 +1309,67 @@ int32 CreateTarget(const FContainerTarget& Target)
 		Ar->Close();
 	}
 
-	for (FExportData& GlobalExport : GlobalExports)
+	int32 NumScriptImports = 0;
 	{
-		int32* FindGlobalImport = GlobalImportsByFullName.Find(GlobalExport.FullName);
-		if (FindGlobalImport)
+		// Sort imports by script objects first
+		GlobalImports.Sort();
+
+		// build remap from old global import index to new sorted global import index
+		TMap<int32, int32> Remap;
+		Remap.Reserve(GlobalImports.Num() + 1);
+		Remap.Add(-1, -1);
+		for (int32 I = 0; I < GlobalImports.Num(); ++I)
 		{
-			GlobalExport.GlobalImportIndex = *FindGlobalImport;
+			FImportData& Import = GlobalImports[I];
+			Remap.Add(Import.GlobalIndex, I);
 		}
+
+		// remap all global import indices and lookup package pointers and export indices
+		FPackage* LastPackage = nullptr; 
+		for (int32 I = 0; I < GlobalImports.Num(); ++I)
+		{
+			FImportData& GlobalImport = GlobalImports[I];
+
+			GlobalImport.GlobalIndex = Remap[GlobalImport.GlobalIndex];
+			GlobalImport.OuterIndex = Remap[GlobalImport.OuterIndex];
+			GlobalImport.OutermostIndex = Remap[GlobalImport.OutermostIndex];
+
+			if (!GlobalImport.bIsScript)
+			{
+				if (NumScriptImports == 0)
+				{
+					NumScriptImports = I;
+				}
+
+				if (GlobalImport.bIsPackage)
+				{
+					if (LastPackage)
+					{
+						LastPackage->GlobalImportCount = I - LastPackage->FirstGlobalImport;
+					}
+					FPackage* FindPackage = PackageMap.FindRef(GlobalImport.ObjectName);
+					check(FindPackage);
+					FindPackage->FirstGlobalImport = I;
+					GlobalImport.Package = FindPackage;
+					LastPackage = FindPackage;
+				}
+				else
+				{
+					int32* FindGlobalExport = GlobalExportsByFullName.Find(GlobalImport.FullName);
+					check(FindGlobalExport);
+					GlobalImport.GlobalExportIndex = *FindGlobalExport;
+				}
+			}
+
+			GlobalImportsByFullName[GlobalImport.FullName] = I;
+		}
+		LastPackage->GlobalImportCount = GlobalImports.Num() - LastPackage->FirstGlobalImport;
+	}
+
+	for (FExportData& GlobalExport : GlobalExports)
+	{
+		int32* FindGlobalImport = GlobalImportsByFullName.Find(GlobalExport.FullName);
+		GlobalExport.GlobalImportIndex = FindGlobalImport ? *FindGlobalImport : -1;
 	}
 
 	uint64 ImportSize = Imports.Num() * 28;
@@ -1310,8 +1377,7 @@ int32 CreateTarget(const FContainerTarget& Target)
 	uint64 PreloadDependenciesSize = PreloadDependencies.Num() * 4;
 
 #if OUTPUT_NAMEMAP_CSV
-	FString CsvFilePath = SlimportArchive->GetArchiveName();
-	CsvFilePath.Append(TEXT(".csv"));
+	FString CsvFilePath = OutputDir / TEXT("AllImports.csv");
 	TUniquePtr<FArchive> CsvArchive(IFileManager::Get().CreateFileWriter(*CsvFilePath));
 	if (CsvArchive)
 	{
@@ -1334,25 +1400,6 @@ int32 CreateTarget(const FContainerTarget& Target)
 	}
 #endif
 
-	if (GlimportArchive)
-	{
-		int32 Pad = 0;
-		for (const FImportData& ImportData : GlobalImports)
-		{
-			UniqueImportPackages += (ImportData.OuterIndex == 0);
-			NameMapBuilder.SerializeName(*GlimportArchive, ImportData.ObjectName);
-			FPackageIndex Index;
-			Index = FPackageIndex::FromImport(ImportData.GlobalIndex);
-			*GlimportArchive << Index;
-			Index = ImportData.OuterIndex >= 0 ? FPackageIndex::FromImport(ImportData.OuterIndex) :
-				FPackageIndex();
-			*GlimportArchive << Index;
-			Index = FPackageIndex::FromImport(ImportData.OutermostIndex);
-			*GlimportArchive << Index;
-			*GlimportArchive << Pad;
-		}
-	}
-
 	uint64 SummaryChunkCount = 0;
 	uint64 ExportChunkCount = 0;
 	uint64 BulkChunkCount = 0;
@@ -1363,23 +1410,53 @@ int32 CreateTarget(const FContainerTarget& Target)
 	TSet<FCircularImportChain> CircularChains;
 	int32 PackageIndex = 0;
 
-	// Lookup package pointers for all imports before running postload code
+	// Lookup global indices and package pointers for all imports before adding preload and postload arcs
 	UE_LOG(LogIoStore, Display, TEXT("Looking up import packages..."));
 	for (auto& PackageKV : PackageMap)
 	{
 		FPackage& Package = *PackageKV.Value;
 
-		Package.ImportedPackages.Reserve(Package.ImportedPackageNames.Num());
-		for (const FName& ImportedPackageName : Package.ImportedPackageNames)
+		Package.ImportedPackages.Reserve(Package.ImportCount);
+		for (int32 I = 0; I < Package.ImportCount; ++I)
 		{
-			FPackage* FindPackage = PackageMap.FindRef(ImportedPackageName);
-			if (FindPackage)
+			int32 GlobalImportIndex = *GlobalImportsByFullName.Find(Package.ImportedFullNames[I]);
+
+			FImportData& ImportData = GlobalImports[GlobalImportIndex];
+			Package.Imports.Add(ImportData.GlobalIndex);
+			if (ImportData.Package)
 			{
-				Package.ImportedPackages.Add(FindPackage);
+				Package.ImportedPackages.Add(ImportData.Package);
 			}
 		}
 	}
 
+	UE_LOG(LogIoStore, Display, TEXT("Converting export map import indices..."));
+	for (auto& PackageKV : PackageMap)
+	{
+		FPackage& Package = *PackageKV.Value;
+		for (int32 I = 0; I < Package.ExportCount; ++I)
+		{
+			FObjectExport& ObjectExport = Exports[Package.ExportIndexOffset + I];
+			FExportData& ExportData = GlobalExports[Package.Exports[I]];
+
+			check(!ObjectExport.OuterIndex.IsImport());
+			ExportData.OuterIndex = ObjectExport.OuterIndex;
+			ExportData.ClassIndex =
+				ObjectExport.ClassIndex.IsImport() ?
+				FPackageIndex::FromImport(Package.Imports[ObjectExport.ClassIndex.ToImport()]) :
+				ObjectExport.ClassIndex;
+			ExportData.SuperIndex =
+				ObjectExport.SuperIndex.IsImport() ?
+				FPackageIndex::FromImport(Package.Imports[ObjectExport.SuperIndex.ToImport()]) :
+				ObjectExport.SuperIndex;
+			ExportData.TemplateIndex =
+				ObjectExport.TemplateIndex.IsImport() ?
+				FPackageIndex::FromImport(Package.Imports[ObjectExport.TemplateIndex.ToImport()]) :
+				ObjectExport.TemplateIndex;
+		}
+	}
+
+#if 0
 	UE_LOG(LogIoStore, Display, TEXT("Simulating postload dependencies for stats..."));
 	for (auto& PackageKV : PackageMap)
 	{
@@ -1405,6 +1482,7 @@ int32 CreateTarget(const FContainerTarget& Target)
 		// 		Package.AllReachablePackages.Num());
 		// }
 	}
+#endif
 
 	UE_LOG(LogIoStore, Display, TEXT("Adding optimized postload dependencies..."));
 	for (auto& PackageKV : PackageMap)
@@ -1439,22 +1517,22 @@ int32 CreateTarget(const FContainerTarget& Target)
 				else
 				{
 					FImportData& Import = GlobalImports[Package.Imports[Dep.ToImport()]];
-					ensure(!Import.bIsPackage);
-					int32* FindGlobalExport = GlobalExportsByFullName.Find(Import.FullName);
-					if (FindGlobalExport)
+					check(!Import.bIsPackage);
+					if (Import.bIsScript)
 					{
-						FExportData& Export = GlobalExports[*FindGlobalExport];
+						// Add script arc with null package and global import index as node index
+						AddScriptArc(Package, Import.GlobalIndex, I, PhaseTo);
+					}
+					else
+					{
+						check(Import.GlobalExportIndex != - 1);
+						FExportData& Export = GlobalExports[Import.GlobalExportIndex];
 						FPackage* SourcePackage = PackageMap.FindRef(Export.SourcePackageName);
 						check(SourcePackage);
 						FPackageIndex SourceDep = FPackageIndex::FromExport(Export.SourceIndex);
 						AddExternalExportArc(ExportGraph, *SourcePackage, Export.SourceIndex, PhaseFrom, Package, I, PhaseTo);
 						Package.ImportedPreloadPackages.Add(SourcePackage);
 					}
-					else
-					{
-						// Add script arc with null package and global import index as node index
-						AddScriptArc(Package, Import.GlobalIndex, I, PhaseTo);
-					}
 				}
 			};
 
@@ -1496,22 +1574,80 @@ int32 CreateTarget(const FContainerTarget& Target)
 	UE_LOG(LogIoStore, Display, TEXT("Building bundles..."));
 	BuildBundles(ExportGraph, PackageMap);
 
-	UE_LOG(LogIoStore, Display, TEXT("Serializing..."));
-	for (auto& PackageKV : PackageMap)
+
+	UE_LOG(LogIoStore, Display, TEXT("Serializing global import names..."));
+	if (GlobalImportNamesArchive)
 	{
-		FPackage& Package = *PackageKV.Value;
+		for (const FImportData& ImportData : GlobalImports)
+		{
+			NameMapBuilder.SerializeName(*GlobalImportNamesArchive, ImportData.ObjectName);
+		}
+	}
 
-		// Separate file for script arcs that are only required during initial loading
+	UE_LOG(LogIoStore, Display, TEXT("Serializing initial load..."));
+#if 1
+	// Separate file for script arcs that are only required during initial loading
+	if (InitialLoadArchive)
+	{
+		int32 PackageCount = PackageMap.Num();
+		*InitialLoadArchive << PackageCount;
+		*InitialLoadArchive << NumScriptImports;
+
+		FBufferWriter ScriptArcsArchive(nullptr, 0, EBufferWriterFlags::AllowResize | EBufferWriterFlags::TakeOwnership);
+
+		for (auto& PackageKV : PackageMap)
 		{
-			Package.ScriptArcsOffset = ScriptArcsArchive->Tell();
-			Package.ScriptArcsCount = Package.ScriptArcs.Num();
+			FPackage& Package = *PackageKV.Value;
+
+			int32 ScriptArcsOffset = ScriptArcsArchive.Tell();
+			int32 ScriptArcsCount = Package.ScriptArcs.Num();
+
+			*InitialLoadArchive << ScriptArcsOffset;
+			*InitialLoadArchive << ScriptArcsCount;
+
 			for (FArc& ScriptArc : Package.ScriptArcs)
 			{
-				*ScriptArcsArchive << ScriptArc.FromNodeIndex;
-				*ScriptArcsArchive << ScriptArc.ToNodeIndex;
+				ScriptArcsArchive << ScriptArc.FromNodeIndex;
+				ScriptArcsArchive << ScriptArc.ToNodeIndex;
 			}
 		}
 
+		for (int32 I = 0; I < NumScriptImports; ++I)
+		{
+			FImportData& ImportData = GlobalImports[I];
+			FPackageIndex OuterIndex = 
+				ImportData.OuterIndex >= 0 ?
+				FPackageIndex::FromImport(ImportData.OuterIndex) :
+				FPackageIndex();
+
+			*InitialLoadArchive << OuterIndex;
+		}
+
+		InitialLoadArchive->Serialize(ScriptArcsArchive.GetWriterData(), ScriptArcsArchive.Tell());
+	}
+#endif
+
+	UE_LOG(LogIoStore, Display, TEXT("Serializing..."));
+	for (auto& PackageKV : PackageMap)
+	{
+		FPackage& Package = *PackageKV.Value;
+
+		// ImportedPackages
+		Package.ImportedPackagesCount = Package.ImportedPackages.Num();
+		Package.ImportedPackagesOffset = ImportedPackagesArchive->Tell();
+		for (FPackage* ImportedPackage : Package.ImportedPackages)
+		{
+			*ImportedPackagesArchive << ImportedPackage->GlobalPackageId;
+		}
+
+		// Temporary Archive for serializing ImportMap
+		FBufferWriter ImportMapArchive(nullptr, 0, EBufferWriterFlags::AllowResize | EBufferWriterFlags::TakeOwnership );
+		for (int32 GlobalImportIndex : Package.Imports)
+		{
+			ImportMapArchive << GlobalImportIndex;
+		}
+		Package.ImportMapSize = ImportMapArchive.Tell();
+
 		// Temporary Archive for serializing EDL graph data
 		FBufferWriter ZenGraphArchive(nullptr, 0, EBufferWriterFlags::AllowResize | EBufferWriterFlags::TakeOwnership );
 
@@ -1551,10 +1687,10 @@ int32 CreateTarget(const FContainerTarget& Target)
 
 			ZenExportMapArchive << ObjectExport.SerialSize;
 			NameMapBuilder.SerializeName(ZenExportMapArchive, ObjectExport.ObjectName);
-			ZenExportMapArchive << ObjectExport.OuterIndex;
-			ZenExportMapArchive << ObjectExport.ClassIndex;
-			ZenExportMapArchive << ObjectExport.SuperIndex;
-			ZenExportMapArchive << ObjectExport.TemplateIndex;
+			ZenExportMapArchive << ExportData.OuterIndex;
+			ZenExportMapArchive << ExportData.ClassIndex;
+			ZenExportMapArchive << ExportData.SuperIndex;
+			ZenExportMapArchive << ExportData.TemplateIndex;
 			ZenExportMapArchive << ExportData.GlobalImportIndex;
 			ZenExportMapArchive << (uint32&)ObjectExport.ObjectFlags;
 		}
@@ -1589,16 +1725,18 @@ int32 CreateTarget(const FContainerTarget& Target)
 
 		ensure(Package.GlobalPackageId == PackageIndex++);
 		{
+			int32 Pad = 0;
+			int32 ExportBundleCount = Package.ExportBundles.Num();
+
 			NameMapBuilder.SerializeName(*StoreTocArchive, Package.Name);
 			NameMapBuilder.SerializeName(*StoreTocArchive, RelativeFileName);
-			*StoreTocArchive << Package.ImportCount;
-			*StoreTocArchive << Package.SlimportCount;
-			*StoreTocArchive << Package.SlimportOffset;
 			*StoreTocArchive << Package.ExportCount;
-			int32 ExportBundleCount = Package.ExportBundles.Num();
 			*StoreTocArchive << ExportBundleCount;
-			*StoreTocArchive << Package.ScriptArcsOffset;
-			*StoreTocArchive << Package.ScriptArcsCount;
+			*StoreTocArchive << Package.FirstGlobalImport;
+			*StoreTocArchive << Package.GlobalImportCount;
+			*StoreTocArchive << Package.ImportedPackagesCount;
+			*StoreTocArchive << Package.ImportedPackagesOffset;
+			*StoreTocArchive << Pad;
 		}
 
 		Package.NameMapSize = Package.NameIndices.Num() * Package.NameIndices.GetTypeSize();
@@ -1607,6 +1745,7 @@ int32 CreateTarget(const FContainerTarget& Target)
 			const uint64 ZenSummarySize = 
 						sizeof(FZenPackageSummary)
 						+ Package.NameMapSize
+						+ Package.ImportMapSize
 						+ Package.ExportMapSize
 						+ Package.ExportBundlesSize
 						+ Package.UGraphSize;
@@ -1617,6 +1756,8 @@ int32 CreateTarget(const FContainerTarget& Target)
 			ZenSummary->PackageFlags = Package.PackageFlags;
 			ZenSummary->GraphDataSize = Package.UGraphSize;
 			ZenSummary->BulkDataStartOffset = Package.BulkDataStartOffset;
+			int32* FindGlobalImportIndexForPackage = GlobalImportsByFullName.Find(Package.Name.ToString());
+			ZenSummary->GlobalImportIndex = FindGlobalImportIndexForPackage ? *FindGlobalImportIndexForPackage : -1;
 
 			FBufferWriter ZenAr(ZenSummaryBuffer, ZenSummarySize); 
 			ZenAr.Seek(sizeof(FZenPackageSummary));
@@ -1627,6 +1768,13 @@ int32 CreateTarget(const FContainerTarget& Target)
 				ZenAr.Serialize(Package.NameIndices.GetData(), Package.NameMapSize);
 			}
 
+			// ImportMap data
+			{
+				check(ImportMapArchive.Tell() == Package.ImportMapSize);	
+				ZenSummary->ImportMapOffset = ZenAr.Tell();
+				ZenAr.Serialize(ImportMapArchive.GetWriterData(), ImportMapArchive.Tell());
+			}
+
 			// ExportMap data
 			{
 				check(ZenExportMapArchive.Tell() == Package.ExportMapSize);	
@@ -1748,13 +1896,13 @@ int32 CreateTarget(const FContainerTarget& Target)
 		GlobalMetaArchive << StoreTocByteCount;
 		GlobalMetaArchive.Serialize(StoreTocArchive->GetData(), StoreTocByteCount);
 
-		int32 SlimportsByteCount = SlimportArchive->TotalSize();
-		GlobalMetaArchive << SlimportsByteCount;
-		GlobalMetaArchive.Serialize(SlimportArchive->GetData(), SlimportsByteCount);
+		int32 ImportedPackagesByteCount = ImportedPackagesArchive->TotalSize();
+		GlobalMetaArchive << ImportedPackagesByteCount;
+		GlobalMetaArchive.Serialize(ImportedPackagesArchive->GetData(), ImportedPackagesByteCount);
 
-		int32 GlimportsByteCount = GlimportArchive->TotalSize();
-		GlobalMetaArchive << GlimportsByteCount;
-		GlobalMetaArchive.Serialize(GlimportArchive->GetData(), GlimportsByteCount);
+		int32 GlobalImportNamesByteCount = GlobalImportNamesArchive->TotalSize();
+		GlobalMetaArchive << GlobalImportNamesByteCount;
+		GlobalMetaArchive.Serialize(GlobalImportNamesArchive->GetData(), GlobalImportNamesByteCount);
 
 		const FIoStatus Status = AppendToContainer(*IoStoreWriter, GlobalMetaArchive, CreateIoChunkId(0, 0, EIoChunkType::LoaderGlobalMeta));
 		if (!Status.IsOk())
@@ -1765,7 +1913,7 @@ int32 CreateTarget(const FContainerTarget& Target)
 
 	{
 		UE_LOG(LogIoStore, Display, TEXT("Saving initial load meta data to container file"));
-		const FIoStatus Status = AppendToContainer(*IoStoreWriter, *ScriptArcsArchive, CreateIoChunkId(0, 0, EIoChunkType::LoaderInitialLoadMeta));
+		const FIoStatus Status = AppendToContainer(*IoStoreWriter, *InitialLoadArchive, CreateIoChunkId(0, 0, EIoChunkType::LoaderInitialLoadMeta));
 		if (!Status.IsOk())
 		{
 			UE_LOG(LogIoStore, Error, TEXT("Failed to save initial load meta data to container file"));
@@ -1789,14 +1937,16 @@ int32 CreateTarget(const FContainerTarget& Target)
 	uint64 UAssetSize = 0;
 	uint64 SummarySize = 0;
 	uint64 UGraphSize = 0;
+	uint64 ImportMapSize = 0;
 	uint64 ExportMapSize = 0;
 	uint64 NameMapSize = 0;
 	uint64 NameMapCount = 0;
 	uint64 ZenPackageSummarySize = FileNames.Num() * sizeof(FZenPackageSummary);
 	uint64 StoreTocSize = StoreTocArchive->Tell();
-	uint64 GlimportSize = GlimportArchive->Tell();
-	uint64 SlimportSize = SlimportArchive->Tell();
-	uint64 ScriptArcsSize = ScriptArcsArchive->Tell();
+	uint64 GlobalImportNamesSize = GlobalImportNamesArchive->Tell();
+	uint64 PackageImportSize = ImportedPackagesArchive->Tell();
+	uint64 ImportedPackagesCount = 0;
+	uint64 InitialLoadSize = InitialLoadArchive->Tell();
 	uint64 ScriptArcsCount = 0;
 	uint64 CircularPackagesCount = 0;
 	uint64 TotalInternalArcCount = 0;
@@ -1809,6 +1959,14 @@ int32 CreateTarget(const FContainerTarget& Target)
 	uint64 BundleCount = 0;
 	uint64 BundleEntryCount = 0;
 
+	uint64 ZenPackageHeaderSize = 0;
+
+	uint64 UniqueImportPackages = 0;
+	for (const FImportData& ImportData : GlobalImports)
+	{
+		UniqueImportPackages += (ImportData.OuterIndex == 0);
+	}
+
 	for (auto& PackageKV : PackageMap)
 	{
 		FPackage& Package = *PackageKV.Value;
@@ -1817,12 +1975,14 @@ int32 CreateTarget(const FContainerTarget& Target)
 		UAssetSize += Package.UAssetSize;
 		SummarySize += Package.SummarySize;
 		UGraphSize += Package.UGraphSize;
+		ImportMapSize += Package.ImportMapSize;
 		ExportMapSize += Package.ExportMapSize;
 		NameMapSize += Package.NameMapSize;
 		NameMapCount += Package.NameIndices.Num();
-		ScriptArcsCount += Package.ScriptArcsCount;
+		ScriptArcsCount += Package.ScriptArcs.Num();
 		CircularPackagesCount += Package.bHasCircularImportDependencies;
 		TotalInternalArcCount += Package.InternalArcs.Num();
+		ImportedPackagesCount += Package.ImportedPackages.Num();
 		NameCount += Package.NameMap.Num();
 		PackagesWithCircularDependenciesCount += Package.bHasCircularPostLoadDependencies;
 		PackagesWithoutPreloadDependenciesCount += Package.ImportedPreloadPackages.Num() == 0;
@@ -1841,6 +2001,8 @@ int32 CreateTarget(const FContainerTarget& Target)
 		}
 	}
 
+	ZenPackageHeaderSize = ZenPackageSummarySize + NameMapSize + ImportMapSize + ExportMapSize + UGraphSize;
+
 	UE_LOG(LogIoStore, Display, TEXT("-------------------- IoStore Summary: %s --------------------"), *Target.TargetPlatform->PlatformName());
 	UE_LOG(LogIoStore, Display, TEXT("Packages: %8d total, %d circular dependencies, %d no preload dependencies, %d no import dependencies"),
 		PackageMap.Num(), PackagesWithCircularDependenciesCount, PackagesWithoutPreloadDependenciesCount, PackagesWithoutImportDependenciesCount);
@@ -1856,17 +2018,19 @@ int32 CreateTarget(const FContainerTarget& Target)
 	UE_LOG(LogIoStore, Display, TEXT("Pak: %12.2f MB PreloadDependencies, %d imports + %d exports = %d"),
 		(double)PreloadDependenciesSize / 1024.0 / 1024.0, ImportPreloadCount, ExportPreloadCount, PreloadDependencies.Num());
 
-	UE_LOG(LogIoStore, Display, TEXT("IoStore: %8.2f MB GlobalToc"), (double)StoreTocSize / 1024.0 / 1024.0);
 	UE_LOG(LogIoStore, Display, TEXT("IoStore: %8.2f MB GlobalNameMap, %d unique names"), (double)GlobalNameMapSize / 1024.0 / 1024.0, NameMapBuilder.Num());
-	UE_LOG(LogIoStore, Display, TEXT("IoStore: %8.2f MB GlobalImportMap, %d unique imports, %d unique import packages"),
-		(double)GlimportSize / 1024.0 / 1024.0, GlobalImportsByFullName.Num(), UniqueImportPackages);
+	UE_LOG(LogIoStore, Display, TEXT("IoStore: %8.2f MB GlobalPackageData"), (double)StoreTocSize / 1024.0 / 1024.0);
+	UE_LOG(LogIoStore, Display, TEXT("IoStore: %8.2f MB GlobalImportedPackages, %d imported packages"), (double)PackageImportSize / 1024.0 / 1024.0, ImportedPackagesCount);
+	UE_LOG(LogIoStore, Display, TEXT("IoStore: %8.2f MB GlobalImportNames, %d total imports, %d script imports, %d UPackage imports"),
+		(double)GlobalImportNamesSize / 1024.0 / 1024.0, GlobalImportsByFullName.Num(), NumScriptImports, UniqueImportPackages);
+	UE_LOG(LogIoStore, Display, TEXT("IoStore: %8.2f MB InitialLoadData, %d script arcs, %d script outers, %d packages"), (double)InitialLoadSize / 1024.0 / 1024.0, ScriptArcsCount, NumScriptImports, FileNames.Num());
+	UE_LOG(LogIoStore, Display, TEXT("IoStore: %8.2f MB PackageHeader, %d packages"), (double)ZenPackageHeaderSize / 1024.0 / 1024.0, FileNames.Num());
 	UE_LOG(LogIoStore, Display, TEXT("IoStore: %8.2f MB PackageSummary"), (double)ZenPackageSummarySize / 1024.0 / 1024.0);
 	UE_LOG(LogIoStore, Display, TEXT("IoStore: %8.2f MB PackageNameMap, %d indices"), (double)NameMapSize / 1024.0 / 1024.0, NameMapCount);
-	UE_LOG(LogIoStore, Display, TEXT("IoStore: %8.2f MB PackageImportMap"), (double)SlimportSize / 1024.0 / 1024.0);
+	UE_LOG(LogIoStore, Display, TEXT("IoStore: %8.2f MB PackageImportMap"), (double)ImportMapSize / 1024.0 / 1024.0);
 	UE_LOG(LogIoStore, Display, TEXT("IoStore: %8.2f MB PackageExportMap"), (double)ExportMapSize / 1024.0 / 1024.0);
-	UE_LOG(LogIoStore, Display, TEXT("IoStore: %8.2f MB SerializedArcs, %d internal arcs, %d external arcs, %d circular packages (%d chains)"),
+	UE_LOG(LogIoStore, Display, TEXT("IoStore: %8.2f MB PackageArcs, %d internal arcs, %d external arcs, %d circular packages (%d chains)"),
 		(double)UGraphSize / 1024.0 / 1024.0, TotalInternalArcCount, TotalExternalArcCount, CircularPackagesCount, CircularChains.Num());
-	UE_LOG(LogIoStore, Display, TEXT("IoStore: %8.2f MB ScriptArcs, %d arcs"), (double)ScriptArcsSize / 1024.0 / 1024.0, ScriptArcsCount);
 
 	return 0;
 }
diff --git a/Engine/Source/Runtime/CoreUObject/Private/Serialization/AsyncLoading2.cpp b/Engine/Source/Runtime/CoreUObject/Private/Serialization/AsyncLoading2.cpp
--- a/Engine/Source/Runtime/CoreUObject/Private/Serialization/AsyncLoading2.cpp
+++ b/Engine/Source/Runtime/CoreUObject/Private/Serialization/AsyncLoading2.cpp
@@ -155,40 +155,120 @@ struct FExportBundleEntry
 	uint32 LocalExportIndex;
 };
 
-struct FGlobalImport
+class FGlobalImport
 {
-	FName ObjectName;
-	FPackageIndex GlobalIndex;
-	FPackageIndex OuterIndex;
-	FPackageIndex OutermostIndex;
-	int32 Pad;
+private:
+	union
+	{
+		UObject* Object = nullptr;
+		struct
+		{
+			int32 ObjectIndex;
+			int32 SerialNumber;
+		} WeakPointer;
+	};
+
+	TAtomic<int32> RefCount { 0 };
+	bool bIsWeakPointer = false;
+
+public:
+	inline void AddRef()
+	{
+		++RefCount;
+	}
+
+	inline void ReleaseRef()
+	{
+		--RefCount;
+	}
+
+	inline int32 GetRefCount()
+	{
+		return RefCount;
+	}
+
+	UObject* GetObject()
+	{
+		if (bIsWeakPointer)
+		{
+			InternalMakeRaw();
+		}
+		return Object;
+	}
+
+	UObject* GetObjectIfRawPointer()
+	{
+		return bIsWeakPointer ? nullptr : Object;
+	}
+
+	void SetObject(UObject* InObject)
+	{
+#if DO_CHECK
+		if (bIsWeakPointer)
+		{
+			InternalMakeRaw();
+		}
+		check(!Object || Object == InObject);
+#endif
+		Object = InObject;
+		bIsWeakPointer = false;
+	}
+
+	void MakeWeak()
+	{
+		check(RefCount == 0);
+		check(!bIsWeakPointer);
+		int32 ObjectIndex = GUObjectArray.ObjectToIndex((UObjectBase*)Object);
+		int32 SerialNumber = GUObjectArray.AllocateSerialNumber(ObjectIndex);
+		WeakPointer.ObjectIndex = ObjectIndex;
+		WeakPointer.SerialNumber = SerialNumber;
+		bIsWeakPointer = true;
+	}
+
+private:
+	void InternalMakeRaw()
+	{
+		check(bIsWeakPointer);
+		FUObjectItem* ObjectItem = GUObjectArray.IndexToObject(WeakPointer.ObjectIndex);
+		int32 ActualSerialNumber = GUObjectArray.GetSerialNumber(WeakPointer.ObjectIndex);
+		Object = (ObjectItem && ActualSerialNumber == WeakPointer.SerialNumber) ? (UObject*)ObjectItem->Object : nullptr;
+		bIsWeakPointer = false;
+	}
 };
 
 struct FPackageStoreEntry
 {
 	FName Name;
 	FName FileName;
-	int32 ImportCount;
-	int32 SlimportCount;
-	int32 SlimportOffset;
 	int32 ExportCount;
 	int32 ExportBundleCount;
-	int32 ScriptArcsOffset;
-	int32 ScriptArcsCount;
+	int32 FirstGlobalImport;
+	int32 GlobalImportCount;
+	int32 ImportedPackagesCount;
+	int32 ImportedPackagesOffset;
+	int32 Pad;
 };
 
 struct FPackageSummary
 {
 	uint32 PackageFlags;
 	int32 NameMapOffset;
+	int32 ImportMapOffset;
 	int32 ExportMapOffset;
 	int32 ExportBundlesOffset;
 	int32 GraphDataOffset;
 	int32 GraphDataSize;
 	int32 BulkDataStartOffset;
+	int32 GlobalImportIndex;
 	int32 Pad;
 };
 
+struct FPackageStoreInitialLoadEntry
+{
+	int32 ScriptArcsOffset;
+	int32 ScriptArcsCount;
+};
+
 class FGlobalNameMap
 {
 public:
@@ -247,52 +327,62 @@ private:
 
 struct FGlobalImportStore
 {
-	/** Persistent data */
+	/** Serialized data */
 	int32 Count = 0;
+	int32 ScriptImportCount = 0;
 	FName* Names = nullptr;
-	FPackageIndex* Outers = nullptr;
-	FPackageIndex* Packages = nullptr;
-	UObject** Objects = nullptr;
+	FPackageIndex* ScriptImportOuters = nullptr;
+	/** Runtime data */
+	FGlobalImport* Imports = nullptr;
 	/** Reference tracking for GC management */
-	TAtomic<int32>* RefCounts = nullptr;
 	TArray<UObject*> KeepAliveObjects;
 	bool bNeedToRunGarbageCollect = false;
+
 	void OnPreGarbageCollect(bool bInIsLoadingPackages);
 	void OnPostGarbageCollect();
+
+	void FindAllScriptImports();
 };
 
 class FPackageStore
 {
+public:
 	TMap<FName, FGlobalPackageId> PackageNameToGlobalPackageId;
 	FGlobalImportStore ImportStore;
+	FIoBuffer InitialLoadIoBuffer;
 	FPackageStoreEntry* StoreEntries = nullptr;
-	int32* Slimports = nullptr;
+	FPackageStoreInitialLoadEntry* InitialLoadStoreEntries = nullptr;
+	int32* ImportedPackages = nullptr;
 	int32* ScriptArcs = nullptr;
-	int32 SlimportCount = 0;
-	int32 ScriptArcsCount = 0;
 	int32 PackageCount = 0;
+	int32 ImportedPackagesCount = 0;
+	int32 ScriptArcsCount = 0;
 
 public:
 	void Load(FIoDispatcher& IoDispatcher, FGlobalNameMap& GlobalNameMap)
 	{
-		FEvent* Event = FPlatformProcess::GetSynchEventFromPool();
 		FIoBuffer IoBuffer;
+		FEvent* Event = FPlatformProcess::GetSynchEventFromPool();
 
-		IoDispatcher.ReadWithCallback(
-			CreateIoChunkId(0, 0, EIoChunkType::LoaderGlobalMeta),
-			FIoReadOptions(),
-			[Event, &IoBuffer](TIoStatusOr<FIoBuffer> Result)
-			{
-				IoBuffer = Result.ConsumeValueOrDie();
-				Event->Trigger();
-			});
+		{
+			TRACE_CPUPROFILER_EVENT_SCOPE(LoadPackageStoreTocIo);
 
-		Event->Wait();
+			IoDispatcher.ReadWithCallback(
+				CreateIoChunkId(0, 0, EIoChunkType::LoaderGlobalMeta),
+				FIoReadOptions(),
+				[Event, &IoBuffer](TIoStatusOr<FIoBuffer> Result)
+				{
+			 		IoBuffer = Result.ConsumeValueOrDie();
+					Event->Trigger();
+				});
+
+			Event->Wait();
+		}
 		
 		FLargeMemoryReader GlobalMetaAr(IoBuffer.Data(), IoBuffer.DataSize());
 		
 		{
-			TRACE_CPUPROFILER_EVENT_SCOPE(LoadPackageStoreToc);
+			TRACE_CPUPROFILER_EVENT_SCOPE(LoadPackageStoreTocFixup);
 
 			int32 PackageByteCount = 0;
 			GlobalMetaAr << PackageByteCount;
@@ -302,7 +392,6 @@ public:
 
 			// In-place loading
 			GlobalMetaAr.Serialize(StoreEntries, PackageByteCount);
-			PackageNameToGlobalPackageId.Reserve(PackageCount);
 
 			// FName fixup
 			{
@@ -317,6 +406,7 @@ public:
 			// Global package index
 			{
 				TRACE_CPUPROFILER_EVENT_SCOPE(LoadPackageStoreGlobalIds);
+				PackageNameToGlobalPackageId.Reserve(PackageCount);
 				for (int I = 0; I < PackageCount; ++I)
 				{
 					PackageNameToGlobalPackageId.Add(StoreEntries[I].Name, FGlobalPackageId { I });
@@ -325,73 +415,86 @@ public:
 		}
 
 		{
-			TRACE_CPUPROFILER_EVENT_SCOPE(LoadPackageStoreSlimports);
+			TRACE_CPUPROFILER_EVENT_SCOPE(LoadPackageStoreImportedPackagesFixup);
 
-			int32 SlimportByteCount = 0;
-			GlobalMetaAr << SlimportByteCount;
+			int32 PackageImportByteCount = 0;
+			GlobalMetaAr << PackageImportByteCount;
 
-			SlimportCount = SlimportByteCount / sizeof(int32);
-			Slimports = reinterpret_cast<int32*>(FMemory::Malloc(SlimportByteCount));
-			GlobalMetaAr.Serialize(Slimports, SlimportByteCount);
+			ImportedPackagesCount = PackageImportByteCount / sizeof(int32);
+			ImportedPackages = reinterpret_cast<int32*>(FMemory::Malloc(PackageImportByteCount));
+			GlobalMetaAr.Serialize(ImportedPackages, PackageImportByteCount);
 		}
 
 		{
-			TRACE_CPUPROFILER_EVENT_SCOPE(LoadPackageStoreGlimports);
+			TRACE_CPUPROFILER_EVENT_SCOPE(LoadPackageStoreGlobalImportNamesFixup);
 
 			int32 ImportByteCount = 0;
 			GlobalMetaAr << ImportByteCount; 
 
-			ImportStore.Count = ImportByteCount / sizeof(FGlobalImport);
-			FGlobalImport* Imports = reinterpret_cast<FGlobalImport*>(FMemory::Malloc(ImportByteCount));
-			GlobalMetaAr.Serialize(Imports, ImportByteCount);
+			ImportStore.Count = ImportByteCount / sizeof(FName);
 
-			TRACE_CPUPROFILER_EVENT_SCOPE(LoadPackageStoreGlimportsFixup);
-			{
-				ImportStore.Names = new FName[ImportStore.Count];
-				ImportStore.Outers = new FPackageIndex[ImportStore.Count];
-				ImportStore.Packages = new FPackageIndex[ImportStore.Count];
-				ImportStore.Objects = new UObject*[ImportStore.Count];
-				ImportStore.RefCounts = new TAtomic<int32>[ImportStore.Count];
+			FName* SerializedNames = reinterpret_cast<FName*>(IoBuffer.Data() + GlobalMetaAr.Tell());
+			ImportStore.Names = reinterpret_cast<FName*>(FMemory::Malloc(ImportByteCount));
 
-				for (int I = 0; I < ImportStore.Count; ++I)
-				{
-					FGlobalImport& Import = Imports[I];
-					Import.ObjectName = GlobalNameMap.FromSerializedName(Import.ObjectName);
-					ImportStore.Names[I] = Import.ObjectName;
-					ImportStore.Outers[I] = Import.OuterIndex;
-					ImportStore.Packages[I] = Import.OutermostIndex;
-					ImportStore.Objects[I] = nullptr;
-					ImportStore.RefCounts[I] = 0;
-				}
+			for (int I = 0; I < ImportStore.Count; ++I)
+			{
+				ImportStore.Names[I] = GlobalNameMap.FromSerializedName(SerializedNames[I]);
 			}
 		}
 
+		{
+			TRACE_CPUPROFILER_EVENT_SCOPE(LoadPackageStoreInitGlobalImports);
+			ImportStore.Imports = new FGlobalImport[ImportStore.Count];
+		}
+
 		// Load initial loading meta data
-		IoDispatcher.ReadWithCallback(
-			CreateIoChunkId(0, 0, EIoChunkType::LoaderInitialLoadMeta),
-			FIoReadOptions(),
-			[Event, &IoBuffer](TIoStatusOr<FIoBuffer> Result)
-			{
-				IoBuffer = Result.ConsumeValueOrDie();
-				Event->Trigger();
-			});
+		{
+			FIoBuffer TempIoBuffer;
+			TRACE_CPUPROFILER_EVENT_SCOPE(LoadPackageStoreInitialLoadIo);
+			IoDispatcher.ReadWithCallback(
+				CreateIoChunkId(0, 0, EIoChunkType::LoaderInitialLoadMeta),
+				FIoReadOptions(),
+				[Event, &TempIoBuffer](TIoStatusOr<FIoBuffer> Result)
+				{
+					TempIoBuffer = Result.ConsumeValueOrDie();
+					Event->Trigger();
+				});
 
-		Event->Wait();
-		FPlatformProcess::ReturnSynchEventToPool(Event);
+			Event->Wait();
+			FPlatformProcess::ReturnSynchEventToPool(Event);
+			InitialLoadIoBuffer = TempIoBuffer;
+		}
 
 		{
-			TRACE_CPUPROFILER_EVENT_SCOPE(LoadPackageStoreScriptArcs);
+			TRACE_CPUPROFILER_EVENT_SCOPE(LoadPackageStoreInitialLoadFixup);
+			uint8* Data = InitialLoadIoBuffer.Data();
+
+			int32 TempPackageCount = *reinterpret_cast<int32*>(Data);
+			Data += 4;
+			ImportStore.ScriptImportCount = *reinterpret_cast<int32*>(Data);
+			Data += 4;
 
-			FLargeMemoryReader InitialLoadAr(IoBuffer.Data(), IoBuffer.DataSize());
+			check(TempPackageCount == PackageCount);
 
-			const int32 ScriptArcsByteCount = InitialLoadAr.TotalSize();
-			ScriptArcsCount = ScriptArcsByteCount / sizeof(int32);
-			ScriptArcs = reinterpret_cast<int32*>(FMemory::Malloc(ScriptArcsByteCount));
-			InitialLoadAr.Serialize(ScriptArcs, ScriptArcsByteCount);
+			InitialLoadStoreEntries = reinterpret_cast<FPackageStoreInitialLoadEntry*>(Data);
+			Data += TempPackageCount * sizeof(FPackageStoreInitialLoadEntry);
+
+			ImportStore.ScriptImportOuters = reinterpret_cast<FPackageIndex*>(Data);
+			Data += ImportStore.ScriptImportCount * sizeof(FPackageIndex);
+
+			ScriptArcs = reinterpret_cast<int32*>(Data);
 		}
 	}
 
-	inline FGlobalImportStore& GetImportStore()
+	void FinalizeInitialLoad()
+	{
+		ImportStore.FindAllScriptImports();
+		InitialLoadStoreEntries = nullptr;
+		ScriptArcs = nullptr;
+		InitialLoadIoBuffer = FIoBuffer();
+	}
+
+	inline FGlobalImportStore& GetGlobalImportStore()
 	{
 		return ImportStore;
 	}
@@ -402,16 +505,15 @@ public:
 		return Id;
 	}
 
-	inline UObject** GetGlobalImportObjects(int32& OutCount)
-	{
-		OutCount = ImportStore.Count;
-		return ImportStore.Objects;
-	}
-
-	inline FPackageIndex* GetGlobalImportOuters(int32& OutCount) const
+	inline FPackageIndex* GetGlobalImportScriptOuters(int32& OutCount) const
 	{
-		OutCount = ImportStore.Count;
-		return ImportStore.Outers;
+		if (GIsInitialLoad)
+		{
+			OutCount = ImportStore.ScriptImportCount;
+			return ImportStore.ScriptImportOuters;
+		}
+		OutCount = 0;
+		return nullptr;
 	}
 
 	inline FName* GetGlobalImportNames(int32& OutCount) const
@@ -420,34 +522,37 @@ public:
 		return ImportStore.Names;
 	}
 
-	inline FPackageIndex* GetGlobalImportPackages(int32& OutCount) const 
+	inline FGlobalImport* GetGlobalImports()
 	{
-		OutCount = ImportStore.Count;
-		return ImportStore.Packages;
+		return ImportStore.Imports;
 	}
 
-	inline TAtomic<int32>* GetGlobalImportObjectRefCounts()
+	inline FPackageStoreEntry* GetGlobalStoreEntries()
 	{
-		return ImportStore.RefCounts;
+		return StoreEntries;
 	}
 
-	inline int32* GetPackageSlimports(FGlobalPackageId GlobalPackageId, int32& OutCount) const 
+	inline int32* GetPackageImportedPackages(FGlobalPackageId GlobalPackageId, int32& OutCount) const 
 	{
 		FPackageStoreEntry& Entry = StoreEntries[GlobalPackageId.Id];
-		OutCount = Entry.SlimportCount;
-		return Slimports + Entry.SlimportOffset / sizeof(int32);
+		OutCount = Entry.ImportedPackagesCount;
+		if (Entry.ImportedPackagesCount >= 0)
+		{
+			return ImportedPackages + Entry.ImportedPackagesOffset / sizeof(int32);
+		}
+		return nullptr;
 	}
 
 	inline int32* GetPackageScriptArcs(FGlobalPackageId GlobalPackageId, int32& OutCount) const 
 	{
-		FPackageStoreEntry& Entry = StoreEntries[GlobalPackageId.Id];
-		OutCount = Entry.ScriptArcsCount;
-		return ScriptArcs + Entry.ScriptArcsOffset / sizeof(int32);
-	}
-
-	inline int32 GetPackageImportCount(FGlobalPackageId GlobalPackageId) const 
-	{
-		return StoreEntries[GlobalPackageId.Id].ImportCount;
+		if (ensure(GIsInitialLoad))
+		{
+			FPackageStoreInitialLoadEntry& Entry = InitialLoadStoreEntries[GlobalPackageId.Id];
+			OutCount = Entry.ScriptArcsCount;
+			return ScriptArcs + Entry.ScriptArcsOffset / sizeof(int32);
+		}
+		OutCount = 0;
+		return nullptr;
 	}
 
 	inline int32 GetPackageExportCount(FGlobalPackageId GlobalPackageId) const 
@@ -469,70 +574,126 @@ public:
 
 struct FPackageImportStore
 {
-	int32 Count = 0;
-	const int32* ImportMap = nullptr;
+	FPackageStoreEntry* GlobalStoreEntries = nullptr;
+	FGlobalImport* GlobalImports = nullptr;
 	const FName* GlobalImportNames = nullptr;
-	const FPackageIndex* GlobalImportOuters = nullptr;
-	const FPackageIndex* GlobalImportPackages = nullptr;
-	UObject** GlobalImportObjects = nullptr;
-	TAtomic<int32>* GlobalImportObjectRefCounts = nullptr;
+	const FPackageIndex* GlobalImportScriptOuters = nullptr;
+	const int32* ImportedPackages = nullptr;
+	const int32* ImportMap = nullptr;
+	FGlobalPackageId GlobalPackageId;
+	int32 GlobalScriptImportCount = 0;
+	int32 ImportedPackagesCount = 0;
+	int32 ImportMapCount = 0;
 
-	FPackageImportStore(FPackageStore& PackageStore, FGlobalPackageId GlobalPackageId)
+	FPackageImportStore(FPackageStore& GlobalPackageStore, FGlobalPackageId InGlobalPackageId) : GlobalPackageId(InGlobalPackageId)
 	{
 		TRACE_CPUPROFILER_EVENT_SCOPE(NewPackageImportStore);
-		int32 GlobalImportCount = 0;
-		ImportMap = PackageStore.GetPackageSlimports(GlobalPackageId, Count);
-		GlobalImportNames = PackageStore.GetGlobalImportNames(GlobalImportCount);
-		GlobalImportOuters = PackageStore.GetGlobalImportOuters(GlobalImportCount);
-		GlobalImportPackages = PackageStore.GetGlobalImportPackages(GlobalImportCount);
-		GlobalImportObjects = PackageStore.GetGlobalImportObjects(GlobalImportCount);
-		GlobalImportObjectRefCounts = PackageStore.GetGlobalImportObjectRefCounts();
+		{
+			int32 GlobalImportCount = 0;
+			GlobalStoreEntries = GlobalPackageStore.GetGlobalStoreEntries();
+			GlobalImports = GlobalPackageStore.GetGlobalImports();
+			GlobalImportNames = GlobalPackageStore.GetGlobalImportNames(GlobalImportCount);
+		}
+		GlobalImportScriptOuters = GlobalPackageStore.GetGlobalImportScriptOuters(GlobalScriptImportCount);
+		ImportedPackages = GlobalPackageStore.GetPackageImportedPackages(GlobalPackageId, ImportedPackagesCount);
 		AddGlobalImportObjectReferences();
 	}
 
 	~FPackageImportStore()
 	{
-		Clear();
+		check(!ImportMap);
 	}
 
-	UObject* FindExistingSlimport(int32 LocalImportIndex);
-	UObject* FindExistingGlimport(int32 GlobalImportIndex);
+	UObject* FindImportFromGlobalIndex(int32 GlobalImportIndex);
 
-	UObject* GetObject(int32 LocalImportIndex)
+	inline UObject* GetImportFromGlobalIndex(FPackageIndex GlobalIndex)
 	{
-		return GlobalImportObjects[ImportMap[LocalImportIndex]];
+		check(GlobalIndex.IsImport());
+		const int32 GlobalImportIndex = GlobalIndex.ToImport();
+		return GlobalImports[GlobalImportIndex].GetObject();
 	}
 
-	UObject* GetObject(FPackageIndex Index)
+	inline UObject* GetImportFromGlobalIndex(int32 GlobalImportIndex)
 	{
-		return GlobalImportObjects[ImportMap[Index.ToImport()]];
+		return GlobalImports[GlobalImportIndex].GetObject();
 	}
 
-	void Clear()
+	inline UObject* FindOrGetImportFromLocalIndex(FPackageIndex LocalIndex)
 	{
-		if (ImportMap)
+		check(LocalIndex.IsImport());
+		check(ImportMap);
+		const int32 LocalImportIndex = LocalIndex.ToImport();
+		check(LocalImportIndex < ImportMapCount);
+		const int32 GlobalImportIndex = ImportMap[LocalIndex.ToImport()];
+		return
+			GIsInitialLoad && GlobalImportIndex < GlobalScriptImportCount ?
+			FindImportFromGlobalIndex(GlobalImportIndex) :
+			GetImportFromGlobalIndex(GlobalImportIndex);
+	}
+
+	inline UObject* FindOrGetImportFromGlobalIndex(FPackageIndex GlobalIndex)
+	{
+		check(GlobalIndex.IsImport());
+		const int32 GlobalImportIndex = GlobalIndex.ToImport();
+		return
+			GIsInitialLoad && GlobalImportIndex < GlobalScriptImportCount ?
+			FindImportFromGlobalIndex(GlobalImportIndex) :
+			GetImportFromGlobalIndex(GlobalImportIndex);
+	}
+
+	void StoreGlobalImportObject(int32 GlobalImportIndex, UObject* InObject)
+	{
+		if (GlobalImportIndex >= 0)
 		{
-			ReleaseGlobalImportObjectReferences();
-			ImportMap = nullptr;
+			GlobalImports[GlobalImportIndex].SetObject(InObject);
 		}
 	}
 
+	void ClearReferences()
+	{
+		ReleaseGlobalImportObjectReferences();
+	}
+
 private:
 	void AddGlobalImportObjectReferences()
 	{
-		for (int32 LocalImportIndex = 0; LocalImportIndex < Count; ++LocalImportIndex)
+		for (int32 LocalPackageIndex = 0; LocalPackageIndex < ImportedPackagesCount; ++LocalPackageIndex)
 		{
-			const int32 GlobalImportIndex = ImportMap[LocalImportIndex];
-			++GlobalImportObjectRefCounts[GlobalImportIndex];
+			int32 GlobalPackageIndex = ImportedPackages[LocalPackageIndex];
+			FPackageStoreEntry& Entry = GlobalStoreEntries[GlobalPackageIndex];
+			for (int32 ObjectImportIndex = 0; ObjectImportIndex < Entry.GlobalImportCount; ObjectImportIndex++)
+			{
+				int32 GlobalImportIndex = Entry.FirstGlobalImport + ObjectImportIndex;
+				GlobalImports[GlobalImportIndex].AddRef();
+			}
+		}
+		// Add ref counts to own imports to speed up OnPreGarbageCollect
+		FPackageStoreEntry& Entry = GlobalStoreEntries[GlobalPackageId.Id];
+		for (int32 ObjectImportIndex = 0; ObjectImportIndex < Entry.GlobalImportCount; ObjectImportIndex++)
+		{
+			int32 GlobalImportIndex = Entry.FirstGlobalImport + ObjectImportIndex;
+			GlobalImports[GlobalImportIndex].AddRef();
 		}
 	}
 
 	void ReleaseGlobalImportObjectReferences()
 	{
-		for (int32 LocalImportIndex = 0; LocalImportIndex < Count; ++LocalImportIndex)
+		for (int32 LocalPackageIndex = 0; LocalPackageIndex < ImportedPackagesCount; ++LocalPackageIndex)
+		{
+			int32 GlobalPackageIndex = ImportedPackages[LocalPackageIndex];
+			FPackageStoreEntry& Entry = GlobalStoreEntries[GlobalPackageIndex];
+			for (int32 ObjectImportIndex = 0; ObjectImportIndex < Entry.GlobalImportCount; ObjectImportIndex++)
+			{
+				int32 GlobalImportIndex = Entry.FirstGlobalImport + ObjectImportIndex;
+				GlobalImports[GlobalImportIndex].ReleaseRef();
+			}
+		}
+		// Release ref counts to own imports to speed up OnPreGarbageCollect
+		FPackageStoreEntry& Entry = GlobalStoreEntries[GlobalPackageId.Id];
+		for (int32 ObjectImportIndex = 0; ObjectImportIndex < Entry.GlobalImportCount; ObjectImportIndex++)
 		{
-			const int32 GlobalImportIndex = ImportMap[LocalImportIndex];
-			--GlobalImportObjectRefCounts[GlobalImportIndex];
+			int32 GlobalImportIndex = Entry.FirstGlobalImport + ObjectImportIndex;
+			GlobalImports[GlobalImportIndex].ReleaseRef();
 		}
 	}
 };
@@ -575,8 +736,7 @@ public:
 		}
 		else
 		{
-			// Object = ImportStore->GetObject(Index);
-			Object = ImportStore->FindExistingSlimport(Index.ToImport());
+			Object = ImportStore->FindOrGetImportFromLocalIndex(Index);
 		}
 		return *this;
 	}
@@ -612,11 +772,7 @@ public:
 		int32 Number = 0;
 		Ar << Number;
 
-		// TODO: Remove support for old global namemap
-		if (PackageNameMap)
-		{
-			NameIndex = PackageNameMap[NameIndex];
-		}
+		NameIndex = PackageNameMap[NameIndex];
 
 		if (GlobalNameMap->IsValidIndex(NameIndex))
 		{
@@ -659,7 +815,6 @@ enum class EAsyncPackageLoadingState2 : uint8
 
 enum EEventLoadNode2 : uint8
 {
-	Package_CreateLinker, // TODO: remove
 	Package_ProcessSummary,
 	Package_ExportsSerialized,
 	Package_StartPostLoad,
@@ -903,15 +1058,7 @@ struct FAsyncPackage2
 		}
 	}
 
-	void ClearImportedPackages()
-	{
-		for (FAsyncPackage2* ImportedAsyncPackage : ImportedAsyncPackages)
-		{
-			ImportedAsyncPackage->ReleaseRef();
-		}
-		ImportedAsyncPackages.Empty();
-		ImportStore.Clear();
-	}
+	void ClearImportedPackages();
 
 	/** Marks a specific request as complete */
 	void MarkRequestIDsAsComplete();
@@ -1108,7 +1255,7 @@ private:
 	TArray<UObject*> DeferredClusterObjects;
 
 	/** List of all request handles */
-	TArray<int32> RequestIDs;
+	TArray<int32, TInlineAllocator<2>> RequestIDs;
 #if WITH_EDITORONLY_DATA
 	/** Index of the meta-data object within the linkers export table (unset if not yet processed, although may still be INDEX_NONE if there is no meta-data) */
 	TOptional<int32> MetaDataIndex;
@@ -1133,7 +1280,7 @@ private:
 
 	uint64 PackageSummarySize = 0;
 	FIoBuffer PackageSummaryIoBuffer;
-	TArray<FIoBuffer> ExportBundleIoBuffers;
+	TArray<FIoBuffer, TInlineAllocator<2>> ExportBundleIoBuffers;
 
 	// FZenLinkerLoad
 	TArray<FExternalReadCallback> ExternalReadDependencies;
@@ -1401,7 +1548,7 @@ private:
 	FGlobalNameMap GlobalNameMap;
 	FIoDispatcher& IoDispatcher;
 
-	FPackageStore PackageStore;
+	FPackageStore GlobalPackageStore;
 public:
 
 	//~ Begin FRunnable Interface.
@@ -1620,7 +1767,7 @@ public:
 	/**
 	 * [ASYNC/GAME THREAD] Removes a request ID from the list of pending requests
 	 */
-	void RemovePendingRequests(TArray<int32>& RequestIDs)
+	void RemovePendingRequests(TArray<int32, TInlineAllocator<2>>& RequestIDs)
 	{
 		FScopeLock Lock(&PendingRequestsCritical);
 		for (int32 ID : RequestIDs)
@@ -1650,7 +1797,7 @@ private:
 
 	FAsyncPackage2* CreateAsyncPackage(const FAsyncPackageDesc& Desc)
 	{
-		FGlobalPackageId* GlobalPackageId = PackageStore.GetGlobalPackageId(Desc.NameToLoad);
+		FGlobalPackageId* GlobalPackageId = GlobalPackageStore.GetGlobalPackageId(Desc.NameToLoad);
 		if (!GlobalPackageId)
 		{
 			return nullptr;
@@ -1745,7 +1892,7 @@ void FAsyncLoadingThread2Impl::InitializeLoading()
 
 	{
 		TRACE_CPUPROFILER_EVENT_SCOPE(LoadPackageStore);
-		PackageStore.Load(IoDispatcher, GlobalNameMap);
+		GlobalPackageStore.Load(IoDispatcher, GlobalNameMap);
 	}
 
 	AsyncThreadReady.Increment();
@@ -1754,7 +1901,7 @@ void FAsyncLoadingThread2Impl::InitializeLoading()
 void FAsyncLoadingThread2Impl::QueuePackage(FAsyncPackageDesc& Package)
 {
 	//TRACE_CPUPROFILER_EVENT_SCOPE(QueuePackage);
-	if (!PackageStore.GetGlobalPackageId(Package.NameToLoad))
+	if (!GlobalPackageStore.GetGlobalPackageId(Package.NameToLoad))
 	{
 		UE_LOG(LogStreaming, Warning, TEXT("QueuePackage: Skipping unknown package: %s"), *Package.NameToLoad.ToString());
 		return;
@@ -2218,15 +2365,16 @@ void FAsyncPackage2::SetupSerializedArcs()
 	}
 }
 
-static UObject* GFindExistingGlimport(int32 GlobalImportIndex,
-	UObject** GlobalImportObjects,
-	const FPackageIndex* GlobalImportOuters,
+static UObject* GFindExistingScriptImport(int32 GlobalImportIndex,
+	FGlobalImport* GlobalImports,
+	const FPackageIndex* GlobalImportScriptOuters,
 	const FName* GlobalImportNames)
 {
-	UObject*& Object = GlobalImportObjects[GlobalImportIndex];
+	FGlobalImport& GlobalImport = GlobalImports[GlobalImportIndex];
+	UObject* Object = GlobalImport.GetObject();
 	if (!Object)
 	{
-		const FPackageIndex& OuterIndex = GlobalImportOuters[GlobalImportIndex];
+		const FPackageIndex& OuterIndex = GlobalImportScriptOuters[GlobalImportIndex];
 		const FName& ObjectName = GlobalImportNames[GlobalImportIndex];
 		if (OuterIndex.IsNull())
 		{
@@ -2234,32 +2382,44 @@ static UObject* GFindExistingGlimport(int32 GlobalImportIndex,
 		}
 		else
 		{
-			UObject* Outer = GFindExistingGlimport(OuterIndex.ToImport(), GlobalImportObjects, GlobalImportOuters, GlobalImportNames);
+			UObject* Outer = GFindExistingScriptImport(OuterIndex.ToImport(), GlobalImports, GlobalImportScriptOuters, GlobalImportNames);
 			if (Outer)
 			{
 				Object = StaticFindObjectFast(UObject::StaticClass(), Outer, ObjectName, false, true);
 			}
 		}
+		GlobalImport.SetObject(Object);
 	}
 	return Object;
 }
 
-UObject* FPackageImportStore::FindExistingGlimport(int32 GlobalImportIndex)
+UObject* FPackageImportStore::FindImportFromGlobalIndex(int32 GlobalImportIndex)
 {
-	return GFindExistingGlimport(GlobalImportIndex,
-		GlobalImportObjects,
-		GlobalImportOuters,
+	check(GlobalImportScriptOuters);
+	check(GlobalImportIndex < GlobalScriptImportCount);
+	return GFindExistingScriptImport(
+		GlobalImportIndex,
+		GlobalImports,
+		GlobalImportScriptOuters,
 		GlobalImportNames);
 }
 
-UObject* FPackageImportStore::FindExistingSlimport(int32 LocalImportIndex)
+void FGlobalImportStore::FindAllScriptImports()
 {
-	int32 GlobalImportIndex = ImportMap[LocalImportIndex];
-	return GFindExistingGlimport(GlobalImportIndex,
-		GlobalImportObjects,
-		GlobalImportOuters,
-		GlobalImportNames);
+	TRACE_CPUPROFILER_EVENT_SCOPE(FindAllScriptImports);
+	check(ScriptImportOuters);
+	for (int32 GlobalImportIndex = 0; GlobalImportIndex < ScriptImportCount; ++GlobalImportIndex)
+	{
+		UObject* Object = GFindExistingScriptImport(GlobalImportIndex, Imports, ScriptImportOuters, Names);
+#if DO_CHECK
+		UE_CLOG(!Object, LogStreaming, Warning, TEXT("Failed to find script import after initial load: %s in %s"),
+			*Names[GlobalImportIndex].ToString(),
+			*Names[ScriptImportOuters[GlobalImportIndex].ToImport()].ToString());
+#endif
+	}
+	ScriptImportOuters = nullptr;
 }
+
 void FAsyncPackage2::ImportPackagesRecursive()
 {
 	if (bHasImportedPackagesRecursive)
@@ -2268,51 +2428,50 @@ void FAsyncPackage2::ImportPackagesRecursive()
 	}
 	bHasImportedPackagesRecursive = true;
 
-	FMemStack& MemStack = FMemStack::Get();
-	FMemMark Mark(MemStack);
-	TArrayView<FPackageIndex> LocalImportedPackages(new(MemStack)FPackageIndex[ImportStore.Count], ImportStore.Count);
-
-	for (int32 LocalImportIndex = 0; LocalImportIndex < ImportStore.Count; ++LocalImportIndex)
+	FPackageStore& GlobalPackageStore = AsyncLoadingThread.GlobalPackageStore;
+	int32 ImportedPackageCount = 0;
+	int32* Imports = GlobalPackageStore.GetPackageImportedPackages(GlobalPackageId, ImportedPackageCount);
+	if (!Imports)
 	{
-		const int32 GlobalImportIndex = ImportStore.ImportMap[LocalImportIndex];
-		const FPackageIndex ImportedPackageIndex = ImportStore.GlobalImportPackages[GlobalImportIndex];
-		UObject* ImportedObject = ImportStore.FindExistingGlimport(GlobalImportIndex);
-		const UPackage* Package = (UPackage*)ImportStore.GlobalImportObjects[ImportedPackageIndex.ToImport()];
-
-		if (Package && Package->HasAnyPackageFlags(PKG_CompiledIn))
-		{
-			continue;
-		}
+		return;
+	}
 
-		if (ImportedObject)
+	for (int32 LocalImportIndex = 0; LocalImportIndex < ImportedPackageCount; ++LocalImportIndex)
+	{
+		// AreAllImportsInImportPackageAlreadyLoaded?
+		int32 GlobalPackageIndex = Imports[LocalImportIndex];
+		FPackageStoreEntry& Entry = GlobalPackageStore.StoreEntries[GlobalPackageIndex];
+		bool bNeedToLoadPackage = false;
+		// start from 1 to skip checking the UPackage
+		for (int32 ObjectImportIndex = 1; ObjectImportIndex < Entry.GlobalImportCount; ObjectImportIndex++)
 		{
-			if (IsFullyLoadedObj(ImportedObject))
+			int32 GlobalImportIndex = Entry.FirstGlobalImport + ObjectImportIndex;
+			UObject* ImportedObject = ImportStore.GetImportFromGlobalIndex(GlobalImportIndex);
+			if (!ImportedObject || !ImportedObject->HasAllFlags(RF_WasLoaded | RF_LoadCompleted))
 			{
-				continue;
+				bNeedToLoadPackage = true;
+				break;
 			}
 		}
 
-		if (LocalImportedPackages.Contains(ImportedPackageIndex))
-		{
-			continue;
-		}
-
-		LocalImportedPackages[LocalImportIndex] = ImportedPackageIndex;
-
-		const FName ImportedPackageName = ImportStore.GlobalImportNames[ImportedPackageIndex.ToImport()];
-		FAsyncPackageDesc Info(INDEX_NONE, ImportedPackageName);
-		Info.Priority = Desc.Priority;
-		bool bInserted;
-		FAsyncPackage2* ImportedAsyncPackage = AsyncLoadingThread.FindOrInsertPackage(&Info, bInserted);
-		if (ImportedAsyncPackage)
+		if (bNeedToLoadPackage)
 		{
-			TRACE_LOADTIME_ASYNC_PACKAGE_IMPORT_DEPENDENCY(this, ImportedAsyncPackage);
-			ImportedAsyncPackage->AddRef();
-			ImportedAsyncPackages.Add(ImportedAsyncPackage);
-			if (bInserted)
+			const FName ImportedPackageName = Entry.Name;
+			FAsyncPackageDesc Info(INDEX_NONE, ImportedPackageName);
+			Info.Priority = Desc.Priority;
+			bool bInserted;
+			FAsyncPackage2* ImportedAsyncPackage = AsyncLoadingThread.FindOrInsertPackage(&Info, bInserted);
+			if (ImportedAsyncPackage)
 			{
-				ImportedAsyncPackage->ImportPackagesRecursive();
-				ImportedAsyncPackage->StartLoading();
+				TRACE_LOADTIME_ASYNC_PACKAGE_IMPORT_DEPENDENCY(this, ImportedAsyncPackage);
+				ImportedAsyncPackage->AddRef();
+				ImportedAsyncPackages.Reserve(ImportedPackageCount);
+				ImportedAsyncPackages.Add(ImportedAsyncPackage);
+				if (bInserted)
+				{
+					ImportedAsyncPackage->ImportPackagesRecursive();
+					ImportedAsyncPackage->StartLoading();
+				}
 			}
 		}
 	}
@@ -2320,11 +2479,12 @@ void FAsyncPackage2::ImportPackagesRecursive()
 
 void FAsyncPackage2::StartLoading()
 {
+	TRACE_CPUPROFILER_EVENT_SCOPE(StartLoading);
 	TRACE_LOADTIME_BEGIN_LOAD_ASYNC_PACKAGE(this);
 	check(AsyncPackageLoadingState == EAsyncPackageLoadingState2::NewPackage);
 	AsyncPackageLoadingState = EAsyncPackageLoadingState2::WaitingForSummary;
 
-	LoadStartTime = FPlatformTime::Seconds();;
+	LoadStartTime = FPlatformTime::Seconds();
 
 	FIoReadOptions ReadOptions;
 	AsyncLoadingThread.IoDispatcher.ReadWithCallback(CreateIoChunkId(GlobalPackageId.Id, 0, EIoChunkType::PackageSummary),
@@ -2341,60 +2501,38 @@ void FAsyncPackage2::SetupScriptArcs()
 {
 	TRACE_CPUPROFILER_EVENT_SCOPE(SetupScriptArcs);
 
-	FPackageStore& GlobalPackageStore = AsyncLoadingThread.PackageStore;
+	FPackageStore& GlobalPackageStore = AsyncLoadingThread.GlobalPackageStore;
 	int32 ScriptArcsCount = 0;
 	int32* ScriptArcs = GlobalPackageStore.GetPackageScriptArcs(GlobalPackageId, ScriptArcsCount);
 	for (int32 ScriptArcIndex = 0; ScriptArcIndex < ScriptArcsCount * 2;)
 	{
 		int32 GlobalImportIndex = ScriptArcs[ScriptArcIndex++];
 		int32 ToNodeIndex = ScriptArcs[ScriptArcIndex++];
+		FEventLoadNode2* ToNode = PackageNodes + ToNodeIndex;
 
-		// skip packages
-		if (ImportStore.GlobalImportOuters[GlobalImportIndex].IsNull())
+		FPackageIndex OuterMostIndex = FPackageIndex::FromImport(GlobalImportIndex);
+		FPackageIndex OuterMostNonPackageIndex = OuterMostIndex;
+		while (true)
 		{
-			continue;
-		}
+			check(!OuterMostIndex.IsNull() && OuterMostIndex.IsImport());
+			FPackageIndex NextOuterMostIndex = ImportStore.GlobalImportScriptOuters[OuterMostIndex.ToImport()];
 
-		// find package of import object
-		FPackageIndex PackageIndex = ImportStore.GlobalImportPackages[GlobalImportIndex];
-		const FName& PackageName = ImportStore.GlobalImportNames[PackageIndex.ToImport()];
-		UObject*& Package = ImportStore.GlobalImportObjects[PackageIndex.ToImport()];
-		UPackage* ImportPackage = Package ? CastChecked<UPackage>(Package) : nullptr;
-		if (!ImportPackage)
-		{
-			ImportPackage = FindObjectFast<UPackage>(NULL, PackageName, false, false);
-			Package = ImportPackage; // This will write to global import table!!!
-		}
-		check(ImportPackage);
-
-		// do initial loading stuff for compiled in packages
-		if (ImportPackage->HasAnyPackageFlags(PKG_CompiledIn))
-		{
-			FPackageIndex OuterMostIndex = FPackageIndex::FromImport(GlobalImportIndex);
-			FPackageIndex OuterMostNonPackageIndex = OuterMostIndex;
-			while (true)
-			{
-				check(!OuterMostIndex.IsNull() && OuterMostIndex.IsImport());
-				FPackageIndex NextOuterMostIndex = ImportStore.GlobalImportOuters[OuterMostIndex.ToImport()];
-
-				if (NextOuterMostIndex.IsNull())
-				{
-					break;
-				}
-				OuterMostNonPackageIndex = OuterMostIndex;
-				OuterMostIndex = NextOuterMostIndex;
-			}
-			FName OuterMostNonPackageObjectName = ImportStore.GlobalImportNames[OuterMostNonPackageIndex.ToImport()];
-			check(ImportStore.GlobalImportOuters[OuterMostIndex.ToImport()].IsNull());
-			check(PackageName == ImportStore.GlobalImportNames[OuterMostIndex.ToImport()]);
-			// OuterMostNonPackageIndex is used here because if it is a CDO or subobject, etc,
-			// we wait for the outermost thing that is not a package
-			const bool bWaitingForCompiledInImport = EDLBootNotificationManager.AddWaitingPackage(
-				this, PackageName, OuterMostNonPackageObjectName, FPackageIndex::FromImport(ToNodeIndex));
-			if (bWaitingForCompiledInImport)
+			if (NextOuterMostIndex.IsNull())
 			{
-				PackageNodes[ToNodeIndex].AddBarrier();
+				break;
 			}
+			OuterMostNonPackageIndex = OuterMostIndex;
+			OuterMostIndex = NextOuterMostIndex;
+		}
+		FName PackageName = ImportStore.GlobalImportNames[OuterMostIndex.ToImport()];
+		FName OuterMostNonPackageObjectName = ImportStore.GlobalImportNames[OuterMostNonPackageIndex.ToImport()];
+		// OuterMostNonPackageIndex is used here because if it is a CDO or subobject, etc,
+		// we wait for the outermost thing that is not a package
+		const bool bWaitingForCompiledInImport = EDLBootNotificationManager.AddWaitingPackage(
+			this, PackageName, OuterMostNonPackageObjectName, FPackageIndex::FromImport(ToNodeIndex));
+		if (bWaitingForCompiledInImport)
+		{
+			PackageNodes[ToNodeIndex].AddBarrier();
 		}
 	}
 }
@@ -2473,7 +2611,7 @@ UObject* FAsyncPackage2::EventDrivenIndexToObject(FPackageIndex Index, bool bChe
 	}
 	else if (Index.IsImport())
 	{
-		Result = ImportStore.FindExistingSlimport(Index.ToImport());
+		Result = ImportStore.FindOrGetImportFromGlobalIndex(Index);
 		check(Result);
 	}
 #if DO_CHECK
@@ -2656,18 +2794,7 @@ void FAsyncPackage2::EventDrivenCreateExport(int32 LocalExportIndex)
 	}
 
 	check(Object);
-	if (Export.GlobalImportIndex != -1)
-	{
-		UObject*& StoreObject = ImportStore.GlobalImportObjects[Export.GlobalImportIndex];
-		if (!StoreObject)
-		{
-			StoreObject = Object;
-		}
-		else
-		{
-			check(StoreObject == Object);
-		}
-	}
+	ImportStore.StoreGlobalImportObject(Export.GlobalImportIndex, Object);
 }
 
 void FAsyncPackage2::EventDrivenSerializeExport(int32 LocalExportIndex, const uint8* ExportData, uint64 ExportDataSize)
@@ -2765,6 +2892,8 @@ EAsyncPackageState::Type FAsyncPackage2::Event_ExportsDone(FAsyncPackage2* Packa
 	TRACE_CPUPROFILER_EVENT_SCOPE(Event_ExportsDone);
 
 	check(Package->AsyncPackageLoadingState == EAsyncPackageLoadingState2::ProcessNewImportsAndExports);
+	Package->ImportStore.ImportMap = nullptr;
+	Package->ImportStore.ImportMapCount = 0;
 	Package->bAllExportsSerialized = true;
 	Package->AsyncPackageLoadingState = EAsyncPackageLoadingState2::PostLoad_Etc;
 
@@ -3231,26 +3360,14 @@ EAsyncPackageState::Type FAsyncLoadingThread2Impl::TickAsyncLoadingFromGameThrea
 					FDeferredMessageLog::Flush();
 				}
 			}
-			if (!bDidSomething)
+
+			if (!bDidSomething && GIsInitialLoad)
 			{
-				if (bIsMultithreaded)
-				{
-					if (GIsInitialLoad)
-					{
-						bDidSomething = EDLBootNotificationManager.ConstructWaitingBootObjects(); // with the ASL, we always create new boot objects when we have nothing else to do
-					}
-				}
-				else
-				{
-					if (GIsInitialLoad)
-					{
-						bDidSomething = EDLBootNotificationManager.FireCompletedCompiledInImports(); // no ASL, first try to fire any completed boot objects, and if there are none, then create some boot objects
-						if (!bDidSomething)
-						{
-							bDidSomething = EDLBootNotificationManager.ConstructWaitingBootObjects();
-						}
-					}
-				}
+				 bDidSomething = EDLBootNotificationManager.FireCompletedCompiledInImports();
+				 if (!bDidSomething)
+				 {
+				 	bDidSomething = EDLBootNotificationManager.ConstructWaitingBootObjects();
+				 }
 			}
 		}
 
@@ -3428,6 +3545,8 @@ uint32 FAsyncLoadingThread2Impl::Run()
 	FMemory::SetupTLSCachesOnCurrentThread();
 
 	FAsyncLoadingThreadState2& ThreadState = *FAsyncLoadingThreadState2::Get();
+	
+	GlobalPackageStore.FinalizeInitialLoad();
 
 	FZenaphoreWaiter Waiter(AltZenaphore, TEXT("WaitForEvents"));
 	bool bIsSuspended = false;
@@ -3448,7 +3567,7 @@ uint32 FAsyncLoadingThread2Impl::Run()
 		}
 		else
 		{
-			PackageStore.GetImportStore().bNeedToRunGarbageCollect |= IsAsyncLoadingPackages();
+			GlobalPackageStore.GetGlobalImportStore().bNeedToRunGarbageCollect |= IsAsyncLoadingPackages();
 			bool bDidSomething = false;
 			{
 				FGCScopeGuard GCGuard;
@@ -3515,10 +3634,6 @@ EAsyncPackageState::Type FAsyncLoadingThread2Impl::TickAsyncThreadFromGameThread
 	int32 ProcessedRequests = 0;
 	if (AsyncThreadReady.GetValue())
 	{
-		if (GIsInitialLoad)
-		{
-			EDLBootNotificationManager.FireCompletedCompiledInImports();
-		}
 		if (IsGarbageCollectionWaiting() || FAsyncLoadingThreadState2::Get()->IsTimeLimitExceeded())
 		{
 			Result = EAsyncPackageState::TimeOut;
@@ -3622,6 +3737,7 @@ static void VerifyLoadFlagsWhenFinishedLoading()
 	}
 	UE_LOG(LogStreaming, Log, TEXT("Verified load flags when finished loading"));
 }
+
 void FGlobalImportStore::OnPreGarbageCollect(bool bInIsLoadingPackages)
 {
 	if (!bNeedToRunGarbageCollect && !bInIsLoadingPackages)
@@ -3630,44 +3746,33 @@ void FGlobalImportStore::OnPreGarbageCollect(bool bInIsLoadingPackages)
 	}
 	bNeedToRunGarbageCollect = bInIsLoadingPackages;
 
-	int32 NumCleared = 0;
-	for (int32 GlobalImportIndex = 0; GlobalImportIndex < Count; ++GlobalImportIndex)
+	int32 NumWeak = 0;
+	for (int32 GlobalImportIndex = ScriptImportCount; GlobalImportIndex < Count; ++GlobalImportIndex)
 	{
-		UObject*& Object = Objects[GlobalImportIndex];
+		FGlobalImport& GlobalImport = Imports[GlobalImportIndex];
+		UObject* Object = GlobalImport.GetObjectIfRawPointer();
 		if (!Object)
 		{
 			continue;
 		}
 
-		// Import objects in packages currently being loaded already have the Async flag set.
-		// They will never be destroyed during GC, and the object pointers are safe to keep.
-		const bool bHasAsyncFlag = Object->HasAnyInternalFlags(EInternalObjectFlags::Async);
-		if (bHasAsyncFlag)
-		{
-			continue;
-		}
-
-		const int32 RefCount = RefCounts[GlobalImportIndex];
-		check(RefCount >= 0);
-
-		if (RefCount > 0)
+		if (GlobalImport.GetRefCount() > 0)
 		{
-			// Import objects in native packages will never be garbage collected and does not need marking.
-			FPackageIndex PackageIndex = Packages[GlobalImportIndex];
-			UPackage* Package = CastChecked<UPackage>(Objects[PackageIndex.ToImport()]);
-			if (Package->HasAnyPackageFlags(PKG_CompiledIn))
+			// Import objects in packages currently being loaded already have the Async flag set.
+			// They will never be destroyed during GC, and the object pointers are safe to keep.
+			if (!Object->HasAnyInternalFlags(EInternalObjectFlags::Async))
 			{
-				continue;
+				// Mark object to be kept alive during GC
+				Object->SetInternalFlags(EInternalObjectFlags::Async);
+				KeepAliveObjects.Add(Object);
 			}
-			// Mark object to be kept alive during GC
-			Object->SetInternalFlags(EInternalObjectFlags::Async);
-			KeepAliveObjects.Add(Object);
 		}
 		else
 		{
-			// Clear object pointer since object may get destroyed during GC
-			Object = nullptr;
-			++NumCleared;
+			// Convert object pointer to weak since object may get destroyed during GC
+			check(!Object->HasAnyInternalFlags(EInternalObjectFlags::Async));
+			GlobalImport.MakeWeak();
+			++NumWeak;
 		}
 	}
 
@@ -3681,15 +3786,15 @@ void FGlobalImportStore::OnPreGarbageCollect(bool bInIsLoadingPackages)
 	{
 		for (int32 GlobalImportIndex = 0; GlobalImportIndex < Count; ++GlobalImportIndex)
 		{
-			check(Objects[GlobalImportIndex] == nullptr);
-			check(RefCounts[GlobalImportIndex] == 0);
+			check(GlobalImportIndex < ScriptImportCount || Imports[GlobalImportIndex].GetObjectIfRawPointer() == nullptr);
+			check(Imports[GlobalImportIndex].GetRefCount() == 0);
 		}
 		VerifyLoadFlagsWhenFinishedLoading();
 	}
 #endif
 
-	UE_LOG(LogStreaming, Log, TEXT("FGlobalImportStore::OnPreGarbageCollect - Marked %d objects, cleared %d object pointers"),
-		KeepAliveObjects.Num(), NumCleared);
+	UE_LOG(LogStreaming, Display, TEXT("FGlobalImportStore::OnPreGarbageCollect - Marked %d objects to keep, made %d object pointers weak"),
+		KeepAliveObjects.Num(), NumWeak);
 }
 
 void FGlobalImportStore::OnPostGarbageCollect()
@@ -3715,13 +3820,13 @@ void FAsyncLoadingThread2Impl::OnPreGarbageCollect()
 {
 	TRACE_CPUPROFILER_EVENT_SCOPE(AltPreGC);
 	const bool bIsAsyncLoadingPackages = IsAsyncLoadingPackages();
-	PackageStore.GetImportStore().OnPreGarbageCollect(bIsAsyncLoadingPackages);
+	GlobalPackageStore.GetGlobalImportStore().OnPreGarbageCollect(bIsAsyncLoadingPackages);
 }
 
 void FAsyncLoadingThread2Impl::OnPostGarbageCollect()
 {
 	TRACE_CPUPROFILER_EVENT_SCOPE(AltPostGC);
-	PackageStore.GetImportStore().OnPostGarbageCollect();
+	GlobalPackageStore.GetGlobalImportStore().OnPostGarbageCollect();
 }
 
 /**
@@ -3776,7 +3881,7 @@ FAsyncPackage2::FAsyncPackage2(const FAsyncPackageDesc& InDesc, int32 InSerialNu
 , GraphAllocator(InGraphAllocator)
 , GlobalPackageId(InGlobalPackageId)
 // Begin EDL specific properties
-, ImportStore(AsyncLoadingThread.PackageStore, InGlobalPackageId)
+, ImportStore(AsyncLoadingThread.GlobalPackageStore, InGlobalPackageId)
 , AsyncPackageLoadingState(EAsyncPackageLoadingState2::NewPackage)
 , SerialNumber(InSerialNumber)
 , bAllExportsSerialized(false)
@@ -3785,8 +3890,8 @@ FAsyncPackage2::FAsyncPackage2(const FAsyncPackageDesc& InDesc, int32 InSerialNu
 	TRACE_LOADTIME_NEW_ASYNC_PACKAGE(this, InDesc.Name);
 	AddRequestID(InDesc.RequestID);
 
-	ExportBundleCount = AsyncLoadingThread.PackageStore.GetPackageExportBundleCount(GlobalPackageId);
-	ExportCount = AsyncLoadingThread.PackageStore.GetPackageExportCount(GlobalPackageId);
+	ExportBundleCount = AsyncLoadingThread.GlobalPackageStore.GetPackageExportBundleCount(GlobalPackageId);
+	ExportCount = AsyncLoadingThread.GlobalPackageStore.GetPackageExportCount(GlobalPackageId);
 	Exports.AddZeroed(ExportCount);
 
 	ExportBundleIoBuffers.SetNum(ExportBundleCount);
@@ -3850,6 +3955,36 @@ FAsyncPackage2::~FAsyncPackage2()
 	check(OwnedObjects.Num() == 0);
 }
 
+void FAsyncPackage2::ClearImportedPackages()
+{
+	TRACE_CPUPROFILER_EVENT_SCOPE(ClearImportedPackages);
+	// TODO: Use AsyncPackage pointer from global id, namelookup entry has already been removed
+#if 0
+	FPackageStore& GlobalPackageStore = AsyncLoadingThread.GlobalPackageStore;
+	int32 ImportedPackageCount = 0;
+	int32* Imports = GlobalPackageStore.GetPackageImportedPackages(GlobalPackageId, ImportedPackageCount);
+	TArray<FAsyncPackage2*> TempImportedAsyncPackages;
+	for (int32 LocalImportIndex = 0; LocalImportIndex < ImportedPackageCount; ++LocalImportIndex)
+	{
+		int32 EntryIndex = Imports[LocalImportIndex];
+		FPackageStoreEntry& Entry = GlobalPackageStore.StoreEntries[EntryIndex];
+		FAsyncPackage2* ImportedAsyncPackage = AsyncLoadingThread.FindAsyncPackage(Entry.Name);
+		TempImportedAsyncPackages.Add(ImportedAsyncPackage);
+		if (ImportedAsyncPackage)
+		{
+			// ImportedAsyncPackage->ReleaseRef();
+		}
+	}
+	ensure(TempImportedAsyncPackages == ImportedAsyncPackages);
+#else
+	for (FAsyncPackage2* ImportedAsyncPackage : ImportedAsyncPackages)
+	{
+		ImportedAsyncPackage->ReleaseRef();
+	}
+	ImportedAsyncPackages.Empty();
+#endif
+	ImportStore.ClearReferences();
+}
 
 void FAsyncPackage2::ClearOwnedObjects()
 {
@@ -3988,7 +4123,7 @@ EAsyncPackageState::Type FAsyncPackage2::CreateLinker()
 
 	{
 		uint32 LinkerFlags = LOAD_None | LOAD_Async | LOAD_NoVerify;
-		Linker = new FLinkerLoad(Package, *AsyncLoadingThread.PackageStore.GetPackageFileName(GlobalPackageId), LinkerFlags);
+		Linker = new FLinkerLoad(Package, *AsyncLoadingThread.GlobalPackageStore.GetPackageFileName(GlobalPackageId), LinkerFlags);
 		Linker->bIsAsyncLoader = false;
 		Linker->bLockoutLegacyOperations = true;
 		Linker->SetIsLoading(true);
@@ -4013,17 +4148,15 @@ EAsyncPackageState::Type FAsyncPackage2::FinishLinker()
 	const uint8* PackageSummaryData = PackageSummaryIoBuffer.Data();
 	const FPackageSummary* PackageSummary = reinterpret_cast<const FPackageSummary*>(PackageSummaryData);
 	
-	// TODO: Remove support for old global namemap
-	PackageNameMap = (PackageSummary->NameMapOffset == PackageSummary->ExportMapOffset) ? 
-		nullptr :
-		reinterpret_cast<const int32*>(PackageSummaryData + PackageSummary->NameMapOffset);
-
+	PackageNameMap = reinterpret_cast<const int32*>(PackageSummaryData + PackageSummary->NameMapOffset);
+	ImportStore.ImportMap = reinterpret_cast<const int32*>(PackageSummaryData + PackageSummary->ImportMapOffset);
+	ImportStore.ImportMapCount = (PackageSummary->ExportMapOffset - PackageSummary->ImportMapOffset) / sizeof(int32);
 	ExportMap = reinterpret_cast<const FExportMapEntry*>(PackageSummaryData + PackageSummary->ExportMapOffset);
 	ExportBundles = reinterpret_cast<const FExportBundleHeader*>(PackageSummaryData + PackageSummary->ExportBundlesOffset);
 	ExportBundleEntries = reinterpret_cast<const FExportBundleEntry*>(ExportBundles + ExportBundleCount);
 
 
-	TRACE_LOADTIME_PACKAGE_SUMMARY(this, PackageSummarySize, ImportStore.Count, ExportCount);
+	TRACE_LOADTIME_PACKAGE_SUMMARY(this, PackageSummarySize, ImportStore.ImportMapCount, ExportCount);
 
 	FPackageFileSummary& Summary	= Linker->Summary;
 	Summary.PackageFlags			= PackageSummary->PackageFlags;
@@ -4031,6 +4164,7 @@ EAsyncPackageState::Type FAsyncPackage2::FinishLinker()
 	Summary.SetFileVersions(GPackageFileUE4Version, GPackageFileLicenseeUE4Version, /*unversioned*/true);
 	Linker->UpdateFromPackageFileSummary();
 
+	ImportStore.StoreGlobalImportObject(PackageSummary->GlobalImportIndex, LinkerRoot);
 	return EAsyncPackageState::Complete;
 }
 
