BulkData streaming will now return a valid pointer as a parameter in it's callbacks rather than nullptr, allowing systems that were actually using it to function.

- Add a new callback type FBulkDataIORequestCallBack to replace FAsyncFileCallBack when making streaming requests from BulkData objects.
 - This callback will provide a pointer to the IBulkDataIORequest interface rather than IAsyncReadRequest and so will be valid when running with the IoDispatcher enabled.
- Change all BulkData streaming methods (old and new) to take FBulkDataIORequestCallBack instead.
- Update all existing systems to use this.
- When using the old file system we need to convert from the new callback to the old callback, this is encapsulated in FBulkDataIORequest::MakeReadRequest and does add a small amount of overhead.

#rnx

diff --git a/Engine/Plugins/Experimental/GeometryCache/Source/GeometryCache/Private/StreamingGeometryCacheData.cpp b/Engine/Plugins/Experimental/GeometryCache/Source/GeometryCache/Private/StreamingGeometryCacheData.cpp
--- a/Engine/Plugins/Experimental/GeometryCache/Source/GeometryCache/Private/StreamingGeometryCacheData.cpp
+++ b/Engine/Plugins/Experimental/GeometryCache/Source/GeometryCache/Private/StreamingGeometryCacheData.cpp
@@ -91,7 +91,7 @@ void FStreamingGeometryCacheData::RemoveResidentChunk(FResidentChunk& LoadedChun
 /**
 This is called from some random thread when reading is complete.
 */
-void FStreamingGeometryCacheData::OnAsyncReadComplete(int32 LoadedChunkIndex, IAsyncReadRequest* ReadRequest)
+void FStreamingGeometryCacheData::OnAsyncReadComplete(int32 LoadedChunkIndex, IBulkDataIORequest* ReadRequest)
 {
 	// We should do the least amount of work possible here as to not stall the async io threads.
 	// We also cannot take the critical section here as this would lead to a deadlock between the
@@ -228,7 +228,7 @@ void FStreamingGeometryCacheData::UpdateStreamingStatus()
 				// Kick of a load
 				check(Chunk.BulkData.GetBulkDataSize() == ResidentChunk.DataSize);
 
-				FAsyncFileCallBack AsyncFileCallBack = [this, NeededIndex](bool bWasCancelled, IAsyncReadRequest* Req)
+				FBulkDataIORequestCallBack AsyncFileCallBack = [this, NeededIndex](bool bWasCancelled, IBulkDataIORequest* Req)
 				{
 					this->OnAsyncReadComplete(NeededIndex, Req);
 				};
diff --git a/Engine/Plugins/Experimental/GeometryCache/Source/GeometryCache/Private/StreamingGeometryCacheData.h b/Engine/Plugins/Experimental/GeometryCache/Source/GeometryCache/Private/StreamingGeometryCacheData.h
--- a/Engine/Plugins/Experimental/GeometryCache/Source/GeometryCache/Private/StreamingGeometryCacheData.h
+++ b/Engine/Plugins/Experimental/GeometryCache/Source/GeometryCache/Private/StreamingGeometryCacheData.h
@@ -29,11 +29,11 @@ struct FResidentChunk
 */
 struct FCompletedChunk
 {
-	IAsyncReadRequest* ReadRequest;
+	IBulkDataIORequest* ReadRequest;
 	int32 LoadedChunkIndex;
 
 	FCompletedChunk() : ReadRequest(nullptr), LoadedChunkIndex(0) {}
-	FCompletedChunk(int32 SetLoadedChunkIndex, IAsyncReadRequest* SetReadRequest) : ReadRequest(SetReadRequest), LoadedChunkIndex(SetLoadedChunkIndex) {}
+	FCompletedChunk(int32 SetLoadedChunkIndex, IBulkDataIORequest* SetReadRequest) : ReadRequest(SetReadRequest), LoadedChunkIndex(SetLoadedChunkIndex) {}
 };
 
 /**
@@ -63,7 +63,7 @@ private:
 
 	FResidentChunk &AddResidentChunk(int32 ChunkId, const FStreamedGeometryCacheChunk &ChunkInfo);
 	void RemoveResidentChunk(FResidentChunk& LoadedChunk);
-	void OnAsyncReadComplete(int32 LoadedChunkIndex, IAsyncReadRequest* ReadRequest);
+	void OnAsyncReadComplete(int32 LoadedChunkIndex, IBulkDataIORequest* ReadRequest);
 	void ProcessCompletedChunks();
 
 	// The track we are associated with
diff --git a/Engine/Source/Runtime/CoreUObject/Private/Serialization/BulkData.cpp b/Engine/Source/Runtime/CoreUObject/Private/Serialization/BulkData.cpp
--- a/Engine/Source/Runtime/CoreUObject/Private/Serialization/BulkData.cpp
+++ b/Engine/Source/Runtime/CoreUObject/Private/Serialization/BulkData.cpp
@@ -1594,14 +1594,14 @@ void FUntypedBulkData::SerializeBulkData( FArchive& Ar, void* Data )
 	}
 }
 
-IBulkDataIORequest* FUntypedBulkData::CreateStreamingRequest(EAsyncIOPriorityAndFlags Priority, FAsyncFileCallBack* CompleteCallback, uint8* UserSuppliedMemory) const
+IBulkDataIORequest* FUntypedBulkData::CreateStreamingRequest(EAsyncIOPriorityAndFlags Priority, FBulkDataIORequestCallBack* CompleteCallback, uint8* UserSuppliedMemory) const
 {
 	const int64 DataSize = GetBulkDataSize();
 
 	return CreateStreamingRequest(0, DataSize, Priority, CompleteCallback, UserSuppliedMemory);
 }
 
-IBulkDataIORequest* FUntypedBulkData::CreateStreamingRequest(int64 OffsetInBulkData, int64 BytesToRead, EAsyncIOPriorityAndFlags Priority, FAsyncFileCallBack* CompleteCallback, uint8* UserSuppliedMemory) const
+IBulkDataIORequest* FUntypedBulkData::CreateStreamingRequest(int64 OffsetInBulkData, int64 BytesToRead, EAsyncIOPriorityAndFlags Priority, FBulkDataIORequestCallBack* CompleteCallback, uint8* UserSuppliedMemory) const
 {
 	check(Filename.IsEmpty() == false);
 
@@ -1633,14 +1633,15 @@ IBulkDataIORequest* FUntypedBulkData::CreateStreamingRequest(int64 OffsetInBulkD
 
 	const int64 OffsetInFile = AdjustedBulkDataOffsetInFile + OffsetInBulkData;
 
-	IAsyncReadRequest* ReadRequest = IORequestHandle->ReadRequest(OffsetInFile, BytesToRead, Priority, CompleteCallback, UserSuppliedMemory);
-	if (ReadRequest != nullptr)
+	FBulkDataIORequest* IORequest = new FBulkDataIORequest(IORequestHandle);
+
+	if (IORequest->MakeReadRequest(OffsetInFile, BytesToRead, Priority, CompleteCallback, UserSuppliedMemory))
 	{
-		return new FBulkDataIORequest(IORequestHandle, ReadRequest, BytesToRead);
+		return IORequest;
 	}
 	else
 	{
-		delete IORequestHandle;
+		delete IORequest;
 		return nullptr;
 	}
 }
@@ -1658,7 +1659,7 @@ FBulkDataStreamingToken FUntypedBulkData::CreateStreamingToken() const
 	return FBulkDataStreamingToken((uint32)GetBulkDataOffsetInFile(), (uint32)GetBulkDataSize());
 }
 
-IBulkDataIORequest* FUntypedBulkData::CreateStreamingRequestForRange(const FString& Filename, const BulkDataRangeArray& RangeArray, EAsyncIOPriorityAndFlags Priority, FAsyncFileCallBack* CompleteCallback)
+IBulkDataIORequest* FUntypedBulkData::CreateStreamingRequestForRange(const FString& Filename, const BulkDataRangeArray& RangeArray, EAsyncIOPriorityAndFlags Priority, FBulkDataIORequestCallBack* CompleteCallback)
 {
 	check(Filename.IsEmpty() == false);
 	check(RangeArray.Num() > 0);
@@ -1682,14 +1683,15 @@ IBulkDataIORequest* FUntypedBulkData::CreateStreamingRequestForRange(const FStri
 
 	check(ReadSize > 0);
 
-	IAsyncReadRequest* ReadRequest = IORequestHandle->ReadRequest(ReadOffset, ReadSize, Priority, CompleteCallback, nullptr);
-	if (ReadRequest != nullptr)
+	FBulkDataIORequest* IORequest = new FBulkDataIORequest(IORequestHandle);
+
+	if (IORequest->MakeReadRequest(ReadOffset, ReadSize, Priority, CompleteCallback, nullptr))
 	{
-		return new FBulkDataIORequest(IORequestHandle, ReadRequest, ReadSize);
+		return IORequest;
 	}
 	else
 	{
-		delete IORequestHandle;
+		delete IORequest;
 		return nullptr;
 	}
 }
@@ -1702,6 +1704,30 @@ FBulkDataIORequest::~FBulkDataIORequest()
 	delete FileHandle;
 }
 
+bool FBulkDataIORequest::MakeReadRequest(int64 Offset, int64 BytesToRead, EAsyncIOPriorityAndFlags PriorityAndFlags, FBulkDataIORequestCallBack* CompleteCallback, uint8* UserSuppliedMemory)
+{
+	check(ReadRequest == nullptr);
+	check(Size == 0);
+
+	FBulkDataIORequestCallBack LocalCallback = *CompleteCallback;
+	FAsyncFileCallBack AsyncFileCallBack = [LocalCallback, this](bool bWasCancelled, IAsyncReadRequest*)
+	{
+		LocalCallback(bWasCancelled, this);
+	};
+
+	ReadRequest = FileHandle->ReadRequest(Offset, BytesToRead, PriorityAndFlags, &AsyncFileCallBack, UserSuppliedMemory);
+	
+	if (ReadRequest != nullptr)
+	{
+		Size = BytesToRead;
+		return true;
+	}
+	else
+	{
+		return false;
+	}
+}
+
 bool FBulkDataIORequest::PollCompletion() const
 {
 	return ReadRequest->PollCompletion();
diff --git a/Engine/Source/Runtime/CoreUObject/Private/Serialization/BulkData2.cpp b/Engine/Source/Runtime/CoreUObject/Private/Serialization/BulkData2.cpp
--- a/Engine/Source/Runtime/CoreUObject/Private/Serialization/BulkData2.cpp
+++ b/Engine/Source/Runtime/CoreUObject/Private/Serialization/BulkData2.cpp
@@ -219,7 +219,7 @@ public:
 	using ChunkIdArray = TArray<FIoChunkId, TInlineAllocator<8>>;
 
 public:
-	FBulkDataIoDispatcherRequest(const FIoChunkId& InChunkID, int64 InOffsetInBulkData, int64 InBytesToRead, FAsyncFileCallBack* InCompleteCallback, uint8* InUserSuppliedMemory)
+	FBulkDataIoDispatcherRequest(const FIoChunkId& InChunkID, int64 InOffsetInBulkData, int64 InBytesToRead, FBulkDataIORequestCallBack* InCompleteCallback, uint8* InUserSuppliedMemory)
 		: UserSuppliedMemory(InUserSuppliedMemory)
 	{
 		RequestArray.Push({ InChunkID, (uint64)InOffsetInBulkData , (uint64)InBytesToRead });
@@ -230,7 +230,7 @@ public:
 		}
 	}
 
-	FBulkDataIoDispatcherRequest(const ChunkIdArray& ChunkIDs, FAsyncFileCallBack* InCompleteCallback)
+	FBulkDataIoDispatcherRequest(const ChunkIdArray& ChunkIDs, FBulkDataIORequestCallBack* InCompleteCallback)
 		: UserSuppliedMemory(nullptr)
 	{
 		for (const FIoChunkId& ChunkId : ChunkIDs)
@@ -390,7 +390,7 @@ private:
 
 		if (CompleteCallback)
 		{
-			CompleteCallback(bIsCanceled, nullptr);
+			CompleteCallback(bIsCanceled, this);
 		}
 
 		FPlatformMisc::MemoryBarrier();
@@ -454,7 +454,7 @@ private:
 
 	TArray<Request, TInlineAllocator<8>> RequestArray;
 
-	FAsyncFileCallBack CompleteCallback;
+	FBulkDataIORequestCallBack CompleteCallback;
 	uint8* UserSuppliedMemory = nullptr;
 
 	uint8* DataResult = nullptr;
@@ -890,14 +890,14 @@ bool FBulkDataBase::IsUsingIODispatcher() const
 	return (BulkDataFlags & BULKDATA_UsesIoDispatcher) != 0;
 }
 
-IBulkDataIORequest* FBulkDataBase::CreateStreamingRequest(EAsyncIOPriorityAndFlags Priority, FAsyncFileCallBack* CompleteCallback, uint8* UserSuppliedMemory) const
+IBulkDataIORequest* FBulkDataBase::CreateStreamingRequest(EAsyncIOPriorityAndFlags Priority, FBulkDataIORequestCallBack* CompleteCallback, uint8* UserSuppliedMemory) const
 {
 	const int64 DataSize = GetBulkDataSize();
 
 	return CreateStreamingRequest(0, DataSize, Priority, CompleteCallback, UserSuppliedMemory);
 }
 
-IBulkDataIORequest* FBulkDataBase::CreateStreamingRequest(int64 OffsetInBulkData, int64 BytesToRead, EAsyncIOPriorityAndFlags Priority, FAsyncFileCallBack* CompleteCallback, uint8* UserSuppliedMemory) const
+IBulkDataIORequest* FBulkDataBase::CreateStreamingRequest(int64 OffsetInBulkData, int64 BytesToRead, EAsyncIOPriorityAndFlags Priority, FBulkDataIORequestCallBack* CompleteCallback, uint8* UserSuppliedMemory) const
 {
 	if (IsUsingIODispatcher())
 	{
@@ -927,20 +927,21 @@ IBulkDataIORequest* FBulkDataBase::CreateStreamingRequest(int64 OffsetInBulkData
 
 		const int64 OffsetInFile = FileData.BulkDataOffsetInFile + OffsetInBulkData;
 
-		IAsyncReadRequest* ReadRequest = IORequestHandle->ReadRequest(OffsetInFile, BytesToRead, Priority, CompleteCallback, UserSuppliedMemory);
-		if (ReadRequest != nullptr)
+		FBulkDataIORequest* IORequest = new FBulkDataIORequest(IORequestHandle);
+
+		if (IORequest->MakeReadRequest(OffsetInFile, BytesToRead, Priority, CompleteCallback, UserSuppliedMemory))
 		{
-			return new FBulkDataIORequest(IORequestHandle, ReadRequest, BytesToRead);
+			return IORequest;
 		}
 		else
 		{
-			delete IORequestHandle;
+			delete IORequest;
 			return nullptr;
 		}
 	}
 }
 
-IBulkDataIORequest* FBulkDataBase::CreateStreamingRequestForRange(const BulkDataRangeArray& RangeArray, EAsyncIOPriorityAndFlags Priority, FAsyncFileCallBack* CompleteCallback)
+IBulkDataIORequest* FBulkDataBase::CreateStreamingRequestForRange(const BulkDataRangeArray& RangeArray, EAsyncIOPriorityAndFlags Priority, FBulkDataIORequestCallBack* CompleteCallback)
 {
 	check(RangeArray.Num() > 0);
 
diff --git a/Engine/Source/Runtime/CoreUObject/Public/Serialization/BulkData.h b/Engine/Source/Runtime/CoreUObject/Public/Serialization/BulkData.h
--- a/Engine/Source/Runtime/CoreUObject/Public/Serialization/BulkData.h
+++ b/Engine/Source/Runtime/CoreUObject/Public/Serialization/BulkData.h
@@ -162,8 +162,17 @@ public:
 
 	}
 
+	FBulkDataIORequest(IAsyncReadFileHandle* InFileHandle)
+		: FileHandle(InFileHandle)
+		, ReadRequest(nullptr)
+		, Size(0)
+	{
+	}
+
 	virtual ~FBulkDataIORequest();
 
+	bool MakeReadRequest(int64 Offset, int64 BytesToRead, EAsyncIOPriorityAndFlags PriorityAndFlags, FBulkDataIORequestCallBack* CompleteCallback, uint8* UserSuppliedMemory);
+
 	virtual bool PollCompletion() const override;
 	virtual bool WaitCompletion( float TimeLimitSeconds = 0.0f ) const override;
 
@@ -692,7 +701,7 @@ public:
 	* @param UserSuppliedMemory A pointer to memory for the IO request to be written to, it is up to the caller to make sure that it is large enough. If the pointer is null then the system will allocate memory instead.
 	* @return					A request for the read. This is owned by the caller and must be deleted by the caller.
 	**/
-	IBulkDataIORequest* CreateStreamingRequest(EAsyncIOPriorityAndFlags Priority, FAsyncFileCallBack* CompleteCallback, uint8* UserSuppliedMemory) const;
+	IBulkDataIORequest* CreateStreamingRequest(EAsyncIOPriorityAndFlags Priority, FBulkDataIORequestCallBack* CompleteCallback, uint8* UserSuppliedMemory) const;
 
 	/**
 	* Create an async read request for the bulk data.
@@ -705,7 +714,7 @@ public:
 	* @param UserSuppliedMemory A pointer to memory for the IO request to be written to, it is up to the caller to make sure that it is large enough. If the pointer is null then the system will allocate memory instead.
 	* @return					A request for the read. This is owned by the caller and must be deleted by the caller.
 	**/
-	IBulkDataIORequest* CreateStreamingRequest(int64 OffsetInBulkData, int64 BytesToRead, EAsyncIOPriorityAndFlags Priority, FAsyncFileCallBack* CompleteCallback, uint8* UserSuppliedMemory) const;
+	IBulkDataIORequest* CreateStreamingRequest(int64 OffsetInBulkData, int64 BytesToRead, EAsyncIOPriorityAndFlags Priority, FBulkDataIORequestCallBack* CompleteCallback, uint8* UserSuppliedMemory) const;
 
 #if USE_BULKDATA_STREAMING_TOKEN 
 
@@ -730,7 +739,7 @@ public:
 	* @param CompleteCallback	Called from an arbitrary thread when the request is complete. Can be nullptr, if non-null, must remain valid until it is called. It will always be called.
 	* @return					A request for the read. This is owned by the caller and must be deleted by the caller.
 	**/
-	static IBulkDataIORequest* CreateStreamingRequestForRange(const FString& Filename, const BulkDataRangeArray& RangeArray, EAsyncIOPriorityAndFlags Priority, FAsyncFileCallBack* CompleteCallback);
+	static IBulkDataIORequest* CreateStreamingRequestForRange(const FString& Filename, const BulkDataRangeArray& RangeArray, EAsyncIOPriorityAndFlags Priority, FBulkDataIORequestCallBack* CompleteCallback);
 #endif
 
 	/*-----------------------------------------------------------------------------
diff --git a/Engine/Source/Runtime/CoreUObject/Public/Serialization/BulkData2.h b/Engine/Source/Runtime/CoreUObject/Public/Serialization/BulkData2.h
--- a/Engine/Source/Runtime/CoreUObject/Public/Serialization/BulkData2.h
+++ b/Engine/Source/Runtime/CoreUObject/Public/Serialization/BulkData2.h
@@ -21,8 +21,6 @@ struct FOwnedBulkDataPtr;
  *
  * It functions pretty much the same as IAsyncReadRequest expect that it also holds
  * the file handle as well.
- *
- *
  */
 class COREUOBJECT_API IBulkDataIORequest
 {
@@ -38,6 +36,11 @@ public:
 	virtual void Cancel() = 0;
 };
 
+/**
+ * Callback to use when making streaming requests
+ */
+typedef TFunction<void(bool bWasCancelled, IBulkDataIORequest*)> FBulkDataIORequestCallBack;
+
 /**
  * @documentation @todo documentation
  */
@@ -113,10 +116,10 @@ public:
 	bool IsMemoryMapped() const;
 	bool IsUsingIODispatcher() const;
 
-	IBulkDataIORequest* CreateStreamingRequest(EAsyncIOPriorityAndFlags Priority, FAsyncFileCallBack* CompleteCallback, uint8* UserSuppliedMemory) const;
-	IBulkDataIORequest* CreateStreamingRequest(int64 OffsetInBulkData, int64 BytesToRead, EAsyncIOPriorityAndFlags Priority, FAsyncFileCallBack* CompleteCallback, uint8* UserSuppliedMemory) const;
+	IBulkDataIORequest* CreateStreamingRequest(EAsyncIOPriorityAndFlags Priority, FBulkDataIORequestCallBack* CompleteCallback, uint8* UserSuppliedMemory) const;
+	IBulkDataIORequest* CreateStreamingRequest(int64 OffsetInBulkData, int64 BytesToRead, EAsyncIOPriorityAndFlags Priority, FBulkDataIORequestCallBack* CompleteCallback, uint8* UserSuppliedMemory) const;
 	
-	static IBulkDataIORequest* CreateStreamingRequestForRange(const BulkDataRangeArray& RangeArray, EAsyncIOPriorityAndFlags Priority, FAsyncFileCallBack* CompleteCallback);
+	static IBulkDataIORequest* CreateStreamingRequestForRange(const BulkDataRangeArray& RangeArray, EAsyncIOPriorityAndFlags Priority, FBulkDataIORequestCallBack* CompleteCallback);
 
 	void RemoveBulkData();
 
@@ -176,6 +179,8 @@ class COREUOBJECT_API FUntypedBulkData2 : public FBulkDataBase
 	// used or not.
 	static_assert(TIsPODType<ElementType>::Value, "FUntypedBulkData2 is limited to POD types!");
 public:
+	FUntypedBulkData2() = default;
+
 	void Serialize(FArchive& Ar, UObject* Owner, int32 Index, bool bAttemptFileMapping)
 	{
 		FBulkDataBase::Serialize(Ar, Owner, Index, bAttemptFileMapping, sizeof(ElementType));
diff --git a/Engine/Source/Runtime/Engine/Private/Animation/AnimationStreaming.cpp b/Engine/Source/Runtime/Engine/Private/Animation/AnimationStreaming.cpp
--- a/Engine/Source/Runtime/Engine/Private/Animation/AnimationStreaming.cpp
+++ b/Engine/Source/Runtime/Engine/Private/Animation/AnimationStreaming.cpp
@@ -215,8 +215,8 @@ void FStreamingAnimationData::BeginPendingRequests(const TArray<uint32>& Indices
 			int64 ChunkSize = Chunk.BulkData.GetBulkDataSize();
 			ChunkStorage.RequestStart = FPlatformTime::Seconds();
 			UE_LOG(LogAnimation, Warning, TEXT("Anim Streaming Request Started %s Chunk:%i At:%.3f\n"), *StreamableAnim->GetName(), ChunkIndex, ChunkStorage.RequestStart);
-			FAsyncFileCallBack AsyncFileCallBack =
-				[this, ChunkIndex, ChunkSize](bool bWasCancelled, IAsyncReadRequest* Req)
+			FBulkDataIORequestCallBack AsyncFileCallBack =
+				[this, ChunkIndex, ChunkSize](bool bWasCancelled, IBulkDataIORequest* Req)
 			{
 				AnimationStreamingManager->OnAsyncFileCallback(this, ChunkIndex, ChunkSize, Req, bWasCancelled);
 			};
@@ -325,7 +325,7 @@ FAnimationStreamingManager::~FAnimationStreamingManager()
 {
 }
 
-void FAnimationStreamingManager::OnAsyncFileCallback(FStreamingAnimationData* StreamingAnimData, int32 ChunkIndex, int64 ReadSize, IAsyncReadRequest* ReadRequest, bool bWasCancelled)
+void FAnimationStreamingManager::OnAsyncFileCallback(FStreamingAnimationData* StreamingAnimData, int32 ChunkIndex, int64 ReadSize, IBulkDataIORequest* ReadRequest, bool bWasCancelled)
 {
 	// Check to see if we successfully managed to load anything
 	uint8* Mem = ReadRequest->GetReadResults();
diff --git a/Engine/Source/Runtime/Engine/Private/AudioStreaming.cpp b/Engine/Source/Runtime/Engine/Private/AudioStreaming.cpp
--- a/Engine/Source/Runtime/Engine/Private/AudioStreaming.cpp
+++ b/Engine/Source/Runtime/Engine/Private/AudioStreaming.cpp
@@ -355,8 +355,8 @@ void FStreamingWaveData::BeginPendingRequests(const TArray<uint32>& IndicesToLoa
 				check(!ChunkStorage->Data); // Make sure we do not already have data
 				check(Chunk.BulkData.GetBulkDataSize() == ChunkStorage->DataSize); // Make sure that the bulkdata size matches
 				
-				FAsyncFileCallBack AsyncFileCallBack =
-					[this, LoadedChunkStorageIndex](bool bWasCancelled, IAsyncReadRequest* Req)
+				FBulkDataIORequestCallBack AsyncFileCallBack =
+					[this, LoadedChunkStorageIndex](bool bWasCancelled, IBulkDataIORequest* Req)
 				{
 					AudioStreamingManager->OnAsyncFileCallback(this, LoadedChunkStorageIndex, Req);
 
@@ -491,7 +491,7 @@ FLegacyAudioStreamingManager::~FLegacyAudioStreamingManager()
 {
 }
 
-void FLegacyAudioStreamingManager::OnAsyncFileCallback(FStreamingWaveData* StreamingWaveData, int32 LoadedAudioChunkIndex, IAsyncReadRequest* ReadRequest)
+void FLegacyAudioStreamingManager::OnAsyncFileCallback(FStreamingWaveData* StreamingWaveData, int32 LoadedAudioChunkIndex, IBulkDataIORequest* ReadRequest)
 {
 	// Check to see if we successfully managed to load anything
 	uint8* Mem = ReadRequest->GetReadResults();
diff --git a/Engine/Source/Runtime/Engine/Private/AudioStreamingCache.cpp b/Engine/Source/Runtime/Engine/Private/AudioStreamingCache.cpp
--- a/Engine/Source/Runtime/Engine/Private/AudioStreamingCache.cpp
+++ b/Engine/Source/Runtime/Engine/Private/AudioStreamingCache.cpp
@@ -1002,7 +1002,7 @@ void FAudioChunkCache::KickOffAsyncLoad(FCacheElement* CacheElement, const FChun
 		
 		NumberOfLoadsInFlight.Increment();
 
-		FAsyncFileCallBack AsyncFileCallBack = [this, OnLoadCompleted, CacheElement, InKey, ChunkDataSize](bool bWasCancelled, IAsyncReadRequest* Request)
+		FBulkDataIORequestCallBack AsyncFileCallBack = [this, OnLoadCompleted, CacheElement, InKey, ChunkDataSize](bool bWasCancelled, IBulkDataIORequest*)
 		{
 			// Populate key and DataSize. The async read request was set up to write directly into CacheElement->ChunkData.
 			CacheElement->Key = InKey;
diff --git a/Engine/Source/Runtime/Engine/Private/Streaming/SkeletalMeshUpdate.cpp b/Engine/Source/Runtime/Engine/Private/Streaming/SkeletalMeshUpdate.cpp
--- a/Engine/Source/Runtime/Engine/Private/Streaming/SkeletalMeshUpdate.cpp
+++ b/Engine/Source/Runtime/Engine/Private/Streaming/SkeletalMeshUpdate.cpp
@@ -362,7 +362,7 @@ FString FSkeletalMeshStreamIn_IO::GetIOFilename(const FContext& Context)
 
 void FSkeletalMeshStreamIn_IO::SetAsyncFileCallback(const FContext& Context)
 {
-	AsyncFileCallback = [this, Context](bool bWasCancelled, IAsyncReadRequest* Req)
+	AsyncFileCallback = [this, Context](bool bWasCancelled, IBulkDataIORequest*)
 	{
 		// At this point task synchronization would hold the number of pending requests.
 		TaskSynchronization.Decrement();
diff --git a/Engine/Source/Runtime/Engine/Private/Streaming/SkeletalMeshUpdate.h b/Engine/Source/Runtime/Engine/Private/Streaming/SkeletalMeshUpdate.h
--- a/Engine/Source/Runtime/Engine/Private/Streaming/SkeletalMeshUpdate.h
+++ b/Engine/Source/Runtime/Engine/Private/Streaming/SkeletalMeshUpdate.h
@@ -193,7 +193,7 @@ protected:
 	void CancelIORequest();
 
 	class IBulkDataIORequest* IORequest;
-	FAsyncFileCallBack AsyncFileCallback;
+	FBulkDataIORequestCallBack AsyncFileCallback;
 	bool bHighPrioIORequest;
 };
 
diff --git a/Engine/Source/Runtime/Engine/Private/Streaming/StaticMeshUpdate.cpp b/Engine/Source/Runtime/Engine/Private/Streaming/StaticMeshUpdate.cpp
--- a/Engine/Source/Runtime/Engine/Private/Streaming/StaticMeshUpdate.cpp
+++ b/Engine/Source/Runtime/Engine/Private/Streaming/StaticMeshUpdate.cpp
@@ -341,7 +341,7 @@ FString FStaticMeshStreamIn_IO::GetIOFilename(const FContext& Context)
 
 void FStaticMeshStreamIn_IO::SetAsyncFileCallback(const FContext& Context)
 {
-	AsyncFileCallback = [this, Context](bool bWasCancelled, IAsyncReadRequest* Req)
+	AsyncFileCallback = [this, Context](bool bWasCancelled, IBulkDataIORequest*)
 	{
 		// At this point task synchronization would hold the number of pending requests.
 		TaskSynchronization.Decrement();
diff --git a/Engine/Source/Runtime/Engine/Private/Streaming/StaticMeshUpdate.h b/Engine/Source/Runtime/Engine/Private/Streaming/StaticMeshUpdate.h
--- a/Engine/Source/Runtime/Engine/Private/Streaming/StaticMeshUpdate.h
+++ b/Engine/Source/Runtime/Engine/Private/Streaming/StaticMeshUpdate.h
@@ -184,7 +184,7 @@ protected:
 	void CancelIORequest();
 
 	class IBulkDataIORequest* IORequest;
-	FAsyncFileCallBack AsyncFileCallback;
+	FBulkDataIORequestCallBack AsyncFileCallback;
 	bool bHighPrioIORequest;
 };
 
diff --git a/Engine/Source/Runtime/Engine/Private/Streaming/Texture2DStreamIn_IO.cpp b/Engine/Source/Runtime/Engine/Private/Streaming/Texture2DStreamIn_IO.cpp
--- a/Engine/Source/Runtime/Engine/Private/Streaming/Texture2DStreamIn_IO.cpp
+++ b/Engine/Source/Runtime/Engine/Private/Streaming/Texture2DStreamIn_IO.cpp
@@ -113,7 +113,7 @@ void FTexture2DStreamIn_IO::ClearIORequests(const FContext& Context)
 
 void FTexture2DStreamIn_IO::SetAsyncFileCallback()
 {
-	AsyncFileCallBack = [this](bool bWasCancelled, IAsyncReadRequest* Req)
+	AsyncFileCallBack = [this](bool bWasCancelled, IBulkDataIORequest*)
 	{
 		// At this point task synchronization would hold the number of pending requests.
 		TaskSynchronization.Decrement();
diff --git a/Engine/Source/Runtime/Engine/Private/Streaming/Texture2DStreamIn_IO.h b/Engine/Source/Runtime/Engine/Private/Streaming/Texture2DStreamIn_IO.h
--- a/Engine/Source/Runtime/Engine/Private/Streaming/Texture2DStreamIn_IO.h
+++ b/Engine/Source/Runtime/Engine/Private/Streaming/Texture2DStreamIn_IO.h
@@ -72,6 +72,6 @@ private:
 	FString IOFilename;
 #endif
 
-	FAsyncFileCallBack AsyncFileCallBack;
+	FBulkDataIORequestCallBack AsyncFileCallBack;
 };
 
diff --git a/Engine/Source/Runtime/Engine/Public/Animation/AnimationStreaming.h b/Engine/Source/Runtime/Engine/Public/Animation/AnimationStreaming.h
--- a/Engine/Source/Runtime/Engine/Public/Animation/AnimationStreaming.h
+++ b/Engine/Source/Runtime/Engine/Public/Animation/AnimationStreaming.h
@@ -157,7 +157,7 @@ struct FAnimationStreamingManager : public IAnimationStreamingManager
 	// End IAudioStreamingManager interface
 
 	/** Called when an async callback is made on an async loading audio chunk request. */
-	void OnAsyncFileCallback(FStreamingAnimationData* StreamingAnimData, int32 ChunkIndex, int64 ReadSize, IAsyncReadRequest* ReadRequest, bool bWasCancelled);
+	void OnAsyncFileCallback(FStreamingAnimationData* StreamingAnimData, int32 ChunkIndex, int64 ReadSize, IBulkDataIORequest* ReadRequest, bool bWasCancelled);
 
 protected:
 
diff --git a/Engine/Source/Runtime/Engine/Public/AudioStreaming.h b/Engine/Source/Runtime/Engine/Public/AudioStreaming.h
--- a/Engine/Source/Runtime/Engine/Public/AudioStreaming.h
+++ b/Engine/Source/Runtime/Engine/Public/AudioStreaming.h
@@ -280,7 +280,7 @@ struct FLegacyAudioStreamingManager : public IAudioStreamingManager
 	// End IAudioStreamingManager interface
 
 	/** Called when an async callback is made on an async loading audio chunk request. */
-	void OnAsyncFileCallback(FStreamingWaveData* StreamingWaveData, int32 LoadedAudioChunkIndex, IAsyncReadRequest* ReadRequest);
+	void OnAsyncFileCallback(FStreamingWaveData* StreamingWaveData, int32 LoadedAudioChunkIndex, IBulkDataIORequest* ReadRequest);
 
 	/** Processes pending async file IO results. */
 	void ProcessPendingAsyncFileResults();
