Remove a line of code that was not supposed to be submitted in CL 10478374

#rnx

diff --git a/Engine/Source/Runtime/CoreUObject/Public/Serialization/BulkData2.h b/Engine/Source/Runtime/CoreUObject/Public/Serialization/BulkData2.h
--- a/Engine/Source/Runtime/CoreUObject/Public/Serialization/BulkData2.h
+++ b/Engine/Source/Runtime/CoreUObject/Public/Serialization/BulkData2.h
@@ -179,7 +179,6 @@ class COREUOBJECT_API FUntypedBulkData2 : public FBulkDataBase
 	// used or not.
 	static_assert(TIsPODType<ElementType>::Value, "FUntypedBulkData2 is limited to POD types!");
 public:
-	FUntypedBulkData2() = default;
 
 	void Serialize(FArchive& Ar, UObject* Owner, int32 Index, bool bAttemptFileMapping)
 	{
