Some minor compile fixes for some of our older projects that were not linking some templated methods correctly.

#rnx

diff --git a/Engine/Source/Runtime/CoreUObject/Public/Serialization/BulkData2.h b/Engine/Source/Runtime/CoreUObject/Public/Serialization/BulkData2.h
--- a/Engine/Source/Runtime/CoreUObject/Public/Serialization/BulkData2.h
+++ b/Engine/Source/Runtime/CoreUObject/Public/Serialization/BulkData2.h
@@ -179,6 +179,7 @@ class COREUOBJECT_API FUntypedBulkData2 : public FBulkDataBase
 	// used or not.
 	static_assert(TIsPODType<ElementType>::Value, "FUntypedBulkData2 is limited to POD types!");
 public:
+	FORCEINLINE FUntypedBulkData2() {}
 
 	void Serialize(FArchive& Ar, UObject* Owner, int32 Index, bool bAttemptFileMapping)
 	{
@@ -244,7 +245,7 @@ public:
 	 *
 	 * @return A FBulkDataBuffer that owns a copy of the BulkData, this might be a subset of the data depending on the value of RequestedSize.
 	 */
-	FBulkDataBuffer<ElementType> GetCopyAsBuffer(int64 RequestedElementCount, bool bDiscardInternalCopy)
+	FORCEINLINE FBulkDataBuffer<ElementType> GetCopyAsBuffer(int64 RequestedElementCount, bool bDiscardInternalCopy)
 	{
 		const int64 MaxElementCount = GetElementCount();
 
