AsyncLoading2:
Mostly rely on global bundle load order instead of adding EDL arcs between bundles
Merge all IO chunks for a package into one

diff --git a/Engine/Source/Developer/IoStoreUtilities/Private/IoStoreUtilities.cpp b/Engine/Source/Developer/IoStoreUtilities/Private/IoStoreUtilities.cpp
--- a/Engine/Source/Developer/IoStoreUtilities/Private/IoStoreUtilities.cpp
+++ b/Engine/Source/Developer/IoStoreUtilities/Private/IoStoreUtilities.cpp
@@ -222,15 +222,13 @@ enum EPreloadDependencyType
 
 enum EEventLoadNode2 : uint8
 {
-	Package_ProcessSummary,
 	Package_ExportsSerialized,
 	Package_StartPostLoad,
 	Package_Tick,
 	Package_Delete,
 	Package_NumPhases,
 
-	ExportBundle_StartIo = 0,
-	ExportBundle_Process,
+	ExportBundle_Process = 0,
 	ExportBundle_NumPhases,
 };
 
@@ -261,6 +259,7 @@ struct FExportGraphNode;
 struct FExportBundle
 {
 	TArray<FExportGraphNode*> Nodes;
+	uint32 LoadOrder;
 };
 
 struct FPackage;
@@ -419,7 +418,6 @@ struct FPackage
 	TArray<FExportGraphNode*> CreateExportNodes;
 	TArray<FExportGraphNode*> SerializeExportNodes;
 
-	uint64 LoadOrder = uint64(-1);
 	TArray<FExportGraphNode*> NodesWithNoIncomingEdges;
 	FPackageGraphNode* Node = nullptr;
 };
@@ -641,25 +639,17 @@ static void AddExportsDoneArc(FPackage& FromPackage, FPackage& ToPackage)
 	ExternalArcs.Add({ EEventLoadNode2::Package_ExportsSerialized, EEventLoadNode2::Package_StartPostLoad });
 }
 
-static void AddInternalBundleArc(FPackage& Package, uint32 FromBundleIndex, uint32 ToBundleIndex)
-{
-	uint32 FromNodeIndex = EEventLoadNode2::Package_NumPhases + FromBundleIndex * EEventLoadNode2::ExportBundle_NumPhases + EEventLoadNode2::ExportBundle_Process;
-	uint32 ToNodeIndex = EEventLoadNode2::Package_NumPhases + ToBundleIndex * EEventLoadNode2::ExportBundle_NumPhases + EEventLoadNode2::ExportBundle_StartIo;
-	check(!Package.InternalArcs.Contains(FArc({FromNodeIndex, ToNodeIndex})));
-	Package.InternalArcs.Add({ FromNodeIndex, ToNodeIndex });
-}
-
 static void AddUniqueExternalBundleArc(FPackage& FromPackage, uint32 FromBundleIndex, FPackage& ToPackage, uint32 ToBundleIndex)
 {
 	uint32 FromNodeIndex = EEventLoadNode2::Package_NumPhases + FromBundleIndex * EEventLoadNode2::ExportBundle_NumPhases + EEventLoadNode2::ExportBundle_Process;
-	uint32 ToNodeIndex = EEventLoadNode2::Package_NumPhases + ToBundleIndex * EEventLoadNode2::ExportBundle_NumPhases + EEventLoadNode2::ExportBundle_StartIo;
+	uint32 ToNodeIndex = EEventLoadNode2::Package_NumPhases + ToBundleIndex * EEventLoadNode2::ExportBundle_NumPhases + EEventLoadNode2::ExportBundle_Process;
 	TArray<FArc>& ExternalArcs = ToPackage.ExternalArcs.FindOrAdd(FromPackage.Name);
 	ExternalArcs.AddUnique({ FromNodeIndex, ToNodeIndex });
 }
 
 static void AddUniqueScriptBundleArc(FPackage& Package, uint32 GlobalImportIndex, uint32 BundleIndex)
 {
-	uint32 NodeIndex = EEventLoadNode2::Package_NumPhases + BundleIndex * EEventLoadNode2::ExportBundle_NumPhases + EEventLoadNode2::ExportBundle_StartIo;
+	uint32 NodeIndex = EEventLoadNode2::Package_NumPhases + BundleIndex * EEventLoadNode2::ExportBundle_NumPhases + EEventLoadNode2::ExportBundle_Process;
 	Package.ScriptArcs.AddUnique({ GlobalImportIndex, NodeIndex });
 }
 
@@ -820,36 +810,43 @@ static void AddPostLoadDependencies(
 
 static void BuildBundles(FExportGraph& ExportGraph, TMap<FName, FPackage*>& PackageMap)
 {
+
 	TArray<FExportGraphNode*> ExportLoadOrder = ExportGraph.ComputeLoadOrder(PackageMap);
 	FPackage* LastPackage = nullptr;
+	uint32 BundleLoadOrder = 0;
 	for (FExportGraphNode* Node : ExportLoadOrder)
 	{
 		FPackage* Package = Node->Package;
 		check(Package);
 
+		uint32 BundleIndex;
+		FExportBundle* Bundle;
 		if (Package != LastPackage)
 		{
-			Package->ExportBundles.AddDefaulted();
-			if (Package->ExportBundles.Num() > 1)
-			{
-				AddInternalBundleArc(*Package, Package->ExportBundles.Num() - 2, Package->ExportBundles.Num() - 1);
-			}
+			BundleIndex = Package->ExportBundles.Num();
+			Bundle = &Package->ExportBundles.AddDefaulted_GetRef();
+			Bundle->LoadOrder = BundleLoadOrder++;
+			LastPackage = Package;
+		}
+		else
+		{
+			BundleIndex = Package->ExportBundles.Num() - 1;
+			Bundle = &Package->ExportBundles[BundleIndex];
 		}
-		LastPackage = Package;
-
-		uint32 BundleIndex = Package->ExportBundles.Num() - 1;
 		for (FExportGraphNode* ExternalDependency : Node->ExternalDependencies)
 		{
 			uint32* FindDependentBundleIndex = ExternalDependency->Package->ExportBundleMap.Find(ExternalDependency);
 			check(FindDependentBundleIndex);
-			AddUniqueExternalBundleArc(*ExternalDependency->Package, *FindDependentBundleIndex, *Package, BundleIndex);
+			if (BundleIndex > 0)
+			{
+				AddUniqueExternalBundleArc(*ExternalDependency->Package, *FindDependentBundleIndex, *Package, BundleIndex);
+			}
 		}
 		for (uint32 ScriptDependencyGlobalImportIndex : Node->ScriptDependencies)
 		{
 			AddUniqueScriptBundleArc(*Package, ScriptDependencyGlobalImportIndex, BundleIndex);
 		}
-
-		Package->ExportBundles[BundleIndex].Nodes.Add(Node);
+		Bundle->Nodes.Add(Node);
 		Package->ExportBundleMap.Add(Node, BundleIndex);
 	}
 }
@@ -1027,6 +1024,7 @@ static void FindExport(TArray<FExportData>& GlobalExports, TMap<FString, int32>&
 {
 	const FObjectExport* Export = ExportMap + LocalExportIndex;
 	FString& FullName = TempFullNames[LocalExportIndex];
+
 	if (FullName.Len() == 0)
 	{
 		if (Export->OuterIndex.IsNull())
@@ -1628,6 +1626,7 @@ int32 CreateTarget(const FContainerTarget& Target)
 #endif
 
 	UE_LOG(LogIoStore, Display, TEXT("Serializing..."));
+	TArray<uint32> ExportBundleOrderEntries;
 	for (auto& PackageKV : PackageMap)
 	{
 		FPackage& Package = *PackageKV.Value;
@@ -1708,6 +1707,7 @@ int32 CreateTarget(const FContainerTarget& Target)
 		}
 		for (FExportBundle& ExportBundle : Package.ExportBundles)
 		{
+			ExportBundleOrderEntries.Add(ExportBundle.LoadOrder);
 			for (FExportGraphNode* ExportNode : ExportBundle.Nodes)
 			{
 				uint32 CommandType = uint32(ExportNode->BundleEntry.CommandType);
@@ -1725,18 +1725,18 @@ int32 CreateTarget(const FContainerTarget& Target)
 
 		ensure(Package.GlobalPackageId == PackageIndex++);
 		{
-			int32 Pad = 0;
 			int32 ExportBundleCount = Package.ExportBundles.Num();
 
 			NameMapBuilder.SerializeName(*StoreTocArchive, Package.Name);
 			NameMapBuilder.SerializeName(*StoreTocArchive, RelativeFileName);
 			*StoreTocArchive << Package.ExportCount;
 			*StoreTocArchive << ExportBundleCount;
+			int32 FirstExportBundleIndex = ExportBundleOrderEntries.Num() - ExportBundleCount;
+			*StoreTocArchive << FirstExportBundleIndex;
 			*StoreTocArchive << Package.FirstGlobalImport;
 			*StoreTocArchive << Package.GlobalImportCount;
 			*StoreTocArchive << Package.ImportedPackagesCount;
 			*StoreTocArchive << Package.ImportedPackagesOffset;
-			*StoreTocArchive << Pad;
 		}
 
 		Package.NameMapSize = Package.NameIndices.Num() * Package.NameIndices.GetTypeSize();
@@ -1796,32 +1796,23 @@ int32 CreateTarget(const FContainerTarget& Target)
 				ZenAr.Serialize(ZenGraphArchive.GetWriterData(), ZenGraphArchive.Tell());
 			}
 
-			// Package summary chunk
+			// Export bundle chunks
 			{
-				FIoBuffer IoBuffer(FIoBuffer::AssumeOwnership, ZenSummaryBuffer, ZenSummarySize);
-#if !SKIP_WRITE_CONTAINER
-				IoStoreWriter->Append(CreateChunkId(Package.GlobalPackageId, 0, EIoChunkType::PackageSummary, TEXT("PackageSummary")), IoBuffer);
-#endif
-				++SummaryChunkCount;
-			}
+				check(Package.ExportBundles.Num());
 
-			// Export chunks
-			if (Package.ExportCount)
-			{
 				FString UExpFileName = FPaths::ChangeExtension(Package.FileName, TEXT(".uexp"));
 				TUniquePtr<FArchive> ExpAr(IFileManager::Get().CreateFileReader(*UExpFileName));
-				const int64 TotalExportsSize = ExpAr->TotalSize();
+				Package.UExpSize = ExpAr->TotalSize();
 #if !SKIP_WRITE_CONTAINER
-				uint8* ExportsBuffer = static_cast<uint8*>(FMemory::Malloc(TotalExportsSize));
-				ExpAr->Serialize(ExportsBuffer, TotalExportsSize);
+				uint8* ExportsBuffer = static_cast<uint8*>(FMemory::Malloc(ExpAr->TotalSize()));
+				ExpAr->Serialize(ExportsBuffer, ExpAr->TotalSize());
 #endif
-				Package.UExpSize = TotalExportsSize;
+				ExpAr->Close();
 
+				uint64 BundleBufferSize = ZenSummarySize;
 				for (int32 I = 0; I < Package.ExportBundles.Num(); ++I)
 				{
-					check(I < UINT16_MAX);
 					FExportBundle& ExportBundle = Package.ExportBundles[I];
-					uint64 BundleBufferSize = 0;
 					for (FExportGraphNode* Node : ExportBundle.Nodes)
 					{
 						if (Node->BundleEntry.CommandType == FExportBundleEntry::ExportCommandType_Serialize)
@@ -1830,32 +1821,37 @@ int32 CreateTarget(const FContainerTarget& Target)
 							BundleBufferSize += ObjectExport.SerialSize;
 						}
 					}
-					uint8* BundleBuffer = static_cast<uint8*>(FMemory::Malloc(BundleBufferSize));
-					uint64 BundleBufferOffset = 0;
+				}
+
+#if !SKIP_WRITE_CONTAINER
+				uint8* BundleBuffer = static_cast<uint8*>(FMemory::Malloc(BundleBufferSize));
+				FMemory::Memcpy(BundleBuffer, ZenSummaryBuffer, ZenSummarySize);
+#endif
+				FMemory::Free(ZenSummaryBuffer);
+				uint64 BundleBufferOffset = ZenSummarySize;
+				for (FExportBundle& ExportBundle : Package.ExportBundles)
+				{
 					for (FExportGraphNode* Node : ExportBundle.Nodes)
 					{
 						if (Node->BundleEntry.CommandType == FExportBundleEntry::ExportCommandType_Serialize)
 						{
 							FObjectExport& ObjectExport = Exports[Package.ExportIndexOffset + Node->BundleEntry.LocalExportIndex];
 							const int64 Offset = ObjectExport.SerialOffset - Package.UAssetSize;
+#if !SKIP_WRITE_CONTAINER
 							FMemory::Memcpy(BundleBuffer + BundleBufferOffset, ExportsBuffer + Offset, ObjectExport.SerialSize);
+#endif
 							BundleBufferOffset += ObjectExport.SerialSize;
 						}
 					}
-#if !SKIP_WRITE_CONTAINER
-					FObjectExport& ObjectExport = Exports[Package.ExportIndexOffset + I];
-					const int64 Offset = ObjectExport.SerialOffset - Package.UAssetSize;
-					FIoBuffer IoBuffer(FIoBuffer::Wrap, BundleBuffer, BundleBufferSize);
-					IoStoreWriter->Append(CreateChunkId(Package.GlobalPackageId, I, EIoChunkType::ExportBundleData, *Package.FileName), IoBuffer);
-#endif
-					++ExportChunkCount;
-					FMemory::Free(BundleBuffer);
 				}
 
 #if !SKIP_WRITE_CONTAINER
+				FIoBuffer IoBuffer(FIoBuffer::Wrap, BundleBuffer, BundleBufferSize);
+				IoStoreWriter->Append(CreateChunkId(Package.GlobalPackageId, 0, EIoChunkType::ExportBundleData, *Package.FileName), IoBuffer);
+				FMemory::Free(BundleBuffer);
 				FMemory::Free(ExportsBuffer);
-				ExpAr->Close();
 #endif
+				++ExportChunkCount;
 			}
 
 #if !SKIP_BULKDATA
@@ -1904,6 +1900,10 @@ int32 CreateTarget(const FContainerTarget& Target)
 		GlobalMetaArchive << GlobalImportNamesByteCount;
 		GlobalMetaArchive.Serialize(GlobalImportNamesArchive->GetData(), GlobalImportNamesByteCount);
 
+		int32 ExportBundleOrderByteCount = ExportBundleOrderEntries.Num() * sizeof(uint32);
+		GlobalMetaArchive << ExportBundleOrderByteCount;
+		GlobalMetaArchive.Serialize(ExportBundleOrderEntries.GetData(), ExportBundleOrderByteCount);
+
 		const FIoStatus Status = AppendToContainer(*IoStoreWriter, GlobalMetaArchive, CreateIoChunkId(0, 0, EIoChunkType::LoaderGlobalMeta));
 		if (!Status.IsOk())
 		{
diff --git a/Engine/Source/Runtime/CoreUObject/Private/Serialization/AsyncLoading2.cpp b/Engine/Source/Runtime/CoreUObject/Private/Serialization/AsyncLoading2.cpp
--- a/Engine/Source/Runtime/CoreUObject/Private/Serialization/AsyncLoading2.cpp
+++ b/Engine/Source/Runtime/CoreUObject/Private/Serialization/AsyncLoading2.cpp
@@ -242,11 +242,11 @@ struct FPackageStoreEntry
 	FName FileName;
 	int32 ExportCount;
 	int32 ExportBundleCount;
+	int32 FirstExportBundleIndex;
 	int32 FirstGlobalImport;
 	int32 GlobalImportCount;
 	int32 ImportedPackagesCount;
 	int32 ImportedPackagesOffset;
-	int32 Pad;
 };
 
 struct FPackageSummary
@@ -353,6 +353,7 @@ public:
 	FPackageStoreEntry* StoreEntries = nullptr;
 	FPackageStoreInitialLoadEntry* InitialLoadStoreEntries = nullptr;
 	int32* ImportedPackages = nullptr;
+	uint32* ExportBundleOrderEntries = nullptr;
 	int32* ScriptArcs = nullptr;
 	int32 PackageCount = 0;
 	int32 ImportedPackagesCount = 0;
@@ -434,6 +435,7 @@ public:
 			ImportStore.Count = ImportByteCount / sizeof(FName);
 
 			FName* SerializedNames = reinterpret_cast<FName*>(IoBuffer.Data() + GlobalMetaAr.Tell());
+			GlobalMetaAr.Seek(GlobalMetaAr.Tell() + ImportByteCount);
 			ImportStore.Names = reinterpret_cast<FName*>(FMemory::Malloc(ImportByteCount));
 
 			for (int I = 0; I < ImportStore.Count; ++I)
@@ -447,6 +449,16 @@ public:
 			ImportStore.Imports = new FGlobalImport[ImportStore.Count];
 		}
 
+		{
+			TRACE_CPUPROFILER_EVENT_SCOPE(LoadPackageStoreBundleOrder);
+			
+			int32 BundleOrderByteCount;
+			GlobalMetaAr << BundleOrderByteCount;
+			ExportBundleOrderEntries = reinterpret_cast<uint32*>(FMemory::Malloc(BundleOrderByteCount));
+			GlobalMetaAr.Serialize(ExportBundleOrderEntries, BundleOrderByteCount);
+		}
+
+
 		// Load initial loading meta data
 		{
 			FIoBuffer TempIoBuffer;
@@ -560,9 +572,10 @@ public:
 		return StoreEntries[GlobalPackageId.Id].ExportCount;
 	}
 
-	inline int32 GetPackageExportBundleCount(FGlobalPackageId GlobalPackageId) const
+	inline TTuple<uint32*, uint32> GetPackageExportBundles(FGlobalPackageId GlobalPackageId) const
 	{
-		return StoreEntries[GlobalPackageId.Id].ExportBundleCount;
+		const FPackageStoreEntry* StoreEntry = StoreEntries + GlobalPackageId.Id;
+		return MakeTuple(ExportBundleOrderEntries + StoreEntry->FirstExportBundleIndex, uint32(StoreEntry->ExportBundleCount));
 	}
 
 	inline FString GetPackageFileName(FGlobalPackageId GlobalPackageId)
@@ -815,15 +828,13 @@ enum class EAsyncPackageLoadingState2 : uint8
 
 enum EEventLoadNode2 : uint8
 {
-	Package_ProcessSummary,
 	Package_ExportsSerialized,
 	Package_StartPostLoad,
 	Package_Tick,
 	Package_Delete,
 	Package_NumPhases,
 
-	ExportBundle_StartIo = 0,
-	ExportBundle_Process,
+	ExportBundle_Process = 0,
 	ExportBundle_NumPhases,
 };
 
@@ -866,7 +877,7 @@ private:
 	TAtomic<int32> BarrierCount { 0 };
 	TAtomic<uint8> DependencyWriterCount { 0 };
 	TAtomic<uint8> bDone { 0 };
-#if UE_BUILD_DEVELOPMENT
+#if UE_BUILD_DEBUG || UE_BUILD_DEVELOPMENT
 	TAtomic<uint8> bFired { 0 };
 #endif
 
@@ -1037,6 +1048,7 @@ uint32 FAsyncLoadingThreadState2::TlsSlot;
 struct FAsyncPackage2
 {
 	friend struct FScopedAsyncPackageEvent2;
+	friend class FAsyncLoadingThread2Impl;
 
 	FAsyncPackage2(const FAsyncPackageDesc& InDesc, int32 InSerialNumber, FAsyncLoadingThread2Impl& InAsyncLoadingThread, IEDLBootNotificationManager& InEDLBootNotificationManager, FAsyncLoadEventGraphAllocator& InGraphAllocator, const FAsyncLoadEventSpec* EventSpecs, FGlobalPackageId InGlobalPackageId);
 	virtual ~FAsyncPackage2();
@@ -1278,9 +1290,8 @@ private:
 	FEventLoadNode2* ExportBundleNodes = nullptr;
 	uint32 ExportBundleNodeCount = 0;
 
-	uint64 PackageSummarySize = 0;
-	FIoBuffer PackageSummaryIoBuffer;
-	TArray<FIoBuffer, TInlineAllocator<2>> ExportBundleIoBuffers;
+	FIoBuffer IoBuffer;
+	const uint8* SerialDataPtr = nullptr;
 
 	// FZenLinkerLoad
 	TArray<FExternalReadCallback> ExternalReadDependencies;
@@ -1290,6 +1301,7 @@ private:
 	TArray<UObject*> Exports;
 	FPackageImportStore ImportStore;
 
+	uint32* ExportBundleOrderEntries = nullptr;
 	int32 ExportBundleCount = 0;
 	const FExportBundleHeader* ExportBundles = nullptr;
 	const FExportBundleEntry* ExportBundleEntries = nullptr;
@@ -1314,8 +1326,6 @@ public:
 
 	bool bAllExportsSerialized;
 
-	static EAsyncPackageState::Type Event_ProcessSummary(FAsyncPackage2* Package, int32);
-	static EAsyncPackageState::Type Event_StartExportBundleIo(FAsyncPackage2* Package, int32 ExportBundleIndex);
 	static EAsyncPackageState::Type Event_ProcessExportBundle(FAsyncPackage2* Package, int32 ExportBundleIndex);
 	static EAsyncPackageState::Type Event_ExportsDone(FAsyncPackage2* Package, int32);
 	static EAsyncPackageState::Type Event_StartPostLoad(FAsyncPackage2* Package, int32);
@@ -1355,7 +1365,7 @@ public:
 private:
 	void CreateNodes(const FAsyncLoadEventSpec* EventSpecs);
 	void SetupScriptArcs();
-	void SetupSerializedArcs();
+	void SetupSerializedArcs(const uint8* GraphData, uint64 GraphDataSize);
 
 	/**
 	 * Begin async loading process. Simulates parts of BeginLoad.
@@ -1379,7 +1389,7 @@ private:
 	 *
 	 * @return true if linker is finished being created, false otherwise
 	 */
-	EAsyncPackageState::Type FinishLinker();
+	EAsyncPackageState::Type FinishLinker(const FPackageSummary* PackageSummary);
 
 	/**
 	 * Route PostLoad to all loaded objects. This might load further objects!
@@ -1549,6 +1559,20 @@ private:
 	FIoDispatcher& IoDispatcher;
 
 	FPackageStore GlobalPackageStore;
+
+	struct FBundleIoRequest
+	{
+		bool operator<(const FBundleIoRequest& Other) const
+		{
+			return BundeOrder < Other.BundeOrder;
+		}
+
+		FAsyncPackage2* Package;
+		uint32 BundeOrder;
+	};
+	TArray<FBundleIoRequest> WaitingIoRequests;
+	uint32 PendingBundleIoRequestsCount = 0;
+
 public:
 
 	//~ Begin FRunnable Interface.
@@ -1794,6 +1818,9 @@ private:
 	EAsyncPackageState::Type ProcessLoadedPackagesFromGameThread(bool& bDidSomething, int32 FlushRequestID = INDEX_NONE);
 
 	bool CreateAsyncPackagesFromQueue();
+	void AddBundleIoRequest(FAsyncPackage2* Package, uint32 BundleOrder);
+	void BundleIoRequestCompleted();
+	void StartBundleIoRequests();
 
 	FAsyncPackage2* CreateAsyncPackage(const FAsyncPackageDesc& Desc)
 	{
@@ -1978,11 +2005,50 @@ bool FAsyncLoadingThread2Impl::CreateAsyncPackagesFromQueue()
 			Package->StartLoading();
 		}
 		delete PackageRequest;
+
+		StartBundleIoRequests();
 	}
 
 	return QueueCopy.Num() > 0;
 }
 
+void FAsyncLoadingThread2Impl::AddBundleIoRequest(FAsyncPackage2* Package, uint32 BundleOrder)
+{
+	WaitingIoRequests.HeapPush({ Package, BundleOrder });
+}
+
+void FAsyncLoadingThread2Impl::BundleIoRequestCompleted()
+{
+	--PendingBundleIoRequestsCount;
+	if (WaitingIoRequests.Num())
+	{
+		StartBundleIoRequests();
+	}
+}
+
+void FAsyncLoadingThread2Impl::StartBundleIoRequests()
+{
+	TRACE_CPUPROFILER_EVENT_SCOPE(StartBundleIoRequests);
+	const uint32 MaxPendingRequestsCount = 1024;
+	while (WaitingIoRequests.Num() && PendingBundleIoRequestsCount < MaxPendingRequestsCount)
+	{
+		FBundleIoRequest BundleIoRequest;
+		WaitingIoRequests.HeapPop(BundleIoRequest, false);
+		FAsyncPackage2* Package = BundleIoRequest.Package;
+
+		FIoReadOptions ReadOptions;
+		IoDispatcher.ReadWithCallback(CreateIoChunkId(Package->GlobalPackageId.Id, 0, EIoChunkType::ExportBundleData),
+			ReadOptions,
+			[Package](TIoStatusOr<FIoBuffer> Result)
+		{
+			Package->IoBuffer = Result.ConsumeValueOrDie();
+			Package->GetExportBundleNode(EEventLoadNode2::ExportBundle_Process, 0)->ReleaseBarrier();
+		});
+
+		++PendingBundleIoRequestsCount;
+	}
+}
+
 FEventLoadNode2::FEventLoadNode2(const FAsyncLoadEventSpec* InSpec, FAsyncPackage2* InPackage, int32 InImportOrExportIndex)
 	: Spec(InSpec)
 	, Package(InPackage)
@@ -1995,7 +2061,7 @@ FEventLoadNode2::FEventLoadNode2(const FAsyncLoadEventSpec* InSpec, FAsyncPackag
 void FEventLoadNode2::DependsOn(FEventLoadNode2* Other)
 {
 	TRACE_CPUPROFILER_EVENT_SCOPE(DependsOn);
-#if UE_BUILD_DEVELOPMENT
+#if UE_BUILD_DEBUG || UE_BUILD_DEVELOPMENT
 	check(!bDone);
 	check(!bFired);
 #endif
@@ -2045,7 +2111,7 @@ void FEventLoadNode2::DependsOn(FEventLoadNode2* Other)
 
 void FEventLoadNode2::AddBarrier()
 {
-#if UE_BUILD_DEVELOPMENT
+#if UE_BUILD_DEBUG || UE_BUILD_DEVELOPMENT
 	check(!bDone);
 	check(!bFired);
 #endif
@@ -2054,7 +2120,7 @@ void FEventLoadNode2::AddBarrier()
 
 void FEventLoadNode2::AddBarrier(int32 Count)
 {
-#if UE_BUILD_DEVELOPMENT
+#if UE_BUILD_DEBUG || UE_BUILD_DEVELOPMENT
 	check(!bDone);
 	check(!bFired);
 #endif
@@ -2074,7 +2140,7 @@ void FEventLoadNode2::Fire()
 {
 	//TRACE_CPUPROFILER_EVENT_SCOPE(Fire);
 
-#if UE_BUILD_DEVELOPMENT
+#if UE_BUILD_DEBUG || UE_BUILD_DEVELOPMENT
 	bFired.Store(1);
 #endif
 
@@ -2294,30 +2360,7 @@ FUObjectSerializeContext* FAsyncPackage2::GetSerializeContext()
 	return FUObjectThreadContext::Get().GetSerializeContext();
 }
 
-EAsyncPackageState::Type FAsyncPackage2::Event_ProcessSummary(FAsyncPackage2* Package, int32)
-{
-	TRACE_CPUPROFILER_EVENT_SCOPE(Event_ProcessSummary);
-	FScopedAsyncPackageEvent2 Scope(Package);
-	Package->CreateLinker();
-	Package->FinishLinker();
-
-	Package->SetupSerializedArcs();
-	if (GIsInitialLoad)
-	{
-		Package->SetupScriptArcs();
-	}
-
-	check(Package->AsyncPackageLoadingState == EAsyncPackageLoadingState2::WaitingForSummary);
-	Package->AsyncPackageLoadingState = EAsyncPackageLoadingState2::ProcessNewImportsAndExports;
-
-	for (int32 ExportBundleIndex = 0; ExportBundleIndex < Package->ExportBundleCount; ++ExportBundleIndex)
-	{
-		Package->GetExportBundleNode(EEventLoadNode2::ExportBundle_StartIo, ExportBundleIndex)->ReleaseBarrier();
-	}
-	return EAsyncPackageState::Complete;
-}
-
-void FAsyncPackage2::SetupSerializedArcs()
+void FAsyncPackage2::SetupSerializedArcs(const uint8* GraphData, uint64 GraphDataSize)
 {
 	TRACE_CPUPROFILER_EVENT_SCOPE(SetupSerializedArcs);
 
@@ -2325,38 +2368,34 @@ void FAsyncPackage2::SetupSerializedArcs()
 
 	const FGlobalNameMap& GlobalNameMap = AsyncLoadingThread.GlobalNameMap;
 
-	const uint8* PackageSummaryData = PackageSummaryIoBuffer.Data();
-	const FPackageSummary* PackageSummary = reinterpret_cast<const FPackageSummary*>(PackageSummaryData);
-	uint64 ZenHeaderDataSize = PackageSummary->GraphDataSize;
-	const uint8* ZenHeaderData = PackageSummaryData + PackageSummary->GraphDataOffset;
-	FSimpleArchive ZenHeaderArchive(ZenHeaderData, ZenHeaderDataSize);
+	FSimpleArchive GraphArchive(GraphData, GraphDataSize);
 	int32 InternalArcCount;
-	ZenHeaderArchive << InternalArcCount;
+	GraphArchive << InternalArcCount;
 	for (int32 InternalArcIndex = 0; InternalArcIndex < InternalArcCount; ++InternalArcIndex)
 	{
 		int32 FromNodeIndex;
 		int32 ToNodeIndex;
-		ZenHeaderArchive << FromNodeIndex;
-		ZenHeaderArchive << ToNodeIndex;
+		GraphArchive << FromNodeIndex;
+		GraphArchive << ToNodeIndex;
 		PackageNodes[ToNodeIndex].DependsOn(PackageNodes + FromNodeIndex);
 	}
 	int32 ImportedPackagesCount;
-	ZenHeaderArchive << ImportedPackagesCount;
+	GraphArchive << ImportedPackagesCount;
 	for (int32 ImportedPackageIndex = 0; ImportedPackageIndex < ImportedPackagesCount; ++ImportedPackageIndex)
 	{
 		int32 ImportedPackageNameIndex;
 		int32 ImportedPackageNameNumber;
-		ZenHeaderArchive << ImportedPackageNameIndex << ImportedPackageNameNumber;
+		GraphArchive << ImportedPackageNameIndex << ImportedPackageNameNumber;
 		FName ImportedPackageName = GlobalNameMap.GetName(ImportedPackageNameIndex, ImportedPackageNameNumber);
 		FAsyncPackage2* ImportedPackage = AsyncLoadingThread.FindAsyncPackage(ImportedPackageName);
 		int32 ExternalArcCount;
-		ZenHeaderArchive << ExternalArcCount;
+		GraphArchive << ExternalArcCount;
 		for (int32 ExternalArcIndex = 0; ExternalArcIndex < ExternalArcCount; ++ExternalArcIndex)
 		{
 			int32 FromNodeIndex;
 			int32 ToNodeIndex;
-			ZenHeaderArchive << FromNodeIndex;
-			ZenHeaderArchive << ToNodeIndex;
+			GraphArchive << FromNodeIndex;
+			GraphArchive << ToNodeIndex;
 			if (ImportedPackage)
 			{
 				PackageNodes[ToNodeIndex].DependsOn(ImportedPackage->PackageNodes + FromNodeIndex);
@@ -2486,15 +2525,9 @@ void FAsyncPackage2::StartLoading()
 
 	LoadStartTime = FPlatformTime::Seconds();
 
-	FIoReadOptions ReadOptions;
-	AsyncLoadingThread.IoDispatcher.ReadWithCallback(CreateIoChunkId(GlobalPackageId.Id, 0, EIoChunkType::PackageSummary),
-		ReadOptions,
-		[this](TIoStatusOr<FIoBuffer> Result)
-		{
-			PackageSummaryIoBuffer = Result.ConsumeValueOrDie();
-			PackageSummarySize = PackageSummaryIoBuffer.DataSize();
-			GetPackageNode(EEventLoadNode2::Package_ProcessSummary)->ReleaseBarrier();
-		});
+	check(ExportBundleCount > 0);
+	uint32 BundleOrder = ExportBundleOrderEntries[0];
+	AsyncLoadingThread.AddBundleIoRequest(this, BundleOrder);
 }
 
 void FAsyncPackage2::SetupScriptArcs()
@@ -2537,33 +2570,49 @@ void FAsyncPackage2::SetupScriptArcs()
 	}
 }
 
-EAsyncPackageState::Type FAsyncPackage2::Event_StartExportBundleIo(FAsyncPackage2* Package, int32 ExportBundleIndex)
-{
-	TRACE_CPUPROFILER_EVENT_SCOPE(Event_StartExportBundleIo);
-
-	FIoReadOptions ReadOptions;
-	Package->AsyncLoadingThread.IoDispatcher.ReadWithCallback(CreateIoChunkId(Package->GlobalPackageId.Id, ExportBundleIndex, EIoChunkType::ExportBundleData),
-		ReadOptions,
-		[Package, ExportBundleIndex](TIoStatusOr<FIoBuffer> Result)
-		{
-			Package->ExportBundleIoBuffers[ExportBundleIndex] = Result.ConsumeValueOrDie();
-			Package->GetExportBundleNode(EEventLoadNode2::ExportBundle_Process, ExportBundleIndex)->ReleaseBarrier();
-		});
-	return EAsyncPackageState::Complete;
-}
-
 EAsyncPackageState::Type FAsyncPackage2::Event_ProcessExportBundle(FAsyncPackage2* Package, int32 ExportBundleIndex)
 {
 	TRACE_CPUPROFILER_EVENT_SCOPE(Event_ProcessExportBundle);
 
-	check(Package->AsyncPackageLoadingState == EAsyncPackageLoadingState2::ProcessNewImportsAndExports);
-
-	FIoBuffer& ExportBundleIoBuffer = Package->ExportBundleIoBuffers[ExportBundleIndex];
-	const uint8* ExportData = ExportBundleIoBuffer.Data();
-
 	FScopedAsyncPackageEvent2 Scope(Package);
 
 	check(ExportBundleIndex < Package->ExportBundleCount);
+
+	if (ExportBundleIndex == 0)
+	{
+		TRACE_CPUPROFILER_EVENT_SCOPE(ProcessPackageSummary);
+
+		check(Package->AsyncPackageLoadingState == EAsyncPackageLoadingState2::WaitingForSummary);
+
+		Package->SerialDataPtr = Package->IoBuffer.Data();
+		const uint8* PackageSummaryData = Package->SerialDataPtr;
+		const FPackageSummary* PackageSummary = reinterpret_cast<const FPackageSummary*>(PackageSummaryData);
+
+		Package->PackageNameMap = reinterpret_cast<const int32*>(PackageSummaryData + PackageSummary->NameMapOffset);
+		Package->ImportStore.ImportMap = reinterpret_cast<const int32*>(PackageSummaryData + PackageSummary->ImportMapOffset);
+		Package->ImportStore.ImportMapCount = (PackageSummary->ExportMapOffset - PackageSummary->ImportMapOffset) / sizeof(int32);
+		Package->ExportMap = reinterpret_cast<const FExportMapEntry*>(PackageSummaryData + PackageSummary->ExportMapOffset);
+		Package->ExportBundles = reinterpret_cast<const FExportBundleHeader*>(PackageSummaryData + PackageSummary->ExportBundlesOffset);
+		Package->ExportBundleEntries = reinterpret_cast<const FExportBundleEntry*>(Package->ExportBundles + Package->ExportBundleCount);
+
+		const uint8* GraphData = PackageSummaryData + PackageSummary->GraphDataOffset;
+
+		uint64 PackageSummarySize = GraphData + PackageSummary->GraphDataSize - PackageSummaryData;
+		Package->SerialDataPtr += PackageSummarySize;
+
+		Package->CreateLinker();
+		Package->FinishLinker(PackageSummary);
+
+		Package->SetupSerializedArcs(GraphData, PackageSummary->GraphDataSize);
+
+		check(Package->AsyncPackageLoadingState == EAsyncPackageLoadingState2::WaitingForSummary);
+		Package->AsyncPackageLoadingState = EAsyncPackageLoadingState2::ProcessNewImportsAndExports;
+
+		TRACE_LOADTIME_PACKAGE_SUMMARY(Package, PackageSummarySize, Package->ImportStore.ImportMapCount, Package->ExportCount);
+	}
+
+	check(Package->AsyncPackageLoadingState == EAsyncPackageLoadingState2::ProcessNewImportsAndExports);
+
 	const FExportBundleHeader* ExportBundle = Package->ExportBundles + ExportBundleIndex;
 	
 	const FExportBundleEntry* BundleEntry = Package->ExportBundleEntries + ExportBundle->FirstEntryIndex;
@@ -2579,22 +2628,34 @@ EAsyncPackageState::Type FAsyncPackage2::Event_ProcessExportBundle(FAsyncPackage
 		{
 			check(BundleEntry->CommandType == FExportBundleEntry::ExportBundleEntryType_Serialize);
 			uint64 ExportSerialSize = Package->ExportMap[BundleEntry->LocalExportIndex].SerialSize;
-			check(ExportData + ExportSerialSize <= ExportBundleIoBuffer.Data() + ExportBundleIoBuffer.DataSize());
+			check(Package->SerialDataPtr + ExportSerialSize <= Package->IoBuffer.Data() + Package->IoBuffer.DataSize());
 			UObject* Object = Package->Exports[BundleEntry->LocalExportIndex];
 			check(Object);
 			if (Object->HasAnyFlags(RF_NeedLoad))
 			{
-				Package->EventDrivenSerializeExport(BundleEntry->LocalExportIndex, ExportData, ExportSerialSize);
+				Package->EventDrivenSerializeExport(BundleEntry->LocalExportIndex, Package->SerialDataPtr, ExportSerialSize);
 			}
 			check(!Object->HasAnyFlags(RF_NeedLoad));
 
-			ExportData += ExportSerialSize;
+			Package->SerialDataPtr += ExportSerialSize;
 		}
 		++BundleEntry;
 	}
-	ExportBundleIoBuffer = FIoBuffer();
 
-	Package->GetNode(Package_ExportsSerialized)->ReleaseBarrier();
+	if (ExportBundleIndex + 1 < Package->ExportBundleCount)
+	{
+		Package->GetExportBundleNode(ExportBundle_Process, ExportBundleIndex + 1)->ReleaseBarrier();
+	}
+	else
+	{
+		Package->GetNode(Package_ExportsSerialized)->ReleaseBarrier();
+	}
+
+	if (ExportBundleIndex == 0)
+	{
+		Package->AsyncLoadingThread.BundleIoRequestCompleted();
+	}
+
 	return EAsyncPackageState::Complete;
 }
 
@@ -2998,17 +3059,20 @@ EAsyncPackageState::Type FAsyncPackage2::Event_Delete(FAsyncPackage2* Package, i
 
 FEventLoadNode2* FAsyncPackage2::GetPackageNode(EEventLoadNode2 Phase)
 {
+	check(Phase < EEventLoadNode2::Package_NumPhases);
 	return PackageNodes + Phase;
 }
 
 FEventLoadNode2* FAsyncPackage2::GetExportBundleNode(EEventLoadNode2 Phase, uint32 ExportBundleIndex)
 {
+	check(ExportBundleIndex < uint32(ExportBundleCount));
 	uint32 ExportBundleNodeIndex = ExportBundleIndex * EEventLoadNode2::ExportBundle_NumPhases + Phase;
 	return ExportBundleNodes + ExportBundleNodeIndex;
 }
 
 FEventLoadNode2* FAsyncPackage2::GetNode(int32 NodeIndex)
 { 
+	check(uint32(NodeIndex) < EEventLoadNode2::Package_NumPhases + ExportBundleNodeCount);
 	return &PackageNodes[NodeIndex];
 }
 
@@ -3398,13 +3462,11 @@ FAsyncLoadingThread2Impl::FAsyncLoadingThread2Impl(IEDLBootNotificationManager&
 	}
 
 	EventSpecs.AddDefaulted(EEventLoadNode2::Package_NumPhases + EEventLoadNode2::ExportBundle_NumPhases);
-	EventSpecs[EEventLoadNode2::Package_ProcessSummary] = { &FAsyncPackage2::Event_ProcessSummary, &AsyncEventQueue, false };
 	EventSpecs[EEventLoadNode2::Package_ExportsSerialized] = { &FAsyncPackage2::Event_ExportsDone, &AsyncEventQueue, true };
 	EventSpecs[EEventLoadNode2::Package_StartPostLoad] = { &FAsyncPackage2::Event_StartPostLoad, &AsyncEventQueue, false };
 	EventSpecs[EEventLoadNode2::Package_Tick] = { &FAsyncPackage2::Event_Tick, &EventQueue, false };
 	EventSpecs[EEventLoadNode2::Package_Delete] = { &FAsyncPackage2::Event_Delete, &AsyncEventQueue, false };
 
-	EventSpecs[EEventLoadNode2::Package_NumPhases + EEventLoadNode2::ExportBundle_StartIo] = { &FAsyncPackage2::Event_StartExportBundleIo, &AsyncEventQueue, false };
 	EventSpecs[EEventLoadNode2::Package_NumPhases + EEventLoadNode2::ExportBundle_Process] = { &FAsyncPackage2::Event_ProcessExportBundle, &ProcessExportBundlesEventQueue, false };
 
 	CancelLoadingEvent = FPlatformProcess::GetSynchEventFromPool();
@@ -3890,12 +3952,17 @@ FAsyncPackage2::FAsyncPackage2(const FAsyncPackageDesc& InDesc, int32 InSerialNu
 	TRACE_LOADTIME_NEW_ASYNC_PACKAGE(this, InDesc.Name);
 	AddRequestID(InDesc.RequestID);
 
-	ExportBundleCount = AsyncLoadingThread.GlobalPackageStore.GetPackageExportBundleCount(GlobalPackageId);
+	TTuple<uint32*, uint32> PackageExportBundles = AsyncLoadingThread.GlobalPackageStore.GetPackageExportBundles(GlobalPackageId);
+	ExportBundleOrderEntries = PackageExportBundles.Get<0>();
+	ExportBundleCount = PackageExportBundles.Get<1>();
 	ExportCount = AsyncLoadingThread.GlobalPackageStore.GetPackageExportCount(GlobalPackageId);
 	Exports.AddZeroed(ExportCount);
 
-	ExportBundleIoBuffers.SetNum(ExportBundleCount);
 	CreateNodes(EventSpecs);
+	if (GIsInitialLoad)
+	{
+		SetupScriptArcs();
+	}
 }
 
 void FAsyncPackage2::CreateNodes(const FAsyncLoadEventSpec* EventSpecs)
@@ -3910,10 +3977,6 @@ void FAsyncPackage2::CreateNodes(const FAsyncLoadEventSpec* EventSpecs)
 			new (PackageNodes + Phase) FEventLoadNode2(EventSpecs + Phase, this, -1);
 		}
 
-
-		FEventLoadNode2* ProcessSummaryNode = PackageNodes + EEventLoadNode2::Package_ProcessSummary;
-		ProcessSummaryNode->AddBarrier();
-
 		FEventLoadNode2* ExportsSerializedNode = PackageNodes + EEventLoadNode2::Package_ExportsSerialized;
 		FEventLoadNode2* StartPostLoadNode = PackageNodes + EEventLoadNode2::Package_StartPostLoad;
 		FEventLoadNode2* TickNode = PackageNodes + EEventLoadNode2::Package_Tick;
@@ -3928,14 +3991,11 @@ void FAsyncPackage2::CreateNodes(const FAsyncLoadEventSpec* EventSpecs)
 		for (int32 ExportBundleIndex = 0; ExportBundleIndex < ExportBundleCount; ++ExportBundleIndex)
 		{
 			uint32 NodeIndex = EEventLoadNode2::ExportBundle_NumPhases * ExportBundleIndex;
-			FEventLoadNode2* StartIoNode = ExportBundleNodes + NodeIndex + EEventLoadNode2::ExportBundle_StartIo;
-			new (StartIoNode) FEventLoadNode2(EventSpecs + EEventLoadNode2::Package_NumPhases + EEventLoadNode2::ExportBundle_StartIo, this, ExportBundleIndex);
 			FEventLoadNode2* ProcessNode = ExportBundleNodes + NodeIndex + EEventLoadNode2::ExportBundle_Process;
 			new (ProcessNode) FEventLoadNode2(EventSpecs + EEventLoadNode2::Package_NumPhases + EEventLoadNode2::ExportBundle_Process, this, ExportBundleIndex);
 			ProcessNode->AddBarrier();
-			StartIoNode->AddBarrier();
 		}
-		ExportsSerializedNode->AddBarrier(ExportBundleCount);
+		ExportsSerializedNode->AddBarrier();
 	}
 }
 
@@ -4140,24 +4200,11 @@ EAsyncPackageState::Type FAsyncPackage2::CreateLinker()
  *
  * @return true if linker is finished being created, false otherwise
  */
-EAsyncPackageState::Type FAsyncPackage2::FinishLinker()
+EAsyncPackageState::Type FAsyncPackage2::FinishLinker(const FPackageSummary* PackageSummary)
 {
 	TRACE_CPUPROFILER_EVENT_SCOPE(FinishLinker);
 	LLM_SCOPE(ELLMTag::AsyncLoading);
 
-	const uint8* PackageSummaryData = PackageSummaryIoBuffer.Data();
-	const FPackageSummary* PackageSummary = reinterpret_cast<const FPackageSummary*>(PackageSummaryData);
-	
-	PackageNameMap = reinterpret_cast<const int32*>(PackageSummaryData + PackageSummary->NameMapOffset);
-	ImportStore.ImportMap = reinterpret_cast<const int32*>(PackageSummaryData + PackageSummary->ImportMapOffset);
-	ImportStore.ImportMapCount = (PackageSummary->ExportMapOffset - PackageSummary->ImportMapOffset) / sizeof(int32);
-	ExportMap = reinterpret_cast<const FExportMapEntry*>(PackageSummaryData + PackageSummary->ExportMapOffset);
-	ExportBundles = reinterpret_cast<const FExportBundleHeader*>(PackageSummaryData + PackageSummary->ExportBundlesOffset);
-	ExportBundleEntries = reinterpret_cast<const FExportBundleEntry*>(ExportBundles + ExportBundleCount);
-
-
-	TRACE_LOADTIME_PACKAGE_SUMMARY(this, PackageSummarySize, ImportStore.ImportMapCount, ExportCount);
-
 	FPackageFileSummary& Summary	= Linker->Summary;
 	Summary.PackageFlags			= PackageSummary->PackageFlags;
 	Summary.BulkDataStartOffset		= PackageSummary->BulkDataStartOffset;
