AsyncLoading2:
Fixed bundles being loaded out of order during initial loading
Fixed static analysis warning

diff --git a/Engine/Source/Developer/IoStoreUtilities/Private/IoStoreUtilities.cpp b/Engine/Source/Developer/IoStoreUtilities/Private/IoStoreUtilities.cpp
--- a/Engine/Source/Developer/IoStoreUtilities/Private/IoStoreUtilities.cpp
+++ b/Engine/Source/Developer/IoStoreUtilities/Private/IoStoreUtilities.cpp
@@ -818,6 +818,10 @@ static void BuildBundles(FExportGraph& ExportGraph, TMap<FName, FPackage*>& Pack
 	{
 		FPackage* Package = Node->Package;
 		check(Package);
+		if (!Package)
+		{
+			continue;
+		}
 
 		uint32 BundleIndex;
 		FExportBundle* Bundle;
diff --git a/Engine/Source/Runtime/CoreUObject/Private/Serialization/AsyncLoading2.cpp b/Engine/Source/Runtime/CoreUObject/Private/Serialization/AsyncLoading2.cpp
--- a/Engine/Source/Runtime/CoreUObject/Private/Serialization/AsyncLoading2.cpp
+++ b/Engine/Source/Runtime/CoreUObject/Private/Serialization/AsyncLoading2.cpp
@@ -2030,12 +2030,19 @@ void FAsyncLoadingThread2Impl::StartBundleIoRequests()
 {
 	TRACE_CPUPROFILER_EVENT_SCOPE(StartBundleIoRequests);
 	const uint32 MaxPendingRequestsCount = 1024;
+	FAsyncPackage2* PreviousPackage = nullptr;
 	while (WaitingIoRequests.Num() && PendingBundleIoRequestsCount < MaxPendingRequestsCount)
 	{
 		FBundleIoRequest BundleIoRequest;
 		WaitingIoRequests.HeapPop(BundleIoRequest, false);
 		FAsyncPackage2* Package = BundleIoRequest.Package;
 
+		if (GIsInitialLoad && PreviousPackage)
+		{
+			Package->GetExportBundleNode(ExportBundle_Process, 0)->DependsOn(PreviousPackage->GetExportBundleNode(ExportBundle_Process, 0));
+		}
+		PreviousPackage = Package;
+
 		FIoReadOptions ReadOptions;
 		IoDispatcher.ReadWithCallback(CreateIoChunkId(Package->GlobalPackageId.Id, 0, EIoChunkType::ExportBundleData),
 			ReadOptions,
