Change an assert to a log warning with additional info to make it easier to track an issue on the build machines.

(a build machine was triggering the assert but I cannot reproduce it locally, I suspect an issue with the data that machine has stored in the DDC but I do not have a way to debug that)

#rnx

diff --git a/Engine/Source/Runtime/CoreUObject/Private/UObject/SavePackage.cpp b/Engine/Source/Runtime/CoreUObject/Private/UObject/SavePackage.cpp
--- a/Engine/Source/Runtime/CoreUObject/Private/UObject/SavePackage.cpp
+++ b/Engine/Source/Runtime/CoreUObject/Private/UObject/SavePackage.cpp
@@ -447,18 +447,34 @@ private:
 	TSet<FNameEntryId> ReferencedNames;
 };
 
-void FPackageStoreBulkDataManifest::PackageDesc::AddData(EIoChunkType InType, uint64 InChunkId, uint64 InOffset, uint64 InSize)
+void FPackageStoreBulkDataManifest::PackageDesc::AddData(EIoChunkType InType, uint64 InChunkId, uint64 InOffset, uint64 InSize, const FString& DebugFilename)
 {
 	// The ChunkId is supposed to be unique for each type
 	auto func = [=](const BulkDataDesc& Entry) { return Entry.ChunkId == InChunkId && Entry.Type == InType; };
-	check(Data.FindByPredicate(func) == nullptr);
+	if (Data.FindByPredicate(func) != nullptr)
+	{
+		FString Message;
+		for (int i = 0; i < Data.Num(); ++i)
+		{
+			const BulkDataDesc& Entry = Data[i];
+			Message.Appendf(TEXT("[%3d] ID: %20llu Offset: %8llu Size: %8llu Type: %d\n"),
+								 i, Entry.ChunkId, Entry.Offset, Entry.Size, Entry.Type);
+		}
+
+		Message.Appendf(TEXT(	"[New] ID: %20llu Offset: %8llu Size: %8llu Type: %d\n"),
+								InChunkId, InOffset, InSize, InType);
 
-	BulkDataDesc& Entry = Data.Emplace_GetRef();
+		UE_LOG(	LogSavePackage, Warning, TEXT("Duplicate BulkData description found in Package '%s', this will cause issues trying to run IoStore!\n%s"),*DebugFilename , *Message);
+	}
+	else
+	{
+		BulkDataDesc& Entry = Data.Emplace_GetRef();
 
-	Entry.ChunkId = InChunkId;
-	Entry.Offset = InOffset;
-	Entry.Size = InSize;
-	Entry.Type = InType;
+		Entry.ChunkId = InChunkId;
+		Entry.Offset = InOffset;
+		Entry.Size = InSize;
+		Entry.Type = InType;
+	}
 }
 
 FPackageStoreBulkDataManifest::FPackageStoreBulkDataManifest(const FString& ProjectPath)
@@ -520,7 +536,7 @@ void FPackageStoreBulkDataManifest::AddFileAccess(const FString& PackageFilename
 	const FString NormalizedFilename = FixFilename(PackageFilename);
 	
 	PackageDesc& Entry = GetOrCreateFileAccess(NormalizedFilename);
-	Entry.AddData(InType, InChunkId, InOffset, InSize);
+	Entry.AddData(InType, InChunkId, InOffset, InSize, NormalizedFilename);
 }
 
 FPackageStoreBulkDataManifest::PackageDesc& FPackageStoreBulkDataManifest::GetOrCreateFileAccess(const FString& PackageFilename)
diff --git a/Engine/Source/Runtime/CoreUObject/Public/UObject/SavePackage.h b/Engine/Source/Runtime/CoreUObject/Public/UObject/SavePackage.h
--- a/Engine/Source/Runtime/CoreUObject/Public/UObject/SavePackage.h
+++ b/Engine/Source/Runtime/CoreUObject/Public/UObject/SavePackage.h
@@ -37,7 +37,7 @@ public:
 			EIoChunkType Type;
 		};
 
-		void AddData(EIoChunkType InType, uint64 InChunkId, uint64 InOffset, uint64 InSize);
+		void AddData(EIoChunkType InType, uint64 InChunkId, uint64 InOffset, uint64 InSize, const FString& DebugFilename);
 		const TArray<BulkDataDesc>& GetDataArray() const { return Data; }
 	private:
 		friend FArchive& operator<<(FArchive& Ar, PackageDesc& Entry);
