A small clean up to BulkData serialize

- Remove a FString allocation (use a pointer rather than taking a copy as we no longer modify the filename)
- Move the offset adjustment code to a single location and do it only a single time

#rnx

diff --git a/Engine/Source/Runtime/CoreUObject/Private/Serialization/BulkData2.cpp b/Engine/Source/Runtime/CoreUObject/Private/Serialization/BulkData2.cpp
--- a/Engine/Source/Runtime/CoreUObject/Private/Serialization/BulkData2.cpp
+++ b/Engine/Source/Runtime/CoreUObject/Private/Serialization/BulkData2.cpp
@@ -577,19 +577,20 @@ void FBulkDataBase::Serialize(FArchive& Ar, UObject* Owner, int32 /*Index*/, boo
 		check(Package != nullptr);
 
 		if (!IsInlined() && bUseZenLoader)
-		{		
-			EIoChunkType Type = IsOptional() ? EIoChunkType::OptionalBulkData : EIoChunkType::BulkData;
+		{
+			const EIoChunkType Type = IsOptional() ? EIoChunkType::OptionalBulkData : EIoChunkType::BulkData;
 			ChunkID = CreateBulkdataChunkId(Package->GetGlobalPackageId(), BulkDataOffsetInFile, Type);
 
-			BulkDataFlags |= BULKDATA_UsesIoDispatcher;
+			BulkDataFlags |= BULKDATA_UsesIoDispatcher; // Indicates that this BulkData should use the FIoChunkId rather than a filename
 		}
 		else
 		{
+			// Invalidate the Token and then set the BulkDataSize for fast retrieval
 			Fallback.Token = InvalidToken;
 			Fallback.BulkDataSize = BulkDataSize;
 		}
 
-		FString Filename;
+		const FString* Filename = nullptr;
 		FName PackageName;
 		const FLinkerLoad* Linker = nullptr;
 
@@ -598,13 +599,7 @@ void FBulkDataBase::Serialize(FArchive& Ar, UObject* Owner, int32 /*Index*/, boo
 			Linker = FLinkerLoad::FindExistingLinkerForPackage(Package);
 			check(Linker != nullptr);
 
-			// Fix up the file offset, but only if not stored inline (TODO: Look into getting rid of this and cook with the correct values?)
-			if (!IsInlined())
-			{
-				BulkDataOffsetInFile += Linker->Summary.BulkDataStartOffset;
-			}
-
-			Filename = Linker->Filename;
+			Filename = &Linker->Filename;
 			PackageName = Package->FileName;
 		}
 
@@ -638,15 +633,15 @@ void FBulkDataBase::Serialize(FArchive& Ar, UObject* Owner, int32 /*Index*/, boo
 #endif
 					};
 
-					if (DoesOptionalDataExist(Filename))
+					if (DoesOptionalDataExist(*Filename))
 					{
-						SerializeDuplicateData(Ar, Linker, BulkDataFlags, BulkDataSizeOnDisk, BulkDataOffsetInFile);
+						SerializeDuplicateData(Ar, BulkDataFlags, BulkDataSizeOnDisk, BulkDataOffsetInFile);
 
-						// Update the fallback data again if needed
 						if (!IsInlined() && bUseZenLoader)
 						{
+							// Regenerate the FIoChunkId to find the optional BulkData instead!
 							ChunkID = CreateBulkdataChunkId(Package->GetGlobalPackageId(), BulkDataOffsetInFile, EIoChunkType::OptionalBulkData);
-							BulkDataFlags |= BULKDATA_UsesIoDispatcher;
+							BulkDataFlags |= BULKDATA_UsesIoDispatcher; // Indicates that this BulkData should use the FIoChunkId rather than a filename
 						}
 						else
 						{
@@ -660,17 +655,23 @@ void FBulkDataBase::Serialize(FArchive& Ar, UObject* Owner, int32 /*Index*/, boo
 						// out if we read things as 32bit or 64bit)
 						uint32 DummyValue32;
 						int64 DummyValue64;
-						SerializeDuplicateData(Ar, Linker, DummyValue32, DummyValue64, DummyValue64);
+						SerializeDuplicateData(Ar, DummyValue32, DummyValue64, DummyValue64);
 					}
 				}
+
+				// Fix up the file offset if we have a linker (if we do not then we will be loading via FIoDispatcher anyway)
+				if (Linker != nullptr)
+				{
+					BulkDataOffsetInFile += Linker->Summary.BulkDataStartOffset;
+				}
 			}
 
+			// If we are not using the FIoDispatcher then we need to make sure we can retrieve the filename later
 			if (bUseZenLoader == false)
 			{
-				Fallback.Token = FileTokenSystem::RegisterFileToken( PackageName, Filename, BulkDataOffsetInFile);
+				check(Filename != nullptr);
+				Fallback.Token = FileTokenSystem::RegisterFileToken( PackageName, *Filename, BulkDataOffsetInFile);
 			}
-
-		//	DebugFilename = ConvertFilenameFromFlags(Filename);
 		}
 		else
 		{
@@ -1084,7 +1085,7 @@ void FBulkDataBase::LoadDataDirectly(void** DstBuffer)
 	}
 }
 
-void FBulkDataBase::SerializeDuplicateData(FArchive& Ar, const FLinkerLoad* Linker, uint32& OutBulkDataFlags, int64& OutBulkDataSizeOnDisk, int64& OutBulkDataOffsetInFile)
+void FBulkDataBase::SerializeDuplicateData(FArchive& Ar, uint32& OutBulkDataFlags, int64& OutBulkDataSizeOnDisk, int64& OutBulkDataOffsetInFile)
 {
 	Ar << OutBulkDataFlags;
 
@@ -1107,12 +1108,6 @@ void FBulkDataBase::SerializeDuplicateData(FArchive& Ar, const FLinkerLoad* Link
 		uint16 DummyBulkDataIndex = InvalidBulkDataIndex;
 		Ar << DummyBulkDataIndex;
 	}
-
-	// Fix up the file offset if we have a valid linker
-	if (Linker != nullptr)
-	{
-		OutBulkDataOffsetInFile += Linker->Summary.BulkDataStartOffset;
-	}
 }
 
 void FBulkDataBase::SerializeBulkData(FArchive& Ar, void* DstBuffer, int64 DataLength)
diff --git a/Engine/Source/Runtime/CoreUObject/Public/Serialization/BulkData2.h b/Engine/Source/Runtime/CoreUObject/Public/Serialization/BulkData2.h
--- a/Engine/Source/Runtime/CoreUObject/Public/Serialization/BulkData2.h
+++ b/Engine/Source/Runtime/CoreUObject/Public/Serialization/BulkData2.h
@@ -137,7 +137,7 @@ public:
 private:
 	void LoadDataDirectly(void** DstBuffer);
 
-	void SerializeDuplicateData(FArchive& Ar, const FLinkerLoad* Linker, uint32& OutBulkDataFlags, int64& OutBulkDataSizeOnDisk, int64& OutBulkDataOffsetInFile);
+	void SerializeDuplicateData(FArchive& Ar, uint32& OutBulkDataFlags, int64& OutBulkDataSizeOnDisk, int64& OutBulkDataOffsetInFile);
 	void SerializeBulkData(FArchive& Ar, void* DstBuffer, int64 DataLength);
 
 	void AllocateData(SIZE_T SizeInBytes);
