AsyncLoading2: Use the same FSimpleExportArchive for all exports in a bundle.
Remove the unused internal CompletionCallbacks.
Tiny optimizations.

diff --git a/Engine/Source/Runtime/CoreUObject/Private/Serialization/AsyncLoading2.cpp b/Engine/Source/Runtime/CoreUObject/Private/Serialization/AsyncLoading2.cpp
--- a/Engine/Source/Runtime/CoreUObject/Private/Serialization/AsyncLoading2.cpp
+++ b/Engine/Source/Runtime/CoreUObject/Private/Serialization/AsyncLoading2.cpp
@@ -717,6 +717,7 @@ class FSimpleExportArchive
 public:
 	FSimpleExportArchive(const uint8* BufferPtr, uint64 BufferSize) : FSimpleArchive(BufferPtr, BufferSize) {}
 
+	void UsingCustomVersion(const FGuid& Key) override {};
 	using FArchive::operator<<; // For visibility of the overloads we don't override
 
 	//~ Begin FArchive::FArchiveUObject Interface
@@ -1104,7 +1105,7 @@ struct FAsyncPackage2
 		return Desc.NameToLoad;
 	}
 
-	void AddCompletionCallback(TUniquePtr<FLoadPackageAsyncDelegate>&& Callback, bool bInternal);
+	void AddCompletionCallback(TUniquePtr<FLoadPackageAsyncDelegate>&& Callback);
 
 	FORCEINLINE UPackage* GetLinkerRoot() const
 	{
@@ -1216,25 +1217,6 @@ private:
 	/** Checks if all dependencies (imported packages) of this package have been fully loaded */
 	bool AreAllDependenciesFullyLoadedInternal(FAsyncPackage2* Package, TSet<UPackage*>& VisitedPackages, FString& OutError);
 
-	struct FCompletionCallback
-	{
-		bool bIsInternal;
-		bool bCalled;
-		TUniquePtr<FLoadPackageAsyncDelegate> Callback;
-
-		FCompletionCallback()
-			: bIsInternal(false)
-			, bCalled(false)
-		{
-		}
-		FCompletionCallback(bool bInInternal, TUniquePtr<FLoadPackageAsyncDelegate>&& InCallback)
-			: bIsInternal(bInInternal)
-			, bCalled(false)
-			, Callback(MoveTemp(InCallback))
-		{
-		}
-	};
-
 	TAtomic<int32> RefCount{ 0 };
 
 	/** Basic information associated with this package */
@@ -1244,7 +1226,8 @@ private:
 	/** Package which is going to have its exports and imports loaded									*/
 	UPackage*				LinkerRoot;
 	/** Call backs called when we finished loading this package											*/
-	TArray<FCompletionCallback>	CompletionCallbacks;
+	using FCompletionCallback = TUniquePtr<FLoadPackageAsyncDelegate>;
+	TArray<FCompletionCallback, TInlineAllocator<2>> CompletionCallbacks;
 	/** Current index into ObjLoaded array used to spread routing PostLoad over several frames			*/
 	int32							PostLoadIndex;
 	/** Current index into DeferredPostLoadObjects array used to spread routing PostLoad over several frames			*/
@@ -1333,7 +1316,7 @@ public:
 	static EAsyncPackageState::Type Event_Delete(FAsyncPackage2* Package, int32);
 
 	void EventDrivenCreateExport(int32 LocalExportIndex);
-	void EventDrivenSerializeExport(int32 LocalExportIndex, const uint8* ExportData, uint64 ExportDataSize);
+	void EventDrivenSerializeExport(int32 LocalExportIndex, FSimpleExportArchive& Ar);
 
 	UObject* EventDrivenIndexToObject(FPackageIndex Index, bool bCheckSerialized);
 	template<class T>
@@ -1353,7 +1336,7 @@ public:
 
 	/** [EDL] End Event driven loader specific stuff */
 
-	void CallCompletionCallbacks(bool bInternalOnly, EAsyncLoadingResult::Type LoadingResult);
+	void CallCompletionCallbacks(EAsyncLoadingResult::Type LoadingResult);
 
 	/**
 	* Route PostLoad to deferred objects.
@@ -1974,8 +1957,7 @@ FAsyncPackage2* FAsyncLoadingThread2Impl::FindOrInsertPackage(FAsyncPackageDesc*
 		}
 		if (InDesc->PackageLoadedDelegate.IsValid())
 		{
-			const bool bInternalCallback = false;
-			Package->AddCompletionCallback(MoveTemp(InDesc->PackageLoadedDelegate), bInternalCallback);
+			Package->AddCompletionCallback(MoveTemp(InDesc->PackageLoadedDelegate));
 		}
 	}
 	return Package;
@@ -2620,6 +2602,29 @@ EAsyncPackageState::Type FAsyncPackage2::Event_ProcessExportBundle(FAsyncPackage
 
 	check(Package->AsyncPackageLoadingState == EAsyncPackageLoadingState2::ProcessNewImportsAndExports);
 
+	uint64 ExportsBufferSize = Package->IoBuffer.DataSize() - (Package->SerialDataPtr - Package->IoBuffer.Data());
+	FSimpleExportArchive Ar(Package->SerialDataPtr, ExportsBufferSize);
+	{
+		Ar.SetUE4Ver(Package->LinkerRoot->LinkerPackageVersion);
+		Ar.SetLicenseeUE4Ver(Package->LinkerRoot->LinkerLicenseeVersion);
+		// Ar.SetEngineVer(Summary.SavedByEngineVersion); // very old versioning scheme
+		// Ar.SetCustomVersions(LinkerRoot->LinkerCustomVersion); // only if not cooking with -unversioned
+		Ar.SetUseUnversionedPropertySerialization(CanUseUnversionedPropertySerialization());
+		Ar.SetIsLoading(true);
+		Ar.SetIsPersistent(true);
+		if (Package->LinkerRoot->GetPackageFlags() & PKG_FilterEditorOnly)
+		{
+			Ar.SetFilterEditorOnly(true);
+		}
+		Ar.ArAllowLazyLoading = true;
+
+		// FSimpleExportArchive special fields
+		Ar.PackageNameMap = Package->PackageNameMap;
+		Ar.GlobalNameMap = &Package->AsyncLoadingThread.GlobalNameMap.GetNameEntries();
+		Ar.ImportStore = &Package->ImportStore;
+		Ar.Exports = &Package->Exports;
+		Ar.ExternalReadDependencies = &Package->ExternalReadDependencies;
+	}
 	const FExportBundleHeader* ExportBundle = Package->ExportBundles + ExportBundleIndex;
 	
 	const FExportBundleEntry* BundleEntry = Package->ExportBundleEntries + ExportBundle->FirstEntryIndex;
@@ -2640,7 +2645,11 @@ EAsyncPackageState::Type FAsyncPackage2::Event_ProcessExportBundle(FAsyncPackage
 			check(Object);
 			if (Object->HasAnyFlags(RF_NeedLoad))
 			{
-				Package->EventDrivenSerializeExport(BundleEntry->LocalExportIndex, Package->SerialDataPtr, ExportSerialSize);
+				TRACE_LOADTIME_SERIALIZE_EXPORT_SCOPE(Object, ExportSerialSize);
+				const int64 Pos = Ar.Tell();
+				check(ExportSerialSize <= uint64(Ar.TotalSize() - Pos));
+				Package->EventDrivenSerializeExport(BundleEntry->LocalExportIndex, Ar);
+				check(ExportSerialSize == uint64(Ar.Tell() - Pos));
 			}
 			check(!Object->HasAnyFlags(RF_NeedLoad));
 
@@ -2865,7 +2874,7 @@ void FAsyncPackage2::EventDrivenCreateExport(int32 LocalExportIndex)
 	ImportStore.StoreGlobalImportObject(Export.GlobalImportIndex, Object);
 }
 
-void FAsyncPackage2::EventDrivenSerializeExport(int32 LocalExportIndex, const uint8* ExportData, uint64 ExportDataSize)
+void FAsyncPackage2::EventDrivenSerializeExport(int32 LocalExportIndex, FSimpleExportArchive& Ar)
 {
 	TRACE_CPUPROFILER_EVENT_SCOPE(SerializeExport);
 
@@ -2893,31 +2902,8 @@ void FAsyncPackage2::EventDrivenSerializeExport(int32 LocalExportIndex, const ui
 		}
 	}
 
-	FSimpleExportArchive Ar(ExportData, ExportDataSize);
-	Ar.SetUE4Ver(LinkerRoot->LinkerPackageVersion);
-	Ar.SetLicenseeUE4Ver(LinkerRoot->LinkerLicenseeVersion);
-	// Ar.SetEngineVer(Summary.SavedByEngineVersion); // very old versioning scheme
-	// Ar.SetCustomVersions(LinkerRoot->LinkerCustomVersion); // only if not cooking with -unversioned
-	Ar.SetIsLoading(true);
-	Ar.SetIsPersistent(true);
-	if (LinkerRoot->GetPackageFlags() & PKG_FilterEditorOnly)
-	{
-		Ar.SetFilterEditorOnly(true);
-	}
-	Ar.ArAllowLazyLoading = true;
-
-	// FSimpleExportArchive special fields
-	Ar.PackageNameMap = PackageNameMap;
-	Ar.GlobalNameMap = &AsyncLoadingThread.GlobalNameMap.GetNameEntries();
-	Ar.ImportStore = &ImportStore;
-	Ar.Exports = &Exports;
-	Ar.ExternalReadDependencies = &ExternalReadDependencies;
-	Ar.SetUseUnversionedPropertySerialization(CanUseUnversionedPropertySerialization());
-
 	Object->ClearFlags(RF_NeedLoad);
 
-	TRACE_LOADTIME_SERIALIZE_EXPORT_SCOPE(Object, ExportDataSize);
-
 	FUObjectSerializeContext* LoadContext = GetSerializeContext();
 	UObject* PrevSerializedObject = LoadContext->SerializedObject;
 	LoadContext->SerializedObject = Object;
@@ -2940,6 +2926,8 @@ void FAsyncPackage2::EventDrivenSerializeExport(int32 LocalExportIndex, const ui
 		Object->Serialize(Ar);
 	}
 
+	Ar.TemplateForGetArchetypeFromLoader = nullptr;
+
 	Object->SetFlags(RF_LoadCompleted);
 	LoadContext->SerializedObject = PrevSerializedObject;
 
@@ -3267,11 +3255,10 @@ EAsyncPackageState::Type FAsyncLoadingThread2Impl::ProcessLoadedPackagesFromGame
 			TRACE_LOADTIME_END_LOAD_ASYNC_PACKAGE(Package);
 
 			// Call external callbacks
-			const bool bInternalCallbacks = false;
 			const EAsyncLoadingResult::Type LoadingResult = Package->HasLoadFailed() ? EAsyncLoadingResult::Failed : EAsyncLoadingResult::Succeeded;
 			{
 				TRACE_CPUPROFILER_EVENT_SCOPE(PackageCompletionCallbacks);
-				Package->CallCompletionCallbacks(bInternalCallbacks, LoadingResult);
+				Package->CallCompletionCallbacks(LoadingResult);
 			}
 #if WITH_EDITOR
 			// In the editor we need to find any assets and add them to list for later callback
@@ -3964,6 +3951,7 @@ FAsyncPackage2::FAsyncPackage2(const FAsyncPackageDesc& InDesc, int32 InSerialNu
 	ExportBundleCount = PackageExportBundles.Get<1>();
 	ExportCount = AsyncLoadingThread.GlobalPackageStore.GetPackageExportCount(GlobalPackageId);
 	Exports.AddZeroed(ExportCount);
+	OwnedObjects.Reserve(ExportCount + 1); // +1 for UPackage
 
 	CreateNodes(EventSpecs);
 	if (GIsInitialLoad)
@@ -4469,9 +4457,6 @@ EAsyncPackageState::Type FAsyncPackage2::FinishObjects()
 	// FLinkerManager::Get().DissociateImportsAndForcedExports(); //@todo: this should be avoidable
 	PostLoadIndex = 0;
 
-	const bool bInternalCallbacks = true;
-	CallCompletionCallbacks(bInternalCallbacks, LoadingResult);
-
 	for (UObject* Object : OwnedObjects)
 	{
 		if (!Object->HasAnyFlags(RF_NeedPostLoad | RF_NeedPostLoadSubobjects))
@@ -4483,18 +4468,14 @@ EAsyncPackageState::Type FAsyncPackage2::FinishObjects()
 	return EAsyncPackageState::Complete;
 }
 
-void FAsyncPackage2::CallCompletionCallbacks(bool bInternal, EAsyncLoadingResult::Type LoadingResult)
+void FAsyncPackage2::CallCompletionCallbacks(EAsyncLoadingResult::Type LoadingResult)
 {
-	checkSlow(bInternal || !IsInAsyncLoadingThread());
+	checkSlow(!IsInAsyncLoadingThread());
 
 	UPackage* LoadedPackage = (!bLoadHasFailed) ? LinkerRoot : nullptr;
 	for (FCompletionCallback& CompletionCallback : CompletionCallbacks)
 	{
-		if (CompletionCallback.bIsInternal == bInternal && !CompletionCallback.bCalled)
-		{
-			CompletionCallback.bCalled = true;
-			CompletionCallback.Callback->ExecuteIfBound(Desc.Name, LoadedPackage, LoadingResult);
-		}
+		CompletionCallback->ExecuteIfBound(Desc.Name, LoadedPackage, LoadingResult);
 	}
 }
 
@@ -4509,8 +4490,7 @@ void FAsyncPackage2::Cancel()
 	// Call any completion callbacks specified.
 	bLoadHasFailed = true;
 	const EAsyncLoadingResult::Type Result = EAsyncLoadingResult::Canceled;
-	CallCompletionCallbacks(true, Result);
-	CallCompletionCallbacks(false, Result);
+	CallCompletionCallbacks(Result);
 
 	if (LinkerRoot)
 	{
@@ -4523,11 +4503,11 @@ void FAsyncPackage2::Cancel()
 	}
 }
 
-void FAsyncPackage2::AddCompletionCallback(TUniquePtr<FLoadPackageAsyncDelegate>&& Callback, bool bInternal)
+void FAsyncPackage2::AddCompletionCallback(TUniquePtr<FLoadPackageAsyncDelegate>&& Callback)
 {
 	// This is to ensure that there is no one trying to subscribe to a already loaded package
 	//check(!bLoadHasFinished && !bLoadHasFailed);
-	CompletionCallbacks.Emplace(bInternal, MoveTemp(Callback));
+	CompletionCallbacks.Emplace(MoveTemp(Callback));
 }
 
 void FAsyncPackage2::UpdateLoadPercentage()
