AsyncLoading2: Remove FLinkerLoad completely (was only used for old bulkdata path) .
Remove relative file names, reducing global namemap from ~95 MB to ~65 MB.
Merge CreateLinker and FinishLinker into CreateUPackage.

Todo: Reimplement CreateClustersForPackage without using FLinkerLoad

diff --git a/Engine/Source/Developer/IoStoreUtilities/Private/IoStoreUtilities.cpp b/Engine/Source/Developer/IoStoreUtilities/Private/IoStoreUtilities.cpp
--- a/Engine/Source/Developer/IoStoreUtilities/Private/IoStoreUtilities.cpp
+++ b/Engine/Source/Developer/IoStoreUtilities/Private/IoStoreUtilities.cpp
@@ -1067,7 +1067,6 @@ int32 CreateTarget(const FContainerTarget& Target)
 
 	const FString CookedDir = Target.CookedDirectory;
 	const FString OutputDir = Target.OutputDirectory;
-	const FString RelativePrefixForLegacyFilename = TEXT("../../../");
 	
 	FPackageStoreBulkDataManifest BulkDataManifest(Target.CookedProjectDirectory);
 	const bool bWithBulkDataManifest = BulkDataManifest.Load();
@@ -1159,6 +1158,7 @@ int32 CreateTarget(const FContainerTarget& Target)
 		FPackage& Package = *NewPackage;
 		Package.Name = PackageFName;
 		Package.FileName = FileName;
+		Package.RelativeFileName = MoveTemp(RelativeFileName);
 		Package.GlobalPackageId = FileIndex;
 		Package.UAssetSize = Ar->TotalSize();
 		Package.SummarySize = Ar->Tell() - SummaryStartPos;
@@ -1168,9 +1168,6 @@ int32 CreateTarget(const FContainerTarget& Target)
 		Package.PackageFlags = Summary.PackageFlags;
 		Package.BulkDataStartOffset = Summary.BulkDataStartOffset;
 
-		Package.RelativeFileName = RelativePrefixForLegacyFilename;
-		Package.RelativeFileName.Append(*FileName + CookedDir.Len());
-
 		if (Summary.NameCount > 0)
 		{
 			Ar->Seek(Summary.NameOffset);
@@ -1296,7 +1293,7 @@ int32 CreateTarget(const FContainerTarget& Target)
 			TempFullNames.SetNum(Summary.ExportCount, false);
 			for (int32 I = 0; I < Summary.ExportCount; ++I)
 			{
-				FindExport(GlobalExports, GlobalExportsByFullName, TempFullNames, Exports.GetData() + BaseIndex, I, PackageFName);
+				FindExport(GlobalExports, GlobalExportsByFullName, TempFullNames, Exports.GetData() + BaseIndex, I, Package.Name);
 
 				FExportData& ExportData = GlobalExports[*GlobalExportsByFullName.Find(TempFullNames[I])];
 				Package.Exports.Add(ExportData.GlobalIndex);
@@ -1721,21 +1718,17 @@ int32 CreateTarget(const FContainerTarget& Target)
 		}
 		Package.ExportBundlesSize = ZenExportBundlesArchive.Tell();
 
-		FName RelativeFileName(*Package.RelativeFileName);
-		NameMapBuilder.MarkNameAsReferenced(Package.Name);
-		int32 PackageNameIndex = NameMapBuilder.MapName(Package.Name);
-		int32 PackageNameNumber = Package.Name.GetNumber();
-		NameMapBuilder.MarkNameAsReferenced(RelativeFileName);
 
 		ensure(Package.GlobalPackageId == PackageIndex++);
 		{
 			int32 ExportBundleCount = Package.ExportBundles.Num();
+			int32 FirstExportBundleIndex = ExportBundleOrderEntries.Num() - ExportBundleCount;
+
+			NameMapBuilder.MarkNameAsReferenced(Package.Name);
 
 			NameMapBuilder.SerializeName(*StoreTocArchive, Package.Name);
-			NameMapBuilder.SerializeName(*StoreTocArchive, RelativeFileName);
 			*StoreTocArchive << Package.ExportCount;
 			*StoreTocArchive << ExportBundleCount;
-			int32 FirstExportBundleIndex = ExportBundleOrderEntries.Num() - ExportBundleCount;
 			*StoreTocArchive << FirstExportBundleIndex;
 			*StoreTocArchive << Package.FirstGlobalImport;
 			*StoreTocArchive << Package.GlobalImportCount;
@@ -1888,23 +1881,23 @@ int32 CreateTarget(const FContainerTarget& Target)
 		return IoWriter.Append(ChunkId, IoBuffer);
 	};
 
+	int32 StoreTocByteCount = StoreTocArchive->TotalSize();
+	int32 ImportedPackagesByteCount = ImportedPackagesArchive->TotalSize();
+	int32 GlobalImportNamesByteCount = GlobalImportNamesArchive->TotalSize();
+	int32 ExportBundleOrderByteCount = ExportBundleOrderEntries.Num() * sizeof(uint32);
 	{
 		UE_LOG(LogIoStore, Display, TEXT("Saving global meta data to container file"));
 		FLargeMemoryWriter GlobalMetaArchive(0, true);
 
-		int32 StoreTocByteCount = StoreTocArchive->TotalSize();
 		GlobalMetaArchive << StoreTocByteCount;
 		GlobalMetaArchive.Serialize(StoreTocArchive->GetData(), StoreTocByteCount);
 
-		int32 ImportedPackagesByteCount = ImportedPackagesArchive->TotalSize();
 		GlobalMetaArchive << ImportedPackagesByteCount;
 		GlobalMetaArchive.Serialize(ImportedPackagesArchive->GetData(), ImportedPackagesByteCount);
 
-		int32 GlobalImportNamesByteCount = GlobalImportNamesArchive->TotalSize();
 		GlobalMetaArchive << GlobalImportNamesByteCount;
 		GlobalMetaArchive.Serialize(GlobalImportNamesArchive->GetData(), GlobalImportNamesByteCount);
 
-		int32 ExportBundleOrderByteCount = ExportBundleOrderEntries.Num() * sizeof(uint32);
 		GlobalMetaArchive << ExportBundleOrderByteCount;
 		GlobalMetaArchive.Serialize(ExportBundleOrderEntries.GetData(), ExportBundleOrderByteCount);
 
@@ -1946,9 +1939,6 @@ int32 CreateTarget(const FContainerTarget& Target)
 	uint64 NameMapSize = 0;
 	uint64 NameMapCount = 0;
 	uint64 ZenPackageSummarySize = FileNames.Num() * sizeof(FZenPackageSummary);
-	uint64 StoreTocSize = StoreTocArchive->Tell();
-	uint64 GlobalImportNamesSize = GlobalImportNamesArchive->Tell();
-	uint64 PackageImportSize = ImportedPackagesArchive->Tell();
 	uint64 ImportedPackagesCount = 0;
 	uint64 InitialLoadSize = InitialLoadArchive->Tell();
 	uint64 ScriptArcsCount = 0;
@@ -2023,10 +2013,11 @@ int32 CreateTarget(const FContainerTarget& Target)
 		(double)PreloadDependenciesSize / 1024.0 / 1024.0, ImportPreloadCount, ExportPreloadCount, PreloadDependencies.Num());
 
 	UE_LOG(LogIoStore, Display, TEXT("IoStore: %8.2f MB GlobalNameMap, %d unique names"), (double)GlobalNameMapSize / 1024.0 / 1024.0, NameMapBuilder.Num());
-	UE_LOG(LogIoStore, Display, TEXT("IoStore: %8.2f MB GlobalPackageData"), (double)StoreTocSize / 1024.0 / 1024.0);
-	UE_LOG(LogIoStore, Display, TEXT("IoStore: %8.2f MB GlobalImportedPackages, %d imported packages"), (double)PackageImportSize / 1024.0 / 1024.0, ImportedPackagesCount);
+	UE_LOG(LogIoStore, Display, TEXT("IoStore: %8.2f MB GlobalPackageData"), (double)StoreTocByteCount / 1024.0 / 1024.0);
+	UE_LOG(LogIoStore, Display, TEXT("IoStore: %8.2f MB GlobalImportedPackages, %d imported packages"), (double)ImportedPackagesByteCount / 1024.0 / 1024.0, ImportedPackagesCount);
+	UE_LOG(LogIoStore, Display, TEXT("IoStore: %8.2f MB GlobalBundleOrders, %d bundles"), (double)ExportBundleOrderByteCount / 1024.0 / 1024.0, ExportBundleOrderEntries.Num());
 	UE_LOG(LogIoStore, Display, TEXT("IoStore: %8.2f MB GlobalImportNames, %d total imports, %d script imports, %d UPackage imports"),
-		(double)GlobalImportNamesSize / 1024.0 / 1024.0, GlobalImportsByFullName.Num(), NumScriptImports, UniqueImportPackages);
+		(double)GlobalImportNamesByteCount / 1024.0 / 1024.0, GlobalImportsByFullName.Num(), NumScriptImports, UniqueImportPackages);
 	UE_LOG(LogIoStore, Display, TEXT("IoStore: %8.2f MB InitialLoadData, %d script arcs, %d script outers, %d packages"), (double)InitialLoadSize / 1024.0 / 1024.0, ScriptArcsCount, NumScriptImports, FileNames.Num());
 	UE_LOG(LogIoStore, Display, TEXT("IoStore: %8.2f MB PackageHeader, %d packages"), (double)ZenPackageHeaderSize / 1024.0 / 1024.0, FileNames.Num());
 	UE_LOG(LogIoStore, Display, TEXT("IoStore: %8.2f MB PackageSummary"), (double)ZenPackageSummarySize / 1024.0 / 1024.0);
diff --git a/Engine/Source/Runtime/CoreUObject/Private/Serialization/AsyncLoading2.cpp b/Engine/Source/Runtime/CoreUObject/Private/Serialization/AsyncLoading2.cpp
--- a/Engine/Source/Runtime/CoreUObject/Private/Serialization/AsyncLoading2.cpp
+++ b/Engine/Source/Runtime/CoreUObject/Private/Serialization/AsyncLoading2.cpp
@@ -19,9 +19,7 @@
 #include "Misc/MessageDialog.h"
 #include "Misc/PackageName.h"
 #include "UObject/PackageFileSummary.h"
-#include "UObject/Linker.h"
 #include "UObject/SoftObjectPath.h"
-#include "UObject/LinkerLoad.h"
 #include "Serialization/DeferredMessageLog.h"
 #include "UObject/UObjectThreadContext.h"
 #include "Misc/Paths.h"
@@ -239,7 +237,6 @@ private:
 struct FPackageStoreEntry
 {
 	FName Name;
-	FName FileName;
 	int32 ExportCount;
 	int32 ExportBundleCount;
 	int32 FirstExportBundleIndex;
@@ -400,7 +397,6 @@ public:
 				for (int I = 0; I < PackageCount; ++I)
 				{
 					StoreEntries[I].Name = GlobalNameMap.FromSerializedName(StoreEntries[I].Name);
-					StoreEntries[I].FileName = GlobalNameMap.FromSerializedName(StoreEntries[I].FileName);
 				}
 			}
 
@@ -577,12 +573,6 @@ public:
 		const FPackageStoreEntry* StoreEntry = StoreEntries + GlobalPackageId.Id;
 		return MakeTuple(ExportBundleOrderEntries + StoreEntry->FirstExportBundleIndex, uint32(StoreEntry->ExportBundleCount));
 	}
-
-	inline FString GetPackageFileName(FGlobalPackageId GlobalPackageId)
-	{
-		return StoreEntries[GlobalPackageId.Id].FileName.ToString();
-	}
-
 };
 
 struct FPackageImportStore
@@ -1221,8 +1211,6 @@ private:
 
 	/** Basic information associated with this package */
 	FAsyncPackageDesc Desc;
-	/** Linker which is going to have its exports and imports loaded									*/
-	FLinkerLoad*				Linker;
 	/** Package which is going to have its exports and imports loaded									*/
 	UPackage*				LinkerRoot;
 	/** Call backs called when we finished loading this package											*/
@@ -1362,17 +1350,11 @@ private:
 	 */
 	void EndAsyncLoad();
 	/**
-	 * Create linker async. Linker is not finalized at this point.
+	 * Create UPackage
 	 *
 	 * @return true
 	 */
-	EAsyncPackageState::Type CreateLinker();
-	/**
-	 * Finalizes linker creation till time limit is exceeded.
-	 *
-	 * @return true if linker is finished being created, false otherwise
-	 */
-	EAsyncPackageState::Type FinishLinker(const FPackageSummary* PackageSummary);
+	void CreateUPackage(const FPackageSummary* PackageSummary);
 
 	/**
 	 * Route PostLoad to all loaded objects. This might load further objects!
@@ -2589,8 +2571,7 @@ EAsyncPackageState::Type FAsyncPackage2::Event_ProcessExportBundle(FAsyncPackage
 		uint64 PackageSummarySize = GraphData + PackageSummary->GraphDataSize - PackageSummaryData;
 		Package->SerialDataPtr += PackageSummarySize;
 
-		Package->CreateLinker();
-		Package->FinishLinker(PackageSummary);
+		Package->CreateUPackage(PackageSummary);
 
 		Package->SetupSerializedArcs(GraphData, PackageSummary->GraphDataSize);
 
@@ -2707,7 +2688,7 @@ UObject* FAsyncPackage2::EventDrivenIndexToObject(FPackageIndex Index, bool bChe
 		{
 			UE_LOG(LogStreaming, Fatal, TEXT("Missing Dependency, request for %s but it was still has RF_NeedLoad."), *Linker->GetPathName(Index));
 		}*/
-		UE_LOG(LogStreaming, Fatal, TEXT("Missing Dependency"), *Linker->GetPathName(Index));
+		UE_LOG(LogStreaming, Fatal, TEXT("Missing Dependency"));
 	}
 	if (Result)
 	{
@@ -3017,11 +2998,6 @@ EAsyncPackageState::Type FAsyncPackage2::Event_Tick(FAsyncPackage2* Package, int
 		}
 	} while (!FAsyncLoadingThreadState2::Get()->IsTimeLimitExceeded() && LoadingState == EAsyncPackageState::TimeOut);
 
-	if (Package->LinkerRoot && LoadingState == EAsyncPackageState::Complete)
-	{
-		Package->LinkerRoot->MarkAsFullyLoaded();
-	}
-
 	// Mark this package as loaded if everything completed.
 	Package->bLoadHasFinished = (LoadingState == EAsyncPackageState::Complete);
 
@@ -3920,7 +3896,6 @@ void FAsyncLoadingThread2Impl::FireCompletedCompiledInImport(void* AsyncPackage,
 */
 FAsyncPackage2::FAsyncPackage2(const FAsyncPackageDesc& InDesc, int32 InSerialNumber, FAsyncLoadingThread2Impl& InAsyncLoadingThread, IEDLBootNotificationManager& InEDLBootNotificationManager, FAsyncLoadEventGraphAllocator& InGraphAllocator, const FAsyncLoadEventSpec* EventSpecs, FGlobalPackageId InGlobalPackageId)
 : Desc(InDesc)
-, Linker(nullptr)
 , LinkerRoot(nullptr)
 , PostLoadIndex(0)
 , DeferredPostLoadIndex(0)
@@ -4063,7 +4038,6 @@ void FAsyncPackage2::ClearOwnedObjects()
 		Object->ClearInternalFlags(InternalFlagsToClear);
 	}
 	OwnedObjects.Empty();
-	LinkerRoot->LinkerLoad = nullptr; // TEMP 
 }
 
 void FAsyncPackage2::AddRequestID(int32 Id)
@@ -4134,80 +4108,44 @@ void FAsyncPackage2::EndAsyncLoad()
 	{
 		AsyncLoadingThread.LeaveAsyncLoadingTick();
 	}
-
-	if (!bLoadHasFailed)
-	{
-		// Mark the package as loaded, if we succeeded
-		LinkerRoot->SetFlags(RF_WasLoaded);
-	}
 }
 
-/**
- * Create linker async. Linker is not finalized at this point.
- *
- * @return true
- */
-EAsyncPackageState::Type FAsyncPackage2::CreateLinker()
+void FAsyncPackage2::CreateUPackage(const FPackageSummary* PackageSummary)
 {
-	TRACE_CPUPROFILER_EVENT_SCOPE(CreateLinker);
-	check(Linker == nullptr);
+	check(!LinkerRoot);
 
 	// Try to find existing package or create it if not already present.
-	UPackage* Package = nullptr;
 	{
-		{
-			TRACE_CPUPROFILER_EVENT_SCOPE(UPackageFind);
-			Package = FindObjectFast<UPackage>(nullptr, Desc.Name);
-		}
-		if (!Package)
-		{
-			TRACE_CPUPROFILER_EVENT_SCOPE(UPackageCreate);
-			Package = NewObject<UPackage>(/*Outer*/nullptr, Desc.Name, RF_Public);
-			Package->SetPackageFlags(Desc.PackageFlags);
-			Package->FileName = Desc.NameToLoad;
-			Package->SetGlobalPackageId(GlobalPackageId.Id);
-			bCreatedLinkerRoot = true;
-		}
+		TRACE_CPUPROFILER_EVENT_SCOPE(UPackageFind);
+		LinkerRoot = FindObjectFast<UPackage>(nullptr, Desc.Name);
 	}
-	check(!IsNativeCodePackage(Package));
-	AddOwnedObjectWithAsyncFlag(Package, /*bForceAdd*/ !bCreatedLinkerRoot);
-	check(Package->HasAnyInternalFlags(EInternalObjectFlags::Async));
-	LinkerRoot = Package;
-
-	check(!FLinkerLoad::FindExistingLinkerForPackage(Package));
-
+	if (!LinkerRoot)
 	{
-		uint32 LinkerFlags = LOAD_None | LOAD_Async | LOAD_NoVerify;
-		Linker = new FLinkerLoad(Package, *AsyncLoadingThread.GlobalPackageStore.GetPackageFileName(GlobalPackageId), LinkerFlags);
-		Linker->bIsAsyncLoader = false;
-		Linker->bLockoutLegacyOperations = true;
-		Linker->SetIsLoading(true);
-		Linker->SetIsPersistent(true);
+		TRACE_CPUPROFILER_EVENT_SCOPE(UPackageCreate);
+		LinkerRoot = NewObject<UPackage>(/*Outer*/nullptr, Desc.Name, RF_Public);
+		// LinkerRoot->FileName = Desc.NameToLoad;
+		LinkerRoot->SetGlobalPackageId(GlobalPackageId.Id);
+		LinkerRoot->SetPackageFlagsTo(PackageSummary->PackageFlags);
+		LinkerRoot->LinkerPackageVersion = GPackageFileUE4Version;
+		LinkerRoot->LinkerLicenseeVersion = GPackageFileLicenseeUE4Version;
+		// LinkerRoot->LinkerCustomVersion = PackageSummaryVersions; // only if (!bCustomVersionIsLatest)
+		LinkerRoot->SetFlags(RF_WasLoaded);
+		bCreatedLinkerRoot = true;
+	}
+	else
+	{
+		check(LinkerRoot->GetGlobalPackageId() == GlobalPackageId.Id);
+		check(LinkerRoot->GetPackageFlags() == PackageSummary->PackageFlags);
+		check(LinkerRoot->LinkerPackageVersion == GPackageFileUE4Version);
+		check(LinkerRoot->LinkerLicenseeVersion == GPackageFileLicenseeUE4Version);
+		check(LinkerRoot->HasAnyFlags(RF_WasLoaded));
 	}
-	Package->LinkerLoad = Linker;
-
-	UE_LOG(LogStreaming, Verbose, TEXT("FAsyncPackage::CreateLinker for %s finished."), *Desc.NameToLoad.ToString());
-	return EAsyncPackageState::Complete;
-}
-
-/**
- * Finalizes linker creation till time limit is exceeded.
- *
- * @return true if linker is finished being created, false otherwise
- */
-EAsyncPackageState::Type FAsyncPackage2::FinishLinker(const FPackageSummary* PackageSummary)
-{
-	TRACE_CPUPROFILER_EVENT_SCOPE(FinishLinker);
-	LLM_SCOPE(ELLMTag::AsyncLoading);
-
-	FPackageFileSummary& Summary	= Linker->Summary;
-	Summary.PackageFlags			= PackageSummary->PackageFlags;
-	Summary.BulkDataStartOffset		= PackageSummary->BulkDataStartOffset;
-	Summary.SetFileVersions(GPackageFileUE4Version, GPackageFileLicenseeUE4Version, /*unversioned*/true);
-	Linker->UpdateFromPackageFileSummary();
 
 	ImportStore.StoreGlobalImportObject(PackageSummary->GlobalImportIndex, LinkerRoot);
-	return EAsyncPackageState::Complete;
+	AddOwnedObjectWithAsyncFlag(LinkerRoot, /*bForceAdd*/ !bCreatedLinkerRoot);
+	check(LinkerRoot->HasAnyInternalFlags(EInternalObjectFlags::Async));
+
+	UE_LOG(LogStreaming, Verbose, TEXT("FAsyncPackage::CreateUPackage for %s finished."), *Desc.NameToLoad.ToString());
 }
 
 EAsyncPackageState::Type FAsyncPackage2::FinishExternalReadDependencies()
@@ -4273,7 +4211,7 @@ EAsyncPackageState::Type FAsyncPackage2::PostLoadObjects()
 	return (PostLoadIndex == ExportCount) ? EAsyncPackageState::Complete : EAsyncPackageState::TimeOut;
 }
 
-void CreateClustersFromPackage(FLinkerLoad* PackageLinker, TArray<UObject*>& OutClusterObjects);
+// void CreateClustersFromPackage(FLinkerLoad* PackageLinker, TArray<UObject*>& OutClusterObjects);
 
 EAsyncPackageState::Type FAsyncPackage2::PostLoadDeferredObjects()
 {
@@ -4387,10 +4325,10 @@ EAsyncPackageState::Type FAsyncPackage2::PostLoadDeferredObjects()
 			LinkerRoot->MarkAsFullyLoaded();			
 			LinkerRoot->SetLoadTime(FPlatformTime::Seconds() - LoadStartTime);
 
-			if (Linker)
-			{
-				CreateClustersFromPackage(Linker, DeferredClusterObjects);
-			}
+			// if (Linker)
+			// {
+			// 	CreateClustersFromPackage(Linker, DeferredClusterObjects);
+			// }
 		}
 
 		FSoftObjectPath::InvalidateTag();
@@ -4429,9 +4367,6 @@ EAsyncPackageState::Type FAsyncPackage2::FinishObjects()
 {
 	SCOPED_LOADTIMER(FinishObjectsTime);
 
-	FUObjectSerializeContext* LoadContext = GetSerializeContext();
-	check(!Linker || LoadContext == Linker->GetSerializeContext());		
-
 	EAsyncLoadingResult::Type LoadingResult;
 	if (!bLoadHasFailed)
 	{
@@ -4453,10 +4388,6 @@ EAsyncPackageState::Type FAsyncPackage2::FinishObjects()
 		LoadingResult = EAsyncLoadingResult::Failed;
 	}
 
-	// Simulate what EndLoad does.
-	// FLinkerManager::Get().DissociateImportsAndForcedExports(); //@todo: this should be avoidable
-	PostLoadIndex = 0;
-
 	for (UObject* Object : OwnedObjects)
 	{
 		if (!Object->HasAnyFlags(RF_NeedPostLoad | RF_NeedPostLoadSubobjects))
