AsyncLoading2: Include fix for 10482167

diff --git a/Engine/Source/Runtime/CoreUObject/Private/Serialization/AsyncLoading2.cpp b/Engine/Source/Runtime/CoreUObject/Private/Serialization/AsyncLoading2.cpp
--- a/Engine/Source/Runtime/CoreUObject/Private/Serialization/AsyncLoading2.cpp
+++ b/Engine/Source/Runtime/CoreUObject/Private/Serialization/AsyncLoading2.cpp
@@ -18,7 +18,7 @@
 #include "Misc/App.h"
 #include "Misc/MessageDialog.h"
 #include "Misc/PackageName.h"
-#include "UObject/PackageFileSummary.h"
+#include "UObject/ObjectResource.h"
 #include "UObject/SoftObjectPath.h"
 #include "Serialization/DeferredMessageLog.h"
 #include "UObject/UObjectThreadContext.h"
