AsyncLoading2: Add FPackageId and FAsyncPackageDesc2 to remove all internal use of FNames.
Optimize LoadPackageAsync by simplifying package name checks, from ~260 ms to ~60 ms.
Simplify QueuePackage and CreateAsyncPackage based on always having only valid FAsyncPackageDesc2 requests.
Use FPackageId in SetupSerializedArcs.
Remove AsyncPackageSerialNumber (a leftover from when the async package was a weak pointer).
Some fixes for static analysis warnings.

diff --git a/Engine/Source/Developer/IoStoreUtilities/Private/IoStoreUtilities.cpp b/Engine/Source/Developer/IoStoreUtilities/Private/IoStoreUtilities.cpp
--- a/Engine/Source/Developer/IoStoreUtilities/Private/IoStoreUtilities.cpp
+++ b/Engine/Source/Developer/IoStoreUtilities/Private/IoStoreUtilities.cpp
@@ -399,7 +399,7 @@ struct FPackage
 	TArray<int32> Imports;
 	TArray<int32> Exports;
 	TArray<FArc> InternalArcs;
-	TMap<FName, TArray<FArc>> ExternalArcs;
+	TMap<FPackage*, TArray<FArc>> ExternalArcs;
 	TArray<FArc> ScriptArcs;
 	
 	TArray<FExportBundle> ExportBundles;
@@ -615,7 +615,7 @@ static void AddPostLoadArc(FPackage& FromPackage, FPackage& ToPackage)
 
 static void AddStartPostLoadArc(FPackage& FromPackage, FPackage& ToPackage)
 {
-	TArray<FArc>& ExternalArcs = ToPackage.ExternalArcs.FindOrAdd(FromPackage.Name);
+	TArray<FArc>& ExternalArcs = ToPackage.ExternalArcs.FindOrAdd(&FromPackage);
 	check(!ExternalArcs.Contains(FArc({EEventLoadNode2::Package_ExportsSerialized, EEventLoadNode2::Package_StartPostLoad})));
 	check(!ExternalArcs.Contains(FArc({EEventLoadNode2::Package_StartPostLoad, EEventLoadNode2::Package_StartPostLoad})));
 	ExternalArcs.Add({ EEventLoadNode2::Package_StartPostLoad, EEventLoadNode2::Package_StartPostLoad });
@@ -623,7 +623,7 @@ static void AddStartPostLoadArc(FPackage& FromPackage, FPackage& ToPackage)
 
 static void AddExportsDoneArc(FPackage& FromPackage, FPackage& ToPackage)
 {
-	TArray<FArc>& ExternalArcs = ToPackage.ExternalArcs.FindOrAdd(FromPackage.Name);
+	TArray<FArc>& ExternalArcs = ToPackage.ExternalArcs.FindOrAdd(&FromPackage);
 	check(!ExternalArcs.Contains(FArc({EEventLoadNode2::Package_ExportsSerialized, EEventLoadNode2::Package_StartPostLoad})));
 	check(!ExternalArcs.Contains(FArc({EEventLoadNode2::Package_StartPostLoad, EEventLoadNode2::Package_StartPostLoad})));
 	ExternalArcs.Add({ EEventLoadNode2::Package_ExportsSerialized, EEventLoadNode2::Package_StartPostLoad });
@@ -633,7 +633,7 @@ static void AddUniqueExternalBundleArc(FPackage& FromPackage, uint32 FromBundleI
 {
 	uint32 FromNodeIndex = EEventLoadNode2::Package_NumPhases + FromBundleIndex * EEventLoadNode2::ExportBundle_NumPhases + EEventLoadNode2::ExportBundle_Process;
 	uint32 ToNodeIndex = EEventLoadNode2::Package_NumPhases + ToBundleIndex * EEventLoadNode2::ExportBundle_NumPhases + EEventLoadNode2::ExportBundle_Process;
-	TArray<FArc>& ExternalArcs = ToPackage.ExternalArcs.FindOrAdd(FromPackage.Name);
+	TArray<FArc>& ExternalArcs = ToPackage.ExternalArcs.FindOrAdd(&FromPackage);
 	ExternalArcs.AddUnique({ FromNodeIndex, ToNodeIndex });
 }
 
@@ -1352,7 +1352,10 @@ int32 CreateTarget(const FContainerTarget& Target)
 
 			GlobalImportsByFullName[GlobalImport.FullName] = I;
 		}
-		LastPackage->GlobalImportCount = GlobalImports.Num() - LastPackage->FirstGlobalImport;
+		if (LastPackage)
+		{
+			LastPackage->GlobalImportCount = GlobalImports.Num() - LastPackage->FirstGlobalImport;
+		}
 	}
 
 	for (FExportData& GlobalExport : GlobalExports)
@@ -1653,12 +1656,11 @@ int32 CreateTarget(const FContainerTarget& Target)
 		ZenGraphArchive << ImportedPackagesCount;
 		for (auto& KV : Package.ExternalArcs)
 		{
-			const FName& ImportedPackageName = KV.Key;
-			NameMapBuilder.SerializeName(ZenGraphArchive, ImportedPackageName);
-
+			FPackage* ImportedPackage = KV.Key;
 			TArray<FArc>& Arcs = KV.Value;
-
 			int32 ExternalArcCount = Arcs.Num();
+
+			ZenGraphArchive << ImportedPackage->GlobalPackageId;
 			ZenGraphArchive << ExternalArcCount;
 			for (FArc& ExternalArc : Arcs)
 			{
diff --git a/Engine/Source/Runtime/CoreUObject/Private/Serialization/AsyncLoading2.cpp b/Engine/Source/Runtime/CoreUObject/Private/Serialization/AsyncLoading2.cpp
--- a/Engine/Source/Runtime/CoreUObject/Private/Serialization/AsyncLoading2.cpp
+++ b/Engine/Source/Runtime/CoreUObject/Private/Serialization/AsyncLoading2.cpp
@@ -110,21 +110,6 @@ public:
 	}
 };
 
-struct FGlobalPackageId
-{
-	int32 Id;
-
-	bool operator==(FGlobalPackageId RHS) const
-	{
-		return Id == RHS.Id;
-	}
-
-	friend uint32 GetTypeHash(FGlobalPackageId GlobalPackageId)
-	{
-		return GetTypeHash(GlobalPackageId.Id);
-	}
-};
-
 struct FExportMapEntry
 {
 	uint64 SerialSize;
@@ -341,7 +326,7 @@ struct FGlobalImportStore
 class FPackageStore
 {
 public:
-	TMap<FName, FGlobalPackageId> PackageNameToGlobalPackageId;
+	TMap<FName, FPackageId> PackageNameToPackageId;
 	FGlobalImportStore ImportStore;
 	FIoBuffer InitialLoadIoBuffer;
 	FPackageStoreEntry* StoreEntries = nullptr;
@@ -400,10 +385,10 @@ public:
 			// Global package index
 			{
 				TRACE_CPUPROFILER_EVENT_SCOPE(LoadPackageStoreGlobalIds);
-				PackageNameToGlobalPackageId.Reserve(PackageCount);
+				PackageNameToPackageId.Reserve(PackageCount);
 				for (int I = 0; I < PackageCount; ++I)
 				{
-					PackageNameToGlobalPackageId.Add(StoreEntries[I].Name, FGlobalPackageId { I });
+					PackageNameToPackageId.Add(StoreEntries[I].Name, FPackageId::FromIndex(I));
 				}
 			}
 		}
@@ -504,9 +489,9 @@ public:
 		return ImportStore;
 	}
 
-	inline FGlobalPackageId* GetGlobalPackageId(FName Name)
+	inline FPackageId* GetPackageId(FName Name)
 	{
-		FGlobalPackageId* Id = PackageNameToGlobalPackageId.Find(Name);
+		FPackageId* Id = PackageNameToPackageId.Find(Name);
 		return Id;
 	}
 
@@ -537,9 +522,9 @@ public:
 		return StoreEntries;
 	}
 
-	inline int32* GetPackageImportedPackages(FGlobalPackageId GlobalPackageId, int32& OutCount) const 
+	inline int32* GetPackageImportedPackages(FPackageId PackageId, int32& OutCount) const 
 	{
-		FPackageStoreEntry& Entry = StoreEntries[GlobalPackageId.Id];
+		FPackageStoreEntry& Entry = StoreEntries[PackageId.ToIndex()];
 		OutCount = Entry.ImportedPackagesCount;
 		if (Entry.ImportedPackagesCount >= 0)
 		{
@@ -548,11 +533,11 @@ public:
 		return nullptr;
 	}
 
-	inline int32* GetPackageScriptArcs(FGlobalPackageId GlobalPackageId, int32& OutCount) const 
+	inline int32* GetPackageScriptArcs(FPackageId PackageId, int32& OutCount) const 
 	{
 		if (ensure(GIsInitialLoad))
 		{
-			FPackageStoreInitialLoadEntry& Entry = InitialLoadStoreEntries[GlobalPackageId.Id];
+			FPackageStoreInitialLoadEntry& Entry = InitialLoadStoreEntries[PackageId.ToIndex()];
 			OutCount = Entry.ScriptArcsCount;
 			return ScriptArcs + Entry.ScriptArcsOffset / sizeof(int32);
 		}
@@ -560,14 +545,14 @@ public:
 		return nullptr;
 	}
 
-	inline int32 GetPackageExportCount(FGlobalPackageId GlobalPackageId) const 
+	inline int32 GetPackageExportCount(FPackageId PackageId) const 
 	{
-		return StoreEntries[GlobalPackageId.Id].ExportCount;
+		return StoreEntries[PackageId.ToIndex()].ExportCount;
 	}
 
-	inline TTuple<uint32*, uint32> GetPackageExportBundles(FGlobalPackageId GlobalPackageId) const
+	inline TTuple<uint32*, uint32> GetPackageExportBundles(FPackageId PackageId) const
 	{
-		const FPackageStoreEntry* StoreEntry = StoreEntries + GlobalPackageId.Id;
+		const FPackageStoreEntry* StoreEntry = StoreEntries + PackageId.ToIndex();
 		return MakeTuple(ExportBundleOrderEntries + StoreEntry->FirstExportBundleIndex, uint32(StoreEntry->ExportBundleCount));
 	}
 };
@@ -580,12 +565,12 @@ struct FPackageImportStore
 	const FPackageIndex* GlobalImportScriptOuters = nullptr;
 	const int32* ImportedPackages = nullptr;
 	const int32* ImportMap = nullptr;
-	FGlobalPackageId GlobalPackageId;
+	FPackageId PackageId;
 	int32 GlobalScriptImportCount = 0;
 	int32 ImportedPackagesCount = 0;
 	int32 ImportMapCount = 0;
 
-	FPackageImportStore(FPackageStore& GlobalPackageStore, FGlobalPackageId InGlobalPackageId) : GlobalPackageId(InGlobalPackageId)
+	FPackageImportStore(FPackageStore& GlobalPackageStore, FPackageId InPackageId) : PackageId(InPackageId)
 	{
 		TRACE_CPUPROFILER_EVENT_SCOPE(NewPackageImportStore);
 		{
@@ -595,7 +580,7 @@ struct FPackageImportStore
 			GlobalImportNames = GlobalPackageStore.GetGlobalImportNames(GlobalImportCount);
 		}
 		GlobalImportScriptOuters = GlobalPackageStore.GetGlobalImportScriptOuters(GlobalScriptImportCount);
-		ImportedPackages = GlobalPackageStore.GetPackageImportedPackages(GlobalPackageId, ImportedPackagesCount);
+		ImportedPackages = GlobalPackageStore.GetPackageImportedPackages(PackageId, ImportedPackagesCount);
 		AddGlobalImportObjectReferences();
 	}
 
@@ -668,7 +653,7 @@ private:
 			}
 		}
 		// Add ref counts to own imports to speed up OnPreGarbageCollect
-		FPackageStoreEntry& Entry = GlobalStoreEntries[GlobalPackageId.Id];
+		FPackageStoreEntry& Entry = GlobalStoreEntries[PackageId.ToIndex()];
 		for (int32 ObjectImportIndex = 0; ObjectImportIndex < Entry.GlobalImportCount; ObjectImportIndex++)
 		{
 			int32 GlobalImportIndex = Entry.FirstGlobalImport + ObjectImportIndex;
@@ -689,7 +674,7 @@ private:
 			}
 		}
 		// Release ref counts to own imports to speed up OnPreGarbageCollect
-		FPackageStoreEntry& Entry = GlobalStoreEntries[GlobalPackageId.Id];
+		FPackageStoreEntry& Entry = GlobalStoreEntries[PackageId.ToIndex()];
 		for (int32 ObjectImportIndex = 0; ObjectImportIndex < Entry.GlobalImportCount; ObjectImportIndex++)
 		{
 			int32 GlobalImportIndex = Entry.FirstGlobalImport + ObjectImportIndex;
@@ -1029,6 +1014,48 @@ struct FAsyncLoadingThreadState2
 
 uint32 FAsyncLoadingThreadState2::TlsSlot;
 
+struct FAsyncPackageDesc2
+{
+	int32 RequestID;
+	FPackageId PackageId;
+	FName Name;
+	/** Delegate called on completion of loading. This delegate can only be created and consumed on the game thread */
+	TUniquePtr<FLoadPackageAsyncDelegate> PackageLoadedDelegate;
+
+	FAsyncPackageDesc2(
+		int32 InRequestID,
+		FPackageId InPackageId,
+		const FName& InName,
+		TUniquePtr<FLoadPackageAsyncDelegate>&& InCompletionDelegate = TUniquePtr<FLoadPackageAsyncDelegate>())
+		: RequestID(InRequestID)
+		, PackageId(InPackageId)
+		, Name(InName)
+		, PackageLoadedDelegate(MoveTemp(InCompletionDelegate))
+	{
+	}
+
+	/** This constructor does not modify the package loaded delegate as this is not safe outside the game thread */
+	FAsyncPackageDesc2(const FAsyncPackageDesc2& OldPackage)
+		: RequestID(OldPackage.RequestID)
+		, PackageId(OldPackage.PackageId)
+		, Name(OldPackage.Name)
+	{
+	}
+
+	/** This constructor will explicitly copy the package loaded delegate and invalidate the old one */
+	FAsyncPackageDesc2(const FAsyncPackageDesc2& OldPackage, TUniquePtr<FLoadPackageAsyncDelegate>&& InPackageLoadedDelegate)
+		: FAsyncPackageDesc2(OldPackage)
+	{
+		PackageLoadedDelegate = MoveTemp(InPackageLoadedDelegate);
+	}
+
+#if DO_GUARD_SLOW
+	~FAsyncPackageDesc2()
+	{
+		checkSlow(!PackageLoadedDelegate.IsValid() || IsInGameThread());
+	}
+#endif
+};
 /**
 * Structure containing intermediate data required for async loading of all exports of a package.
 */
@@ -1038,7 +1065,11 @@ struct FAsyncPackage2
 	friend struct FScopedAsyncPackageEvent2;
 	friend class FAsyncLoadingThread2Impl;
 
-	FAsyncPackage2(const FAsyncPackageDesc& InDesc, int32 InSerialNumber, FAsyncLoadingThread2Impl& InAsyncLoadingThread, IEDLBootNotificationManager& InEDLBootNotificationManager, FAsyncLoadEventGraphAllocator& InGraphAllocator, const FAsyncLoadEventSpec* EventSpecs, FGlobalPackageId InGlobalPackageId);
+	FAsyncPackage2(const FAsyncPackageDesc2& InDesc,
+		FAsyncLoadingThread2Impl& InAsyncLoadingThread,
+		IEDLBootNotificationManager& InEDLBootNotificationManager,
+		FAsyncLoadEventGraphAllocator& InGraphAllocator,
+		const FAsyncLoadEventSpec* EventSpecs);
 	virtual ~FAsyncPackage2();
 
 
@@ -1084,12 +1115,9 @@ struct FAsyncPackage2
 		return Desc.Name;
 	}
 
-	/**
-	 * Returns the name to load of the package.
-	 */
-	FORCEINLINE const FName& GetPackageNameToLoad() const
+	FORCEINLINE FPackageId GetPackageId() const
 	{
-		return Desc.NameToLoad;
+		return Desc.PackageId;
 	}
 
 	void AddCompletionCallback(TUniquePtr<FLoadPackageAsyncDelegate>&& Callback);
@@ -1105,18 +1133,6 @@ struct FAsyncPackage2
 		return bLoadHasFinished;
 	}
 
-	/** Returns package loading priority. */
-	FORCEINLINE TAsyncLoadPriority GetPriority() const
-	{
-		return Desc.Priority;
-	}
-
-	/** Returns package loading priority. */
-	FORCEINLINE void SetPriority(TAsyncLoadPriority InPriority)
-	{
-		Desc.Priority = InPriority;
-	}
-
 	/** Returns true if loading has failed */
 	FORCEINLINE bool HasLoadFailed() const
 	{
@@ -1207,7 +1223,7 @@ private:
 	TAtomic<int32> RefCount{ 0 };
 
 	/** Basic information associated with this package */
-	FAsyncPackageDesc Desc;
+	FAsyncPackageDesc2 Desc;
 	/** Package which is going to have its exports and imports loaded									*/
 	UPackage*				LinkerRoot;
 	/** Call backs called when we finished loading this package											*/
@@ -1249,8 +1265,6 @@ private:
 	FAsyncLoadingThread2Impl& AsyncLoadingThread;
 	IEDLBootNotificationManager& EDLBootNotificationManager;
 	FAsyncLoadEventGraphAllocator& GraphAllocator;
-	/** Global package id used to index into the package store											*/
-	FGlobalPackageId GlobalPackageId;
 	/** Packages that have been imported by this async package */
 	TSet<UPackage*> ImportedPackages;
 
@@ -1288,7 +1302,6 @@ public:
 	/** [EDL] Begin Event driven loader specific stuff */
 
 	EAsyncPackageLoadingState2 AsyncPackageLoadingState;
-	int32 SerialNumber;
 
 	bool bHasImportedPackagesRecursive = false;
 
@@ -1475,7 +1488,7 @@ private:
 	/** [ASYNC/GAME THREAD] Event used to signal that the async loading thread has resumed */
 	FEvent* ThreadResumedEvent;
 	/** [ASYNC/GAME THREAD] List of queued packages to stream */
-	TArray<FAsyncPackageDesc*> QueuedPackages;
+	TArray<FAsyncPackageDesc2*> QueuedPackages;
 	/** [ASYNC/GAME THREAD] Package queue critical section */
 	FCriticalSection QueueCritical;
 	/** [ASYNC/GAME THREAD] Event used to signal there's queued packages to stream */
@@ -1489,7 +1502,7 @@ private:
 #endif
 
 	FCriticalSection AsyncPackagesCritical;
-	TMap<FName, FAsyncPackage2*> AsyncPackageNameLookup;
+	TMap<FPackageId, FAsyncPackage2*> AsyncPackageLookup;
 
 	IEDLBootNotificationManager& EDLBootNotificationManager;
 
@@ -1506,7 +1519,7 @@ private:
 	FThreadSafeCounter AsyncThreadReady;
 
 	/** When cancelling async loading: list of package requests to cancel */
-	TArray<FAsyncPackageDesc*> QueuedPackagesToCancel;
+	TArray<FAsyncPackageDesc2*> QueuedPackagesToCancel;
 	/** When cancelling async loading: list of packages to cancel */
 	TSet<FAsyncPackage2*> PackagesToCancel;
 
@@ -1514,7 +1527,6 @@ private:
 	uint32 AsyncLoadingThreadID;
 
 	FThreadSafeCounter PackageRequestID;
-	FThreadSafeCounter AsyncPackageSerialNumber;
 
 	/** I/O Dispatcher */
 	FGlobalNameMap GlobalNameMap;
@@ -1642,9 +1654,22 @@ public:
 	FORCEINLINE FAsyncPackage2* FindAsyncPackage(const FName& PackageName)
 	{
 		TRACE_CPUPROFILER_EVENT_SCOPE(FindAsyncPackage);
+		FPackageId* PackageId = GlobalPackageStore.GetPackageId(PackageName);
+		if (PackageId)
+		{
+			FScopeLock LockAsyncPackages(&AsyncPackagesCritical);
+			//checkSlow(IsInAsyncLoadThread());
+			return AsyncPackageLookup.FindRef(*PackageId);
+		}
+		return nullptr;
+	}
+
+	FORCEINLINE FAsyncPackage2* GetAsyncPackage(const FPackageId& PackageId)
+	{
+		TRACE_CPUPROFILER_EVENT_SCOPE(GetAsyncPackage);
 		FScopeLock LockAsyncPackages(&AsyncPackagesCritical);
 		//checkSlow(IsInAsyncLoadThread());
-		return AsyncPackageNameLookup.FindRef(PackageName);
+		return AsyncPackageLookup.FindRef(PackageId);
 	}
 
 	/**
@@ -1655,14 +1680,14 @@ public:
 	*/
 	void InsertPackage(FAsyncPackage2* Package, bool bReinsert = false);
 
-	FAsyncPackage2* FindOrInsertPackage(FAsyncPackageDesc* InDesc, bool& bInserted);
+	FAsyncPackage2* FindOrInsertPackage(FAsyncPackageDesc2* InDesc, bool& bInserted);
 
 	/**
 	* [ASYNC/GAME THREAD] Queues a package for streaming.
 	*
 	* @param Package package descriptor.
 	*/
-	void QueuePackage(FAsyncPackageDesc& Package);
+	void QueuePackage(FAsyncPackageDesc2& Package);
 
 	/**
 	* [ASYNC* THREAD] Loads all packages
@@ -1784,22 +1809,11 @@ private:
 	void BundleIoRequestCompleted();
 	void StartBundleIoRequests();
 
-	FAsyncPackage2* CreateAsyncPackage(const FAsyncPackageDesc& Desc)
+	FAsyncPackage2* CreateAsyncPackage(const FAsyncPackageDesc2& Desc)
 	{
-		FGlobalPackageId* GlobalPackageId = GlobalPackageStore.GetGlobalPackageId(Desc.NameToLoad);
-		if (!GlobalPackageId)
-		{
-			return nullptr;
-		}
-		check(GlobalPackageId);
-		return new FAsyncPackage2(Desc, AsyncPackageSerialNumber.Increment(), *this, EDLBootNotificationManager, GraphAllocator, EventSpecs.GetData(), *GlobalPackageId);
+		return new FAsyncPackage2(Desc, *this, EDLBootNotificationManager, GraphAllocator, EventSpecs.GetData());
 	}
 
-	/**
-	* [ASYNC THREAD] Internal helper function for updating the priorities of an existing package and all its dependencies
-	*/
-	void UpdateExistingPackagePriorities(FAsyncPackage2* InPackage, TAsyncLoadPriority InNewPriority);
-
 	/**
 	* [ASYNC THREAD] Adds a package to a list of packages that have finished loading on the async thread
 	*/
@@ -1887,56 +1901,44 @@ void FAsyncLoadingThread2Impl::InitializeLoading()
 	AsyncThreadReady.Increment();
 }
 
-void FAsyncLoadingThread2Impl::QueuePackage(FAsyncPackageDesc& Package)
+void FAsyncLoadingThread2Impl::QueuePackage(FAsyncPackageDesc2& Package)
 {
 	//TRACE_CPUPROFILER_EVENT_SCOPE(QueuePackage);
-	if (!GlobalPackageStore.GetGlobalPackageId(Package.NameToLoad))
-	{
-		UE_LOG(LogStreaming, Warning, TEXT("QueuePackage: Skipping unknown package: %s"), *Package.NameToLoad.ToString());
-		return;
-	}
 	{
 		FScopeLock QueueLock(&QueueCritical);
 		++QueuedPackagesCounter;
-		QueuedPackages.Add(new FAsyncPackageDesc(Package, MoveTemp(Package.PackageLoadedDelegate)));
+		QueuedPackages.Add(new FAsyncPackageDesc2(Package, MoveTemp(Package.PackageLoadedDelegate)));
 	}
 	AltZenaphore.NotifyOne();
 }
 
-void FAsyncLoadingThread2Impl::UpdateExistingPackagePriorities(FAsyncPackage2* InPackage, TAsyncLoadPriority InNewPriority)
-{
-	check(!IsInGameThread() || !IsMultithreaded());
-	InPackage->SetPriority(InNewPriority);
-	return;
-}
-
-FAsyncPackage2* FAsyncLoadingThread2Impl::FindOrInsertPackage(FAsyncPackageDesc* InDesc, bool& bInserted)
+FAsyncPackage2* FAsyncLoadingThread2Impl::FindOrInsertPackage(FAsyncPackageDesc2* Desc, bool& bInserted)
 {
 	TRACE_CPUPROFILER_EVENT_SCOPE(FindOrInsertPackage);
 	FAsyncPackage2* Package = nullptr;
 	bInserted = false;
 	{
 		FScopeLock LockAsyncPackages(&AsyncPackagesCritical);
-		Package = AsyncPackageNameLookup.FindRef(InDesc->Name);
+		Package = AsyncPackageLookup.FindRef(Desc->PackageId);
 		if (!Package)
 		{
-			Package = CreateAsyncPackage(*InDesc);
+			Package = CreateAsyncPackage(*Desc);
 			if (!Package)
 			{
 				return nullptr;
 			}
 			Package->AddRef();
 			ExistingAsyncPackagesCounter.Increment();
-			AsyncPackageNameLookup.Add(Package->GetPackageName(), Package);
+			AsyncPackageLookup.Add(Desc->PackageId, Package);
 			bInserted = true;
 		}
-		else if (InDesc->RequestID > 0)
+		else if (Desc->RequestID > 0)
 		{
-			Package->AddRequestID(InDesc->RequestID);
+			Package->AddRequestID(Desc->RequestID);
 		}
-		if (InDesc->PackageLoadedDelegate.IsValid())
+		if (Desc->PackageLoadedDelegate.IsValid())
 		{
-			Package->AddCompletionCallback(MoveTemp(InDesc->PackageLoadedDelegate));
+			Package->AddCompletionCallback(MoveTemp(Desc->PackageLoadedDelegate));
 		}
 	}
 	return Package;
@@ -1945,14 +1947,14 @@ FAsyncPackage2* FAsyncLoadingThread2Impl::FindOrInsertPackage(FAsyncPackageDesc*
 bool FAsyncLoadingThread2Impl::CreateAsyncPackagesFromQueue()
 {
 	TRACE_CPUPROFILER_EVENT_SCOPE(CreateAsyncPackagesFromQueue);
-	TArray<FAsyncPackageDesc*> QueueCopy;
+	TArray<FAsyncPackageDesc2*> QueueCopy;
 	{
 		FScopeLock QueueLock(&QueueCritical);
 		QueueCopy.Append(QueuedPackages);
 		QueuedPackages.Reset();
 	}
 
-	for (FAsyncPackageDesc* PackageRequest : QueueCopy)
+	for (FAsyncPackageDesc2* PackageRequest : QueueCopy)
 	{
 		bool bInserted;
 		FAsyncPackage2* Package = FindOrInsertPackage(PackageRequest, bInserted);
@@ -1997,6 +1999,7 @@ void FAsyncLoadingThread2Impl::StartBundleIoRequests()
 		FBundleIoRequest BundleIoRequest;
 		WaitingIoRequests.HeapPop(BundleIoRequest, false);
 		FAsyncPackage2* Package = BundleIoRequest.Package;
+		check(Package);
 
 		if (GIsInitialLoad && PreviousPackage)
 		{
@@ -2005,7 +2008,7 @@ void FAsyncLoadingThread2Impl::StartBundleIoRequests()
 		PreviousPackage = Package;
 
 		FIoReadOptions ReadOptions;
-		IoDispatcher.ReadWithCallback(CreateIoChunkId(Package->GlobalPackageId.Id, 0, EIoChunkType::ExportBundleData),
+		IoDispatcher.ReadWithCallback(CreateIoChunkId(Package->Desc.PackageId.ToIndex(), 0, EIoChunkType::ExportBundleData),
 			ReadOptions,
 			[Package](TIoStatusOr<FIoBuffer> Result)
 		{
@@ -2334,8 +2337,6 @@ void FAsyncPackage2::SetupSerializedArcs(const uint8* GraphData, uint64 GraphDat
 
 	LLM_SCOPE(ELLMTag::AsyncLoading);
 
-	const FGlobalNameMap& GlobalNameMap = AsyncLoadingThread.GlobalNameMap;
-
 	FSimpleArchive GraphArchive(GraphData, GraphDataSize);
 	int32 InternalArcCount;
 	GraphArchive << InternalArcCount;
@@ -2351,13 +2352,12 @@ void FAsyncPackage2::SetupSerializedArcs(const uint8* GraphData, uint64 GraphDat
 	GraphArchive << ImportedPackagesCount;
 	for (int32 ImportedPackageIndex = 0; ImportedPackageIndex < ImportedPackagesCount; ++ImportedPackageIndex)
 	{
-		int32 ImportedPackageNameIndex;
-		int32 ImportedPackageNameNumber;
-		GraphArchive << ImportedPackageNameIndex << ImportedPackageNameNumber;
-		FName ImportedPackageName = GlobalNameMap.GetName(ImportedPackageNameIndex, ImportedPackageNameNumber);
-		FAsyncPackage2* ImportedPackage = AsyncLoadingThread.FindAsyncPackage(ImportedPackageName);
+		FPackageId ImportedPackageId;
 		int32 ExternalArcCount;
+		GraphArchive << ImportedPackageId;
 		GraphArchive << ExternalArcCount;
+
+		FAsyncPackage2* ImportedPackage = AsyncLoadingThread.GetAsyncPackage(ImportedPackageId);
 		for (int32 ExternalArcIndex = 0; ExternalArcIndex < ExternalArcCount; ++ExternalArcIndex)
 		{
 			int32 FromNodeIndex;
@@ -2437,7 +2437,7 @@ void FAsyncPackage2::ImportPackagesRecursive()
 
 	FPackageStore& GlobalPackageStore = AsyncLoadingThread.GlobalPackageStore;
 	int32 ImportedPackageCount = 0;
-	int32* Imports = GlobalPackageStore.GetPackageImportedPackages(GlobalPackageId, ImportedPackageCount);
+	int32* Imports = GlobalPackageStore.GetPackageImportedPackages(Desc.PackageId, ImportedPackageCount);
 	if (!Imports)
 	{
 		return;
@@ -2463,9 +2463,7 @@ void FAsyncPackage2::ImportPackagesRecursive()
 
 		if (bNeedToLoadPackage)
 		{
-			const FName ImportedPackageName = Entry.Name;
-			FAsyncPackageDesc Info(INDEX_NONE, ImportedPackageName);
-			Info.Priority = Desc.Priority;
+			FAsyncPackageDesc2 Info(INDEX_NONE, FPackageId::FromIndex(GlobalPackageIndex), Entry.Name);
 			bool bInserted;
 			FAsyncPackage2* ImportedAsyncPackage = AsyncLoadingThread.FindOrInsertPackage(&Info, bInserted);
 			if (ImportedAsyncPackage)
@@ -2504,7 +2502,7 @@ void FAsyncPackage2::SetupScriptArcs()
 
 	FPackageStore& GlobalPackageStore = AsyncLoadingThread.GlobalPackageStore;
 	int32 ScriptArcsCount = 0;
-	int32* ScriptArcs = GlobalPackageStore.GetPackageScriptArcs(GlobalPackageId, ScriptArcsCount);
+	int32* ScriptArcs = GlobalPackageStore.GetPackageScriptArcs(Desc.PackageId, ScriptArcsCount);
 	for (int32 ScriptArcIndex = 0; ScriptArcIndex < ScriptArcsCount * 2;)
 	{
 		int32 GlobalImportIndex = ScriptArcs[ScriptArcIndex++];
@@ -2722,8 +2720,8 @@ void FAsyncPackage2::EventDrivenCreateExport(int32 LocalExportIndex)
 	UClass* LoadClass = Export.ClassIndex.IsNull() ? UClass::StaticClass() : CastEventDrivenIndexToObject<UClass>(Export.ClassIndex, true);
 	UObject* ThisParent = Export.OuterIndex.IsNull() ? LinkerRoot : EventDrivenIndexToObject(Export.OuterIndex, false);
 
-	checkf(LoadClass, TEXT("Could not find class object for %s in %s"), *ObjectName.ToString(), *Desc.NameToLoad.ToString());
-	checkf(ThisParent, TEXT("Could not find outer object for %s in %s"), *ObjectName.ToString(), *Desc.NameToLoad.ToString());
+	checkf(LoadClass, TEXT("Could not find class object for %s in %s"), *ObjectName.ToString(), *Desc.Name.ToString());
+	checkf(ThisParent, TEXT("Could not find outer object for %s in %s"), *ObjectName.ToString(), *Desc.Name.ToString());
 	check(!dynamic_cast<UObjectRedirector*>(ThisParent));
 
 	// Try to find existing object first as we cannot in-place replace objects, could have been created by other export in this package
@@ -2766,7 +2764,7 @@ void FAsyncPackage2::EventDrivenCreateExport(int32 LocalExportIndex)
 		// Find the Archetype object for the one we are loading.
 		check(!Export.TemplateIndex.IsNull());
 		UObject* Template = EventDrivenIndexToObject(Export.TemplateIndex, true);
-		checkf(Template, TEXT("Could not find template for %s in %s"), *ObjectName.ToString(), *Desc.NameToLoad.ToString());
+		checkf(Template, TEXT("Could not find template for %s in %s"), *ObjectName.ToString(), *Desc.Name.ToString());
 		// we also need to ensure that the template has set up any instances
 		Template->ConditionalPostLoadSubobjects();
 
@@ -3127,7 +3125,7 @@ bool FAsyncPackage2::AreAllDependenciesFullyLoadedInternal(FAsyncPackage2* Packa
 		}
 		VisitedPackages.Add(ImportPackage);
 
-		FAsyncPackage2* AsyncRoot = AsyncLoadingThread.FindAsyncPackage(ImportPackage->GetFName());
+		FAsyncPackage2* AsyncRoot = AsyncLoadingThread.GetAsyncPackage(ImportPackage->GetPackageId());
 		if (AsyncRoot)
 		{
 			if (!AsyncRoot->bAllExportsSerialized)
@@ -3212,7 +3210,7 @@ EAsyncPackageState::Type FAsyncLoadingThread2Impl::ProcessLoadedPackagesFromGame
 		{
 			{
 				FScopeLock LockAsyncPackages(&AsyncPackagesCritical);
-				AsyncPackageNameLookup.Remove(Package->GetPackageName());
+				AsyncPackageLookup.Remove(Package->GetPackageId());
 				Package->ClearOwnedObjects();
 			}
 
@@ -3891,7 +3889,12 @@ void FAsyncLoadingThread2Impl::FireCompletedCompiledInImport(void* AsyncPackage,
 /**
 * Constructor
 */
-FAsyncPackage2::FAsyncPackage2(const FAsyncPackageDesc& InDesc, int32 InSerialNumber, FAsyncLoadingThread2Impl& InAsyncLoadingThread, IEDLBootNotificationManager& InEDLBootNotificationManager, FAsyncLoadEventGraphAllocator& InGraphAllocator, const FAsyncLoadEventSpec* EventSpecs, FGlobalPackageId InGlobalPackageId)
+FAsyncPackage2::FAsyncPackage2(
+	const FAsyncPackageDesc2& InDesc,
+	FAsyncLoadingThread2Impl& InAsyncLoadingThread,
+	IEDLBootNotificationManager& InEDLBootNotificationManager,
+	FAsyncLoadEventGraphAllocator& InGraphAllocator,
+	const FAsyncLoadEventSpec* EventSpecs)
 : Desc(InDesc)
 , LinkerRoot(nullptr)
 , PostLoadIndex(0)
@@ -3907,21 +3910,19 @@ FAsyncPackage2::FAsyncPackage2(const FAsyncPackageDesc& InDesc, int32 InSerialNu
 , AsyncLoadingThread(InAsyncLoadingThread)
 , EDLBootNotificationManager(InEDLBootNotificationManager)
 , GraphAllocator(InGraphAllocator)
-, GlobalPackageId(InGlobalPackageId)
 // Begin EDL specific properties
-, ImportStore(AsyncLoadingThread.GlobalPackageStore, InGlobalPackageId)
+, ImportStore(AsyncLoadingThread.GlobalPackageStore, Desc.PackageId)
 , AsyncPackageLoadingState(EAsyncPackageLoadingState2::NewPackage)
-, SerialNumber(InSerialNumber)
 , bAllExportsSerialized(false)
 {
 	TRACE_CPUPROFILER_EVENT_SCOPE(NewAsyncPackage);
 	TRACE_LOADTIME_NEW_ASYNC_PACKAGE(this, InDesc.Name);
 	AddRequestID(InDesc.RequestID);
 
-	TTuple<uint32*, uint32> PackageExportBundles = AsyncLoadingThread.GlobalPackageStore.GetPackageExportBundles(GlobalPackageId);
+	TTuple<uint32*, uint32> PackageExportBundles = AsyncLoadingThread.GlobalPackageStore.GetPackageExportBundles(Desc.PackageId);
 	ExportBundleOrderEntries = PackageExportBundles.Get<0>();
 	ExportBundleCount = PackageExportBundles.Get<1>();
-	ExportCount = AsyncLoadingThread.GlobalPackageStore.GetPackageExportCount(GlobalPackageId);
+	ExportCount = AsyncLoadingThread.GlobalPackageStore.GetPackageExportCount(Desc.PackageId);
 	Exports.AddZeroed(ExportCount);
 	OwnedObjects.Reserve(ExportCount + 1); // +1 for UPackage
 
@@ -3977,7 +3978,6 @@ FAsyncPackage2::~FAsyncPackage2()
 	TRACE_LOADTIME_DESTROY_ASYNC_PACKAGE(this);
 
 	MarkRequestIDsAsComplete();
-	SerialNumber = 0; // the weak pointer will always fail now
 	
 	check(OwnedObjects.Num() == 0);
 }
@@ -3989,7 +3989,7 @@ void FAsyncPackage2::ClearImportedPackages()
 #if 0
 	FPackageStore& GlobalPackageStore = AsyncLoadingThread.GlobalPackageStore;
 	int32 ImportedPackageCount = 0;
-	int32* Imports = GlobalPackageStore.GetPackageImportedPackages(GlobalPackageId, ImportedPackageCount);
+	int32* Imports = GlobalPackageStore.GetPackageImportedPackages(PackageId, ImportedPackageCount);
 	TArray<FAsyncPackage2*> TempImportedAsyncPackages;
 	for (int32 LocalImportIndex = 0; LocalImportIndex < ImportedPackageCount; ++LocalImportIndex)
 	{
@@ -4120,8 +4120,8 @@ void FAsyncPackage2::CreateUPackage(const FPackageSummary* PackageSummary)
 	{
 		TRACE_CPUPROFILER_EVENT_SCOPE(UPackageCreate);
 		LinkerRoot = NewObject<UPackage>(/*Outer*/nullptr, Desc.Name, RF_Public);
-		// LinkerRoot->FileName = Desc.NameToLoad;
-		LinkerRoot->SetGlobalPackageId(GlobalPackageId.Id);
+		// LinkerRoot->FileName = Desc.Name;
+		LinkerRoot->SetPackageId(Desc.PackageId);
 		LinkerRoot->SetPackageFlagsTo(PackageSummary->PackageFlags);
 		LinkerRoot->LinkerPackageVersion = GPackageFileUE4Version;
 		LinkerRoot->LinkerLicenseeVersion = GPackageFileLicenseeUE4Version;
@@ -4131,7 +4131,7 @@ void FAsyncPackage2::CreateUPackage(const FPackageSummary* PackageSummary)
 	}
 	else
 	{
-		check(LinkerRoot->GetGlobalPackageId() == GlobalPackageId.Id);
+		check(LinkerRoot->GetPackageId() == Desc.PackageId);
 		check(LinkerRoot->GetPackageFlags() == PackageSummary->PackageFlags);
 		check(LinkerRoot->LinkerPackageVersion == GPackageFileUE4Version);
 		check(LinkerRoot->LinkerLicenseeVersion == GPackageFileLicenseeUE4Version);
@@ -4142,7 +4142,7 @@ void FAsyncPackage2::CreateUPackage(const FPackageSummary* PackageSummary)
 	AddOwnedObjectWithAsyncFlag(LinkerRoot, /*bForceAdd*/ !bCreatedLinkerRoot);
 	check(LinkerRoot->HasAnyInternalFlags(EInternalObjectFlags::Async));
 
-	UE_LOG(LogStreaming, Verbose, TEXT("FAsyncPackage::CreateUPackage for %s finished."), *Desc.NameToLoad.ToString());
+	UE_LOG(LogStreaming, Verbose, TEXT("FAsyncPackage::CreateUPackage for %s finished."), *Desc.Name.ToString());
 }
 
 EAsyncPackageState::Type FAsyncPackage2::FinishExternalReadDependencies()
@@ -4451,8 +4451,6 @@ int32 FAsyncLoadingThread2Impl::LoadPackage(const FString& InName, const FGuid*
 {
 	TRACE_CPUPROFILER_EVENT_SCOPE(LoadPackage);
 
-	int32 RequestID = INDEX_NONE;
-
 	static bool bOnce = false;
 	if (!bOnce)
 	{
@@ -4460,50 +4458,12 @@ int32 FAsyncLoadingThread2Impl::LoadPackage(const FString& InName, const FGuid*
 		FGCObject::StaticInit(); // otherwise this thing is created during async loading, but not associated with a package
 	}
 
-	// The comments clearly state that it should be a package name but we also handle it being a filename as this function is not perf critical
-	// and LoadPackage handles having a filename being passed in as well.
-	FString PackageName;
-	bool bValidPackageName = true;
-
-	if (FPackageName::IsValidLongPackageName(InName, /*bIncludeReadOnlyRoots*/true))
-	{
-		PackageName = InName;
-	}
-	// PackageName got populated by the conditional function
-	else if (!(FPackageName::IsPackageFilename(InName) && FPackageName::TryConvertFilenameToLongPackageName(InName, PackageName)))
-	{
-		// PackageName may get populated by the conditional function
-		FString ClassName;
-
-		if (!FPackageName::ParseExportTextPath(PackageName, &ClassName, &PackageName))
-		{
-			UE_LOG(LogStreaming, Warning, TEXT("LoadPackageAsync failed to begin to load a package because the supplied package name ")
-				TEXT("was neither a valid long package name nor a filename of a map within a content folder: '%s' (%s)"),
-				*PackageName, *InName);
-
-			bValidPackageName = false;
-		}
-	}
-
-	FString PackageNameToLoad(InPackageToLoadFrom);
-
-	if (bValidPackageName)
-	{
-		if (PackageNameToLoad.IsEmpty())
-		{
-			PackageNameToLoad = PackageName;
-		}
-		// Make sure long package name is passed to FAsyncPackage so that it doesn't attempt to 
-		// create a package with short name.
-		if (FPackageName::IsShortPackageName(PackageNameToLoad))
-		{
-			UE_LOG(LogStreaming, Warning, TEXT("Async loading code requires long package names (%s)."), *PackageNameToLoad);
+	int32 RequestID = INDEX_NONE;
 
-			bValidPackageName = false;
-		}
-	}
+	FName PackageName(*InName);
+	FPackageId* PackageId = GlobalPackageStore.GetPackageId(PackageName);
 
-	if (bValidPackageName)
+	if (PackageId)
 	{
 		if (FCoreDelegates::OnAsyncLoadPackage.IsBound())
 		{
@@ -4524,12 +4484,15 @@ int32 FAsyncLoadingThread2Impl::LoadPackage(const FString& InName, const FGuid*
 		}
 
 		// Add new package request
-		FAsyncPackageDesc PackageDesc(RequestID, *PackageName, *PackageNameToLoad, InGuid ? *InGuid : FGuid(), MoveTemp(CompletionDelegatePtr), InPackageFlags, InPIEInstanceID, InPackagePriority);
+		FAsyncPackageDesc2 PackageDesc(RequestID, *PackageId, PackageName, MoveTemp(CompletionDelegatePtr));
 		QueuePackage(PackageDesc);
 	}
 	else
 	{
-		InCompletionDelegate.ExecuteIfBound(FName(*InName), nullptr, EAsyncLoadingResult::Failed);
+		const TCHAR* NameToLoadStr = InPackageToLoadFrom ? InPackageToLoadFrom : TEXT("null");
+		UE_LOG(LogStreaming, Warning, TEXT("FAsyncLoadingThread2Impl::LoadPackage - Skipping unknown package: %s (%s)"),
+			*InName, NameToLoadStr);
+		InCompletionDelegate.ExecuteIfBound(PackageName, nullptr, EAsyncLoadingResult::Failed);
 	}
 
 	return RequestID;
diff --git a/Engine/Source/Runtime/CoreUObject/Private/Serialization/BulkData2.cpp b/Engine/Source/Runtime/CoreUObject/Private/Serialization/BulkData2.cpp
--- a/Engine/Source/Runtime/CoreUObject/Private/Serialization/BulkData2.cpp
+++ b/Engine/Source/Runtime/CoreUObject/Private/Serialization/BulkData2.cpp
@@ -579,7 +579,7 @@ void FBulkDataBase::Serialize(FArchive& Ar, UObject* Owner, int32 /*Index*/, boo
 		if (!IsInlined() && bUseZenLoader)
 		{
 			const EIoChunkType Type = IsOptional() ? EIoChunkType::OptionalBulkData : EIoChunkType::BulkData;
-			ChunkID = CreateBulkdataChunkId(Package->GetGlobalPackageId(), BulkDataOffsetInFile, Type);
+			ChunkID = CreateBulkdataChunkId(Package->GetPackageId().ToIndex(), BulkDataOffsetInFile, Type);
 
 			BulkDataFlags |= BULKDATA_UsesIoDispatcher; // Indicates that this BulkData should use the FIoChunkId rather than a filename
 		}
@@ -640,7 +640,7 @@ void FBulkDataBase::Serialize(FArchive& Ar, UObject* Owner, int32 /*Index*/, boo
 						if (!IsInlined() && bUseZenLoader)
 						{
 							// Regenerate the FIoChunkId to find the optional BulkData instead!
-							ChunkID = CreateBulkdataChunkId(Package->GetGlobalPackageId(), BulkDataOffsetInFile, EIoChunkType::OptionalBulkData);
+							ChunkID = CreateBulkdataChunkId(Package->GetPackageId().ToIndex(), BulkDataOffsetInFile, EIoChunkType::OptionalBulkData);
 							BulkDataFlags |= BULKDATA_UsesIoDispatcher; // Indicates that this BulkData should use the FIoChunkId rather than a filename
 						}
 						else
diff --git a/Engine/Source/Runtime/CoreUObject/Private/UObject/PackageId.cpp b/Engine/Source/Runtime/CoreUObject/Private/UObject/PackageId.cpp
new file mode 100644
--- /dev/null
+++ b/Engine/Source/Runtime/CoreUObject/Private/UObject/PackageId.cpp
@@ -0,0 +1,15 @@
+// Copyright 1998-2019 Epic Games, Inc. All Rights Reserved.
+
+#include "UObject/PackageId.h"
+#include "Serialization/StructuredArchive.h"
+
+FArchive& operator<<(FArchive& Ar, FPackageId& Value)
+{
+	FStructuredArchiveFromArchive(Ar).GetSlot() << Value;
+	return Ar;
+}
+
+void operator<<(FStructuredArchiveSlot Slot, FPackageId& Value)
+{
+	Slot << Value.Id;
+}
diff --git a/Engine/Source/Runtime/CoreUObject/Public/UObject/Package.h b/Engine/Source/Runtime/CoreUObject/Public/UObject/Package.h
--- a/Engine/Source/Runtime/CoreUObject/Public/UObject/Package.h
+++ b/Engine/Source/Runtime/CoreUObject/Public/UObject/Package.h
@@ -6,6 +6,7 @@
 #include "UObject/ObjectMacros.h"
 #include "UObject/UObjectGlobals.h"
 #include "UObject/Object.h"
+#include "UObject/PackageId.h"
 #include "Misc/Guid.h"
 #include "Misc/WorldCompositionUtility.h"
 #include "Misc/OutputDeviceError.h"
@@ -189,7 +190,7 @@ private:
 	uint32	PackageFlagsPrivate;
 	
 	/** Globally unique id used to address I/O chunks within the package */
-	uint32 GlobalPackageId = ~uint32(0); 
+	FPackageId PackageId;
 public:
 
 	/** Editor only: PIE instance ID this package belongs to, INDEX_NONE otherwise */
@@ -513,16 +514,16 @@ public:
 		ChunkIDs = InChunkIDs;
 	}
 
-	/** returns the globally unique package id */
-	FORCEINLINE uint32 GetGlobalPackageId() const
+	/** returns the unique package id */
+	FORCEINLINE FPackageId GetPackageId() const
 	{
-		return GlobalPackageId;
+		return PackageId;
 	}
 
-	/** sets the globally unique package id */
-	FORCEINLINE void SetGlobalPackageId(uint32 InGlobalPackageId)
+	/** sets the unique package id */
+	FORCEINLINE void SetPackageId(FPackageId InPackageId)
 	{
-		GlobalPackageId = InGlobalPackageId;
+		PackageId = InPackageId;
 	}
 
 	////////////////////////////////////////////////////////
diff --git a/Engine/Source/Runtime/CoreUObject/Public/UObject/PackageId.h b/Engine/Source/Runtime/CoreUObject/Public/UObject/PackageId.h
new file mode 100644
--- /dev/null
+++ b/Engine/Source/Runtime/CoreUObject/Public/UObject/PackageId.h
@@ -0,0 +1,43 @@
+// Copyright 1998-2019 Epic Games, Inc. All Rights Reserved.
+
+#pragma once
+
+class FArchive;
+class FStructuredArchiveSlot;
+
+class FPackageId
+{
+	static constexpr uint32 InvalidId = ~uint32(0u);
+	uint32 Id = InvalidId;
+
+	inline explicit FPackageId(int32 InId): Id(InId) {}
+
+public:
+	FPackageId() = default;
+
+	inline static FPackageId FromIndex(uint32 Index)
+	{
+		return FPackageId(Index);
+	}
+
+	inline uint32 ToIndex() const
+	{
+		check(Id != InvalidId);
+		return Id;
+	}
+
+	inline bool operator==(FPackageId Other) const
+	{
+		return Id == Other.Id;
+	}
+
+	inline friend uint32 GetTypeHash(const FPackageId& In)
+	{
+		return In.Id;
+	}
+
+	COREUOBJECT_API friend FArchive& operator<<(FArchive& Ar, FPackageId& Value);
+
+	COREUOBJECT_API friend void operator<<(FStructuredArchiveSlot Slot, FPackageId& Value);
+};
+
