AsyncLoading2: Add missing include for 10505629

diff --git a/Engine/Source/Runtime/CoreUObject/Public/UObject/PackageId.h b/Engine/Source/Runtime/CoreUObject/Public/UObject/PackageId.h
--- a/Engine/Source/Runtime/CoreUObject/Public/UObject/PackageId.h
+++ b/Engine/Source/Runtime/CoreUObject/Public/UObject/PackageId.h
@@ -2,6 +2,8 @@
 
 #pragma once
 
+#include "CoreMinimal.h"
+
 class FArchive;
 class FStructuredArchiveSlot;
 
