Fix code analysis warning in BulkData2.

Technically this was a false positive, but it is probably bad form to dereference a pointer that could be null to a reference just because we know that reference won't be used if the pointer could've been null. I didn't really notice when I converted the FString copy to be a pointer in an earlier changelist.

Instead just pass the pointer in and add a check, which will satisfy the code analyser

#rnx

diff --git a/Engine/Source/Runtime/CoreUObject/Private/Serialization/BulkData2.cpp b/Engine/Source/Runtime/CoreUObject/Private/Serialization/BulkData2.cpp
--- a/Engine/Source/Runtime/CoreUObject/Private/Serialization/BulkData2.cpp
+++ b/Engine/Source/Runtime/CoreUObject/Private/Serialization/BulkData2.cpp
@@ -616,12 +616,13 @@ void FBulkDataBase::Serialize(FArchive& Ar, UObject* Owner, int32 /*Index*/, boo
 			{
 				if (IsDuplicateNonOptional())
 				{
-					auto DoesOptionalDataExist = [this](const FString& PackageFilename)->bool
+					auto DoesOptionalDataExist = [this](const FString* PackageFilename)->bool
 					{
 #if ALLOW_OPTIONAL_DATA
 						if (!IsUsingIODispatcher())
 						{
-							FString OptionalDataFilename = ConvertFilenameFromFlags(PackageFilename);
+							check(PackageFilename != nullptr);
+							FString OptionalDataFilename = ConvertFilenameFromFlags(*PackageFilename);
 							return IFileManager::Get().FileExists(*OptionalDataFilename);
 						}
 						else
@@ -633,7 +634,7 @@ void FBulkDataBase::Serialize(FArchive& Ar, UObject* Owner, int32 /*Index*/, boo
 #endif
 					};
 
-					if (DoesOptionalDataExist(*Filename))
+					if (DoesOptionalDataExist(Filename))
 					{
 						SerializeDuplicateData(Ar, BulkDataFlags, BulkDataSizeOnDisk, BulkDataOffsetInFile);
 
