Cleaned up GC Token Stream generation code and merged token debug info with the token stream to reduce data duplication and simplify API

diff --git a/Engine/Source/Editor/KismetCompiler/Private/KismetCompiler.cpp b/Engine/Source/Editor/KismetCompiler/Private/KismetCompiler.cpp
--- a/Engine/Source/Editor/KismetCompiler/Private/KismetCompiler.cpp
+++ b/Engine/Source/Editor/KismetCompiler/Private/KismetCompiler.cpp
@@ -2390,9 +2390,6 @@ void FKismetCompilerContext::FinishCompilingClass(UClass* Class)
 	{
 		// Propagate the new parent's inheritable class flags
 		Class->ReferenceTokenStream.Empty();
-#if !(UE_BUILD_TEST || UE_BUILD_SHIPPING)
-		Class->DebugTokenMap.Empty();
-#endif
 		Class->ClassFlags &= ~CLASS_RecompilerClear;
 		Class->ClassFlags |= (ParentClass->ClassFlags & CLASS_ScriptInherit);//@TODO: ChangeParentClass had this, but I don't think I want it: | UClass::StaticClassFlags;  // will end up with CLASS_Intrinsic
 		Class->ClassCastFlags |= ParentClass->ClassCastFlags;
diff --git a/Engine/Source/Runtime/CoreUObject/Private/UObject/Class.cpp b/Engine/Source/Runtime/CoreUObject/Private/UObject/Class.cpp
--- a/Engine/Source/Runtime/CoreUObject/Private/UObject/Class.cpp
+++ b/Engine/Source/Runtime/CoreUObject/Private/UObject/Class.cpp
@@ -4912,14 +4912,6 @@ bool UClass::IsClassGroupName(const TCHAR* InGroupName) const
 
 #endif // WITH_EDITOR || HACK_HEADER_GENERATOR
 
-void UClass::PrependStreamWithSuperClass(UClass& SuperClass)
-{
-	ReferenceTokenStream.PrependStream(SuperClass.ReferenceTokenStream);
-
-#if ENABLE_GC_OBJECT_CHECKS
-	DebugTokenMap.PrependWithSuperClass(SuperClass);
-#endif
-}
 
 IMPLEMENT_CORE_INTRINSIC_CLASS(UClass, UStruct,
 	{
diff --git a/Engine/Source/Runtime/CoreUObject/Private/UObject/GarbageCollection.cpp b/Engine/Source/Runtime/CoreUObject/Private/UObject/GarbageCollection.cpp
--- a/Engine/Source/Runtime/CoreUObject/Private/UObject/GarbageCollection.cpp
+++ b/Engine/Source/Runtime/CoreUObject/Private/UObject/GarbageCollection.cpp
@@ -885,7 +885,7 @@ public:
 				FString TokenDebugInfo;
 				if (UClass *Class = (ReferencingObject ? ReferencingObject->GetClass() : nullptr))
 				{
-					auto& TokenInfo = Class->DebugTokenMap.GetTokenInfo(TokenIndex);
+					FTokenInfo TokenInfo = Class->ReferenceTokenStream.GetTokenInfo(TokenIndex);
 					TokenDebugInfo = FString::Printf(TEXT("ReferencingObjectClass: %s, Property Name: %s, Offset: %d"),
 						*Class->GetFullName(), *TokenInfo.Name.GetPlainNameString(), TokenInfo.Offset);
 				}
@@ -2442,11 +2442,7 @@ void UInterfaceProperty::EmitReferenceInfo(UClass& OwnerClass, int32 BaseOffset,
 void UClass::EmitObjectReference(int32 Offset, const FName& DebugName, EGCReferenceType Kind)
 {
 	FGCReferenceInfo ObjectReference(Kind, Offset);
-	int32 TokenIndex = ReferenceTokenStream.EmitReferenceInfo(ObjectReference);
-
-#if ENABLE_GC_OBJECT_CHECKS
-	DebugTokenMap.MapToken(DebugName, Offset, TokenIndex);
-#endif
+	ReferenceTokenStream.EmitReferenceInfo(ObjectReference, DebugName);
 }
 
 void UClass::EmitObjectArrayReference(int32 Offset, const FName& DebugName)
@@ -2529,9 +2525,6 @@ void UClass::AssembleReferenceTokenStream(bool bForce)
 		if (bForce)
 		{
 			ReferenceTokenStream.Empty();
-#if ENABLE_GC_OBJECT_CHECKS
-			DebugTokenMap.Empty();
-#endif
 			ClassFlags &= ~CLASS_TokenStreamAssembled;
 		}
 		TArray<const UStructProperty*> EncounteredStructProps;
@@ -2553,7 +2546,7 @@ void UClass::AssembleReferenceTokenStream(bool bForce)
 			if (!SuperClass->ReferenceTokenStream.IsEmpty())
 			{
 				// Prepend super's stream. This automatically handles removing the EOS token.
-				PrependStreamWithSuperClass(*SuperClass);
+				ReferenceTokenStream.PrependStream(SuperClass->ReferenceTokenStream);
 			}
 		}
 		else
@@ -2586,7 +2579,7 @@ void UClass::AssembleReferenceTokenStream(bool bForce)
 		}
 
 		// Emit end of stream token.
-		static const FName EOSDebugName("EOS");
+		static const FName EOSDebugName("EndOfStreamToken");
 		EmitObjectReference(0, EOSDebugName, GCRT_EndOfStream);
 
 		// Shrink reference token stream to proper size.
@@ -2606,15 +2599,34 @@ void UClass::AssembleReferenceTokenStream(bool bForce)
 void FGCReferenceTokenStream::PrependStream( const FGCReferenceTokenStream& Other )
 {
 	// Remove embedded EOS token if needed.
-	TArray<uint32> TempTokens = Other.Tokens;
-	FGCReferenceInfo EndOfStream( GCRT_EndOfStream, 0 );
-	if( TempTokens.Last() == EndOfStream )
+	FGCReferenceInfo EndOfStream(GCRT_EndOfStream, 0);
+	int32 NumTokensToPrepend = (Other.Tokens.Num() && Other.Tokens.Last() == EndOfStream) ? (Other.Tokens.Num() - 1) : Other.Tokens.Num();
+
+	TArray<uint32> TempTokens;
+	TempTokens.Reserve(NumTokensToPrepend + Tokens.Num());
+
+#if ENABLE_GC_OBJECT_CHECKS
+	check(TokenDebugInfo.Num() == Tokens.Num());
+	check(Other.TokenDebugInfo.Num() == Other.Tokens.Num());
+	TArray<FName> TempTokenDebugInfo;
+	TempTokenDebugInfo.Reserve(NumTokensToPrepend + TokenDebugInfo.Num());
+#endif // ENABLE_GC_OBJECT_CHECKS
+
+	for (int32 TokenIndex = 0; TokenIndex < NumTokensToPrepend; ++TokenIndex)
 	{
-		TempTokens.RemoveAt( TempTokens.Num() - 1 );
+		TempTokens.Add(Other.Tokens[TokenIndex]);
+#if ENABLE_GC_OBJECT_CHECKS
+		TempTokenDebugInfo.Add(Other.TokenDebugInfo[TokenIndex]);
+#endif // ENABLE_GC_OBJECT_CHECKS
 	}
-	// TArray doesn't have a general '+' operator.
-	TempTokens += Tokens;
+
+	TempTokens.Append(Tokens);
 	Tokens = MoveTemp(TempTokens);
+
+#if ENABLE_GC_OBJECT_CHECKS
+	TempTokenDebugInfo.Append(TokenDebugInfo);
+	TokenDebugInfo = MoveTemp(TempTokenDebugInfo);
+#endif // ENABLE_GC_OBJECT_CHECKS
 }
 
 void FGCReferenceTokenStream::Fixup(void (*AddReferencedObjectsPtr)(UObject*, class FReferenceCollector&), bool bKeepOuterToken, bool bKeepClassToken)
@@ -2714,14 +2726,20 @@ void FGCReferenceTokenStream::Fixup(void (*AddReferencedObjectsPtr)(UObject*, cl
 	// ARO is not in the token stream yet.
 	if (!bReplacedARO && AddReferencedObjectsPtr)
 	{
-		EmitReferenceInfo(FGCReferenceInfo(GCRT_AddReferencedObjects, 0));
+		static const FName TokenName("AROToken");
+		EmitReferenceInfo(FGCReferenceInfo(GCRT_AddReferencedObjects, 0), TokenName);
 		EmitPointer((const void*)AddReferencedObjectsPtr);
 	}
 }
 
-int32 FGCReferenceTokenStream::EmitReferenceInfo(FGCReferenceInfo ReferenceInfo)
+int32 FGCReferenceTokenStream::EmitReferenceInfo(FGCReferenceInfo ReferenceInfo, const FName& DebugName)
 {
-	return Tokens.Add(ReferenceInfo);
+	int32 TokenIndex = Tokens.Add(ReferenceInfo);
+#if ENABLE_GC_OBJECT_CHECKS
+	check(TokenDebugInfo.Num() == TokenIndex);
+	TokenDebugInfo.Add(DebugName);
+#endif
+	return TokenIndex;
 }
 
 /**
@@ -2731,7 +2749,13 @@ int32 FGCReferenceTokenStream::EmitReferenceInfo(FGCReferenceInfo ReferenceInfo)
  */
 uint32 FGCReferenceTokenStream::EmitSkipIndexPlaceholder()
 {
-	return Tokens.Add( E_GCSkipIndexPlaceholder );
+	uint32 TokenIndex = Tokens.Add(E_GCSkipIndexPlaceholder);
+#if ENABLE_GC_OBJECT_CHECKS
+	static const FName TokenName("SkipIndexPlaceholder");
+	check(TokenDebugInfo.Num() == TokenIndex);
+	TokenDebugInfo.Add(TokenName);
+#endif
+	return TokenIndex;
 }
 
 /**
@@ -2762,19 +2786,38 @@ void FGCReferenceTokenStream::UpdateSkipIndexPlaceholder( uint32 SkipIndexIndex,
  *
  * @param Count count to emit
  */
-void FGCReferenceTokenStream::EmitCount( uint32 Count )
+int32 FGCReferenceTokenStream::EmitCount( uint32 Count )
 {
-	Tokens.Add( Count );
+	int32 TokenIndex = Tokens.Add( Count );
+#if ENABLE_GC_OBJECT_CHECKS
+	static const FName TokenName("CountToken");
+	check(TokenDebugInfo.Num() == TokenIndex);
+	TokenDebugInfo.Add(TokenName);
+#endif
+	return TokenIndex;
 }
 
-void FGCReferenceTokenStream::EmitPointer( void const* Ptr )
+int32 FGCReferenceTokenStream::EmitPointer( void const* Ptr )
 {
 	const int32 StoreIndex = Tokens.Num();
 	Tokens.AddUninitialized(GNumTokensPerPointer);
 	StorePointer(&Tokens[StoreIndex], Ptr);
+
+#if ENABLE_GC_OBJECT_CHECKS
+	static const FName TokenName("PointerToken");
+	check(TokenDebugInfo.Num() == StoreIndex);
+	for (int32 PointerTokenIndex = 0; PointerTokenIndex < GNumTokensPerPointer; ++PointerTokenIndex)
+	{
+		TokenDebugInfo.Add(TokenName);
+	}
+#endif
+
 	// Now inser the end of pointer marker, this will mostly be used for storing ReturnCount value
 	// if the pointer was stored at the end of struct array stream.
-	EmitReferenceInfo(FGCReferenceInfo(GCRT_EndOfPointer, 0));
+	static const FName EndOfPointerTokenName("EndOfPointerToken");
+	EmitReferenceInfo(FGCReferenceInfo(GCRT_EndOfPointer, 0), EndOfPointerTokenName);
+
+	return StoreIndex;
 }
 
 /**
@@ -2782,9 +2825,17 @@ void FGCReferenceTokenStream::EmitPointer( void const* Ptr )
  *
  * @param Stride stride to emit
  */
-void FGCReferenceTokenStream::EmitStride( uint32 Stride )
+int32 FGCReferenceTokenStream::EmitStride( uint32 Stride )
 {
-	Tokens.Add( Stride );
+	int32 TokenIndex = Tokens.Add( Stride );
+
+#if ENABLE_GC_OBJECT_CHECKS
+	static const FName TokenName("StrideToken");
+	check(TokenDebugInfo.Num() == TokenIndex);
+	TokenDebugInfo.Add(TokenName);
+#endif
+
+	return TokenIndex;
 }
 
 /**
@@ -2803,54 +2854,15 @@ uint32 FGCReferenceTokenStream::EmitReturn()
 
 #if ENABLE_GC_OBJECT_CHECKS
 
-void FGCDebugReferenceTokenMap::MapToken(const FName& DebugName, int32 Offset, int32 TokenIndex)
+FTokenInfo FGCReferenceTokenStream::GetTokenInfo(int32 TokenIndex) const
 {
-	if(TokenMap.Num() <= TokenIndex)
-	{
-		TokenMap.AddZeroed(TokenIndex - TokenMap.Num() + 1);
-
-		auto& TokenInfo = TokenMap[TokenIndex];
-
-		TokenInfo.Offset = Offset;
-		TokenInfo.Name = DebugName;
-	}
-	else
-	{
-		// Token already mapped.
-		checkNoEntry();
-	}
+	FTokenInfo DebugInfo;
+	DebugInfo.Offset = FGCReferenceInfo(Tokens[TokenIndex]).Offset;
+	DebugInfo.Name = TokenDebugInfo[TokenIndex];
+	return DebugInfo;
 }
 
-void FGCDebugReferenceTokenMap::PrependWithSuperClass(const UClass& SuperClass)
-{
-	if (SuperClass.ReferenceTokenStream.Size() == 0)
-	{
-		return;
-	}
-
-	// Check if token stream is already ended with end-of-stream token. If so then something's wrong.
-	checkSlow(TokenMap.Num() == 0 || TokenMap.Last().Name != "EOS");
-
-	int32 OldTokenNumber = TokenMap.Num();
-	int32 NewTokenOffset = SuperClass.ReferenceTokenStream.Size() - 1;
-	TokenMap.AddZeroed(NewTokenOffset);
-
-	for(int32 OldTokenIndex = OldTokenNumber - 1; OldTokenIndex >= 0; --OldTokenIndex)
-	{
-		TokenMap[OldTokenIndex + NewTokenOffset] = TokenMap[OldTokenIndex];
-	}
-
-	for(int32 NewTokenIndex = 0; NewTokenIndex < NewTokenOffset; ++NewTokenIndex)
-	{
-		TokenMap[NewTokenIndex] = SuperClass.DebugTokenMap.GetTokenInfo(NewTokenIndex);
-	}
-}
-
-const FTokenInfo& FGCDebugReferenceTokenMap::GetTokenInfo(int32 TokenIndex) const
-{
-	return TokenMap[TokenIndex];
-}
-#endif // ENABLE_GC_OBJECT_CHECKS
+#endif
 
 
 FGCArrayPool* FGCArrayPool::GetGlobalSingleton()
diff --git a/Engine/Source/Runtime/CoreUObject/Private/UObject/GarbageCollectionVerification.cpp b/Engine/Source/Runtime/CoreUObject/Private/UObject/GarbageCollectionVerification.cpp
--- a/Engine/Source/Runtime/CoreUObject/Private/UObject/GarbageCollectionVerification.cpp
+++ b/Engine/Source/Runtime/CoreUObject/Private/UObject/GarbageCollectionVerification.cpp
@@ -65,7 +65,7 @@ public:
 				FString TokenDebugInfo;
 				if (UClass *Class = (ReferencingObject ? ReferencingObject->GetClass() : nullptr))
 				{
-					const FTokenInfo& TokenInfo = Class->DebugTokenMap.GetTokenInfo(TokenIndex);
+					FTokenInfo TokenInfo = Class->ReferenceTokenStream.GetTokenInfo(TokenIndex);
 					TokenDebugInfo = FString::Printf(TEXT("ReferencingObjectClass: %s, Property Name: %s, Offset: %d"),
 						*Class->GetFullName(), *TokenInfo.Name.GetPlainNameString(), TokenInfo.Offset);
 				}
@@ -187,7 +187,7 @@ public:
 				FString TokenDebugInfo;
 				if (UClass *Class = (ReferencingObject ? ReferencingObject->GetClass() : nullptr))
 				{
-					const FTokenInfo& TokenInfo = Class->DebugTokenMap.GetTokenInfo(TokenIndex);
+					FTokenInfo TokenInfo = Class->ReferenceTokenStream.GetTokenInfo(TokenIndex);
 					TokenDebugInfo = FString::Printf(TEXT("ReferencingObjectClass: %s, Property Name: %s, Offset: %d"),
 						*Class->GetFullName(), *TokenInfo.Name.GetPlainNameString(), TokenInfo.Offset);
 				}
diff --git a/Engine/Source/Runtime/CoreUObject/Private/UObject/ReferenceChainSearch.cpp b/Engine/Source/Runtime/CoreUObject/Private/UObject/ReferenceChainSearch.cpp
--- a/Engine/Source/Runtime/CoreUObject/Private/UObject/ReferenceChainSearch.cpp
+++ b/Engine/Source/Runtime/CoreUObject/Private/UObject/ReferenceChainSearch.cpp
@@ -366,7 +366,7 @@ public:
 #if ENABLE_GC_OBJECT_CHECKS
 			if (TokenIndex >= 0)
 			{
-				const FTokenInfo& TokenInfo = ReferencingObject->GetClass()->DebugTokenMap.GetTokenInfo(TokenIndex);
+				FTokenInfo TokenInfo = ReferencingObject->GetClass()->ReferenceTokenStream.GetTokenInfo(TokenIndex);
 				RefInfo.ReferencerName = TokenInfo.Name;
 				RefInfo.Type = FReferenceChainSearch::EReferenceType::Property;
 			}
diff --git a/Engine/Source/Runtime/CoreUObject/Public/UObject/Class.h b/Engine/Source/Runtime/CoreUObject/Public/UObject/Class.h
--- a/Engine/Source/Runtime/CoreUObject/Public/UObject/Class.h
+++ b/Engine/Source/Runtime/CoreUObject/Public/UObject/Class.h
@@ -2570,23 +2570,11 @@ public:
 	 */
 	TArray<FImplementedInterface> Interfaces;
 
-	/**
-	 * Prepends reference token stream with super class's stream.
-	 *
-	 * @param SuperClass Super class to prepend stream with.
-	 */
-	void PrependStreamWithSuperClass(UClass& SuperClass);
-
 	/** Reference token stream used by realtime garbage collector, finalized in AssembleReferenceTokenStream */
 	FGCReferenceTokenStream ReferenceTokenStream;
 	/** CS for the token stream. Token stream can assemble code can sometimes be called from two threads throuh a web of async loading calls. */
 	FCriticalSection ReferenceTokenStreamCritical;
 
-#if ENABLE_GC_OBJECT_CHECKS
-	/** TokenIndex map to look-up token stream index origin. */
-	FGCDebugReferenceTokenMap DebugTokenMap;
-#endif
-
 	/** This class's native functions. */
 	TArray<FNativeFunctionLookup> NativeFunctionLookupTable;
 
diff --git a/Engine/Source/Runtime/CoreUObject/Public/UObject/GarbageCollection.h b/Engine/Source/Runtime/CoreUObject/Public/UObject/GarbageCollection.h
--- a/Engine/Source/Runtime/CoreUObject/Public/UObject/GarbageCollection.h
+++ b/Engine/Source/Runtime/CoreUObject/Public/UObject/GarbageCollection.h
@@ -167,54 +167,6 @@ struct FTokenInfo
 	FName Name;
 };
 
-/**
- * Stores info about GC token stream debug map.
- */
-class FGCDebugReferenceTokenMap
-{
-public:
-	/**
-	 * Maps token.
-	 *
-	 * @param DebugName Field name.
-	 * @param Offset Field offset.
-	 * @param TokenIndex Token index.
-	 */
-	void MapToken(const FName& DebugName, int32 Offset, int32 TokenIndex);
-
-	/**
-	 * Prepends this token map with its super class token map.
-	 *
-	 * @param SuperClass Super class to obtain its token map.
-	 */
-	void PrependWithSuperClass(const UClass& SuperClass);
-
-	/**
-	 * Returns pointer to token info object at given index.
-	 *
-	 * @param TokenIndex Index of the token.
-	 *
-	 * @returns FTokenInfo object.
-	 */
-	const FTokenInfo& GetTokenInfo(int32 TokenIndex) const;
-
-	/**
-	 * Empties the map.
-	 */
-	void Empty()
-	{
-		TokenMap.Empty();
-	}
-
-	FORCEINLINE int32 GetTokenMapSize() const
-	{
-		return TokenMap.Num();
-	}
-
-private:
-	/* Token map. */
-	TArray<FTokenInfo> TokenMap;
-};
 #endif // ENABLE_GC_OBJECT_CHECKS
 
 /**
@@ -237,12 +189,18 @@ struct FGCReferenceTokenStream
 	void Shrink()
 	{
 		Tokens.Shrink();
+#if ENABLE_GC_OBJECT_CHECKS
+		TokenDebugInfo.Shrink();
+#endif // ENABLE_GC_OBJECT_CHECKS
 	}
 
 	/** Empties the token stream entirely */
 	void Empty()
 	{
 		Tokens.Empty();
+#if ENABLE_GC_OBJECT_CHECKS
+		TokenDebugInfo.Empty();
+#endif // ENABLE_GC_OBJECT_CHECKS
 	}
 
 	/**
@@ -277,7 +235,7 @@ struct FGCReferenceTokenStream
 	 *
 	 * @return Index of the reference info in the token stream.
 	 */
-	int32 EmitReferenceInfo(FGCReferenceInfo ReferenceInfo);
+	int32 EmitReferenceInfo(FGCReferenceInfo ReferenceInfo, const FName& DebugName);
 
 	/**
 	 * Emit placeholder for aray skip index, updated in UpdateSkipIndexPlaceholder
@@ -301,21 +259,21 @@ struct FGCReferenceTokenStream
 	 *
 	 * @param Count count to emit
 	 */
-	void EmitCount( uint32 Count );
+	int32 EmitCount( uint32 Count );
 
 	/**
 	 * Emit a pointer
 	 *
 	 * @param Ptr pointer to emit
 	 */
-	void EmitPointer( void const* Ptr );
+	int32 EmitPointer( void const* Ptr );
 
 	/**
 	 * Emit stride
 	 *
 	 * @param Stride stride to emit
 	 */
-	void EmitStride( uint32 Stride );
+	int32 EmitStride( uint32 Stride );
 
 	/**
 	 * Increase return count on last token.
@@ -423,6 +381,10 @@ struct FGCReferenceTokenStream
 		return CurrentIndex >= (uint32)Tokens.Num();
 	}
 
+#if ENABLE_GC_OBJECT_CHECKS
+	FTokenInfo GetTokenInfo(int32 TokenIndex) const;
+#endif
+
 private:
 
 	/**
@@ -444,7 +406,14 @@ private:
 	}
 
 	/** Token array */
-	TArray<uint32>	Tokens;
+	TArray<uint32> Tokens;
+#if ENABLE_GC_OBJECT_CHECKS
+	/** 
+	 * Name of the proprty that emitted the associated token or token type (pointer etc).
+	 * We want to keep it in a separate array for performance reasons
+	 */
+	TArray<FName> TokenDebugInfo;
+#endif // ENABLE_GC_OBJECT_CHECKS
 };
 
 /** Prevent GC from running in the current scope */
