Missing file from previous check in

diff --git a/Engine/Source/Runtime/CoreUObject/Private/UObject/UObjectGlobals.cpp b/Engine/Source/Runtime/CoreUObject/Private/UObject/UObjectGlobals.cpp
--- a/Engine/Source/Runtime/CoreUObject/Private/UObject/UObjectGlobals.cpp
+++ b/Engine/Source/Runtime/CoreUObject/Private/UObject/UObjectGlobals.cpp
@@ -4555,9 +4555,6 @@ namespace UE4CodeGen_Private
 		{
 			check((NewClass->ClassFlags & CLASS_TokenStreamAssembled) != CLASS_TokenStreamAssembled);
 			NewClass->ReferenceTokenStream.Empty();
-#if ENABLE_GC_OBJECT_CHECKS
-			NewClass->DebugTokenMap.Empty();
-#endif
 		}
 		NewClass->CreateLinkAndAddChildFunctionsToMap(Params.FunctionLinkArray, Params.NumFunctions);
 
