AsyncLoading2:
Limit pending bundle requests on size instead of count

diff --git a/Engine/Source/Developer/IoStoreUtilities/Private/IoStoreUtilities.cpp b/Engine/Source/Developer/IoStoreUtilities/Private/IoStoreUtilities.cpp
--- a/Engine/Source/Developer/IoStoreUtilities/Private/IoStoreUtilities.cpp
+++ b/Engine/Source/Developer/IoStoreUtilities/Private/IoStoreUtilities.cpp
@@ -881,8 +881,7 @@ static bool WriteBulkData(	const FString& Filename, EIoChunkType Type, const FPa
 					const FIoStatus PartialResult = IoStoreWriter->MapPartialRange(BulkDataChunkId, BulkDataDesc.Offset, BulkDataDesc.Size, AccessChunkId);
 					if (!PartialResult.IsOk())
 					{
-						UE_LOG(LogIoStore, Error, TEXT("Failed to map partial range for '%s' due to: %s"), *Package.FileName, *PartialResult.ToString());
-						return false;
+						UE_LOG(LogIoStore, Warning, TEXT("Failed to map partial range for '%s' due to: %s"), *Package.FileName, *PartialResult.ToString());
 					}
 #endif
 				}
@@ -1124,7 +1123,7 @@ int32 CreateTarget(const FContainerTarget& Target)
 		TUniquePtr<FArchive> Ar(IFileManager::Get().CreateFileReader(*FileName));
 		check(Ar);
 
-		UE_CLOG(FileIndex % 1000 == 0, LogIoStore, Display, TEXT("Parsing %d: '%s'"), FileIndex, *FileName);
+		UE_CLOG(FileIndex % 1000 == 0, LogIoStore, Display, TEXT("Parsing %d/%d: '%s'"), FileIndex, FileNames.Num(), *FileName);
 
 		FString RelativeFileName = FileName;
 		RelativeFileName.RemoveFromStart(*CookedDir);
@@ -1620,12 +1619,21 @@ int32 CreateTarget(const FContainerTarget& Target)
 	}
 #endif
 
+	struct FExportBundleMeta
+	{
+		uint32 LoadOrder = -1;
+		uint32 PayloadSize = -1;
+	};
+
 	UE_LOG(LogIoStore, Display, TEXT("Serializing..."));
-	TArray<uint32> ExportBundleOrderEntries;
+	TArray<FExportBundleMeta> ExportBundleMetaEntries;
+	int32 LogIndex = 0;
 	for (auto& PackageKV : PackageMap)
 	{
 		FPackage& Package = *PackageKV.Value;
 
+		UE_CLOG(LogIndex++ % 1000 == 0, LogIoStore, Display, TEXT("Serializing %d/%d: '%s'"), LogIndex, PackageMap.Num(), *Package.Name.ToString());
+
 		// ImportedPackages
 		Package.ImportedPackagesCount = Package.ImportedPackages.Num();
 		Package.ImportedPackagesOffset = ImportedPackagesArchive->Tell();
@@ -1701,7 +1709,6 @@ int32 CreateTarget(const FContainerTarget& Target)
 		}
 		for (FExportBundle& ExportBundle : Package.ExportBundles)
 		{
-			ExportBundleOrderEntries.Add(ExportBundle.LoadOrder);
 			for (FExportGraphNode* ExportNode : ExportBundle.Nodes)
 			{
 				uint32 CommandType = uint32(ExportNode->BundleEntry.CommandType);
@@ -1715,7 +1722,7 @@ int32 CreateTarget(const FContainerTarget& Target)
 		ensure(Package.GlobalPackageId == PackageIndex++);
 		{
 			int32 ExportBundleCount = Package.ExportBundles.Num();
-			int32 FirstExportBundleIndex = ExportBundleOrderEntries.Num() - ExportBundleCount;
+			int32 FirstExportBundleIndex = ExportBundleMetaEntries.Num();
 
 			NameMapBuilder.MarkNameAsReferenced(Package.Name);
 
@@ -1803,6 +1810,8 @@ int32 CreateTarget(const FContainerTarget& Target)
 				for (int32 I = 0; I < Package.ExportBundles.Num(); ++I)
 				{
 					FExportBundle& ExportBundle = Package.ExportBundles[I];
+					FExportBundleMeta& ExportBundleMeta = ExportBundleMetaEntries.AddDefaulted_GetRef();
+					ExportBundleMeta.LoadOrder = ExportBundle.LoadOrder;
 					for (FExportGraphNode* Node : ExportBundle.Nodes)
 					{
 						if (Node->BundleEntry.CommandType == FExportBundleEntry::ExportCommandType_Serialize)
@@ -1811,6 +1820,10 @@ int32 CreateTarget(const FContainerTarget& Target)
 							BundleBufferSize += ObjectExport.SerialSize;
 						}
 					}
+					if (I == 0)
+					{
+						ExportBundleMeta.PayloadSize = BundleBufferSize;
+					}
 				}
 
 #if !SKIP_WRITE_CONTAINER
@@ -1877,7 +1890,7 @@ int32 CreateTarget(const FContainerTarget& Target)
 	int32 StoreTocByteCount = StoreTocArchive->TotalSize();
 	int32 ImportedPackagesByteCount = ImportedPackagesArchive->TotalSize();
 	int32 GlobalImportNamesByteCount = GlobalImportNamesArchive->TotalSize();
-	int32 ExportBundleOrderByteCount = ExportBundleOrderEntries.Num() * sizeof(uint32);
+	int32 ExportBundleMetaByteCount = ExportBundleMetaEntries.Num() * sizeof(FExportBundleMeta);
 	{
 		UE_LOG(LogIoStore, Display, TEXT("Saving global meta data to container file"));
 		FLargeMemoryWriter GlobalMetaArchive(0, true);
@@ -1891,8 +1904,8 @@ int32 CreateTarget(const FContainerTarget& Target)
 		GlobalMetaArchive << GlobalImportNamesByteCount;
 		GlobalMetaArchive.Serialize(GlobalImportNamesArchive->GetData(), GlobalImportNamesByteCount);
 
-		GlobalMetaArchive << ExportBundleOrderByteCount;
-		GlobalMetaArchive.Serialize(ExportBundleOrderEntries.GetData(), ExportBundleOrderByteCount);
+		GlobalMetaArchive << ExportBundleMetaByteCount;
+		GlobalMetaArchive.Serialize(ExportBundleMetaEntries.GetData(), ExportBundleMetaByteCount);
 
 		const FIoStatus Status = AppendToContainer(*IoStoreWriter, GlobalMetaArchive, CreateIoChunkId(0, 0, EIoChunkType::LoaderGlobalMeta));
 		if (!Status.IsOk())
@@ -2024,7 +2037,7 @@ int32 CreateTarget(const FContainerTarget& Target)
 	UE_LOG(LogIoStore, Display, TEXT("IoStore: %8.2f MB GlobalNameHashes"), (double)GlobalNameHashesMB);
 	UE_LOG(LogIoStore, Display, TEXT("IoStore: %8.2f MB GlobalPackageData"), (double)StoreTocByteCount / 1024.0 / 1024.0);
 	UE_LOG(LogIoStore, Display, TEXT("IoStore: %8.2f MB GlobalImportedPackages, %d imported packages"), (double)ImportedPackagesByteCount / 1024.0 / 1024.0, ImportedPackagesCount);
-	UE_LOG(LogIoStore, Display, TEXT("IoStore: %8.2f MB GlobalBundleOrders, %d bundles"), (double)ExportBundleOrderByteCount / 1024.0 / 1024.0, ExportBundleOrderEntries.Num());
+	UE_LOG(LogIoStore, Display, TEXT("IoStore: %8.2f MB GlobalBundleMeta, %d bundles"), (double)ExportBundleMetaByteCount / 1024.0 / 1024.0, ExportBundleMetaEntries.Num());
 	UE_LOG(LogIoStore, Display, TEXT("IoStore: %8.2f MB GlobalImportNames, %d total imports, %d script imports, %d UPackage imports"),
 		(double)GlobalImportNamesByteCount / 1024.0 / 1024.0, GlobalImportsByFullName.Num(), NumScriptImports, UniqueImportPackages);
 	UE_LOG(LogIoStore, Display, TEXT("IoStore: %8.2f MB InitialLoadData, %d script arcs, %d script outers, %d packages"), (double)InitialLoadSize / 1024.0 / 1024.0, ScriptArcsCount, NumScriptImports, FileNames.Num());
diff --git a/Engine/Source/Runtime/CoreUObject/Private/Serialization/AsyncLoading2.cpp b/Engine/Source/Runtime/CoreUObject/Private/Serialization/AsyncLoading2.cpp
--- a/Engine/Source/Runtime/CoreUObject/Private/Serialization/AsyncLoading2.cpp
+++ b/Engine/Source/Runtime/CoreUObject/Private/Serialization/AsyncLoading2.cpp
@@ -252,6 +252,12 @@ struct FPackageStoreInitialLoadEntry
 	int32 ScriptArcsCount;
 };
 
+struct FExportBundleMetaEntry
+{
+	uint32 LoadOrder;
+	uint32 PayloadSize;
+};
+
 class FGlobalNameMap
 {
 public:
@@ -332,7 +338,7 @@ public:
 	FPackageStoreEntry* StoreEntries = nullptr;
 	FPackageStoreInitialLoadEntry* InitialLoadStoreEntries = nullptr;
 	int32* ImportedPackages = nullptr;
-	uint32* ExportBundleOrderEntries = nullptr;
+	FExportBundleMetaEntry* ExportBundleMetaEntries = nullptr;
 	int32* ScriptArcs = nullptr;
 	int32 PackageCount = 0;
 	int32 ImportedPackagesCount = 0;
@@ -428,12 +434,12 @@ public:
 		}
 
 		{
-			TRACE_CPUPROFILER_EVENT_SCOPE(LoadPackageStoreBundleOrder);
+			TRACE_CPUPROFILER_EVENT_SCOPE(LoadPackageStoreBundleMeta);
 			
-			int32 BundleOrderByteCount;
-			GlobalMetaAr << BundleOrderByteCount;
-			ExportBundleOrderEntries = reinterpret_cast<uint32*>(FMemory::Malloc(BundleOrderByteCount));
-			GlobalMetaAr.Serialize(ExportBundleOrderEntries, BundleOrderByteCount);
+			int32 BundleMetaByteCount;
+			GlobalMetaAr << BundleMetaByteCount;
+			ExportBundleMetaEntries = reinterpret_cast<FExportBundleMetaEntry*>(FMemory::Malloc(BundleMetaByteCount));
+			GlobalMetaAr.Serialize(ExportBundleMetaEntries, BundleMetaByteCount);
 		}
 
 
@@ -550,10 +556,10 @@ public:
 		return StoreEntries[PackageId.ToIndex()].ExportCount;
 	}
 
-	inline TTuple<uint32*, uint32> GetPackageExportBundles(FPackageId PackageId) const
+	inline TTuple<FExportBundleMetaEntry*, uint32> GetPackageExportBundleMetaEntries(FPackageId PackageId) const
 	{
 		const FPackageStoreEntry* StoreEntry = StoreEntries + PackageId.ToIndex();
-		return MakeTuple(ExportBundleOrderEntries + StoreEntry->FirstExportBundleIndex, uint32(StoreEntry->ExportBundleCount));
+		return MakeTuple(ExportBundleMetaEntries + StoreEntry->FirstExportBundleIndex, uint32(StoreEntry->ExportBundleCount));
 	}
 };
 
@@ -1283,7 +1289,7 @@ private:
 	TArray<UObject*> Exports;
 	FPackageImportStore ImportStore;
 
-	uint32* ExportBundleOrderEntries = nullptr;
+	FExportBundleMetaEntry* ExportBundleMetaEntries = nullptr;
 	int32 ExportBundleCount = 0;
 	const FExportBundleHeader* ExportBundles = nullptr;
 	const FExportBundleEntry* ExportBundleEntries = nullptr;
@@ -1543,9 +1549,10 @@ private:
 
 		FAsyncPackage2* Package;
 		uint32 BundeOrder;
+		uint32 BundleSize;
 	};
 	TArray<FBundleIoRequest> WaitingIoRequests;
-	uint32 PendingBundleIoRequestsCount = 0;
+	uint64 PendingBundleIoRequestsTotalSize = 0;
 
 public:
 
@@ -1805,8 +1812,8 @@ private:
 	EAsyncPackageState::Type ProcessLoadedPackagesFromGameThread(bool& bDidSomething, int32 FlushRequestID = INDEX_NONE);
 
 	bool CreateAsyncPackagesFromQueue();
-	void AddBundleIoRequest(FAsyncPackage2* Package, uint32 BundleOrder);
-	void BundleIoRequestCompleted();
+	void AddBundleIoRequest(FAsyncPackage2* Package, const FExportBundleMetaEntry& BundleMetaEntry);
+	void BundleIoRequestCompleted(const FExportBundleMetaEntry& BundleMetaEntry);
 	void StartBundleIoRequests();
 
 	FAsyncPackage2* CreateAsyncPackage(const FAsyncPackageDesc2& Desc)
@@ -1973,14 +1980,15 @@ bool FAsyncLoadingThread2Impl::CreateAsyncPackagesFromQueue()
 	return QueueCopy.Num() > 0;
 }
 
-void FAsyncLoadingThread2Impl::AddBundleIoRequest(FAsyncPackage2* Package, uint32 BundleOrder)
+void FAsyncLoadingThread2Impl::AddBundleIoRequest(FAsyncPackage2* Package, const FExportBundleMetaEntry& BundleMetaEntry)
 {
-	WaitingIoRequests.HeapPush({ Package, BundleOrder });
+	WaitingIoRequests.HeapPush({ Package, BundleMetaEntry.LoadOrder, BundleMetaEntry.PayloadSize });
 }
 
-void FAsyncLoadingThread2Impl::BundleIoRequestCompleted()
+void FAsyncLoadingThread2Impl::BundleIoRequestCompleted(const FExportBundleMetaEntry& BundleMetaEntry)
 {
-	--PendingBundleIoRequestsCount;
+	check(PendingBundleIoRequestsTotalSize >= BundleMetaEntry.PayloadSize)
+	PendingBundleIoRequestsTotalSize -= BundleMetaEntry.PayloadSize;
 	if (WaitingIoRequests.Num())
 	{
 		StartBundleIoRequests();
@@ -1990,14 +1998,19 @@ void FAsyncLoadingThread2Impl::BundleIoRequestCompleted()
 void FAsyncLoadingThread2Impl::StartBundleIoRequests()
 {
 	TRACE_CPUPROFILER_EVENT_SCOPE(StartBundleIoRequests);
-	const uint32 MaxPendingRequestsCount = 1024;
+	constexpr uint64 MaxPendingRequestsSize = 256 << 20;
 	FAsyncPackage2* PreviousPackage = nullptr;
-	while (WaitingIoRequests.Num() && PendingBundleIoRequestsCount < MaxPendingRequestsCount)
+	while (WaitingIoRequests.Num())
 	{
-		FBundleIoRequest BundleIoRequest;
-		WaitingIoRequests.HeapPop(BundleIoRequest, false);
+		FBundleIoRequest& BundleIoRequest = WaitingIoRequests.HeapTop();
 		FAsyncPackage2* Package = BundleIoRequest.Package;
 		check(Package);
+		if (PendingBundleIoRequestsTotalSize > 0 && PendingBundleIoRequestsTotalSize + BundleIoRequest.BundleSize > MaxPendingRequestsSize)
+		{
+			break;
+		}
+		PendingBundleIoRequestsTotalSize += BundleIoRequest.BundleSize;
+		WaitingIoRequests.HeapPop(BundleIoRequest, false);
 
 		if (GIsInitialLoad && PreviousPackage)
 		{
@@ -2013,8 +2026,6 @@ void FAsyncLoadingThread2Impl::StartBundleIoRequests()
 			Package->IoBuffer = Result.ConsumeValueOrDie();
 			Package->GetExportBundleNode(EEventLoadNode2::ExportBundle_Process, 0)->ReleaseBarrier();
 		});
-
-		++PendingBundleIoRequestsCount;
 	}
 }
 
@@ -2490,8 +2501,8 @@ void FAsyncPackage2::StartLoading()
 	LoadStartTime = FPlatformTime::Seconds();
 
 	check(ExportBundleCount > 0);
-	uint32 BundleOrder = ExportBundleOrderEntries[0];
-	AsyncLoadingThread.AddBundleIoRequest(this, BundleOrder);
+	FExportBundleMetaEntry& BundleMetaEntry = ExportBundleMetaEntries[0];
+	AsyncLoadingThread.AddBundleIoRequest(this, BundleMetaEntry);
 }
 
 void FAsyncPackage2::SetupScriptArcs()
@@ -2643,7 +2654,7 @@ EAsyncPackageState::Type FAsyncPackage2::Event_ProcessExportBundle(FAsyncPackage
 
 	if (ExportBundleIndex == 0)
 	{
-		Package->AsyncLoadingThread.BundleIoRequestCompleted();
+		Package->AsyncLoadingThread.BundleIoRequestCompleted(Package->ExportBundleMetaEntries[0]);
 	}
 
 	return EAsyncPackageState::Complete;
@@ -2876,17 +2887,19 @@ void FAsyncPackage2::EventDrivenSerializeExport(int32 LocalExportIndex, FSimpleE
 		}
 	}
 
+	// cache archetype
+	// prevents GetArchetype from hitting the expensive GetArchetypeFromRequiredInfoImpl
+	check(!Export.TemplateIndex.IsNull());
+	UObject* Template = EventDrivenIndexToObject(Export.TemplateIndex, true);
+	check(Template);
+	CacheArchetypeForObject(Object, Template);
+
 	Object->ClearFlags(RF_NeedLoad);
 
 	FUObjectSerializeContext* LoadContext = GetSerializeContext();
 	UObject* PrevSerializedObject = LoadContext->SerializedObject;
 	LoadContext->SerializedObject = Object;
 
-	// Find the Archetype object for the one we are loading. This is piped to GetArchetypeFromLoader
-	check(!Export.TemplateIndex.IsNull());
-	UObject* Template = EventDrivenIndexToObject(Export.TemplateIndex, true);
-	check(Template);
-
 	Ar.TemplateForGetArchetypeFromLoader = Template;
 
 	if (Object->HasAnyFlags(RF_ClassDefaultObject))
@@ -2926,6 +2939,7 @@ EAsyncPackageState::Type FAsyncPackage2::Event_ExportsDone(FAsyncPackage2* Packa
 	Package->ImportStore.ImportMapCount = 0;
 	Package->bAllExportsSerialized = true;
 	Package->AsyncPackageLoadingState = EAsyncPackageLoadingState2::PostLoad_Etc;
+	Package->IoBuffer = FIoBuffer();
 
 	Package->GetNode(EEventLoadNode2::Package_StartPostLoad)->ReleaseBarrier();
 	return EAsyncPackageState::Complete;
@@ -3917,8 +3931,8 @@ FAsyncPackage2::FAsyncPackage2(
 	TRACE_LOADTIME_NEW_ASYNC_PACKAGE(this, InDesc.Name);
 	AddRequestID(InDesc.RequestID);
 
-	TTuple<uint32*, uint32> PackageExportBundles = AsyncLoadingThread.GlobalPackageStore.GetPackageExportBundles(Desc.PackageId);
-	ExportBundleOrderEntries = PackageExportBundles.Get<0>();
+	TTuple<FExportBundleMetaEntry*, uint32> PackageExportBundles = AsyncLoadingThread.GlobalPackageStore.GetPackageExportBundleMetaEntries(Desc.PackageId);
+	ExportBundleMetaEntries = PackageExportBundles.Get<0>();
 	ExportBundleCount = PackageExportBundles.Get<1>();
 	ExportCount = AsyncLoadingThread.GlobalPackageStore.GetPackageExportCount(Desc.PackageId);
 	Exports.AddZeroed(ExportCount);
@@ -4192,10 +4206,6 @@ EAsyncPackageState::Type FAsyncPackage2::PostLoadObjects()
 			ThreadContext.CurrentlyPostLoadedObjectByALT = Object;
 			{
 				TRACE_LOADTIME_POSTLOAD_EXPORT_SCOPE(Object);
-				// cache archetype before post load
-				// prevents GetArchetype from hitting the expensive GetArchetypeFromRequiredInfoImpl
-				UObject* Template = EventDrivenIndexToObject(Export.TemplateIndex, true);
-				CacheArchetypeForObject(Object, Template);
 				Object->ConditionalPostLoad();
 				Object->AtomicallyClearInternalFlags(EInternalObjectFlags::AsyncLoading);
 			}
@@ -4220,22 +4230,6 @@ EAsyncPackageState::Type FAsyncPackage2::PostLoadDeferredObjects()
 
 	FUObjectSerializeContext* LoadContext = GetSerializeContext();
 
-	{
-		// cache archetype for all exports before post load
-		// prevents GetArchetype from hitting the expensive GetArchetypeFromRequiredInfoImpl
-		TRACE_CPUPROFILER_EVENT_SCOPE(CacheArchetype);
-		for (int32 LocalExportIndex = 0; LocalExportIndex < ExportCount; ++LocalExportIndex)
-		{
-			UObject* Object = Exports[LocalExportIndex];
-			if (Object->HasAnyFlags(RF_NeedPostLoad))
-			{
-				const FExportMapEntry& Export = ExportMap[LocalExportIndex];
-				UObject* Template = EventDrivenIndexToObject(Export.TemplateIndex, true);
-				CacheArchetypeForObject(Object, Template);
-			}
-		}
-	}
-
 	while (DeferredPostLoadIndex < ExportCount && 
 		!AsyncLoadingThread.IsAsyncLoadingSuspended() &&
 		!FAsyncLoadingThreadState2::Get()->IsTimeLimitExceeded())
