Fix for the multicast delegate being destroyed mid-execution.

diff --git a/Engine/Source/Runtime/CoreUObject/Private/UObject/SparseDelegate.cpp b/Engine/Source/Runtime/CoreUObject/Private/UObject/SparseDelegate.cpp
--- a/Engine/Source/Runtime/CoreUObject/Private/UObject/SparseDelegate.cpp
+++ b/Engine/Source/Runtime/CoreUObject/Private/UObject/SparseDelegate.cpp
@@ -95,6 +95,21 @@ FMulticastScriptDelegate* FSparseDelegateStorage::GetMulticastDelegate(const UOb
 	return nullptr;
 }
 
+TSharedPtr<FMulticastScriptDelegate> FSparseDelegateStorage::GetSharedMulticastDelegate(const UObject* DelegateOwner, const FName DelegateName)
+{
+	FScopeLock SparseDelegateMapLock(&SparseDelegateMapCritical);
+
+	TSharedPtr<FMulticastScriptDelegate> Result;
+	if (FSparseDelegateMap* DelegateMap = SparseDelegates.Find(DelegateOwner))
+	{
+		if (TSharedPtr<FMulticastScriptDelegate>* MulticastDelegatePtr = DelegateMap->Find(DelegateName))
+		{
+			Result = *MulticastDelegatePtr;
+		}
+	}
+	return Result;
+}
+
 void FSparseDelegateStorage::SetMulticastDelegate(const UObject* DelegateOwner, const FName DelegateName, FMulticastScriptDelegate Delegate)
 {
 	FScopeLock SparseDelegateMapLock(&SparseDelegateMapCritical);
diff --git a/Engine/Source/Runtime/CoreUObject/Public/UObject/SparseDelegate.h b/Engine/Source/Runtime/CoreUObject/Public/UObject/SparseDelegate.h
--- a/Engine/Source/Runtime/CoreUObject/Public/UObject/SparseDelegate.h
+++ b/Engine/Source/Runtime/CoreUObject/Public/UObject/SparseDelegate.h
@@ -52,6 +52,9 @@ public:
 	/** Acquires the actual Multicast Delegate from the annotation if any delegates are bound to it. Will be null if no entry exists in the annotation for this object/delegatename. */
 	static FMulticastScriptDelegate* GetMulticastDelegate(const UObject* DelegateOwner, const FName DelegateName);
 
+	/** Acquires the actual Multicast Delegate from the annotation if any delegates are bound to it as a shared pointer. Will be null if no entry exists in the annotation for this object/delegatename. */
+	static TSharedPtr<FMulticastScriptDelegate> GetSharedMulticastDelegate(const UObject* DelegateOwner, const FName DelegateName);
+
 	/** Directly sets the Multicast Delegate for this object/delegatename pair. If the delegate is unbound it will be assigned/inserted anyways. */
 	static void SetMulticastDelegate(const UObject* DelegateOwner, const FName DelegateName, FMulticastScriptDelegate Delegate);
 
@@ -196,9 +199,26 @@ private:
 
 public:
 	/** Returns the multicast delegate if any delegates are bound to the sparse delegate */
+	UE_DEPRECATED(4.25, "This function has been deprecated - please use GetShared() instead")
 	MulticastDelegate* Get() const
 	{
-		return (MulticastDelegate*)(bIsBound ? FSparseDelegateStorage::GetMulticastDelegate(GetDelegateOwner(), GetDelegateName()) : nullptr);
+		MulticastDelegate* Result = nullptr;
+		if (bIsBound)
+		{
+			Result = static_cast<MulticastDelegate*>(FSparseDelegateStorage::GetMulticastDelegate(GetDelegateOwner(), GetDelegateName()));
+		}
+		return Result;
+	}
+
+	/** Returns the multicast delegate if any delegates are bound to the sparse delegate */
+	TSharedPtr<MulticastDelegate> GetShared() const
+	{
+		TSharedPtr<MulticastDelegate> Result;
+		if (bIsBound)
+		{
+			Result = StaticCastSharedPtr<MulticastDelegate>(FSparseDelegateStorage::GetSharedMulticastDelegate(GetDelegateOwner(), GetDelegateName()));
+		}
+		return Result;
 	}
 
 	/**
@@ -298,7 +318,7 @@ public:
 	template<typename... ParamTypes>
 	void Broadcast(ParamTypes... Params)
 	{
-		if (auto MCDelegate = Get())
+		if (TSharedPtr<MulticastDelegate> MCDelegate = GetShared())
 		{
 			MCDelegate->Broadcast(Params...);
 		}
