CIS Fixes

diff --git a/Engine/Source/Runtime/CoreUObject/Private/UObject/FieldPath.cpp b/Engine/Source/Runtime/CoreUObject/Private/UObject/FieldPath.cpp
--- a/Engine/Source/Runtime/CoreUObject/Private/UObject/FieldPath.cpp
+++ b/Engine/Source/Runtime/CoreUObject/Private/UObject/FieldPath.cpp
@@ -9,6 +9,18 @@
 #include "UObject/Package.h"
 #include "UObject/UObjectArray.h"
 
+#if WITH_EDITORONLY_DATA
+FFieldPath::FFieldPath(UField* InField, const FName& InPropertyTypeName)
+{
+	if (InField)
+	{
+		// Must be constructed from the equivalent UField class
+		check(InField->GetClass()->GetFName() == InPropertyTypeName);
+		GenerateFromUField(InField);
+	}
+}
+#endif
+
 void FFieldPath::Generate(FField* InField)
 {
 	Path.Empty();
diff --git a/Engine/Source/Runtime/CoreUObject/Public/UObject/Field.h b/Engine/Source/Runtime/CoreUObject/Public/UObject/Field.h
--- a/Engine/Source/Runtime/CoreUObject/Public/UObject/Field.h
+++ b/Engine/Source/Runtime/CoreUObject/Public/UObject/Field.h
@@ -237,33 +237,19 @@ public:
 	bool IsA(const UClass* InClass) const;
 	bool IsA(const FFieldClass* InClass) const;
 	template <typename T>
-	typename TEnableIf<TIsDerivedFrom<T, UObject>::IsDerived, bool>::Type IsA() const
+	bool IsA() const
 	{
 		static_assert(sizeof(T) > 0, "T must not be an incomplete type");
-		if (bIsUObject && Container.Object)
-		{
-			return Container.Object->IsA(T::StaticClass());
-		}
-		return false;
+		return IsA(T::StaticClass());
 	}
 
-	template <typename T>
-	typename TEnableIf<!TIsDerivedFrom<T, UObject>::IsDerived, bool>::Type IsA() const
-	{
-		static_assert(sizeof(T) > 0, "T must not be an incomplete type");
-		if (!bIsUObject && Container.Field)
-		{
-			return Container.Field->IsA(T::StaticClass());
-		}
-		return false;
-	}
 	template <typename T>
 	typename TEnableIf<TIsDerivedFrom<T, UObject>::IsDerived, T*>::Type Get() const
 	{
 		static_assert(sizeof(T) > 0, "T must not be an incomplete type");
-		if (bIsUObject)
+		if (IsA(T::StaticClass()))
 		{
-			return Cast<T>(Container.Object);
+			return static_cast<T*>(Container.Object);
 		}
 		return nullptr;
 	}
@@ -272,9 +258,9 @@ public:
 	typename TEnableIf<!TIsDerivedFrom<T, UObject>::IsDerived, T*>::Type Get() const
 	{
 		static_assert(sizeof(T) > 0, "T must not be an incomplete type");
-		if (!bIsUObject)
+		if (IsA(T::StaticClass()))
 		{
-			return CastField<T>(Container.Field);
+			return static_cast<T*>(Container.Field);
 		}
 		return nullptr;
 	}
@@ -544,37 +530,26 @@ public:
 	FField* GetTypedOwner(FFieldClass* Target) const;
 
 	template <typename T>
-	typename TEnableIf<TIsDerivedFrom<T, UObject>::IsDerived, T*>::Type GetOwner() const
+	T* GetOwner() const
 	{
 		static_assert(sizeof(T) > 0, "T must not be an incomplete type");
-		return Cast<T>(Owner.ToUObject());
+		return Owner.Get<T>();
 	}
 
 	template <typename T>
-	typename TEnableIf<!TIsDerivedFrom<T, UObject>::IsDerived, T*>::Type GetOwner() const
+	T* GetOwnerChecked() const
 	{
 		static_assert(sizeof(T) > 0, "T must not be an incomplete type");
-		return CastField<T>(Owner.ToField());
-	}
-
-	template <typename T>
-	typename TEnableIf<TIsDerivedFrom<T, UObject>::IsDerived, T*>::Type GetOwnerChecked() const
-	{
-		static_assert(sizeof(T) > 0, "T must not be an incomplete type");
-		return CastChecked<T>(Owner.ToUObject());
-	}
-
-	template <typename T>
-	typename TEnableIf<!TIsDerivedFrom<T, UObject>::IsDerived, T*>::Type GetOwnerChecked() const
-	{
-		static_assert(sizeof(T) > 0, "T must not be an incomplete type");
-		return CastFieldChecked<T>(Owner.ToField());
+		T* Result = Owner.Get<T>();
+		check(Result);
+		return Result;
 	}
 
 	template <typename T>
 	T* GetTypedOwner() const
 	{
-		return Cast<T>(GetTypedOwner(T::StaticClass()));
+		static_assert(sizeof(T) > 0, "T must not be an incomplete type");
+		return static_cast<T*>(GetTypedOwner(T::StaticClass()));
 	}
 
 	FORCEINLINE FName GetFName() const
diff --git a/Engine/Source/Runtime/CoreUObject/Public/UObject/FieldPath.h b/Engine/Source/Runtime/CoreUObject/Public/UObject/FieldPath.h
--- a/Engine/Source/Runtime/CoreUObject/Public/UObject/FieldPath.h
+++ b/Engine/Source/Runtime/CoreUObject/Public/UObject/FieldPath.h
@@ -58,6 +58,10 @@ public:
 	{
 	}
 
+#if WITH_EDITORONLY_DATA
+	FFieldPath(UField* InField, const FName& InPropertyTypeName);
+#endif
+
 	/** Generates path from the passed in field pointer */
 	void Generate(FField* InField);
 
@@ -207,14 +211,8 @@ public:
 
 #if WITH_EDITORONLY_DATA
 	TFieldPath(UField* InField)
-		: FFieldPath()
+		: FFieldPath(InField, PropertyType::StaticClass()->GetFName())
 	{
-		if (InField)
-		{
-			// Must be constructed from the equivalent UField class
-			check(InField->GetClass()->GetFName() == PropertyType::StaticClass()->GetFName());
-			GenerateFromUField(InField);
-		}
 	}
 #endif
 
diff --git a/Engine/Source/Runtime/CoreUObject/Public/UObject/UnrealType.h b/Engine/Source/Runtime/CoreUObject/Public/UObject/UnrealType.h
--- a/Engine/Source/Runtime/CoreUObject/Public/UObject/UnrealType.h
+++ b/Engine/Source/Runtime/CoreUObject/Public/UObject/UnrealType.h
@@ -2156,7 +2156,7 @@ public:
 	TFObjectPropertyBase(FFieldVariant InOwner, const FName& InName, EObjectFlags InObjectFlags)
 		: Super(InOwner, InName, InObjectFlags)
 	{
-		PropertyClass = nullptr;
+		this->PropertyClass = nullptr;
 	}
 
 
@@ -5271,7 +5271,7 @@ public:
 protected:
 	inline void IterateToNext()
 	{
-		T::BaseFieldClass* CurrentField  = Field;
+		typename T::BaseFieldClass* CurrentField  = Field;
 		const UStruct* CurrentStruct = Struct;
 
 		while (CurrentStruct)
