More CIS fixes

diff --git a/Engine/Source/Runtime/CoreUObject/Public/UObject/UnrealType.h b/Engine/Source/Runtime/CoreUObject/Public/UObject/UnrealType.h
--- a/Engine/Source/Runtime/CoreUObject/Public/UObject/UnrealType.h
+++ b/Engine/Source/Runtime/CoreUObject/Public/UObject/UnrealType.h
@@ -5220,7 +5220,7 @@ public:
 	               EFieldIteratorFlags::DeprecatedPropertyFlags InDeprecatedFieldFlags = EFieldIteratorFlags::IncludeDeprecated,
 	               EFieldIteratorFlags::InterfaceClassFlags     InInterfaceFieldFlags  = EFieldIteratorFlags::ExcludeInterfaces)
 		: Struct            ( InStruct )
-		, Field             ( InStruct ? GetChildFieldsFromStruct<T::BaseFieldClass>(InStruct) : NULL )
+		, Field             ( InStruct ? GetChildFieldsFromStruct<typename T::BaseFieldClass>(InStruct) : NULL )
 		, InterfaceIndex    ( -1 )
 		, bIncludeSuper     ( InSuperClassFlags      == EFieldIteratorFlags::IncludeSuper )
 		, bIncludeDeprecated( InDeprecatedFieldFlags == EFieldIteratorFlags::IncludeDeprecated )
@@ -5302,7 +5302,7 @@ protected:
 				if (InterfaceIndex < CurrentClass->Interfaces.Num())
 				{
 					FImplementedInterface& Interface = CurrentClass->Interfaces[InterfaceIndex];
-					CurrentField = Interface.Class ? GetChildFieldsFromStruct<T::BaseFieldClass>(Interface.Class) : nullptr;
+					CurrentField = Interface.Class ? GetChildFieldsFromStruct<typename T::BaseFieldClass>(Interface.Class) : nullptr;
 					continue;
 				}
 			}
@@ -5312,7 +5312,7 @@ protected:
 				CurrentStruct = CurrentStruct->GetInheritanceSuper();
 				if (CurrentStruct)
 				{
-					CurrentField   = GetChildFieldsFromStruct<T::BaseFieldClass>(CurrentStruct);
+					CurrentField   = GetChildFieldsFromStruct<typename T::BaseFieldClass>(CurrentStruct);
 					InterfaceIndex = -1;
 					continue;
 				}
