CIS error fix

diff --git a/Engine/Source/Runtime/CoreUObject/Private/UObject/Class.cpp b/Engine/Source/Runtime/CoreUObject/Private/UObject/Class.cpp
--- a/Engine/Source/Runtime/CoreUObject/Private/UObject/Class.cpp
+++ b/Engine/Source/Runtime/CoreUObject/Private/UObject/Class.cpp
@@ -1906,7 +1906,7 @@ void UStruct::PostLoad()
 		FField* ResolvedProperty = MissingProperty.Key.Get(this);
 		if (ResolvedProperty)
 		{			
-			check(Script.Num() >= (MissingProperty.Value + sizeof(FField*)));
+			check((int32)Script.Num() >= (int32)(MissingProperty.Value + sizeof(FField*)));
 			FField** TargetScriptPropertyPtr = (FField**)(Script.GetData() + MissingProperty.Value);
 			*TargetScriptPropertyPtr = ResolvedProperty;
 		}
