Fixing Linux compile error

diff --git a/Engine/Source/Runtime/CoreUObject/Public/Templates/SubclassOf.h b/Engine/Source/Runtime/CoreUObject/Public/Templates/SubclassOf.h
--- a/Engine/Source/Runtime/CoreUObject/Public/Templates/SubclassOf.h
+++ b/Engine/Source/Runtime/CoreUObject/Public/Templates/SubclassOf.h
@@ -16,6 +16,7 @@ class TSubclassOf
 public:
 
 	typedef typename TChooseClass<TIsDerivedFrom<TClass, FField>::IsDerived, FFieldClass, UClass>::Result TClassType;
+	typedef typename TChooseClass<TIsDerivedFrom<TClass, FField>::IsDerived, FField, UObject>::Result TBaseType;
 
 private:
 
@@ -92,7 +93,9 @@ public:
 	 */
 	FORCEINLINE TClass* GetDefaultObject() const
 	{
-		return Class ? Class->GetDefaultObject<TClass>() : nullptr;
+		TBaseType* Result = Class ? Class->GetDefaultObject() : nullptr;
+		check(Result && Result->IsA(TClass::StaticClass()));
+		return (TClass*)Result;
 	}
 
 	friend FArchive& operator<<(FArchive& Ar, TSubclassOf& SubclassOf)
