Fixing a crash when accessing TSubclassOf default object when its class is null

diff --git a/Engine/Source/Runtime/CoreUObject/Public/Templates/SubclassOf.h b/Engine/Source/Runtime/CoreUObject/Public/Templates/SubclassOf.h
--- a/Engine/Source/Runtime/CoreUObject/Public/Templates/SubclassOf.h
+++ b/Engine/Source/Runtime/CoreUObject/Public/Templates/SubclassOf.h
@@ -93,8 +93,12 @@ public:
 	 */
 	FORCEINLINE TClass* GetDefaultObject() const
 	{
-		TBaseType* Result = Class ? Class->GetDefaultObject() : nullptr;
-		check(Result && Result->IsA(TClass::StaticClass()));
+		TBaseType* Result = nullptr;
+		if (Class)
+		{
+			Result = Class->GetDefaultObject();
+			check(Result && Result->IsA(TClass::StaticClass()));
+		}
 		return (TClass*)Result;
 	}
 
