Making all UField -> FProperty conversion constructors explicit

diff --git a/Engine/Source/Runtime/CoreUObject/Public/UObject/EnumProperty.h b/Engine/Source/Runtime/CoreUObject/Public/UObject/EnumProperty.h
--- a/Engine/Source/Runtime/CoreUObject/Public/UObject/EnumProperty.h
+++ b/Engine/Source/Runtime/CoreUObject/Public/UObject/EnumProperty.h
@@ -23,7 +23,7 @@ public:
 	FEnumProperty(FFieldVariant InOwner, const FName& InName, EObjectFlags InObjectFlags, UEnum* InEnum);
 	FEnumProperty(FFieldVariant InOwner, const FName& InName, EObjectFlags InObjectFlags, int32 InOffset, EPropertyFlags InFlags, UEnum* InEnum);
 #if WITH_EDITORONLY_DATA
-	FEnumProperty(UField* InField);
+	explicit FEnumProperty(UField* InField);
 #endif // WITH_EDITORONLY_DATA
 	virtual ~FEnumProperty();
 
diff --git a/Engine/Source/Runtime/CoreUObject/Public/UObject/Field.h b/Engine/Source/Runtime/CoreUObject/Public/UObject/Field.h
--- a/Engine/Source/Runtime/CoreUObject/Public/UObject/Field.h
+++ b/Engine/Source/Runtime/CoreUObject/Public/UObject/Field.h
@@ -372,7 +372,7 @@ public:
 	FField(EInternal InInernal, FFieldClass* InClass);
 	FField(FFieldVariant InOwner, const FName& InName, EObjectFlags InObjectFlags);
 #if WITH_EDITORONLY_DATA
-	FField(UField* InField);
+	explicit FField(UField* InField);
 #endif // WITH_EDITORONLY_DATA
 	virtual ~FField();
 
diff --git a/Engine/Source/Runtime/CoreUObject/Public/UObject/FieldPathProperty.h b/Engine/Source/Runtime/CoreUObject/Public/UObject/FieldPathProperty.h
--- a/Engine/Source/Runtime/CoreUObject/Public/UObject/FieldPathProperty.h
+++ b/Engine/Source/Runtime/CoreUObject/Public/UObject/FieldPathProperty.h
@@ -32,7 +32,7 @@ public:
 	}
 
 #if WITH_EDITORONLY_DATA
-	FFieldPathProperty(UField* InField);
+	explicit FFieldPathProperty(UField* InField);
 #endif // WITH_EDITORONLY_DATA
 
 	FFieldClass* PropertyClass;
diff --git a/Engine/Source/Runtime/CoreUObject/Public/UObject/TextProperty.h b/Engine/Source/Runtime/CoreUObject/Public/UObject/TextProperty.h
--- a/Engine/Source/Runtime/CoreUObject/Public/UObject/TextProperty.h
+++ b/Engine/Source/Runtime/CoreUObject/Public/UObject/TextProperty.h
@@ -29,7 +29,7 @@ public:
 	}
 
 #if WITH_EDITORONLY_DATA
-	FTextProperty(UField* InField)
+	explicit FTextProperty(UField* InField)
 		: FTextProperty_Super(InField)
 	{
 	}
diff --git a/Engine/Source/Runtime/CoreUObject/Public/UObject/UnrealType.h b/Engine/Source/Runtime/CoreUObject/Public/UObject/UnrealType.h
--- a/Engine/Source/Runtime/CoreUObject/Public/UObject/UnrealType.h
+++ b/Engine/Source/Runtime/CoreUObject/Public/UObject/UnrealType.h
@@ -117,7 +117,7 @@ public:
 	FProperty(FFieldVariant InOwner, const FName& InName, EObjectFlags InObjectFlags);
 	FProperty(FFieldVariant InOwner, const FName& InName, EObjectFlags InObjectFlags, int32 InOffset, EPropertyFlags InFlags);
 #if WITH_EDITORONLY_DATA
-	FProperty(UField* InField);
+	explicit FProperty(UField* InField);
 #endif // WITH_EDITORONLY_DATA
 
 	// UObject interface
@@ -1015,7 +1015,7 @@ public:
 	}
 
 #if WITH_EDITORONLY_DATA
-	TProperty(UField* InField)
+	explicit TProperty(UField* InField)
 		: Super(InField)
 	{
 		SetElementSize();
@@ -1130,7 +1130,7 @@ public:
 	}
 
 #if WITH_EDITORONLY_DATA
-	TProperty_WithEqualityAndSerializer(UField* InField)
+	explicit TProperty_WithEqualityAndSerializer(UField* InField)
 		: Super(InField)
 	{
 	}
@@ -1166,7 +1166,7 @@ class COREUOBJECT_API FNumericProperty : public FProperty
 	{}
 
 #if WITH_EDITORONLY_DATA
-	FNumericProperty(UField* InField)
+	explicit FNumericProperty(UField* InField)
 		: Super(InField)
 	{
 	}
@@ -1303,7 +1303,7 @@ public:
 	}
 
 #if WITH_EDITORONLY_DATA
-	TProperty_Numeric(UField* InField)
+	explicit TProperty_Numeric(UField* InField)
 		: Super(InField)
 	{
 	}
@@ -1528,7 +1528,7 @@ class COREUOBJECT_API FByteProperty : public TProperty_Numeric<uint8>
 	}
 
 #if WITH_EDITORONLY_DATA
-	FByteProperty(UField* InField);
+	explicit FByteProperty(UField* InField);
 #endif // WITH_EDITORONLY_DATA
 
 	// UObject interface.
@@ -1579,7 +1579,7 @@ class COREUOBJECT_API FInt8Property : public TProperty_Numeric<int8>
 	}
 
 #if WITH_EDITORONLY_DATA
-	FInt8Property(UField* InField)
+	explicit FInt8Property(UField* InField)
 		: TProperty_Numeric(InField)
 	{
 	}
@@ -1608,7 +1608,7 @@ class COREUOBJECT_API FInt16Property : public TProperty_Numeric<int16>
 	}
 
 #if WITH_EDITORONLY_DATA
-	FInt16Property(UField* InField)
+	explicit FInt16Property(UField* InField)
 		: TProperty_Numeric(InField)
 	{
 	}
@@ -1638,7 +1638,7 @@ class COREUOBJECT_API FIntProperty : public TProperty_Numeric<int32>
 	}
 
 #if WITH_EDITORONLY_DATA
-	FIntProperty(UField* InField)
+	explicit FIntProperty(UField* InField)
 		: TProperty_Numeric(InField)
 	{
 	}
@@ -1667,7 +1667,7 @@ class COREUOBJECT_API FInt64Property : public TProperty_Numeric<int64>
 	}
 
 #if WITH_EDITORONLY_DATA
-	FInt64Property(UField* InField)
+	explicit FInt64Property(UField* InField)
 		: TProperty_Numeric(InField)
 	{
 	}
@@ -1696,7 +1696,7 @@ class COREUOBJECT_API FUInt16Property : public TProperty_Numeric<uint16>
 	}
 
 #if WITH_EDITORONLY_DATA
-	FUInt16Property(UField* InField)
+	explicit FUInt16Property(UField* InField)
 		: TProperty_Numeric(InField)
 	{
 	}
@@ -1725,7 +1725,7 @@ class COREUOBJECT_API FUInt32Property : public TProperty_Numeric<uint32>
 	}
 
 #if WITH_EDITORONLY_DATA
-	FUInt32Property(UField* InField)
+	explicit FUInt32Property(UField* InField)
 		: TProperty_Numeric(InField)
 	{
 	}
@@ -1754,7 +1754,7 @@ class COREUOBJECT_API FUInt64Property : public TProperty_Numeric<uint64>
 	}
 
 #if WITH_EDITORONLY_DATA
-	FUInt64Property(UField* InField)
+	explicit FUInt64Property(UField* InField)
 		: TProperty_Numeric(InField)
 	{
 	}
@@ -1806,7 +1806,7 @@ class COREUOBJECT_API FFloatProperty : public TProperty_Numeric<float>
 	}
 
 #if WITH_EDITORONLY_DATA
-	FFloatProperty(UField* InField)
+	explicit FFloatProperty(UField* InField)
 		: TProperty_Numeric(InField)
 	{
 	}
@@ -1839,7 +1839,7 @@ class COREUOBJECT_API FDoubleProperty : public TProperty_Numeric<double>
 	}
 
 #if WITH_EDITORONLY_DATA
-	FDoubleProperty(UField* InField)
+	explicit FDoubleProperty(UField* InField)
 		: TProperty_Numeric(InField)
 	{
 	}
@@ -1889,7 +1889,7 @@ public:
 	FBoolProperty(FFieldVariant InOwner, const FName& InName, EObjectFlags InObjectFlags, int32 InOffset, EPropertyFlags InFlags, uint32 InBitMask, uint32 InElementSize, bool bIsNativeBool);
 
 #if WITH_EDITORONLY_DATA
-	FBoolProperty(UField* InField);
+	explicit FBoolProperty(UField* InField);
 #endif // WITH_EDITORONLY_DATA
 
 	// UObject interface.
@@ -2008,7 +2008,7 @@ public:
 	{}
 
 #if WITH_EDITORONLY_DATA
-	FObjectPropertyBase(UField* InField);
+	explicit FObjectPropertyBase(UField* InField);
 #endif // WITH_EDITORONLY_DATA
 
 	// UObject interface
@@ -2167,7 +2167,7 @@ public:
 	}
 
 #if WITH_EDITORONLY_DATA
-	TFObjectPropertyBase(UField* InField)
+	explicit TFObjectPropertyBase(UField* InField)
 		: Super(InField)
 	{
 	}
@@ -2212,7 +2212,7 @@ class COREUOBJECT_API FObjectProperty : public TFObjectPropertyBase<UObject*>
 	}
 
 #if WITH_EDITORONLY_DATA
-	FObjectProperty(UField* InField)
+	explicit FObjectProperty(UField* InField)
 		: TFObjectPropertyBase(InField)
 	{
 	}
@@ -2259,7 +2259,7 @@ class COREUOBJECT_API FWeakObjectProperty : public TFObjectPropertyBase<FWeakObj
 	}
 	
 #if WITH_EDITORONLY_DATA
-	FWeakObjectProperty(UField* InField)
+	explicit FWeakObjectProperty(UField* InField)
 		: TFObjectPropertyBase(InField)
 	{
 	}
@@ -2299,7 +2299,7 @@ class COREUOBJECT_API FLazyObjectProperty : public TFObjectPropertyBase<FLazyObj
 	}
 
 #if WITH_EDITORONLY_DATA
-	FLazyObjectProperty(UField* InField)
+	explicit FLazyObjectProperty(UField* InField)
 		: TFObjectPropertyBase(InField)
 	{
 	}
@@ -2342,7 +2342,7 @@ class COREUOBJECT_API FSoftObjectProperty : public TFObjectPropertyBase<FSoftObj
 	{}
 
 #if WITH_EDITORONLY_DATA
-	FSoftObjectProperty(UField* InField)
+	explicit FSoftObjectProperty(UField* InField)
 		: TFObjectPropertyBase(InField)
 	{
 	}
@@ -2423,7 +2423,7 @@ public:
 	}
 
 #if WITH_EDITORONLY_DATA
-	FClassProperty(UField* InField);
+	explicit FClassProperty(UField* InField);
 #endif // WITH_EDITORONLY_DATA
 
 	// UObject interface
@@ -2491,7 +2491,7 @@ public:
 	{}
 
 #if WITH_EDITORONLY_DATA
-	FSoftClassProperty(UField* InField);
+	explicit FSoftClassProperty(UField* InField);
 #endif // WITH_EDITORONLY_DATA
 
 	// UHT interface
@@ -2566,7 +2566,7 @@ public:
 	}
 
 #if WITH_EDITORONLY_DATA
-	FInterfaceProperty(UField* InField);
+	explicit FInterfaceProperty(UField* InField);
 #endif // WITH_EDITORONLY_DATA
 
 	// UHT interface
@@ -2641,7 +2641,7 @@ public:
 	}
 
 #if WITH_EDITORONLY_DATA
-	FNameProperty(UField* InField)
+	explicit FNameProperty(UField* InField)
 		: FNameProperty_Super(InField)
 	{
 	}
@@ -2685,7 +2685,7 @@ public:
 	}
 
 #if WITH_EDITORONLY_DATA
-	FStrProperty(UField* InField)
+	explicit FStrProperty(UField* InField)
 		: FStrProperty_Super(InField)
 	{
 	}
@@ -2741,7 +2741,7 @@ public:
 	virtual ~FArrayProperty();
 
 #if WITH_EDITORONLY_DATA
-	FArrayProperty(UField* InField);
+	explicit FArrayProperty(UField* InField);
 #endif // WITH_EDITORONLY_DATA
 
 	// UObject interface
@@ -2810,7 +2810,7 @@ public:
 	FMapProperty(FFieldVariant InOwner, const FName& InName, EObjectFlags InObjectFlags, int32 InOffset, EPropertyFlags InFlags);
 
 #if WITH_EDITORONLY_DATA
-	FMapProperty(UField* InField);
+	explicit FMapProperty(UField* InField);
 #endif // WITH_EDITORONLY_DATA
 
 	virtual ~FMapProperty();
@@ -2873,7 +2873,7 @@ public:
 	FSetProperty(FFieldVariant InOwner, const FName& InName, EObjectFlags InObjectFlags, int32 InOffset, EPropertyFlags InFlags);
 
 #if WITH_EDITORONLY_DATA
-	FSetProperty(UField* InField);
+	explicit FSetProperty(UField* InField);
 #endif // WITH_EDITORONLY_DATA
 
 	virtual ~FSetProperty();
@@ -4543,7 +4543,7 @@ public:
 	FStructProperty(FFieldVariant InOwner, const FName& InName, EObjectFlags InObjectFlags, int32 InOffset, EPropertyFlags InFlags, UScriptStruct* InStruct);
 
 #if WITH_EDITORONLY_DATA
-	FStructProperty(UField* InField);
+	explicit FStructProperty(UField* InField);
 #endif // WITH_EDITORONLY_DATA
 
 	// UObject interface
@@ -4637,7 +4637,7 @@ public:
 	}
 
 #if WITH_EDITORONLY_DATA
-	FDelegateProperty(UField* InField);
+	explicit FDelegateProperty(UField* InField);
 #endif // WITH_EDITORONLY_DATA
 
 	// UObject interface
@@ -4692,7 +4692,7 @@ public:
 	}
 
 #if WITH_EDITORONLY_DATA
-	FMulticastDelegateProperty(UField* InField);
+	explicit FMulticastDelegateProperty(UField* InField);
 #endif // WITH_EDITORONLY_DATA
 
 	// UObject interface
@@ -4766,7 +4766,7 @@ public:
 	}
 
 #if WITH_EDITORONLY_DATA
-	TProperty_MulticastDelegate(UField* InField)
+	explicit TProperty_MulticastDelegate(UField* InField)
 		: Super(InField)
 	{
 	}
@@ -4797,7 +4797,7 @@ public:
 	}
 
 #if WITH_EDITORONLY_DATA
-	FMulticastInlineDelegateProperty(UField* InField)
+	explicit FMulticastInlineDelegateProperty(UField* InField)
 		: TProperty_MulticastDelegate(InField)
 	{
 	}
@@ -4838,7 +4838,7 @@ public:
 	}
 
 #if WITH_EDITORONLY_DATA
-	FMulticastSparseDelegateProperty(UField* InField)
+	explicit FMulticastSparseDelegateProperty(UField* InField)
 		: TProperty_MulticastDelegate(InField)
 	{
 	}
