Fixing CIS errors and warnings

diff --git a/Engine/Source/Runtime/CoreUObject/Private/UObject/EnumProperty.cpp b/Engine/Source/Runtime/CoreUObject/Private/UObject/EnumProperty.cpp
--- a/Engine/Source/Runtime/CoreUObject/Private/UObject/EnumProperty.cpp
+++ b/Engine/Source/Runtime/CoreUObject/Private/UObject/EnumProperty.cpp
@@ -39,12 +39,6 @@ namespace UE4EnumProperty_Private
 
 		UnderlyingProp->SetIntPropertyValue(Obj, NewValue);
 	}
-
-	struct FEnumPropertyFriend
-	{
-		static const int32 EnumOffset = STRUCT_OFFSET(FEnumProperty, Enum);
-		static const int32 UnderlyingPropOffset = STRUCT_OFFSET(FEnumProperty, UnderlyingProp);
-	};
 }
 
 IMPLEMENT_FIELD(FEnumProperty)
diff --git a/Engine/Source/Runtime/CoreUObject/Private/UObject/PropertyPrivate.cpp b/Engine/Source/Runtime/CoreUObject/Private/UObject/PropertyPrivate.cpp
--- a/Engine/Source/Runtime/CoreUObject/Private/UObject/PropertyPrivate.cpp
+++ b/Engine/Source/Runtime/CoreUObject/Private/UObject/PropertyPrivate.cpp
@@ -161,7 +161,7 @@ void UEnumProperty::AddReferencedObjects(UObject* InThis, FReferenceCollector& C
 	Super::AddReferencedObjects(InThis, Collector);
 }
 
-namespace UE4EnumProperty_Private
+namespace UE4UEnumProperty_Private
 {
 	struct FEnumPropertyFriend
 	{
@@ -172,8 +172,8 @@ namespace UE4EnumProperty_Private
 
 IMPLEMENT_CORE_INTRINSIC_CLASS(UEnumProperty, UProperty,
 	{
-		Class->EmitObjectReference(UE4EnumProperty_Private::FEnumPropertyFriend::EnumOffset, TEXT("Enum"));
-		Class->EmitObjectReference(UE4EnumProperty_Private::FEnumPropertyFriend::UnderlyingPropOffset, TEXT("UnderlyingProp"));
+		Class->EmitObjectReference(UE4UEnumProperty_Private::FEnumPropertyFriend::EnumOffset, TEXT("Enum"));
+		Class->EmitObjectReference(UE4UEnumProperty_Private::FEnumPropertyFriend::UnderlyingPropOffset, TEXT("UnderlyingProp"));
 	}
 );
 
diff --git a/Engine/Source/Runtime/CoreUObject/Public/UObject/EnumProperty.h b/Engine/Source/Runtime/CoreUObject/Public/UObject/EnumProperty.h
--- a/Engine/Source/Runtime/CoreUObject/Public/UObject/EnumProperty.h
+++ b/Engine/Source/Runtime/CoreUObject/Public/UObject/EnumProperty.h
@@ -9,11 +9,6 @@
 class UEnum;
 class FNumericProperty;
 
-namespace UE4EnumProperty_Private
-{
-	struct FEnumPropertyFriend;
-}
-
 class COREUOBJECT_API FEnumProperty : public FProperty
 {
 	DECLARE_FIELD(FEnumProperty, FProperty, CASTCLASS_FEnumProperty)
@@ -84,8 +79,6 @@ public:
 private:
 	virtual uint32 GetValueTypeHashInternal(const void* Src) const override;
 
-	friend struct UE4EnumProperty_Private::FEnumPropertyFriend;
-	
 #if HACK_HEADER_GENERATOR
 public:
 #endif
diff --git a/Engine/Source/Runtime/CoreUObject/Public/UObject/Field.h b/Engine/Source/Runtime/CoreUObject/Public/UObject/Field.h
--- a/Engine/Source/Runtime/CoreUObject/Public/UObject/Field.h
+++ b/Engine/Source/Runtime/CoreUObject/Public/UObject/Field.h
@@ -788,7 +788,9 @@ FORCEINLINE FieldType* ExactCastField(FField* Src)
 }
 
 template<typename FieldType>
+FUNCTION_NON_NULL_RETURN_START
 FORCEINLINE FieldType* CastFieldChecked(FField* Src)
+FUNCTION_NON_NULL_RETURN_END
 {
 #if !DO_CHECK
 	return static_cast<FieldType*>(Src);
@@ -800,7 +802,9 @@ FORCEINLINE FieldType* CastFieldChecked(FField* Src)
 }
 
 template<typename FieldType>
+FUNCTION_NON_NULL_RETURN_START
 FORCEINLINE const FieldType* CastFieldChecked(const FField* Src)
+FUNCTION_NON_NULL_RETURN_END
 {
 #if !DO_CHECK
 	return static_cast<const FieldType*>(Src);
