Fix a runtime crash when running with the new BulkData api as certain code paths required for dynamically creating UTexture2D were not properly implemented.

- FBulkDataBase::Realloc now takes the total size in bytes rather than an element count as it does not really know what type of element it represents. FUntypedBulkData2 will convert elements to bytes for the caller.
- FBulkDataBase::Realloc is now implemented, although since it is runtime only, it is assumed that ::Realloc is only called on dynamically created BulkData objects and not on BulkData objects that represent files on disk. (otherwise we have to start handling cases where the loaded data does not match the file it represents)
- FBulkDataBase::LoadDataDirectly no longer treats being called on BulkData that cannot load from disk as an error case, instead we just early out to mirror the behaviour of the original BulkData api.

#rnx

diff --git a/Engine/Source/Runtime/CoreUObject/Private/Serialization/BulkData2.cpp b/Engine/Source/Runtime/CoreUObject/Private/Serialization/BulkData2.cpp
--- a/Engine/Source/Runtime/CoreUObject/Private/Serialization/BulkData2.cpp
+++ b/Engine/Source/Runtime/CoreUObject/Private/Serialization/BulkData2.cpp
@@ -730,15 +730,27 @@ bool FBulkDataBase::IsLocked() const
 	return LockStatus != LOCKSTATUS_Unlocked;
 }
 
-void* FBulkDataBase::Realloc(int64 InElementCount)
+void* FBulkDataBase::Realloc(int64 SizeInBytes)
 {
-	BULKDATA_NOT_IMPLEMENTED_FOR_RUNTIME;
-	return nullptr;
+	DECLARE_SCOPE_CYCLE_COUNTER(TEXT("FBulkDataBase::Realloc"), STAT_UBD_Realloc, STATGROUP_Memory);
+
+	check(LockStatus == LOCKSTATUS_ReadWriteLock);
+	checkf(!CanLoadFromDisk(), TEXT("Cannot re-allocate a FBulkDataBase object that represents a file on disk!"));
+
+	AllocateData(SizeInBytes);
+
+	check(!IsUsingIODispatcher());	// This case should get caught above but if someone tries to change that in the
+									// future then this check is a reminder that we need to handle the IoDispatcher 
+									// vs fallback case.
+
+	Fallback.BulkDataSize = SizeInBytes;
+
+	return DataBuffer;
 }
 
 void FBulkDataBase::GetCopy(void** DstBuffer, bool bDiscardInternalCopy)
 {
-	DECLARE_SCOPE_CYCLE_COUNTER(TEXT("FUntypedBulkData::GetCopy"), STAT_UBD_GetCopy, STATGROUP_Memory);
+	DECLARE_SCOPE_CYCLE_COUNTER(TEXT("FBulkDataBase::GetCopy"), STAT_UBD_GetCopy, STATGROUP_Memory);
 
 	check(LockStatus == LOCKSTATUS_Unlocked);
 	check(DstBuffer);
@@ -1045,7 +1057,11 @@ FString FBulkDataBase::GetFilename() const
 void FBulkDataBase::LoadDataDirectly(void** DstBuffer)
 {
 	DECLARE_SCOPE_CYCLE_COUNTER(TEXT("FBulkDataBase::LoadDataDirectly"), STAT_UBD_LoadDataDirectly, STATGROUP_Memory);
-	check(CanLoadFromDisk());
+	
+	if (!CanLoadFromDisk())
+	{
+		return; // Early out if there is nothing to load anyway
+	}
 
 	if (!IsUsingIODispatcher())
 	{
diff --git a/Engine/Source/Runtime/CoreUObject/Public/Serialization/BulkData2.h b/Engine/Source/Runtime/CoreUObject/Public/Serialization/BulkData2.h
--- a/Engine/Source/Runtime/CoreUObject/Public/Serialization/BulkData2.h
+++ b/Engine/Source/Runtime/CoreUObject/Public/Serialization/BulkData2.h
@@ -77,7 +77,7 @@ public:
 	void Unlock();
 	bool IsLocked() const;
 
-	void* Realloc(int64 InElementCount);
+	void* Realloc(int64 SizeInBytes);
 
 	/**
 	 * Retrieves a copy of the bulk data.
@@ -233,7 +233,7 @@ public:
 
 	ElementType* Realloc(int64 InElementCount)
 	{
-		return (ElementType*)FBulkDataBase::Realloc(InElementCount);
+		return (ElementType*)FBulkDataBase::Realloc(InElementCount * sizeof(ElementType));
 	}
 
 	/**
