FProperty optimizations:

- Optimized how MetaData is accessed and stored for FProperties. It's now stored and serialized with FProperties to prevent expensive key generation
- Optimized FProperty Casts
- Cleaned up macro usage for FProperty / UField meta data

diff --git a/Engine/Source/Developer/BlueprintCompilerCppBackend/Private/BlueprintCompilerCppBackendUtils.cpp b/Engine/Source/Developer/BlueprintCompilerCppBackend/Private/BlueprintCompilerCppBackendUtils.cpp
--- a/Engine/Source/Developer/BlueprintCompilerCppBackend/Private/BlueprintCompilerCppBackendUtils.cpp
+++ b/Engine/Source/Developer/BlueprintCompilerCppBackend/Private/BlueprintCompilerCppBackendUtils.cpp
@@ -638,9 +638,9 @@ FString FEmitHelper::HandleMetaData(FFieldVariant Field, bool AddCategory, const
 				ValuesMap = *ValuesMapPtr;
 			}
 		}
-		else
+		else if (const TMap<FName, FString>* FieldMetaMap = Field.ToField()->GetMetaDataMap())
 		{
-			ValuesMap = Field.ToField()->GetMetaDataMap();
+			ValuesMap = *FieldMetaMap;
 		}
 	}
 	
diff --git a/Engine/Source/Editor/DetailCustomizations/Private/PerPlatformPropertyCustomization.cpp b/Engine/Source/Editor/DetailCustomizations/Private/PerPlatformPropertyCustomization.cpp
--- a/Engine/Source/Editor/DetailCustomizations/Private/PerPlatformPropertyCustomization.cpp
+++ b/Engine/Source/Editor/DetailCustomizations/Private/PerPlatformPropertyCustomization.cpp
@@ -84,9 +84,9 @@ TSharedRef<SWidget> FPerPlatformPropertyCustomization<PerPlatformType>::GetWidge
 	// Push down struct metadata to per-platform properties
 	{
 		// First get the source map
-		TMap<FName, FString> SourceMap = StructPropertyHandle->GetMetaDataProperty()->GetMetaDataMap();
+		const TMap<FName, FString>* SourceMap = StructPropertyHandle->GetMetaDataProperty()->GetMetaDataMap();
 		// Iterate through source map, setting each key/value pair in the destination
-		for (const auto& It : SourceMap)
+		for (const auto& It : *SourceMap)
 		{
 				EditProperty->SetInstanceMetaData(*It.Key.ToString(), *It.Value);
 		}
diff --git a/Engine/Source/Programs/UnrealHeaderTool/Private/CodeGenerator.cpp b/Engine/Source/Programs/UnrealHeaderTool/Private/CodeGenerator.cpp
--- a/Engine/Source/Programs/UnrealHeaderTool/Private/CodeGenerator.cpp
+++ b/Engine/Source/Programs/UnrealHeaderTool/Private/CodeGenerator.cpp
@@ -644,7 +644,7 @@ TMap<FName, FString> GenerateMetadataMapForObject(const UObject* Obj)
 			FString Key = MetaKeyValue.Key.ToString();
 			if (!Key.StartsWith(TEXT("/Script")))
 			{
-				Map.Add(*Key, MetaKeyValue.Value);
+				Map.Add(MetaKeyValue.Key, MetaKeyValue.Value);
 			}
 		}
 	}
@@ -653,30 +653,13 @@ TMap<FName, FString> GenerateMetadataMapForObject(const UObject* Obj)
 
 TMap<FName, FString> GenerateMetadataMapForField(const FField* Field)
 {
-	check(Field);
-	UPackage* Package = Field->GetOutermost();
-	check(Package);
-	UMetaData* Metadata = Package->GetMetaData();
-	check(Metadata);
-	UObject* OwnerUObject = Field->GetOwnerUObject();
-	check(OwnerUObject);
-
-	TMap<FName, FString>* PackageMap = Metadata->ObjectMetaDataMap.Find(OwnerUObject);
-	TMap<FName, FString> Map;
-	if (PackageMap)
+	TMap<FName, FString> MetaDataMap;
+	const TMap<FName, FString>* FieldMetaDataMap = Field->GetMetaDataMap();
+	if (FieldMetaDataMap)
 	{
-		FString FieldMetadataPrefix = Field->GetFieldMetadataPrefix();
-		for (const TPair<FName, FString>& MetaKeyValue : *PackageMap)
-		{
-			FString Key = MetaKeyValue.Key.ToString();
-			if (Key.StartsWith(FieldMetadataPrefix))
-			{
-				Key = Key.Mid(FieldMetadataPrefix.Len());
-				Map.Add(*Key, MetaKeyValue.Value);
-			}
-		}
+		MetaDataMap = *FieldMetaDataMap;
 	}
-	return Map;
+	return MetaDataMap;
 }
 
 // Returns the METADATA_PARAMS for this output
@@ -6388,49 +6371,6 @@ ECompilationResult::Type UnrealHeaderTool_Main(const FString& ModuleInfoFilename
 	FDurationTimer MainTimer(MainTime);
 	MainTimer.Start();
 
-	//FFieldVariant Variant;
-	//TFieldPath<FProperty> PropPath;
-	//FProperty* RawProp = nullptr;
-
-	//Variant.IsA<FProperty>();
-	//Variant.IsA<UObject>();
-	//Variant.Get<FProperty>();
-	//Variant.Get<UObject>();
-	//RawProp->GetOwner<FProperty>();
-	//RawProp->GetOwner<UObject>();
-	//RawProp->GetOwnerChecked<FProperty>();
-	//RawProp->GetOwnerChecked<UObject>();
-	//RawProp->GetTypedOwner<FProperty>();
-	//RawProp->GetTypedOwner<UObject>();
-	//{
-	//	UStruct* Struct = nullptr;
-	//	for (TFieldIterator<FProperty> It(Struct); It; ++It)
-	//	{
-	//		check(It);
-	//	}
-	//	for (TFieldIterator<UFunction> It(Struct); It; ++It)
-	//	{
-	//		check(It);
-	//	}
-	//}
-	//{
-	//	FArchive* Ar = nullptr;
-	//	TWeakFieldPtr<FProperty> WeakPtr;
-	//	(*Ar) << WeakPtr;
-	//}
-	//{
-	//	FProperty* Prop = nullptr;
-	//	UObject* Obj = nullptr;
-	//	FTestCollector Collector;
-	//	Collector.AddReferencedObject(Prop);
-	//	Collector.AddReferencedObject(Obj);
-	//}
-	//{
-	//	TSubclassOf<FProperty> SubProp;
-	//	SubProp.GetDefaultObject();
-	//	TSubclassOf<UFunction> SubFn;
-	//	SubFn.GetDefaultObject();
-	//}
 	check(GIsUCCMakeStandaloneHeaderGenerator);
 	ECompilationResult::Type Result = ECompilationResult::Succeeded;
 
diff --git a/Engine/Source/Runtime/CoreUObject/Private/UObject/Class.cpp b/Engine/Source/Runtime/CoreUObject/Private/UObject/Class.cpp
--- a/Engine/Source/Runtime/CoreUObject/Private/UObject/Class.cpp
+++ b/Engine/Source/Runtime/CoreUObject/Private/UObject/Class.cpp
@@ -222,7 +222,7 @@ void UField::AddCppProperty(FProperty* Property)
 	UE_LOG(LogClass, Fatal,TEXT("UField::AddCppProperty"));
 }
 
-#if WITH_EDITOR || HACK_HEADER_GENERATOR
+#if WITH_EDITORONLY_DATA
 
 struct FDisplayNameHelper
 {
@@ -553,7 +553,7 @@ void UField::RemoveMetaData(const FName& Key)
 	return Package->GetMetaData()->RemoveValue(this, Key);
 }
 
-#endif
+#endif // WITH_EDITORONLY_DATA
 
 bool UField::HasAnyCastFlags(const uint64 InCastFlags) const
 {
@@ -1972,7 +1972,7 @@ FString UStruct::GetAuthoredNameForField(const FField* Field) const
 	return FString();
 }
 
-#if WITH_EDITOR || HACK_HEADER_GENERATOR
+#if WITH_EDITORONLY_DATA
 bool UStruct::GetBoolMetaDataHierarchical(const FName& Key) const
 {
 	bool bResult = false;
@@ -2021,7 +2021,7 @@ const UStruct* UStruct::HasMetaDataHierarchical(const FName& Key) const
 	return nullptr;
 }
 
-#endif
+#endif // WITH_EDITORONLY_DATA
 
 #if USE_CIRCULAR_DEPENDENCY_LOAD_DEFERRING
 	/**
@@ -3259,13 +3259,13 @@ int32 FStructUtils::AttemptToFindUninitializedScriptStructMembers()
 			
 			for (const FProperty* Property : TFieldRange<FProperty>(ScriptStruct, EFieldIteratorFlags::ExcludeSuper))
 			{
-#if	WITH_EDITOR || HACK_HEADER_GENERATOR
+#if	WITH_EDITORONLY_DATA
 				static const FName NAME_IgnoreForMemberInitializationTest(TEXT("IgnoreForMemberInitializationTest"));
 				if (Property->HasMetaData(NAME_IgnoreForMemberInitializationTest))
 				{
 					continue;
 				}
-#endif
+#endif // WITH_EDITORONLY_DATA
 
 				if (const FObjectPropertyBase* ObjectProperty = CastField<const FObjectPropertyBase>(Property))
 				{
diff --git a/Engine/Source/Runtime/CoreUObject/Private/UObject/Field.cpp b/Engine/Source/Runtime/CoreUObject/Private/UObject/Field.cpp
--- a/Engine/Source/Runtime/CoreUObject/Private/UObject/Field.cpp
+++ b/Engine/Source/Runtime/CoreUObject/Private/UObject/Field.cpp
@@ -246,7 +246,7 @@ bool FFieldVariant::IsValidLowLevel() const
 	}
 }
 
-#if WITH_EDITOR || HACK_HEADER_GENERATOR
+#if WITH_EDITORONLY_DATA
 bool FFieldVariant::HasMetaData(const FName& Key) const
 {
 	check(Container.Object);
@@ -259,7 +259,7 @@ bool FFieldVariant::HasMetaData(const FName& Key) const
 		return Container.Field->HasMetaData(Key);
 	}
 }
-#endif
+#endif // WITH_EDITORONLY_DATA
 
 FArchive& operator << (FArchive& Ar, FFieldVariant& InOutField)
 {
@@ -300,6 +300,9 @@ FField::FField(EInternal InInernal, FFieldClass* InClass)
 	, Owner((FField*)nullptr)
 	, Next(nullptr)
 	, FlagsPrivate(RF_NoFlags)
+#if WITH_EDITORONLY_DATA
+	, MetaDataMap(nullptr)
+#endif // WITH_EDITORONLY_DATA
 {
 }
 
@@ -308,12 +311,22 @@ FField::FField(FFieldVariant InOwner, const FName& InName, EObjectFlags InObject
 	, Next(nullptr)
 	, NamePrivate(InName)
 	, FlagsPrivate(InObjectFlags)
+#if WITH_EDITORONLY_DATA
+	, MetaDataMap(nullptr)
+#endif // WITH_EDITORONLY_DATA
 {
 
 }
 
 FField::~FField()
 {
+#if WITH_EDITORONLY_DATA
+	if (MetaDataMap)
+	{
+		delete MetaDataMap;
+		MetaDataMap = nullptr;
+	}
+#endif // WITH_EDITORONLY_DATA
 }
 
 #if WITH_EDITORONLY_DATA
@@ -321,6 +334,7 @@ FField::FField(UField* InField)
 	: Next(nullptr)
 	, NamePrivate(InField->GetFName())
 	, FlagsPrivate(InField->GetFlags())
+	, MetaDataMap(nullptr)
 {
 	check(InField);
 	if (InField->HasAnyFlags(RF_NeedLoad))
@@ -347,6 +361,12 @@ FField::FField(UField* InField)
 	{
 		Owner = OriginalOuter;
 	}
+
+	TMap<FName, FString>* FeldMetaDataMap = UMetaData::GetMapForObject(InField);
+	if (FeldMetaDataMap && FeldMetaDataMap->Num())
+	{
+		MetaDataMap = new TMap<FName, FString>(*FeldMetaDataMap);
+	}
 }
 #endif // WITH_EDITORONLY_DATA
 
@@ -412,6 +432,35 @@ void FField::Serialize(FArchive& Ar)
 {
 	Ar << NamePrivate;
 	Ar << (uint32&)FlagsPrivate;
+
+#if WITH_EDITORONLY_DATA	
+	if (!Ar.IsCooking())
+	{
+		UPackage* Package = GetOutermost();
+		if (!Package || !Package->bIsCookedForEditor)
+		{
+			bool bHasMetaData = false;
+			if (Ar.IsLoading())
+			{
+				Ar << bHasMetaData;
+
+			}
+			else
+			{
+				bHasMetaData = MetaDataMap && MetaDataMap->Num();
+				Ar << bHasMetaData;
+			}
+			if (bHasMetaData)
+			{
+				if (!MetaDataMap)
+				{
+					MetaDataMap = new TMap<FName, FString>();
+				}
+				Ar << *MetaDataMap;
+			}			
+		}		
+	}
+#endif
 }
 
 void FField::GetPreloadDependencies(TArray<UObject*>& OutDeps)
@@ -581,7 +630,7 @@ FField* FField::GetTypedOwner(FFieldClass* Target) const
 	return Result;
 }
 
-#if WITH_EDITOR || HACK_HEADER_GENERATOR
+#if WITH_EDITORONLY_DATA
 FString FField::GetFullGroupName(bool bStartWithOuter) const
 {
 	if (bStartWithOuter)
@@ -709,11 +758,6 @@ FText FField::GetToolTipText(bool bShortTooltip) const
 	return LocalizedToolTip;
 }
 
-FString FField::GenerateFieldMetadataKey(const TCHAR* InKey) const
-{
-	return GetFieldMetadataPrefix() + InKey;
-}
-
 /**
 * Determines if the property has any metadata associated with the key
 *
@@ -722,23 +766,12 @@ FString FField::GenerateFieldMetadataKey(const TCHAR* InKey) const
 */
 bool FField::HasMetaData(const TCHAR* Key) const
 {
-	UPackage* Package = GetOutermost();
-	check(Package);
-
-	UMetaData* MetaData = Package->GetMetaData();
-	check(MetaData);
-
-	UObject* OwnerUObject = GetOwnerUObject();
-	check(OwnerUObject);
-
-	bool bHasMetaData = MetaData->HasValue(OwnerUObject, *GenerateFieldMetadataKey(Key));
-
-	return bHasMetaData;
+	return HasMetaData(FName(Key, FNAME_Find));
 }
 
 bool FField::HasMetaData(const FName& Key) const
 {
-	return HasMetaData(*Key.ToString());
+	return MetaDataMap && MetaDataMap->Contains(Key);
 }
 
 /**
@@ -749,23 +782,31 @@ bool FField::HasMetaData(const FName& Key) const
 */
 const FString& FField::GetMetaData(const TCHAR* Key) const
 {
-	UPackage* Package = GetOutermost();
-	check(Package);
+	return GetMetaData(FName(Key, FNAME_Find));
+}
 
-	UMetaData* MetaData = Package->GetMetaData();
-	check(MetaData);
+const FString& FField::GetMetaData(const FName& Key) const
+{
+	// if not found, return a static empty string
+	static FString EmptyString;
 
-	UObject* OwnerUObject = GetOwnerUObject();
-	check(OwnerUObject);
+	// every key needs to be valid and meta data needs to exist
+	if (Key == NAME_None || !MetaDataMap)
+	{
+		return EmptyString;
+	}
 
-	const FString& MetaDataString = MetaData->GetValue(OwnerUObject, *GenerateFieldMetadataKey(Key));
+	// look for the property
+	const FString* ValuePtr = MetaDataMap->Find(Key);
 
-	return MetaDataString;
-}
+	// if we didn't find it, return NULL
+	if (!ValuePtr)
+	{
+		return EmptyString;
+	}
 
-const FString& FField::GetMetaData(const FName& Key) const
-{
-	return GetMetaData(*Key.ToString());
+	// if we found it, return the pointer to the character data
+	return *ValuePtr;
 }
 
 const FText FField::GetMetaDataText(const TCHAR* MetaDataKey, const FString LocalizationNamespace, const FString LocalizationKey) const
@@ -832,21 +873,17 @@ const FText FField::GetMetaDataText(const FName& MetaDataKey, const FString Loca
 */
 void FField::SetMetaData(const TCHAR* Key, const TCHAR* InValue)
 {
-	UPackage* Package = GetOutermost();
-	check(Package);
-
-	UMetaData* MetaData = Package->GetMetaData();
-	check(MetaData);
-
-	UObject* OwnerUObject = GetOwnerUObject();
-	check(OwnerUObject);
-
-	return MetaData->SetValue(OwnerUObject, *GenerateFieldMetadataKey(Key), InValue);
+	SetMetaData(FName(Key), InValue);
 }
 
 void FField::SetMetaData(const FName& Key, const TCHAR* InValue)
 {
-	SetMetaData(*Key.ToString(), InValue);
+	check(Key != NAME_None);
+	if (!MetaDataMap)
+	{
+		MetaDataMap = new TMap<FName, FString>();
+	}
+	MetaDataMap->Add(Key, InValue);
 }
 
 UClass* FField::GetClassMetaData(const TCHAR* Key) const
@@ -865,73 +902,43 @@ UClass* FField::GetClassMetaData(const FName& Key) const
 
 void FField::RemoveMetaData(const TCHAR* Key)
 {
-	UPackage* Package = GetOutermost();
-	check(Package);
-
-	UMetaData* MetaData = Package->GetMetaData();
-	check(MetaData);
-
-	UObject* OwnerUObject = GetOwnerUObject();
-	check(OwnerUObject);
-
-	MetaData->RemoveValue(OwnerUObject, *GenerateFieldMetadataKey(Key));
+	RemoveMetaData(FName(Key));
 }
 
 void FField::RemoveMetaData(const FName& Key)
 {
-	RemoveMetaData(*Key.ToString());
+	check(Key != NAME_None);
+	if (MetaDataMap)
+	{
+		MetaDataMap->Remove(Key);
+	}
 }
 
-TMap<FName, FString> FField::GetMetaDataMap(bool bStripFieldPrefix) const
+const TMap<FName, FString>* FField::GetMetaDataMap() const
 {
-	UObject* OwnerUObject = GetOwnerUObject();
-	check(OwnerUObject);
-
-	TMap<FName, FString> PropertyMetaMap;
-	TMap<FName, FString>* InOwnerMap = UMetaData::GetMapForObject(OwnerUObject);
-	if (InOwnerMap)
-	{		
-		FString FieldPrefix = GetFieldMetadataPrefix();
-		for (TPair<FName, FString>& MetaPair : *InOwnerMap)
-		{
-			FString Key(MetaPair.Key.ToString());
-			if (Key.StartsWith(FieldPrefix))
-			{
-				if (bStripFieldPrefix)
-				{
-					FName KeyName(*Key.Mid(FieldPrefix.Len()));
-					PropertyMetaMap.Add(KeyName, MetaPair.Value);
-				}
-				else
-				{
-					PropertyMetaMap.Add(MetaPair);
-				}
-			}
-		}
-	}
-	return PropertyMetaMap;
+	return MetaDataMap;
 }
 
 void FField::CopyMetaData(const FField* InSourceField, FField* InDestField)
 {
 	check(InSourceField);
 	check(InDestField);
-
-	// First get the source map
-	TMap<FName, FString> SourceMap = InSourceField->GetMetaDataMap(true);
-	if (!SourceMap.Num())
+	if (InSourceField->MetaDataMap)
 	{
-		return;
+		if (!InDestField->MetaDataMap)
+		{
+			InDestField->MetaDataMap = new TMap<FName, FString>();
+		}
+		*InDestField->MetaDataMap = *InSourceField->MetaDataMap;
 	}
-
-	// Iterate through source map, setting each key/value pair in the destination
-	for (const TPair<FName, FString>& It : SourceMap)
+	else if (InDestField->MetaDataMap)
 	{
-		InDestField->SetMetaData(It.Key, *It.Value);
+		delete InDestField->MetaDataMap;
+		InDestField->MetaDataMap = nullptr;
 	}
 }
 
-#endif
+#endif // WITH_EDITORONLY_DATA
 
 void FField::PostDuplicate(const FField& InField)
 {
@@ -1089,7 +1096,7 @@ FField* FField::CreateFromUField(UField* InField)
 
 	return NewField;
 }
-#endif
+#endif // WITH_EDITORONLY_DATA
 
 FString GetFullNameSafe(const FField* InField)
 {
diff --git a/Engine/Source/Runtime/CoreUObject/Public/UObject/Class.h b/Engine/Source/Runtime/CoreUObject/Public/UObject/Class.h
--- a/Engine/Source/Runtime/CoreUObject/Public/UObject/Class.h
+++ b/Engine/Source/Runtime/CoreUObject/Public/UObject/Class.h
@@ -96,7 +96,7 @@ class COREUOBJECT_API UField : public UObject
 	 */
 	FString GetAuthoredName() const;
 
-#if WITH_EDITOR || HACK_HEADER_GENERATOR
+#if WITH_EDITORONLY_DATA
 	/**
 	 * Finds the localized display name or native display name as a fallback.
 	 *
@@ -228,7 +228,7 @@ class COREUOBJECT_API UField : public UObject
 	/** Clear any metadata associated with the key */
 	void RemoveMetaData(const TCHAR* Key);
 	void RemoveMetaData(const FName& Key);
-#endif
+#endif // WITH_EDITORONLY_DATA
 
 	bool HasAnyCastFlags(const uint64 InCastFlags) const;
 	bool HasAllCastFlags(const uint64 InCastFlags) const;
@@ -497,7 +497,7 @@ public:
 	/** Returns a human readable string for a given field, overridden for user defined structs */
 	virtual FString GetAuthoredNameForField(const FField* Field) const;
 
-#if WITH_EDITOR || HACK_HEADER_GENERATOR
+#if WITH_EDITORONLY_DATA
 	/** Try and find boolean metadata with the given key. If not found on this class, work up hierarchy looking for it. */
 	bool GetBoolMetaDataHierarchical(const FName& Key) const;
 
@@ -511,7 +511,7 @@ public:
 	 * @return pointer to the UStruct that has associated metadata, nullptr if Key is not associated with any UStruct in the hierarchy
 	 */
 	const UStruct* HasMetaDataHierarchical(const FName& Key) const;
-#endif
+#endif // WITH_EDITORONLY_DATA
 
 #if HACK_HEADER_GENERATOR
 	// Required by UHT makefiles for internal data serialization.
@@ -2104,7 +2104,7 @@ public:
 	FText GetToolTipText(int32 NameIndex) const { return GetToolTipTextByIndex(NameIndex); }
 #endif
 
-#if WITH_EDITOR || HACK_HEADER_GENERATOR
+#if WITH_EDITORONLY_DATA
 	/**
 	 * Wrapper method for easily determining whether this enum has metadata associated with it.
 	 * 
@@ -2144,7 +2144,7 @@ public:
 	 *
 	 */
 	void RemoveMetaData( const TCHAR* Key, int32 NameIndex=INDEX_NONE ) const;
-#endif
+#endif // WITH_EDITORONLY_DATA
 	
 	/**
 	 * @param EnumPath         Full enum path.
diff --git a/Engine/Source/Runtime/CoreUObject/Public/UObject/Field.h b/Engine/Source/Runtime/CoreUObject/Public/UObject/Field.h
--- a/Engine/Source/Runtime/CoreUObject/Public/UObject/Field.h
+++ b/Engine/Source/Runtime/CoreUObject/Public/UObject/Field.h
@@ -129,10 +129,6 @@ public:
 private: \
 	TClass& operator=(TClass&&);   \
 	TClass& operator=(const TClass&);   \
-	inline static uint64 StaticClassCastFlagsPrivate() \
-	{ \
-		return uint64(TStaticFlags); \
-	} \
 public: \
 	typedef TSuperClass Super;\
 	typedef TClass ThisClass;\
@@ -142,6 +138,10 @@ public: \
 	} \
 	static FFieldClass* StaticClass(); \
 	static FField* Construct(const FFieldVariant& InOwner, const FName& InName, EObjectFlags InObjectFlags); \
+	inline static uint64 StaticClassCastFlagsPrivate() \
+	{ \
+		return uint64(TStaticFlags); \
+	} \
 	inline static uint64 StaticClassCastFlags() \
 	{ \
 		return uint64(TStaticFlags) | Super::StaticClassCastFlags(); \
@@ -315,7 +315,7 @@ public:
 		return Container.Field != Other.Container.Field;
 	}
 
-#if WITH_EDITOR || HACK_HEADER_GENERATOR
+#if WITH_EDITORONLY_DATA
 	bool HasMetaData(const FName& Key) const;
 #endif
 
@@ -335,11 +335,6 @@ class COREUOBJECT_API FField
 {
 	UE_NONCOPYABLE(FField);
 
-	inline static uint64 StaticClassCastFlagsPrivate()
-	{
-		return uint64(CASTCLASS_UField);
-	}
-
 	/** Pointer to the class object representing the type of this FField */
 	FFieldClass* ClassPrivate;
 
@@ -351,6 +346,10 @@ public:
 
 	static FFieldClass* StaticClass();
 
+	inline static uint64 StaticClassCastFlagsPrivate()
+	{
+		return uint64(CASTCLASS_UField);
+	}
 	inline static uint64 StaticClassCastFlags()
 	{
 		return uint64(CASTCLASS_UField);
@@ -475,7 +474,7 @@ public:
 	template<typename T>
 	bool IsA() const
 	{
-		return !!(GetCastFlags() & T::StaticClass()->GetId());
+		return !!(GetCastFlags() & T::StaticClassCastFlagsPrivate());
 	}
 
 	bool HasAnyCastFlags(const uint64 InCastFlags) const
@@ -607,7 +606,14 @@ public:
 		return nullptr;
 	}
 
-#if WITH_EDITOR || HACK_HEADER_GENERATOR
+#if WITH_EDITORONLY_DATA
+
+private:
+	/** Editor-only meta data map */
+	TMap<FName, FString>* MetaDataMap;
+
+public:
+
 	/**
 	* Walks up the chain of packages until it reaches the top level, which it ignores.
 	*
@@ -743,30 +749,20 @@ public:
 	void RemoveMetaData(const FName& Key);
 
 	/** Gets all metadata associated with this field */
-	TMap<FName, FString> GetMetaDataMap(bool bStripFieldPrefix = true) const;
-
-	FString GetFieldMetadataPrefix() const
-	{
-		return FString::Printf(TEXT("%s::"), *GetPathName());
-	}
-
-	FString GenerateFieldMetadataKey(const TCHAR* InKey) const;
+	const TMap<FName, FString>* GetMetaDataMap() const;
 
 	/** Copies all metadata from Source Field to Dest Field */
 	static void CopyMetaData(const FField* InSourceField, FField* InDestField);
 
-#endif
+	/** Creates a new FField from existing UField */
+	static FField* CreateFromUField(UField* InField);
+#endif // WITH_EDITORONLY_DATA
 
 	/** Duplicates an FField */
 	static FField* Duplicate(const FField* InField, FFieldVariant DestOwner, const FName DestName = NAME_None, EObjectFlags FlagMask = RF_AllFlags, EInternalObjectFlags InternalFlagsMask = EInternalObjectFlags::AllFlags);
 
 	/** Generates a name for a Field of a given type. Each generated name is unique in the current runtime */
 	static FName GenerateFFieldName(FFieldVariant InOwner, FFieldClass* InClass);
-
-#if WITH_EDITORONLY_DATA
-	/** Creates a new FField from existing UField */
-	static FField* CreateFromUField(UField* InField);
-#endif
 };
 
 // Support for casting between different FFIeld types
@@ -774,13 +770,13 @@ public:
 template<typename FieldType>
 FORCEINLINE FieldType* CastField(FField* Src)
 {
-	return Src && Src->HasAnyCastFlags(FieldType::StaticClass()->GetId()) ? static_cast<FieldType*>(Src) : nullptr;
+	return Src && Src->HasAnyCastFlags(FieldType::StaticClassCastFlagsPrivate()) ? static_cast<FieldType*>(Src) : nullptr;
 }
 
 template<typename FieldType>
 FORCEINLINE const FieldType* CastField(const FField* Src)
 {
-	return Src && Src->HasAnyCastFlags(FieldType::StaticClass()->GetId()) ? static_cast<const FieldType*>(Src) : nullptr;
+	return Src && Src->HasAnyCastFlags(FieldType::StaticClassCastFlagsPrivate()) ? static_cast<const FieldType*>(Src) : nullptr;
 }
 
 template<typename FieldType>
@@ -797,7 +793,7 @@ FUNCTION_NON_NULL_RETURN_END
 #if !DO_CHECK
 	return static_cast<FieldType*>(Src);
 #else
-	FieldType* CastResult = Src && Src->HasAnyCastFlags(FieldType::StaticClass()->GetId()) ? static_cast<FieldType*>(Src) : nullptr;
+	FieldType* CastResult = Src && Src->HasAnyCastFlags(FieldType::StaticClassCastFlagsPrivate()) ? static_cast<FieldType*>(Src) : nullptr;
 	checkf(CastResult, TEXT("CastFieldChecked failed with 0x%016llx"), (int64)(PTRINT)Src);
 	return CastResult;
 #endif // !DO_CHECK
@@ -811,7 +807,7 @@ FUNCTION_NON_NULL_RETURN_END
 #if !DO_CHECK
 	return static_cast<const FieldType*>(Src);
 #else
-	const FieldType* CastResult = Src && Src->HasAnyCastFlags(FieldType::StaticClass()->GetId()) ? static_cast<const FieldType*>(Src) : nullptr;
+	const FieldType* CastResult = Src && Src->HasAnyCastFlags(FieldType::StaticClassCastFlagsPrivate()) ? static_cast<const FieldType*>(Src) : nullptr;
 	checkf(CastResult, TEXT("CastFieldChecked failed with 0x%016llx"), (int64)(PTRINT)Src);
 	return CastResult;
 #endif // !DO_CHECK
