Reducing the cost of accessing TWeakFieldPtr's FProperty

diff --git a/Engine/Source/Runtime/CoreUObject/Public/UObject/WeakFieldPtr.h b/Engine/Source/Runtime/CoreUObject/Public/UObject/WeakFieldPtr.h
--- a/Engine/Source/Runtime/CoreUObject/Public/UObject/WeakFieldPtr.h
+++ b/Engine/Source/Runtime/CoreUObject/Public/UObject/WeakFieldPtr.h
@@ -105,7 +105,7 @@ public:
 	{
 		if (Owner.Get(bEvenIfPendingKill))
 		{
-			return CastField<T>(Field.TryToResolvePath());
+			return Field.Get();
 		}
 		else
 		{
@@ -122,7 +122,7 @@ public:
 	{
 		if (Owner.Get())
 		{
-			return CastField<T>(Field.TryToResolvePath());
+			return Field.Get();
 		}
 		else
 		{
@@ -137,7 +137,7 @@ public:
 	{
 		if (Owner.GetEvenIfUnreachable())
 		{
-			return CastField<T>(Field.TryToResolvePath());
+			return Field.Get();
 		}
 		else
 		{
