Fixing GC crashes when running with nativized blueprints

diff --git a/Engine/Source/Runtime/CoreUObject/Private/UObject/GarbageCollection.cpp b/Engine/Source/Runtime/CoreUObject/Private/UObject/GarbageCollection.cpp
--- a/Engine/Source/Runtime/CoreUObject/Private/UObject/GarbageCollection.cpp
+++ b/Engine/Source/Runtime/CoreUObject/Private/UObject/GarbageCollection.cpp
@@ -2591,7 +2591,7 @@ void UClass::AssembleReferenceTokenStream(bool bForce)
 		{
 			check(ClassAddReferencedObjects != NULL);
 			const bool bKeepOuter = GetFName() != NAME_Package;
-			const bool bKeepClass = !(GetClassFlags() & CLASS_Native);
+			const bool bKeepClass = !(GetClassFlags() & CLASS_Native) || IsA(UDynamicClass::StaticClass());
 
 			ClassAddReferencedObjectsType AddReferencedObjectsFn = nullptr;
 #if !WITH_EDITOR
