Fix support for reloading inline BulkData from disk.

- Now that we store the original filepath of the package we can work out how to change the offset (which we could not do before) so it is possible to load inline bulkdata again.
- Currently disabled the warning when reloading inline data (notifying the user that it will hard loading time) as we have one case in the engine still that I will make a jira for. Once that is fixed we should enable the warning again.

#rnx

diff --git a/Engine/Source/Runtime/CoreUObject/Private/Serialization/BulkData2.cpp b/Engine/Source/Runtime/CoreUObject/Private/Serialization/BulkData2.cpp
--- a/Engine/Source/Runtime/CoreUObject/Private/Serialization/BulkData2.cpp
+++ b/Engine/Source/Runtime/CoreUObject/Private/Serialization/BulkData2.cpp
@@ -1068,24 +1068,29 @@ void FBulkDataBase::LoadDataDirectly(void** DstBuffer)
 		const int64 BulkDataSize = GetBulkDataSize();
 		FileTokenSystem::Data FileData = FileTokenSystem::GetFileData(Fallback.Token);
 
-		const FString Filename = ConvertFilenameFromFlags(FileData.PackageHeaderFilename);
-
-		// In the older code path if the bulkdata was not in a separate file we could just serialize from the package instead, but the new system
-		// does not keep a reference to the pack so we cannot do that.
-		// TODO: Maybe add a check for this condition?
+		FString Filename; 
+		int64 Offset = FileData.BulkDataOffsetInFile;
 
-		check(!IsInlined());	// TODO: Currently inline data will have the wrong BulkDataOffsetInFile (need to subtract the size of the original
-								// .uasset or .umap file from it but we don't have access to that filepath at this point and doing it during
-								// ::Serialize is too slow!)
+		// Fix up the Filename/Offset to work with streaming if EDL is enabled and the filename is still referencing a uasset or umap
+		if (IsInlined() && (FileData.PackageHeaderFilename.EndsWith(TEXT(".uasset")) || FileData.PackageHeaderFilename.EndsWith(TEXT(".umap"))))
+		{
+			Offset -= IFileManager::Get().FileSize(*FileData.PackageHeaderFilename);
+			Filename = FPaths::GetBaseFilename(FileData.PackageHeaderFilename, false) + TEXT(".uexp");
+		}
+		else
+		{
+			Filename = ConvertFilenameFromFlags(FileData.PackageHeaderFilename);
+		}
 
 		// If the data is inlined then we already loaded is during ::Serialize, this warning should help track cases where data is being discarded then re-requested.
-		UE_CLOG(IsInlined(), LogSerialization, Warning, TEXT("Reloading inlined bulk data directly from disk, consider not discarding it in the first place. Filename: '%s'."), *Filename);
+		// Disabled at the moment
+		//UE_CLOG(IsInlined(), LogSerialization, Warning, TEXT("Reloading inlined bulk data directly from disk, this is detrimental to loading performance. Filename: '%s'."), *Filename);
 
 		FArchive* Ar = IFileManager::Get().CreateFileReader(*Filename, FILEREAD_Silent);
 		checkf(Ar != nullptr, TEXT("Failed to open the file to load bulk data from. Filename: '%s'."), *Filename);
 
 		// Seek to the beginning of the bulk data in the file.
-		Ar->Seek(FileData.BulkDataOffsetInFile);
+		Ar->Seek(Offset);
 
 		if (*DstBuffer == nullptr)
 		{
