Removed superfluous assert

I added it to validate unversioned property serialization (UPS) would always save the same properties as tagged property serialization (TPS) in an equivalent way. There is a small edge case, or bug rather, where FVector_NetQuantize inherits from FVector and inherits the STRUCT_Atomic flag via STRUCT_Inherit but doesn't inherit the STRUCT_SerializeNative flag, which triggers this assert. UPS will respect defaults for this structs but TPS will not. It is very unlikely to be a problem in practice. The net quantize structs are meant for NetSerialize and not vanilla Serialize / SerializeItem / TPS. Apparently some EngineTest test content saved such a struct in source data.

diff --git a/Engine/Source/Runtime/CoreUObject/Private/UObject/Class.cpp b/Engine/Source/Runtime/CoreUObject/Private/UObject/Class.cpp
--- a/Engine/Source/Runtime/CoreUObject/Private/UObject/Class.cpp
+++ b/Engine/Source/Runtime/CoreUObject/Private/UObject/Class.cpp
@@ -1448,10 +1448,6 @@ void UStruct::SerializeVersionedTaggedProperties(FStructuredArchive::FSlot Slot,
 			bUseAtomicSerialization = DefaultsScriptStruct->ShouldSerializeAtomically(UnderlyingArchive);
 		}
 
-		checkf(!(Defaults && UnderlyingArchive.DoDelta() && !UnderlyingArchive.IsTransacting() && bUseAtomicSerialization), 
-				TEXT("Dense/'atomic' structs shouldn't be sparsely delta-serialized. Struct '%s', archive: '%s'."),
-				*DefaultsScriptStruct->GetFName().ToString(), *UnderlyingArchive.GetArchiveName());
-
 		// Save tagged properties.
 
 		// Iterate over properties in the order they were linked and serialize them.
