Fixing TFieldPath comparisons to nullptr which would return false even if the path was valid but unresolved

diff --git a/Engine/Source/Runtime/CoreUObject/Public/UObject/FieldPath.h b/Engine/Source/Runtime/CoreUObject/Public/UObject/FieldPath.h
--- a/Engine/Source/Runtime/CoreUObject/Public/UObject/FieldPath.h
+++ b/Engine/Source/Runtime/CoreUObject/Public/UObject/FieldPath.h
@@ -115,22 +115,8 @@ public:
 		return !Path.Num();
 	}
 
-	/**
-	* Test if this points to a live FField
-	* @param bEvenIfPendingKill, if this is true, pendingkill objects are considered valid
-	* @param bThreadsafeTest, if true then function will just give you information whether referenced
-	*							FField is gone forever (@return false) or if it is still there (@return true, no object flags checked).
-	* @return true if Get() would return a valid non-null pointer
-	**/
-	inline bool IsValid() const
-	{
-		return ResolvedField && (TryToResolvePath() == ResolvedField);
-	}
-
 	/**
 	* Slightly different than !IsValid(), returns true if this used to point to a FField, but doesn't any more and has not been assigned or reset in the mean time.
-	* @param bIncludingIfPendingKill, if this is true, pendingkill objects are considered stale
-	* @param bThreadsafeTest, set it to true when testing outside of Game Thread. Results in false if WeakObjPtr point to an existing object (no flags checked)
 	* @return true if this used to point at a real object but no longer does.
 	**/
 	inline bool IsStale() const
@@ -392,13 +378,13 @@ FORCENOINLINE bool operator==(const LhsT* Lhs, const TFieldPath<RhsT>& Rhs)
 template <typename LhsT>
 FORCENOINLINE bool operator==(const TFieldPath<LhsT>& Lhs, TYPE_OF_NULLPTR)
 {
-	return !Lhs.IsValid();
+	return !Lhs.Get();
 }
 
 template <typename RhsT>
 FORCENOINLINE bool operator==(TYPE_OF_NULLPTR, const TFieldPath<RhsT>& Rhs)
 {
-	return !Rhs.IsValid();
+	return !Rhs.Get();
 }
 
 template <typename LhsT, typename RhsT>
@@ -415,13 +401,13 @@ FORCENOINLINE bool operator!=(const LhsT* Lhs, const TFieldPath<RhsT>& Rhs)
 template <typename LhsT>
 FORCENOINLINE bool operator!=(const TFieldPath<LhsT>& Lhs, TYPE_OF_NULLPTR)
 {
-	return Lhs.IsValid();
+	return Lhs.Get();
 }
 
 template <typename RhsT>
 FORCENOINLINE bool operator!=(TYPE_OF_NULLPTR, const TFieldPath<RhsT>& Rhs)
 {
-	return Rhs.IsValid();
+	return Rhs.Get();
 }
 
 template<class T> struct TIsPODType<TFieldPath<T> > { enum { Value = true }; };
