Fix dangling FProperty pointers in FArchiveStackTraceReader::FSerializeData when loading packages that compile on load.  In the case of those packages, some of the FProperties are deleted before LoadPackage returns.
#rnx

diff --git a/Engine/Source/Editor/UnrealEd/Private/Commandlets/PackageUtilities.cpp b/Engine/Source/Editor/UnrealEd/Private/Commandlets/PackageUtilities.cpp
--- a/Engine/Source/Editor/UnrealEd/Private/Commandlets/PackageUtilities.cpp
+++ b/Engine/Source/Editor/UnrealEd/Private/Commandlets/PackageUtilities.cpp
@@ -1640,7 +1640,7 @@ int32 UPkgInfoCommandlet::Main( const FString& Params )
 						*IndexString,
 						SerializeData.Offset,
 						Indent, *GetFullNameSafe(SerializeData.Object),
-						Indent, *GetFullNameSafe(SerializeData.Property),
+						Indent, *SerializeData.FullPropertyName,
 						Indent, SerializeData.Size,
 						Indent, SerializeData.Count);
 					Out.Logf(ELogVerbosity::Display, TEXT("%s"), *DisplayText);
diff --git a/Engine/Source/Runtime/CoreUObject/Private/Serialization/ArchiveStackTrace.cpp b/Engine/Source/Runtime/CoreUObject/Private/Serialization/ArchiveStackTrace.cpp
--- a/Engine/Source/Runtime/CoreUObject/Private/Serialization/ArchiveStackTrace.cpp
+++ b/Engine/Source/Runtime/CoreUObject/Private/Serialization/ArchiveStackTrace.cpp
@@ -1360,6 +1360,15 @@ void FArchiveStackTrace::DumpPackageHeaderDiffs(const FPackageData& SourcePackag
 #endif // !NO_LOGGING
 }
 
+FArchiveStackTraceReader::FSerializeData::FSerializeData(int64 InOffset, int64 InSize, UObject* InObject, FProperty* InProperty)
+: Offset(InOffset)
+, Size(InSize)
+, Count(1)
+, Object(InObject)
+, PropertyName(InProperty->GetFName())
+, FullPropertyName(GetFullNameSafe(InProperty))
+{}
+
 FArchiveStackTraceReader::FArchiveStackTraceReader(const TCHAR* InFilename, const uint8* InData, const int64 Num)
 	: FLargeMemoryReader(InData, Num, ELargeMemoryReaderFlags::TakeOwnership, InFilename)
 	, ThreadContext(FUObjectThreadContext::Get())
@@ -1374,7 +1383,7 @@ void FArchiveStackTraceReader::Serialize(void* OutData, int64 Num)
 	if (SerializeTrace.Num())
 	{
 		FSerializeData& Last = SerializeTrace.Last();
-		if (Last != NewData)
+		if (NewData.IsContiguousSerialization(Last))
 		{
 			SerializeTrace.Add(NewData);
 		}
diff --git a/Engine/Source/Runtime/CoreUObject/Public/Serialization/ArchiveStackTrace.h b/Engine/Source/Runtime/CoreUObject/Public/Serialization/ArchiveStackTrace.h
--- a/Engine/Source/Runtime/CoreUObject/Public/Serialization/ArchiveStackTrace.h
+++ b/Engine/Source/Runtime/CoreUObject/Public/Serialization/ArchiveStackTrace.h
@@ -254,33 +254,22 @@ public:
 			, Size(0)
 			, Count(0)
 			, Object(nullptr)
-			, Property(nullptr)
-		{}
-		FSerializeData(int64 InOffset, int64 InSize, UObject* InObject, FProperty* InProperty)
-			: Offset(InOffset)
-			, Size(InSize)
-			, Count(1)
-			, Object(InObject)
-			, Property(InProperty)
+			, PropertyName(NAME_None)
 		{}
+		FSerializeData(int64 InOffset, int64 InSize, UObject* InObject, FProperty* InProperty);
 		int64 Offset;
 		int64 Size;
 		int64 Count;
 		UObject* Object;
-		FProperty* Property;
-		bool IsIdentical(const FSerializeData& Other) const
+		FName PropertyName;
+		FString FullPropertyName;
+
+		bool IsContiguousSerialization(const FSerializeData& Other) const
 		{
-			return Object == Other.Object && Property == Other.Property &&
+			// Return whether this and other are neighboring bits of data for the serialization of the same instance of an object\property
+			return Object == Other.Object && PropertyName == Other.PropertyName &&
 				(Offset == Other.Offset || (Offset + Size) == Other.Offset); // This is to merge contiguous blocks
 		}
-		bool operator==(const FSerializeData& Other) const
-		{
-			return IsIdentical(Other);
-		}
-		bool operator!=(const FSerializeData& Other) const
-		{
-			return !IsIdentical(Other);
-		}
 	};
 private:
 	TArray<FSerializeData> SerializeTrace;
