Fixing CIS

diff --git a/Engine/Source/Runtime/CoreUObject/Public/UObject/FieldPath.h b/Engine/Source/Runtime/CoreUObject/Public/UObject/FieldPath.h
--- a/Engine/Source/Runtime/CoreUObject/Public/UObject/FieldPath.h
+++ b/Engine/Source/Runtime/CoreUObject/Public/UObject/FieldPath.h
@@ -401,13 +401,13 @@ FORCENOINLINE bool operator!=(const LhsT* Lhs, const TFieldPath<RhsT>& Rhs)
 template <typename LhsT>
 FORCENOINLINE bool operator!=(const TFieldPath<LhsT>& Lhs, TYPE_OF_NULLPTR)
 {
-	return Lhs.Get();
+	return !!Lhs.Get();
 }
 
 template <typename RhsT>
 FORCENOINLINE bool operator!=(TYPE_OF_NULLPTR, const TFieldPath<RhsT>& Rhs)
 {
-	return Rhs.Get();
+	return !!Rhs.Get();
 }
 
 template<class T> struct TIsPODType<TFieldPath<T> > { enum { Value = true }; };
