Fixing FBoolProperty::ExportTextItem crash

diff --git a/Engine/Source/Runtime/CoreUObject/Private/UObject/PropertyBool.cpp b/Engine/Source/Runtime/CoreUObject/Private/UObject/PropertyBool.cpp
--- a/Engine/Source/Runtime/CoreUObject/Private/UObject/PropertyBool.cpp
+++ b/Engine/Source/Runtime/CoreUObject/Private/UObject/PropertyBool.cpp
@@ -20,6 +20,7 @@ FBoolProperty::FBoolProperty(FFieldVariant InOwner, const FName& InName, EObject
 	, ByteMask(1)
 	, FieldMask(1)
 {
+	SetBoolSize(1, false, 1);
 }
 
 FBoolProperty::FBoolProperty(FFieldVariant InOwner, const FName& InName, EObjectFlags InObjectFlags, int32 InOffset, EPropertyFlags InFlags, uint32 InBitMask, uint32 InElementSize, bool bIsNativeBool)
