Adding deprecation macros for tthe old UProperty classes

diff --git a/Engine/Plugins/Enterprise/VariantManagerContent/Source/VariantManagerContent/Private/PropertyValue.cpp b/Engine/Plugins/Enterprise/VariantManagerContent/Source/VariantManagerContent/Private/PropertyValue.cpp
--- a/Engine/Plugins/Enterprise/VariantManagerContent/Source/VariantManagerContent/Private/PropertyValue.cpp
+++ b/Engine/Plugins/Enterprise/VariantManagerContent/Source/VariantManagerContent/Private/PropertyValue.cpp
@@ -19,6 +19,9 @@
 #include "EdGraphSchema_K2.h"
 #endif
 
+// WARNING: This should always be the last include in any file that needs it (except .generated.h)
+#include "UObject/UndefineUPropertyMacros.h"
+
 #define LOCTEXT_NAMESPACE "PropertyValue"
 
 DEFINE_LOG_CATEGORY(LogVariantContent);
@@ -1694,3 +1697,5 @@ void UPropertyValueVisibility::Serialize(FArchive& Ar)
 }
 
 #undef LOCTEXT_NAMESPACE
+
+#include "UObject/DefineUPropertyMacros.h"
\ No newline at end of file
diff --git a/Engine/Plugins/Enterprise/VariantManagerContent/Source/VariantManagerContent/Public/PropertyValue.h b/Engine/Plugins/Enterprise/VariantManagerContent/Source/VariantManagerContent/Public/PropertyValue.h
--- a/Engine/Plugins/Enterprise/VariantManagerContent/Source/VariantManagerContent/Public/PropertyValue.h
+++ b/Engine/Plugins/Enterprise/VariantManagerContent/Source/VariantManagerContent/Public/PropertyValue.h
@@ -6,6 +6,9 @@
 #include "UObject/ObjectMacros.h"
 #include "UObject/Field.h"
 
+// WARNING: This should always be the last include in any file that needs it (except .generated.h)
+#include "UObject/UndefineUPropertyMacros.h"
+
 #include "PropertyValue.generated.h"
 
 #define PATH_DELIMITER TEXT(" / ")
@@ -233,4 +236,6 @@ class VARIANTMANAGERCONTENT_API UPropertyValueVisibility : public UPropertyValue
 	// UObject Interface
 	virtual void Serialize(FArchive& Ar) override;
 	//~ End UObject Interface
-};
\ No newline at end of file
+};
+
+#include "UObject/DefineUPropertyMacros.h"
\ No newline at end of file
diff --git a/Engine/Plugins/ScriptPlugin/Source/ScriptPlugin/Classes/ScriptBlueprintGeneratedClass.h b/Engine/Plugins/ScriptPlugin/Source/ScriptPlugin/Classes/ScriptBlueprintGeneratedClass.h
--- a/Engine/Plugins/ScriptPlugin/Source/ScriptPlugin/Classes/ScriptBlueprintGeneratedClass.h
+++ b/Engine/Plugins/ScriptPlugin/Source/ScriptPlugin/Classes/ScriptBlueprintGeneratedClass.h
@@ -3,6 +3,10 @@
 #pragma once
 
 #include "Engine/BlueprintGeneratedClass.h"
+
+// WARNING: This should always be the last include in any file that needs it (except .generated.h)
+#include "UObject/UndefineUPropertyMacros.h"
+
 #include "ScriptBlueprintGeneratedClass.generated.h"
 
 class UScriptBlueprintGeneratedClass;
@@ -180,4 +184,6 @@ public:
 	}
 
 	virtual void Serialize(FArchive& Ar) override;
-};
\ No newline at end of file
+};
+
+#include "UObject/DefineUPropertyMacros.h"
\ No newline at end of file
diff --git a/Engine/Plugins/ScriptPlugin/Source/ScriptPlugin/Private/ScriptBlueprintGeneratedClass.cpp b/Engine/Plugins/ScriptPlugin/Source/ScriptPlugin/Private/ScriptBlueprintGeneratedClass.cpp
--- a/Engine/Plugins/ScriptPlugin/Source/ScriptPlugin/Private/ScriptBlueprintGeneratedClass.cpp
+++ b/Engine/Plugins/ScriptPlugin/Source/ScriptPlugin/Private/ScriptBlueprintGeneratedClass.cpp
@@ -8,6 +8,9 @@
 #include "LuaIntegration.h"
 #endif
 
+// WARNING: This should always be the last include in any file that needs it (except .generated.h)
+#include "UObject/UndefineUPropertyMacros.h"
+
 /////////////////////////////////////////////////////
 // UScriptBlueprintGeneratedClass
 
@@ -165,4 +168,6 @@ void FScriptContextBase::GetSupportedScriptFileFormats(TArray<FString>& OutForma
 #if WITH_LUA
 	OutFormats.Add(TEXT("lua;Script"));
 #endif
-}
\ No newline at end of file
+}
+
+#include "UObject/DefineUPropertyMacros.h"
\ No newline at end of file
diff --git a/Engine/Source/Editor/BlueprintGraph/Private/BlueprintBoundEventNodeSpawner.cpp b/Engine/Source/Editor/BlueprintGraph/Private/BlueprintBoundEventNodeSpawner.cpp
--- a/Engine/Source/Editor/BlueprintGraph/Private/BlueprintBoundEventNodeSpawner.cpp
+++ b/Engine/Source/Editor/BlueprintGraph/Private/BlueprintBoundEventNodeSpawner.cpp
@@ -10,8 +10,6 @@
 #include "Kismet2/KismetEditorUtilities.h"
 #include "EditorCategoryUtils.h"
 #include "BlueprintNodeSpawnerUtils.h"
-#include "UObject/UnrealTypePrivate.h" // Only for converting UProperties to FProperties
-#include "UObject/CoreObjectVersion.h"
 
 #define LOCTEXT_NAMESPACE "BlueprintBoundEventNodeSpawner"
 
diff --git a/Engine/Source/Editor/BlueprintGraph/Private/BlueprintFieldNodeSpawner.cpp b/Engine/Source/Editor/BlueprintGraph/Private/BlueprintFieldNodeSpawner.cpp
--- a/Engine/Source/Editor/BlueprintGraph/Private/BlueprintFieldNodeSpawner.cpp
+++ b/Engine/Source/Editor/BlueprintGraph/Private/BlueprintFieldNodeSpawner.cpp
@@ -2,8 +2,6 @@
 
 #include "BlueprintFieldNodeSpawner.h"
 #include "UObject/Package.h"
-#include "UObject/UnrealTypePrivate.h" // only for converting to UProps
-#include "UObject/CoreObjectVersion.h"
 
 #define LOCTEXT_NAMESPACE "BlueprintFieldNodeSpawner"
 
@@ -78,37 +76,4 @@ void UBlueprintFieldNodeSpawner::SetField(FFieldVariant InField)
 	}
 }
 
-// @todo FProp: add collector function
-//void UBlueprintFieldNodeSpawner::AddReferencedObjects(UObject* InThis, FReferenceCollector& Collector)
-//{
-//	UBlueprintFieldNodeSpawner* This = CastChecked<UBlueprintFieldNodeSpawner>(InThis);
-//
-//	if (This->Field.IsUObject())
-//	{
-//		UObject* Object = This->Field.Get<UObject>();
-//		UObject* OldObjectValue = Object;
-//		Collector.AddReferencedObject(Object, This);
-//		if (!Object && OldObjectValue)
-//		{
-//			This->Field = Object;
-//		}
-//	}
-//	else
-//	{
-//		FField* Field = This->Field.Get<FField>();
-//		UObject* Owner = Field ? Field->GetOwnerUObject() : nullptr;
-//		if (Owner)
-//		{
-//			UObject* OldObjectValue = Owner;
-//			Collector.AddReferencedObject(Owner, This);
-//			if (!Owner && OldObjectValue)
-//			{
-//				This->Field = (FField*)nullptr;
-//			}
-//		}
-//	}
-//
-//	Super::AddReferencedObjects(InThis, Collector);
-//}
-
 #undef LOCTEXT_NAMESPACE
diff --git a/Engine/Source/Programs/UnrealHeaderTool/Private/HeaderParser.cpp b/Engine/Source/Programs/UnrealHeaderTool/Private/HeaderParser.cpp
--- a/Engine/Source/Programs/UnrealHeaderTool/Private/HeaderParser.cpp
+++ b/Engine/Source/Programs/UnrealHeaderTool/Private/HeaderParser.cpp
@@ -9090,24 +9090,27 @@ void FHeaderParser::SimplifiedClassParse(const TCHAR* Filename, const TCHAR* InB
 		}
 		else if ( bProcess && FParse::Command(&Str,TEXT("#include")) )
 		{
-			if (bFoundGeneratedInclude)
-			{
-				FFileLineException::Throwf(Filename, CurrentLine, TEXT("#include found after .generated.h file - the .generated.h file should always be the last #include in a header"));
-			}
-
 			// Handle #include directives as if they were 'dependson' keywords.
 			FString DependsOnHeaderName = Str;
 
-			bFoundGeneratedInclude = DependsOnHeaderName.Contains(TEXT(".generated.h"));
-			if (!bFoundGeneratedInclude && DependsOnHeaderName.Len())
+			if (DependsOnHeaderName != TEXT("\"UObject/DefineUPropertyMacros.h\""))
 			{
-				bool  bIsQuotedInclude  = DependsOnHeaderName[0] == '\"';
-				int32 HeaderFilenameEnd = DependsOnHeaderName.Find(bIsQuotedInclude ? TEXT("\"") : TEXT(">"), ESearchCase::CaseSensitive, ESearchDir::FromStart, 1);
+				if (bFoundGeneratedInclude)
+				{
+					FFileLineException::Throwf(Filename, CurrentLine, TEXT("#include found after .generated.h file - the .generated.h file should always be the last #include in a header"));
+				}
 
-				if (HeaderFilenameEnd != INDEX_NONE)
+				bFoundGeneratedInclude = DependsOnHeaderName.Contains(TEXT(".generated.h"));
+				if (!bFoundGeneratedInclude && DependsOnHeaderName.Len())
 				{
-					// Include the extension in the name so that we later know where this entry came from.
-					DependentOn.Add(FHeaderProvider(EHeaderProviderSourceType::FileName, *FPaths::GetCleanFilename(DependsOnHeaderName.Mid(1, HeaderFilenameEnd - 1))));
+					bool  bIsQuotedInclude = DependsOnHeaderName[0] == '\"';
+					int32 HeaderFilenameEnd = DependsOnHeaderName.Find(bIsQuotedInclude ? TEXT("\"") : TEXT(">"), ESearchCase::CaseSensitive, ESearchDir::FromStart, 1);
+
+					if (HeaderFilenameEnd != INDEX_NONE)
+					{
+						// Include the extension in the name so that we later know where this entry came from.
+						DependentOn.Add(FHeaderProvider(EHeaderProviderSourceType::FileName, *FPaths::GetCleanFilename(DependsOnHeaderName.Mid(1, HeaderFilenameEnd - 1))));
+					}
 				}
 			}
 		}
diff --git a/Engine/Source/Runtime/CoreUObject/Private/UObject/Class.cpp b/Engine/Source/Runtime/CoreUObject/Private/UObject/Class.cpp
--- a/Engine/Source/Runtime/CoreUObject/Private/UObject/Class.cpp
+++ b/Engine/Source/Runtime/CoreUObject/Private/UObject/Class.cpp
@@ -52,6 +52,10 @@
 #include "UObject/PropertyProxyArchive.h"
 #include "UObject/FieldPath.h"
 
+
+// WARNING: This should always be the last include in any file that needs it (except .generated.h)
+#include "UObject/UndefineUPropertyMacros.h"
+
 // This flag enables some expensive class tree validation that is meant to catch mutations of 
 // the class tree outside of SetSuperStruct. It has been disabled because loading blueprints 
 // does a lot of mutation of the class tree, and the validation checks impact iteration time.
@@ -6086,3 +6090,6 @@ IMPLEMENT_CORE_INTRINSIC_CLASS(UDynamicClass, UClass,
 		PRAGMA_ENABLE_SHADOW_VARIABLE_WARNINGS
 	#endif
 #endif
+
+
+#include "UObject/DefineUPropertyMacros.h"
\ No newline at end of file
diff --git a/Engine/Source/Runtime/CoreUObject/Private/UObject/EnumProperty.cpp b/Engine/Source/Runtime/CoreUObject/Private/UObject/EnumProperty.cpp
--- a/Engine/Source/Runtime/CoreUObject/Private/UObject/EnumProperty.cpp
+++ b/Engine/Source/Runtime/CoreUObject/Private/UObject/EnumProperty.cpp
@@ -11,6 +11,9 @@
 #include "UObject/LinkerLoad.h"
 #include "Misc/NetworkVersion.h"
 
+// WARNING: This should always be the last include in any file that needs it (except .generated.h)
+#include "UObject/UndefineUPropertyMacros.h"
+
 namespace UE4EnumProperty_Private
 {
 	template <typename OldIntType>
@@ -466,3 +469,5 @@ uint32 FEnumProperty::GetValueTypeHashInternal(const void* Src) const
 	check(UnderlyingProp);
 	return UnderlyingProp->GetValueTypeHash(Src);
 }
+
+#include "UObject/DefineUPropertyMacros.h"
\ No newline at end of file
diff --git a/Engine/Source/Runtime/CoreUObject/Private/UObject/Field.cpp b/Engine/Source/Runtime/CoreUObject/Private/UObject/Field.cpp
--- a/Engine/Source/Runtime/CoreUObject/Private/UObject/Field.cpp
+++ b/Engine/Source/Runtime/CoreUObject/Private/UObject/Field.cpp
@@ -43,6 +43,8 @@ Field.cpp: Defines FField property system fundamentals
 #include "UObject/TextProperty.h"
 #include "UObject/EnumProperty.h"
 
+// WARNING: This should always be the last include in any file that needs it (except .generated.h)
+#include "UObject/UndefineUPropertyMacros.h"
 
 FFieldClass::FFieldClass(const TCHAR* InCPPName, uint64 InId, uint64 InCastFlags, FFieldClass* InSuperClass, FField* (*ConstructFnPtr)(const FFieldVariant&, const FName&, EObjectFlags))
 	: Id(InId)
@@ -1157,4 +1159,6 @@ FField* FindFieldByPath(const TCHAR* InFieldPath)
 		}
 	}
 	return FoundField;
-}
\ No newline at end of file
+}
+
+#include "UObject/DefineUPropertyMacros.h"
\ No newline at end of file
diff --git a/Engine/Source/Runtime/CoreUObject/Private/UObject/FieldPathProperty.cpp b/Engine/Source/Runtime/CoreUObject/Private/UObject/FieldPathProperty.cpp
--- a/Engine/Source/Runtime/CoreUObject/Private/UObject/FieldPathProperty.cpp
+++ b/Engine/Source/Runtime/CoreUObject/Private/UObject/FieldPathProperty.cpp
@@ -6,6 +6,9 @@
 #include "UObject/UnrealTypePrivate.h"
 #include "UObject/LinkerLoad.h"
 
+// WARNING: This should always be the last include in any file that needs it (except .generated.h)
+#include "UObject/UndefineUPropertyMacros.h"
+
 IMPLEMENT_FIELD(FFieldPathProperty)
 
 #if WITH_EDITORONLY_DATA
@@ -154,4 +157,6 @@ FString FFieldPathProperty::GetCPPType(FString* ExtendedTypeText, uint32 CPPExpo
 		InnerTypeText += TEXT(">");
 	}
 	return TEXT("TFieldPath");
-}
\ No newline at end of file
+}
+
+#include "UObject/DefineUPropertyMacros.h"
\ No newline at end of file
diff --git a/Engine/Source/Runtime/CoreUObject/Private/UObject/LinkerPlaceholderBase.cpp b/Engine/Source/Runtime/CoreUObject/Private/UObject/LinkerPlaceholderBase.cpp
--- a/Engine/Source/Runtime/CoreUObject/Private/UObject/LinkerPlaceholderBase.cpp
+++ b/Engine/Source/Runtime/CoreUObject/Private/UObject/LinkerPlaceholderBase.cpp
@@ -7,6 +7,9 @@
 #include "UObject/UnrealTypePrivate.h"
 #include "Blueprint/BlueprintSupport.h"
 
+// WARNING: This should always be the last include in any file that needs it (except .generated.h)
+#include "UObject/UndefineUPropertyMacros.h"
+
 #if USE_DEFERRED_DEPENDENCY_CHECK_VERIFICATION_TESTS
 	#define DEFERRED_DEPENDENCY_ENSURE(EnsueExpr) ensure(EnsueExpr)
 #else  // USE_DEFERRED_DEPENDENCY_CHECK_VERIFICATION_TESTS
@@ -780,3 +783,5 @@ int32 TLinkerImportPlaceholder<UFunction>::ResolvePropertyReferences(UFunction*
 }
 
 #undef DEFERRED_DEPENDENCY_ENSURE
+
+#include "UObject/DefineUPropertyMacros.h"
\ No newline at end of file
diff --git a/Engine/Source/Runtime/CoreUObject/Private/UObject/Property.cpp b/Engine/Source/Runtime/CoreUObject/Private/UObject/Property.cpp
--- a/Engine/Source/Runtime/CoreUObject/Private/UObject/Property.cpp
+++ b/Engine/Source/Runtime/CoreUObject/Private/UObject/Property.cpp
@@ -20,6 +20,9 @@
 #include "Math/Box2D.h"
 #include "UObject/ReleaseObjectVersion.h"
 
+// WARNING: This should always be the last include in any file that needs it (except .generated.h)
+#include "UObject/UndefineUPropertyMacros.h"
+
 DEFINE_LOG_CATEGORY(LogProperty);
 
 // List the core ones here as they have already been included (and can be used without CoreUObject!)
@@ -1533,3 +1536,6 @@ FProperty* UStruct::FindPropertyByName(FName InName) const
 
 	return NULL;
 }
+
+
+#include "UObject/DefineUPropertyMacros.h"
\ No newline at end of file
diff --git a/Engine/Source/Runtime/CoreUObject/Private/UObject/PropertyArray.cpp b/Engine/Source/Runtime/CoreUObject/Private/UObject/PropertyArray.cpp
--- a/Engine/Source/Runtime/CoreUObject/Private/UObject/PropertyArray.cpp
+++ b/Engine/Source/Runtime/CoreUObject/Private/UObject/PropertyArray.cpp
@@ -9,6 +9,9 @@
 #include "UObject/LinkerLoad.h"
 #include "UObject/PropertyHelper.h"
 
+// WARNING: This should always be the last include in any file that needs it (except .generated.h)
+#include "UObject/UndefineUPropertyMacros.h"
+
 /*-----------------------------------------------------------------------------
 	FArrayProperty.
 -----------------------------------------------------------------------------*/
@@ -703,4 +706,6 @@ FField* FArrayProperty::GetInnerFieldByName(const FName& InName)
 		return Inner;
 	}
 	return nullptr;
-}
\ No newline at end of file
+}
+
+#include "UObject/DefineUPropertyMacros.h"
\ No newline at end of file
diff --git a/Engine/Source/Runtime/CoreUObject/Private/UObject/PropertyBaseObject.cpp b/Engine/Source/Runtime/CoreUObject/Private/UObject/PropertyBaseObject.cpp
--- a/Engine/Source/Runtime/CoreUObject/Private/UObject/PropertyBaseObject.cpp
+++ b/Engine/Source/Runtime/CoreUObject/Private/UObject/PropertyBaseObject.cpp
@@ -10,6 +10,9 @@
 #include "UObject/LinkerPlaceholderClass.h"
 #include "UObject/LinkerPlaceholderExportObject.h"
 
+// WARNING: This should always be the last include in any file that needs it (except .generated.h)
+#include "UObject/UndefineUPropertyMacros.h"
+
 /*-----------------------------------------------------------------------------
 	FObjectPropertyBase.
 -----------------------------------------------------------------------------*/
@@ -585,3 +588,5 @@ void FObjectPropertyBase::CopyCompleteValueFromScriptVM( void* Dest, void const*
 		SetObjectPropertyValue(((uint8*)Dest) + Index * ElementSize, ((UObject**)Src)[Index]);
 	}
 }
+
+#include "UObject/DefineUPropertyMacros.h"
\ No newline at end of file
diff --git a/Engine/Source/Runtime/CoreUObject/Private/UObject/PropertyBool.cpp b/Engine/Source/Runtime/CoreUObject/Private/UObject/PropertyBool.cpp
--- a/Engine/Source/Runtime/CoreUObject/Private/UObject/PropertyBool.cpp
+++ b/Engine/Source/Runtime/CoreUObject/Private/UObject/PropertyBool.cpp
@@ -7,6 +7,9 @@
 #include "UObject/UnrealTypePrivate.h"
 #include "UObject/PropertyHelper.h"
 
+// WARNING: This should always be the last include in any file that needs it (except .generated.h)
+#include "UObject/UndefineUPropertyMacros.h"
+
 /*-----------------------------------------------------------------------------
 	FBoolProperty.
 -----------------------------------------------------------------------------*/
@@ -387,3 +390,5 @@ uint32 FBoolProperty::GetValueTypeHashInternal(const void* Src) const
 {
 	return GetTypeHash(*(const bool*)Src);
 }
+
+#include "UObject/DefineUPropertyMacros.h"
\ No newline at end of file
diff --git a/Engine/Source/Runtime/CoreUObject/Private/UObject/PropertyByte.cpp b/Engine/Source/Runtime/CoreUObject/Private/UObject/PropertyByte.cpp
--- a/Engine/Source/Runtime/CoreUObject/Private/UObject/PropertyByte.cpp
+++ b/Engine/Source/Runtime/CoreUObject/Private/UObject/PropertyByte.cpp
@@ -11,6 +11,9 @@
 #include "Algo/Find.h"
 #include "UObject/LinkerLoad.h"
 
+// WARNING: This should always be the last include in any file that needs it (except .generated.h)
+#include "UObject/UndefineUPropertyMacros.h"
+
 /*-----------------------------------------------------------------------------
 	FByteProperty.
 -----------------------------------------------------------------------------*/
@@ -439,3 +442,5 @@ UEnum* FByteProperty::GetIntPropertyEnum() const
 {
 	return Enum;
 }
+
+#include "UObject/DefineUPropertyMacros.h"
\ No newline at end of file
diff --git a/Engine/Source/Runtime/CoreUObject/Private/UObject/PropertyClass.cpp b/Engine/Source/Runtime/CoreUObject/Private/UObject/PropertyClass.cpp
--- a/Engine/Source/Runtime/CoreUObject/Private/UObject/PropertyClass.cpp
+++ b/Engine/Source/Runtime/CoreUObject/Private/UObject/PropertyClass.cpp
@@ -9,6 +9,9 @@
 #include "Misc/ConfigCacheIni.h"
 #include "UObject/PropertyHelper.h"
 
+// WARNING: This should always be the last include in any file that needs it (except .generated.h)
+#include "UObject/UndefineUPropertyMacros.h"
+
 /*-----------------------------------------------------------------------------
 	FClassProperty.
 -----------------------------------------------------------------------------*/
@@ -174,3 +177,5 @@ bool FClassProperty::Identical( const void* A, const void* B, uint32 PortFlags )
 	check(ObjectB == nullptr || ObjectB->IsA<UClass>());
 	return (ObjectA == ObjectB);
 }
+
+#include "UObject/DefineUPropertyMacros.h"
\ No newline at end of file
diff --git a/Engine/Source/Runtime/CoreUObject/Private/UObject/PropertyInterface.cpp b/Engine/Source/Runtime/CoreUObject/Private/UObject/PropertyInterface.cpp
--- a/Engine/Source/Runtime/CoreUObject/Private/UObject/PropertyInterface.cpp
+++ b/Engine/Source/Runtime/CoreUObject/Private/UObject/PropertyInterface.cpp
@@ -10,6 +10,9 @@
 #include "UObject/UnrealTypePrivate.h"
 #include "UObject/LinkerPlaceholderClass.h"
 
+// WARNING: This should always be the last include in any file that needs it (except .generated.h)
+#include "UObject/UndefineUPropertyMacros.h"
+
 /*-----------------------------------------------------------------------------
 	FInterfaceProperty.
 -----------------------------------------------------------------------------*/
@@ -302,3 +305,5 @@ void FInterfaceProperty::AddReferencedObjects(FReferenceCollector& Collector)
 	Collector.AddReferencedObject(InterfaceClass);
 	Super::AddReferencedObjects(Collector);
 }
+
+#include "UObject/DefineUPropertyMacros.h"
\ No newline at end of file
diff --git a/Engine/Source/Runtime/CoreUObject/Private/UObject/PropertyMap.cpp b/Engine/Source/Runtime/CoreUObject/Private/UObject/PropertyMap.cpp
--- a/Engine/Source/Runtime/CoreUObject/Private/UObject/PropertyMap.cpp
+++ b/Engine/Source/Runtime/CoreUObject/Private/UObject/PropertyMap.cpp
@@ -12,6 +12,9 @@
 #include "Serialization/ArchiveUObjectFromStructuredArchive.h"
 #include "UObject/UObjectThreadContext.h"
 
+// WARNING: This should always be the last include in any file that needs it (except .generated.h)
+#include "UObject/UndefineUPropertyMacros.h"
+
 namespace UE4MapProperty_Private
 {
 	/**
@@ -1178,4 +1181,6 @@ FField* FMapProperty::GetInnerFieldByName(const FName& InName)
 		return ValueProp;
 	}
 	return nullptr;
-}
\ No newline at end of file
+}
+
+#include "UObject/DefineUPropertyMacros.h"
\ No newline at end of file
diff --git a/Engine/Source/Runtime/CoreUObject/Private/UObject/PropertyMulticastDelegate.cpp b/Engine/Source/Runtime/CoreUObject/Private/UObject/PropertyMulticastDelegate.cpp
--- a/Engine/Source/Runtime/CoreUObject/Private/UObject/PropertyMulticastDelegate.cpp
+++ b/Engine/Source/Runtime/CoreUObject/Private/UObject/PropertyMulticastDelegate.cpp
@@ -9,6 +9,9 @@
 #include "UObject/LinkerPlaceholderFunction.h"
 #include "Serialization/ArchiveUObjectFromStructuredArchive.h"
 
+// WARNING: This should always be the last include in any file that needs it (except .generated.h)
+#include "UObject/UndefineUPropertyMacros.h"
+
 FMulticastScriptDelegate::FInvocationList FMulticastDelegateProperty::EmptyList;
 
 #if WITH_EDITORONLY_DATA
@@ -599,4 +602,4 @@ void FMulticastSparseDelegateProperty::ClearDelegate(UObject* Parent, void* Prop
 
 IMPLEMENT_FIELD(FMulticastSparseDelegateProperty)
 
-
+#include "UObject/DefineUPropertyMacros.h"
\ No newline at end of file
diff --git a/Engine/Source/Runtime/CoreUObject/Private/UObject/PropertyPrivate.cpp b/Engine/Source/Runtime/CoreUObject/Private/UObject/PropertyPrivate.cpp
--- a/Engine/Source/Runtime/CoreUObject/Private/UObject/PropertyPrivate.cpp
+++ b/Engine/Source/Runtime/CoreUObject/Private/UObject/PropertyPrivate.cpp
@@ -22,6 +22,9 @@
 #include "UObject/LinkerPlaceholderFunction.h"
 #include "UObject/LinkerPlaceholderClass.h"
 
+// WARNING: This should always be the last include in any file that needs it (except .generated.h)
+#include "UObject/UndefineUPropertyMacros.h"
+
 /*-----------------------------------------------------------------------------
 	UProperty implementation.
 -----------------------------------------------------------------------------*/
@@ -914,4 +917,6 @@ IMPLEMENT_CORE_INTRINSIC_CLASS(UPropertyWrapper, UObject,
 IMPLEMENT_CORE_INTRINSIC_CLASS(UMulticastDelegatePropertyWrapper, UPropertyWrapper,
 	{
 	}
-);
\ No newline at end of file
+);
+
+#include "UObject/DefineUPropertyMacros.h"
\ No newline at end of file
diff --git a/Engine/Source/Runtime/CoreUObject/Private/UObject/PropertySet.cpp b/Engine/Source/Runtime/CoreUObject/Private/UObject/PropertySet.cpp
--- a/Engine/Source/Runtime/CoreUObject/Private/UObject/PropertySet.cpp
+++ b/Engine/Source/Runtime/CoreUObject/Private/UObject/PropertySet.cpp
@@ -10,6 +10,9 @@
 #include "Misc/ScopeExit.h"
 #include "Serialization/ArchiveUObjectFromStructuredArchive.h"
 
+// WARNING: This should always be the last include in any file that needs it (except .generated.h)
+#include "UObject/UndefineUPropertyMacros.h"
+
 namespace UE4SetProperty_Private
 {
 	/**
@@ -984,4 +987,6 @@ FField* FSetProperty::GetInnerFieldByName(const FName& InName)
 		return ElementProp;
 	}
 	return nullptr;
-}
\ No newline at end of file
+}
+
+#include "UObject/DefineUPropertyMacros.h"
\ No newline at end of file
diff --git a/Engine/Source/Runtime/CoreUObject/Private/UObject/PropertySoftClassPtr.cpp b/Engine/Source/Runtime/CoreUObject/Private/UObject/PropertySoftClassPtr.cpp
--- a/Engine/Source/Runtime/CoreUObject/Private/UObject/PropertySoftClassPtr.cpp
+++ b/Engine/Source/Runtime/CoreUObject/Private/UObject/PropertySoftClassPtr.cpp
@@ -7,6 +7,9 @@
 #include "UObject/UnrealTypePrivate.h"
 #include "UObject/LinkerPlaceholderClass.h"
 
+// WARNING: This should always be the last include in any file that needs it (except .generated.h)
+#include "UObject/UndefineUPropertyMacros.h"
+
 /*-----------------------------------------------------------------------------
 	FSoftClassProperty.
 -----------------------------------------------------------------------------*/
@@ -116,3 +119,5 @@ bool FSoftClassProperty::SameType(const FProperty* Other) const
 {
 	return Super::SameType(Other) && (MetaClass == ((FSoftClassProperty*)Other)->MetaClass);
 }
+
+#include "UObject/DefineUPropertyMacros.h"
\ No newline at end of file
diff --git a/Engine/Source/Runtime/CoreUObject/Private/UObject/PropertyStruct.cpp b/Engine/Source/Runtime/CoreUObject/Private/UObject/PropertyStruct.cpp
--- a/Engine/Source/Runtime/CoreUObject/Private/UObject/PropertyStruct.cpp
+++ b/Engine/Source/Runtime/CoreUObject/Private/UObject/PropertyStruct.cpp
@@ -10,6 +10,9 @@
 #include "UObject/LinkerPlaceholderBase.h"
 #include "Serialization/ArchiveUObjectFromStructuredArchive.h"
 
+// WARNING: This should always be the last include in any file that needs it (except .generated.h)
+#include "UObject/UndefineUPropertyMacros.h"
+
 static inline void PreloadInnerStructMembers(FStructProperty* StructProperty)
 {
 #if USE_CIRCULAR_DEPENDENCY_LOAD_DEFERRING
@@ -413,3 +416,5 @@ EConvertFromTypeResult FStructProperty::ConvertFromType(const FPropertyTag& Tag,
 	}
 	return EConvertFromTypeResult::UseSerializeItem;
 }
+
+#include "UObject/DefineUPropertyMacros.h"
\ No newline at end of file
diff --git a/Engine/Source/Runtime/CoreUObject/Public/UObject/DefineUPropertyMacros.h b/Engine/Source/Runtime/CoreUObject/Public/UObject/DefineUPropertyMacros.h
new file mode 100644
--- /dev/null
+++ b/Engine/Source/Runtime/CoreUObject/Public/UObject/DefineUPropertyMacros.h
@@ -0,0 +1,133 @@
+// Copyright 1998-2019 Epic Games, Inc. All Rights Reserved.
+
+/*=============================================================================
+	DefineUPropertyMacros.h: macro'ed typedefs for FProperties
+=============================================================================*/
+
+#ifndef UProperty
+	#define UProperty DEPRECATED_MACRO(4.25, "UProperty has been renamed to FProperty") FProperty
+#endif
+
+#ifndef UNumericProperty
+	#define UNumericProperty DEPRECATED_MACRO(4.25, "UNumericProperty has been renamed to FNumericProperty") FNumericProperty
+#endif
+
+#ifndef UByteProperty
+	#define UByteProperty DEPRECATED_MACRO(4.25, "UByteProperty has been renamed to FByteProperty") FByteProperty
+#endif
+
+#ifndef UInt8Property
+	#define UInt8Property DEPRECATED_MACRO(4.25, "UInt8Property has been renamed to FInt8Property") FInt8Property
+#endif
+
+#ifndef UInt16Property
+	#define UInt16Property DEPRECATED_MACRO(4.25, "UInt16Property has been renamed to FInt16Property") FInt16Property
+#endif
+
+#ifndef UIntProperty
+	#define UIntProperty DEPRECATED_MACRO(4.25, "UIntProperty has been renamed to FIntProperty") FIntProperty
+#endif
+
+#ifndef UInt64Property
+	#define UInt64Property DEPRECATED_MACRO(4.25, "UInt64Property has been renamed to FInt64Property") FInt64Property
+#endif
+
+#ifndef UUInt16Property
+	#define UUInt16Property DEPRECATED_MACRO(4.25, "UUInt16Property has been renamed to FUInt16Property") FUInt16Property
+#endif
+
+#ifndef UUInt32Property
+	#define UUInt32Property DEPRECATED_MACRO(4.25, "UUInt32Property has been renamed to FUInt32Property") FUInt32Property
+#endif
+
+#ifndef UUInt64Property
+	#define UUInt64Property DEPRECATED_MACRO(4.25, "UUInt64Property has been renamed to FUInt64Property") FUInt64Property
+#endif
+
+#ifndef UFloatProperty
+	#define UFloatProperty DEPRECATED_MACRO(4.25, "UFloatProperty has been renamed to FFloatProperty") FFloatProperty
+#endif
+
+#ifndef UDoubleProperty
+	#define UDoubleProperty DEPRECATED_MACRO(4.25, "UDoubleProperty has been renamed to FDoubleProperty") FDoubleProperty
+#endif
+
+#ifndef UBoolProperty
+	#define UBoolProperty DEPRECATED_MACRO(4.25, "UBoolProperty has been renamed to FBoolProperty") FBoolProperty
+#endif
+
+#ifndef UObjectPropertyBase
+	#define UObjectPropertyBase DEPRECATED_MACRO(4.25, "UObjectPropertyBase has been renamed to FObjectPropertyBase") FObjectPropertyBase
+#endif
+
+#ifndef UObjectProperty
+	#define UObjectProperty DEPRECATED_MACRO(4.25, "UObjectProperty has been renamed to FObjectProperty") FObjectProperty
+#endif
+
+#ifndef UWeakObjectProperty
+	#define UWeakObjectProperty DEPRECATED_MACRO(4.25, "UWeakObjectProperty has been renamed to FWeakObjectProperty") FWeakObjectProperty
+#endif
+
+#ifndef ULazyObjectProperty
+	#define ULazyObjectProperty DEPRECATED_MACRO(4.25, "ULazyObjectProperty has been renamed to FLazyObjectProperty") FLazyObjectProperty
+#endif
+
+#ifndef USoftObjectProperty
+	#define USoftObjectProperty DEPRECATED_MACRO(4.25, "USoftObjectProperty has been renamed to FSoftObjectProperty") FSoftObjectProperty
+#endif
+
+#ifndef UClassProperty
+	#define UClassProperty DEPRECATED_MACRO(4.25, "UClassProperty has been renamed to FClassProperty") FClassProperty
+#endif
+
+#ifndef USoftClassProperty
+	#define USoftClassProperty DEPRECATED_MACRO(4.25, "USoftClassProperty has been renamed to FSoftClassProperty") FSoftClassProperty
+#endif
+
+#ifndef UInterfaceProperty
+	#define UInterfaceProperty DEPRECATED_MACRO(4.25, "UInterfaceProperty has been renamed to FInterfaceProperty") FInterfaceProperty
+#endif
+
+#ifndef UNameProperty
+	#define UNameProperty DEPRECATED_MACRO(4.25, "UNameProperty has been renamed to FNameProperty") FNameProperty
+#endif
+
+#ifndef UStrProperty
+	#define UStrProperty DEPRECATED_MACRO(4.25, "UStrProperty has been renamed to FStrProperty") FStrProperty
+#endif
+
+#ifndef UArrayProperty
+	#define UArrayProperty DEPRECATED_MACRO(4.25, "UArrayProperty has been renamed to FArrayProperty") FArrayProperty
+#endif
+
+#ifndef UMapProperty
+	#define UMapProperty DEPRECATED_MACRO(4.25, "UMapProperty has been renamed to FMapProperty") FMapProperty
+#endif
+
+#ifndef USetProperty
+	#define USetProperty DEPRECATED_MACRO(4.25, "USetProperty has been renamed to FSetProperty") FSetProperty
+#endif
+
+#ifndef UStructProperty
+	#define UStructProperty DEPRECATED_MACRO(4.25, "UStructProperty has been renamed to FStructProperty") FStructProperty
+#endif
+
+#ifndef UMulticastDelegateProperty
+	#define UMulticastDelegateProperty DEPRECATED_MACRO(4.25, "UMulticastDelegateProperty has been renamed to FMulticastDelegateProperty") FMulticastDelegateProperty
+#endif
+
+#ifndef UMulticastInlineDelegateProperty
+	#define UMulticastInlineDelegateProperty DEPRECATED_MACRO(4.25, "UMulticastInlineDelegateProperty has been renamed to FMulticastInlineDelegateProperty") FMulticastInlineDelegateProperty
+#endif
+
+#ifndef UMulticastSparseDelegateProperty
+	#define UMulticastSparseDelegateProperty DEPRECATED_MACRO(4.25, "UMulticastSparseDelegateProperty has been renamed to FMulticastSparseDelegateProperty") FMulticastSparseDelegateProperty
+#endif
+
+#ifndef UEnumProperty
+	#define UEnumProperty DEPRECATED_MACRO(4.25, "UEnumProperty has been renamed to FEnumProperty") FEnumProperty
+#endif
+
+#ifndef UTextProperty
+	#define UTextProperty DEPRECATED_MACRO(4.25, "UTextProperty has been renamed to FTextProperty") FTextProperty
+#endif
diff --git a/Engine/Source/Runtime/CoreUObject/Public/UObject/UndefineUPropertyMacros.h b/Engine/Source/Runtime/CoreUObject/Public/UObject/UndefineUPropertyMacros.h
new file mode 100644
--- /dev/null
+++ b/Engine/Source/Runtime/CoreUObject/Public/UObject/UndefineUPropertyMacros.h
@@ -0,0 +1,133 @@
+// Copyright 1998-2019 Epic Games, Inc. All Rights Reserved.
+
+/*=============================================================================
+	UndefineUPropertyMacros.h
+=============================================================================*/
+
+#ifdef UProperty
+	#undef UProperty
+#endif
+
+#ifdef UNumericProperty
+	#undef UNumericProperty
+#endif
+
+#ifdef UByteProperty
+	#undef UByteProperty
+#endif
+
+#ifdef UInt8Property
+	#undef UInt8Property
+#endif
+
+#ifdef UInt16Property
+	#undef UInt16Property
+#endif
+
+#ifdef UIntProperty
+	#undef UIntProperty
+#endif
+
+#ifdef UInt64Property
+	#undef UInt64Property
+#endif
+
+#ifdef UUInt16Property
+	#undef UUInt16Property
+#endif
+
+#ifdef UUInt32Property
+	#undef UUInt32Property
+#endif
+
+#ifdef UUInt64Property
+	#undef UUInt64Property
+#endif
+
+#ifdef UFloatProperty
+	#undef UFloatProperty
+#endif
+
+#ifdef UDoubleProperty
+	#undef UDoubleProperty
+#endif
+
+#ifdef UBoolProperty
+	#undef UBoolProperty
+#endif
+
+#ifdef UObjectPropertyBase
+	#undef UObjectPropertyBase
+#endif
+
+#ifdef UObjectProperty
+	#undef UObjectProperty
+#endif
+
+#ifdef UWeakObjectProperty
+	#undef UWeakObjectProperty
+#endif
+
+#ifdef ULazyObjectProperty
+		#undef ULazyObjectProperty
+#endif
+
+#ifdef USoftObjectProperty
+	#undef USoftObjectProperty
+#endif
+
+#ifdef UClassProperty
+	#undef UClassProperty
+#endif
+
+#ifdef USoftClassProperty
+	#undef USoftClassProperty
+#endif
+
+#ifdef UInterfaceProperty
+	#undef UInterfaceProperty
+#endif
+
+#ifdef UNameProperty
+	#undef UNameProperty
+#endif
+
+#ifdef UStrProperty
+	#undef UStrProperty
+#endif
+
+#ifdef UArrayProperty
+	#undef UArrayProperty
+#endif
+
+#ifdef UMapProperty
+	#undef UMapProperty
+#endif
+
+#ifdef USetProperty
+	#undef USetProperty
+#endif
+
+#ifdef UStructProperty
+	#undef UStructProperty
+#endif
+
+#ifdef UMulticastDelegateProperty
+	#undef UMulticastDelegateProperty
+#endif
+
+#ifdef UMulticastInlineDelegateProperty
+	#undef UMulticastInlineDelegateProperty
+#endif
+
+#ifdef UMulticastSparseDelegateProperty
+	#undef UMulticastSparseDelegateProperty
+#endif
+
+#ifdef UEnumProperty
+	#undef UEnumProperty
+#endif
+
+#ifdef UTextProperty
+	#undef UTextProperty
+#endif
diff --git a/Engine/Source/Runtime/CoreUObject/Public/UObject/UnrealType.h b/Engine/Source/Runtime/CoreUObject/Public/UObject/UnrealType.h
--- a/Engine/Source/Runtime/CoreUObject/Public/UObject/UnrealType.h
+++ b/Engine/Source/Runtime/CoreUObject/Public/UObject/UnrealType.h
@@ -32,6 +32,9 @@
 #include "UObject/WeakObjectPtr.h"
 #include "UObject/Field.h"
 
+// WARNING: This should always be the last include in any file that needs it (except .generated.h)
+#include "UObject/UndefineUPropertyMacros.h"
+
 class UPropertyWrapper;
 
 COREUOBJECT_API DECLARE_LOG_CATEGORY_EXTERN(LogType, Log, All);
@@ -5713,4 +5716,7 @@ public:
 class COREUOBJECT_API UMulticastDelegatePropertyWrapper : public UPropertyWrapper
 {
 	DECLARE_CLASS_INTRINSIC(UMulticastDelegatePropertyWrapper, UPropertyWrapper, CLASS_Transient, TEXT("/Script/CoreUObject"));
-};
\ No newline at end of file
+};
+
+
+#include "UObject/DefineUPropertyMacros.h"
\ No newline at end of file
diff --git a/Engine/Source/Runtime/CoreUObject/Public/UObject/UnrealTypePrivate.h b/Engine/Source/Runtime/CoreUObject/Public/UObject/UnrealTypePrivate.h
--- a/Engine/Source/Runtime/CoreUObject/Public/UObject/UnrealTypePrivate.h
+++ b/Engine/Source/Runtime/CoreUObject/Public/UObject/UnrealTypePrivate.h
@@ -9,6 +9,9 @@
 #include "CoreMinimal.h"
 #include "UnrealType.h"
 
+// WARNING: This should always be the last include in any file that needs it (except .generated.h)
+#include "UObject/UndefineUPropertyMacros.h"
+
 #define USE_UPROPERTY_LOAD_DEFERRING (USE_CIRCULAR_DEPENDENCY_LOAD_DEFERRING && WITH_EDITORONLY_DATA)
 
 class COREUOBJECT_API UProperty : public UField
@@ -859,4 +862,7 @@ public:
 	: UProperty(ObjectInitializer, EC_CppProperty, InOffset, InFlags)
 	{
 	}
-};
\ No newline at end of file
+};
+
+// WARNING: This should always be the last include in any file that needs it (except .generated.h)
+#include "UObject/DefineUPropertyMacros.h"
\ No newline at end of file
diff --git a/Engine/Source/Runtime/Engine/Classes/Animation/AnimClassData.h b/Engine/Source/Runtime/Engine/Classes/Animation/AnimClassData.h
--- a/Engine/Source/Runtime/Engine/Classes/Animation/AnimClassData.h
+++ b/Engine/Source/Runtime/Engine/Classes/Animation/AnimClassData.h
@@ -12,7 +12,6 @@
 #include "AnimClassData.generated.h"
 
 class USkeleton;
-class UStructProperty;
 
 UCLASS()
 class ENGINE_API UAnimClassData : public UObject, public IAnimClassInterface
diff --git a/Engine/Source/Runtime/Engine/Classes/Animation/AnimNodeBase.h b/Engine/Source/Runtime/Engine/Classes/Animation/AnimNodeBase.h
--- a/Engine/Source/Runtime/Engine/Classes/Animation/AnimNodeBase.h
+++ b/Engine/Source/Runtime/Engine/Classes/Animation/AnimNodeBase.h
@@ -15,6 +15,9 @@
 #include "Animation/AnimTrace.h"
 #include "UObject/FieldPath.h"
 
+// WARNING: This should always be the last include in any file that needs it (except .generated.h)
+#include "UObject/UndefineUPropertyMacros.h"
+
 #include "AnimNodeBase.generated.h"
 
 #define DECLARE_SCOPE_HIERARCHICAL_COUNTER_ANIMNODE(Method) \
@@ -989,4 +992,7 @@ protected:
 private:		
 	// Reference to the exposed value handler used by this node. Allocated on the class, rather than per instance:
 	const FExposedValueHandler* ExposedValueHandler = nullptr;
-};
\ No newline at end of file
+};
+
+
+#include "UObject/DefineUPropertyMacros.h"
\ No newline at end of file
diff --git a/Engine/Source/Runtime/Engine/Classes/Engine/BlueprintGeneratedClass.h b/Engine/Source/Runtime/Engine/Classes/Engine/BlueprintGeneratedClass.h
--- a/Engine/Source/Runtime/Engine/Classes/Engine/BlueprintGeneratedClass.h
+++ b/Engine/Source/Runtime/Engine/Classes/Engine/BlueprintGeneratedClass.h
@@ -10,6 +10,10 @@
 #include "UObject/Class.h"
 #include "UObject/UnrealType.h"
 #include "EdGraph/EdGraphPin.h"
+
+// WARNING: This should always be the last include in any file that needs it (except .generated.h)
+#include "UObject/UndefineUPropertyMacros.h"
+
 #include "BlueprintGeneratedClass.generated.h"
 
 class AActor;
@@ -864,3 +868,5 @@ private:
 	/** In some cases UObject::ConditionalPostLoad() code calls PostLoadDefaultObject() on a class that's still being serialized. */
 	FCriticalSection SerializeAndPostLoadCritical;
 };
+
+#include "UObject/DefineUPropertyMacros.h"
\ No newline at end of file
diff --git a/Engine/Source/Runtime/Engine/Private/Animation/AnimData.cpp b/Engine/Source/Runtime/Engine/Private/Animation/AnimData.cpp
--- a/Engine/Source/Runtime/Engine/Private/Animation/AnimData.cpp
+++ b/Engine/Source/Runtime/Engine/Private/Animation/AnimData.cpp
@@ -5,4 +5,3 @@
 =============================================================================*/ 
 
 #include "Animation/AnimClassData.h"
-#include "UObject/UnrealTypePrivate.h"
diff --git a/Engine/Source/Runtime/Engine/Private/Animation/AnimNodeBase.cpp b/Engine/Source/Runtime/Engine/Private/Animation/AnimNodeBase.cpp
--- a/Engine/Source/Runtime/Engine/Private/Animation/AnimNodeBase.cpp
+++ b/Engine/Source/Runtime/Engine/Private/Animation/AnimNodeBase.cpp
@@ -5,7 +5,6 @@
 #include "Animation/AnimInstanceProxy.h"
 #include "Animation/AnimTrace.h"
 #include "UObject/CoreObjectVersion.h"
-#include "UObject/UnrealTypePrivate.h" // FProperty is declared here but use it only for upgrading to FProperties
 
 void* FExposedValueCopyRecord::GetDestAddr(FAnimInstanceProxy* Proxy, const FProperty* NodeProperty) const
 {
diff --git a/Engine/Source/Runtime/Engine/Private/BlueprintGeneratedClass.cpp b/Engine/Source/Runtime/Engine/Private/BlueprintGeneratedClass.cpp
--- a/Engine/Source/Runtime/Engine/Private/BlueprintGeneratedClass.cpp
+++ b/Engine/Source/Runtime/Engine/Private/BlueprintGeneratedClass.cpp
@@ -1640,8 +1640,6 @@ ENGINE_API int32 IncrementUberGraphSerialNumber()
 
 void UBlueprintGeneratedClass::Serialize(FArchive& Ar)
 {
-	//Ar.UsingCustomVersion(FCoreObjectVersion::GUID);
-
 #if VALIDATE_UBER_GRAPH_PERSISTENT_FRAME
 	if (Ar.IsLoading() && 0 == (Ar.GetPortFlags() & PPF_Duplicate))
 	{
@@ -1651,11 +1649,6 @@ void UBlueprintGeneratedClass::Serialize(FArchive& Ar)
 
 	Super::Serialize(Ar);
 
-	//if (Ar.CustomVer(FCoreObjectVersion::GUID) >= FCoreObjectVersion::FProperties)
-	//{
-	//	Ar << UberGraphFramePointerProperty;
-	//}
-
 	if (Ar.IsLoading() && 0 == (Ar.GetPortFlags() & PPF_Duplicate))
 	{
 		CreatePersistentUberGraphFrame(ClassDefaultObject, true);
diff --git a/Engine/Source/Runtime/Engine/Private/Interpolation.cpp b/Engine/Source/Runtime/Engine/Private/Interpolation.cpp
--- a/Engine/Source/Runtime/Engine/Private/Interpolation.cpp
+++ b/Engine/Source/Runtime/Engine/Private/Interpolation.cpp
@@ -110,7 +110,6 @@
 #include "Particles/ParticleSystemReplay.h"
 #include "GameFramework/GameState.h"
 #include "UObject/CoreObjectVersion.h"
-#include "UObject/UnrealTypePrivate.h" // Only for FProperty to FProperty conversion
 
 #if WITH_EDITOR
 #include "Sound/SoundCue.h"
