BulkData2 can now work if EDL is disabled.

- Remove the hardcoded assert that forced GEventDrivenLoaderEnabled to be true
- ConvertFilenameFromFlags should only modify the filename if the BULKDATA_PayloadInSeperateFile flag is set

#rnx

diff --git a/Engine/Source/Runtime/CoreUObject/Private/Serialization/BulkData2.cpp b/Engine/Source/Runtime/CoreUObject/Private/Serialization/BulkData2.cpp
--- a/Engine/Source/Runtime/CoreUObject/Private/Serialization/BulkData2.cpp
+++ b/Engine/Source/Runtime/CoreUObject/Private/Serialization/BulkData2.cpp
@@ -525,7 +525,6 @@ void FBulkDataBase::Serialize(FArchive& Ar, UObject* Owner, int32 /*Index*/, boo
 #if WITH_EDITOR == 0 && WITH_EDITORONLY_DATA == 0
 	check(Ar.IsLoading());				// Only support loading from cooked data!
 	check(!GIsEditor);					// The editor path is not supported
-	check(GEventDrivenLoaderEnabled);	// We are assuming the EDL path is enabled ( TODO: Might need to remove this check)
 	check(LockStatus == LOCKSTATUS_Unlocked);
 
 	if (Ar.IsPersistent() && !Ar.IsObjectReferenceCollector() && !Ar.ShouldSkipBulkData())
@@ -1176,7 +1175,11 @@ void FBulkDataBase::FreeData()
 
 FString FBulkDataBase::ConvertFilenameFromFlags(const FString& Filename) const
 {
-	if (IsOptional())
+	if (!InSeperateFile())
+	{
+		return Filename;
+	}
+	else if (IsOptional())
 	{
 		// Optional data should be tested for first as we in theory can have data that would
 		// be marked as inline, also marked as optional and in this case we should treat it as
