AsyncLoading2: Fix crash during logging of warning message about missing import script package.
Add back support for loading a different name than the actual package name (e.g. used by level streaming and localization)

diff --git a/Engine/Source/Runtime/CoreUObject/Private/Serialization/AsyncLoading2.cpp b/Engine/Source/Runtime/CoreUObject/Private/Serialization/AsyncLoading2.cpp
--- a/Engine/Source/Runtime/CoreUObject/Private/Serialization/AsyncLoading2.cpp
+++ b/Engine/Source/Runtime/CoreUObject/Private/Serialization/AsyncLoading2.cpp
@@ -1023,8 +1023,9 @@ uint32 FAsyncLoadingThreadState2::TlsSlot;
 struct FAsyncPackageDesc2
 {
 	int32 RequestID;
-	FPackageId PackageId;
+	FPackageId PackageId; // for NameToLoad
 	FName Name;
+	FName NameToLoad;
 	/** Delegate called on completion of loading. This delegate can only be created and consumed on the game thread */
 	TUniquePtr<FLoadPackageAsyncDelegate> PackageLoadedDelegate;
 
@@ -1032,10 +1033,12 @@ struct FAsyncPackageDesc2
 		int32 InRequestID,
 		FPackageId InPackageId,
 		const FName& InName,
+		const FName& InNameToLoad,
 		TUniquePtr<FLoadPackageAsyncDelegate>&& InCompletionDelegate = TUniquePtr<FLoadPackageAsyncDelegate>())
 		: RequestID(InRequestID)
 		, PackageId(InPackageId)
 		, Name(InName)
+		, NameToLoad(InNameToLoad)
 		, PackageLoadedDelegate(MoveTemp(InCompletionDelegate))
 	{
 	}
@@ -1045,6 +1048,7 @@ struct FAsyncPackageDesc2
 		: RequestID(OldPackage.RequestID)
 		, PackageId(OldPackage.PackageId)
 		, Name(OldPackage.Name)
+		, NameToLoad(OldPackage.NameToLoad)
 	{
 	}
 
@@ -2419,9 +2423,21 @@ void FGlobalImportStore::FindAllScriptImports()
 	{
 		UObject* Object = GFindExistingScriptImport(GlobalImportIndex, Imports, ScriptImportOuters, Names);
 #if DO_CHECK
-		UE_CLOG(!Object, LogStreaming, Warning, TEXT("Failed to find script import after initial load: %s in %s"),
-			*Names[GlobalImportIndex].ToString(),
-			*Names[ScriptImportOuters[GlobalImportIndex].ToImport()].ToString());
+		if (!Object)
+		{
+			FPackageIndex Outer = ScriptImportOuters[GlobalImportIndex];
+			if (Outer.IsNull())
+			{
+				UE_LOG(LogStreaming, Warning, TEXT("Failed to find import script package after initial load: %s"),
+					*Names[GlobalImportIndex].ToString());
+			}
+			else
+			{
+				UE_LOG(LogStreaming, Warning, TEXT("Failed to find import script object after initial load: %s - %s"),
+					*Names[Outer.ToImport()].ToString(),
+					*Names[GlobalImportIndex].ToString());
+			}
+		}
 #endif
 	}
 	ScriptImportOuters = nullptr;
@@ -2463,7 +2479,7 @@ void FAsyncPackage2::ImportPackagesRecursive()
 
 		if (bNeedToLoadPackage)
 		{
-			FAsyncPackageDesc2 Info(INDEX_NONE, FPackageId::FromIndex(GlobalPackageIndex), Entry.Name);
+			FAsyncPackageDesc2 Info(INDEX_NONE, FPackageId::FromIndex(GlobalPackageIndex), Entry.Name, Entry.Name);
 			bool bInserted;
 			FAsyncPackage2* ImportedAsyncPackage = AsyncLoadingThread.FindOrInsertPackage(&Info, bInserted);
 			if (ImportedAsyncPackage)
@@ -2720,8 +2736,8 @@ void FAsyncPackage2::EventDrivenCreateExport(int32 LocalExportIndex)
 	UClass* LoadClass = Export.ClassIndex.IsNull() ? UClass::StaticClass() : CastEventDrivenIndexToObject<UClass>(Export.ClassIndex, true);
 	UObject* ThisParent = Export.OuterIndex.IsNull() ? LinkerRoot : EventDrivenIndexToObject(Export.OuterIndex, false);
 
-	checkf(LoadClass, TEXT("Could not find class object for %s in %s"), *ObjectName.ToString(), *Desc.Name.ToString());
-	checkf(ThisParent, TEXT("Could not find outer object for %s in %s"), *ObjectName.ToString(), *Desc.Name.ToString());
+	checkf(LoadClass, TEXT("Could not find class object for %s in %s"), *ObjectName.ToString(), *Desc.NameToLoad.ToString());
+	checkf(ThisParent, TEXT("Could not find outer object for %s in %s"), *ObjectName.ToString(), *Desc.NameToLoad.ToString());
 	check(!dynamic_cast<UObjectRedirector*>(ThisParent));
 
 	// Try to find existing object first as we cannot in-place replace objects, could have been created by other export in this package
@@ -2764,7 +2780,7 @@ void FAsyncPackage2::EventDrivenCreateExport(int32 LocalExportIndex)
 		// Find the Archetype object for the one we are loading.
 		check(!Export.TemplateIndex.IsNull());
 		UObject* Template = EventDrivenIndexToObject(Export.TemplateIndex, true);
-		checkf(Template, TEXT("Could not find template for %s in %s"), *ObjectName.ToString(), *Desc.Name.ToString());
+		checkf(Template, TEXT("Could not find template for %s in %s"), *ObjectName.ToString(), *Desc.NameToLoad.ToString());
 		// we also need to ensure that the template has set up any instances
 		Template->ConditionalPostLoadSubobjects();
 
@@ -3917,7 +3933,7 @@ FAsyncPackage2::FAsyncPackage2(
 , bAllExportsSerialized(false)
 {
 	TRACE_CPUPROFILER_EVENT_SCOPE(NewAsyncPackage);
-	TRACE_LOADTIME_NEW_ASYNC_PACKAGE(this, InDesc.Name);
+	TRACE_LOADTIME_NEW_ASYNC_PACKAGE(this, InDesc.NameToLoad);
 	AddRequestID(InDesc.RequestID);
 
 	TTuple<FExportBundleMetaEntry*, uint32> PackageExportBundles = AsyncLoadingThread.GlobalPackageStore.GetPackageExportBundleMetaEntries(Desc.PackageId);
@@ -4441,8 +4457,10 @@ int32 FAsyncLoadingThread2Impl::LoadPackage(const FString& InName, const FGuid*
 
 	int32 RequestID = INDEX_NONE;
 
-	FName PackageName(*InName);
-	FPackageId* PackageId = GlobalPackageStore.GetPackageId(PackageName);
+	FName PackageName = FName(*InName);
+	FName PackageNameToLoad = InPackageToLoadFrom ? FName(InPackageToLoadFrom) : PackageName;
+
+	FPackageId* PackageId = GlobalPackageStore.GetPackageId(PackageNameToLoad);
 
 	if (PackageId)
 	{
@@ -4465,14 +4483,13 @@ int32 FAsyncLoadingThread2Impl::LoadPackage(const FString& InName, const FGuid*
 		}
 
 		// Add new package request
-		FAsyncPackageDesc2 PackageDesc(RequestID, *PackageId, PackageName, MoveTemp(CompletionDelegatePtr));
+		FAsyncPackageDesc2 PackageDesc(RequestID, *PackageId, PackageName, PackageNameToLoad, MoveTemp(CompletionDelegatePtr));
 		QueuePackage(PackageDesc);
 	}
 	else
 	{
-		const TCHAR* NameToLoadStr = InPackageToLoadFrom ? InPackageToLoadFrom : TEXT("null");
-		UE_LOG(LogStreaming, Warning, TEXT("FAsyncLoadingThread2Impl::LoadPackage - Skipping unknown package: %s (%s)"),
-			*InName, NameToLoadStr);
+		UE_LOG(LogStreaming, Warning, TEXT("FAsyncLoadingThread2Impl::LoadPackage - Skipping package '%s'. Name to load is unknown: '%s')"),
+			*PackageName.ToString(), *PackageNameToLoad.ToString());
 		InCompletionDelegate.ExecuteIfBound(PackageName, nullptr, EAsyncLoadingResult::Failed);
 	}
 
