Fixed an issue introduced in 10729998 where a 0 sized optional bulkdata object was being treated as inlined and would fail it's 0 byte read instead of not being read at all because it's file does not exist.

First we check for optional and then if it is in a seperate file. If EDL is disabled then optional output wont be cooked anyway.

#rnx

diff --git a/Engine/Source/Runtime/CoreUObject/Private/Serialization/BulkData2.cpp b/Engine/Source/Runtime/CoreUObject/Private/Serialization/BulkData2.cpp
--- a/Engine/Source/Runtime/CoreUObject/Private/Serialization/BulkData2.cpp
+++ b/Engine/Source/Runtime/CoreUObject/Private/Serialization/BulkData2.cpp
@@ -1174,18 +1174,18 @@ void FBulkDataBase::FreeData()
 }
 
 FString FBulkDataBase::ConvertFilenameFromFlags(const FString& Filename) const
-{
-	if (!IsInSeperateFile())
-	{
-		return Filename;
-	}
-	else if (IsOptional())
+{	
+	if (IsOptional())
 	{
 		// Optional data should be tested for first as we in theory can have data that would
 		// be marked as inline, also marked as optional and in this case we should treat it as
 		// optional data first.
 		return FPaths::ChangeExtension(Filename, OptionalExt);
 	}
+	else if (!IsInSeperateFile())
+	{
+		return Filename;
+	}
 	else if (IsInlined())
 	{
 		return FPaths::ChangeExtension(Filename, InlinedExt);
