Fix one hour cook time increase from ResetLoadersForSave change; the cooker should not be calling ResetLoaders because it is saving to a different filename.
#rnx

diff --git a/Engine/Source/Runtime/CoreUObject/Private/UObject/Linker.cpp b/Engine/Source/Runtime/CoreUObject/Private/UObject/Linker.cpp
--- a/Engine/Source/Runtime/CoreUObject/Private/UObject/Linker.cpp
+++ b/Engine/Source/Runtime/CoreUObject/Private/UObject/Linker.cpp
@@ -808,18 +808,15 @@ FLinkerLoad* LoadPackageLinker(UPackage* InOuter, const TCHAR* InLongPackageName
 	return LoadPackageLinker(InOuter, InLongPackageName, LoadFlags, Sandbox, CompatibleGuid, InReaderOverride, [](FLinkerLoad* InLinker) {});
 }
 
-/**
- * 
- * Ensure thumbnails are loaded and then reset the loader in preparation for a package save
- *
- * @param	InOuter			The outer for the package we are saving
- * @param	Filename		The filename we are saving too
- */
-void ResetLoadersForSave(UObject* InOuter, const TCHAR *Filename)
+
+void ResetLoadersForSave(UObject* InOuter, const TCHAR* Filename)
 {
 	UPackage* Package = dynamic_cast<UPackage*>(InOuter);
-	EnsureLoadingComplete(Package);
+	ResetLoadersForSave(Package, Filename);
+}
 
+void ResetLoadersForSave(UPackage* Package, const TCHAR* Filename)
+{
 	FLinkerLoad* Loader = FLinkerLoad::FindExistingLinkerForPackage(Package);
 	if( Loader )
 	{
diff --git a/Engine/Source/Runtime/CoreUObject/Private/UObject/SavePackage.cpp b/Engine/Source/Runtime/CoreUObject/Private/UObject/SavePackage.cpp
--- a/Engine/Source/Runtime/CoreUObject/Private/UObject/SavePackage.cpp
+++ b/Engine/Source/Runtime/CoreUObject/Private/UObject/SavePackage.cpp
@@ -5704,9 +5704,9 @@ FSavePackageResultStruct UPackage::Save(UPackage* InOuter, UObject* Base, EObjec
 				if( Success == true )
 				{
 					{
-						// ResetLoaders on the Package so that we drop the handle to the file on disk and can write to it
+						// If we're writing to the existing file call ResetLoaders on the Package so that we drop the handle to the file on disk and can write to it
 						COOK_STAT(FScopedDurationTimer SaveTimer(SavePackageStats::ResetLoadersTimeSec));
-						ResetLoaders(InOuter);
+						ResetLoadersForSave(InOuter, Filename);
 					}
 
 					if (bSaveAsync)
diff --git a/Engine/Source/Runtime/CoreUObject/Public/UObject/Linker.h b/Engine/Source/Runtime/CoreUObject/Public/UObject/Linker.h
--- a/Engine/Source/Runtime/CoreUObject/Public/UObject/Linker.h
+++ b/Engine/Source/Runtime/CoreUObject/Public/UObject/Linker.h
@@ -597,15 +597,17 @@ COREUOBJECT_API FLinkerLoad* GetPackageLinker(UPackage* InOuter, const TCHAR* In
 COREUOBJECT_API FString GetPrestreamPackageLinkerName(const TCHAR* InLongPackageName, bool bExistSkip = true);
 
 
+UE_DEPRECATED(4.25, "No longer used; use version that takes a UPackage* and call EnsureLoadingComplete separately.")
+COREUOBJECT_API void ResetLoadersForSave(UObject* InOuter, const TCHAR *Filename);
+
 /**
- * 
- * Ensure thumbnails are loaded and then reset the loader in preparation for a package save
  *
- * @param	InOuter			The outer for the package we are saving
+ * Reset the loader for the given package if it is using the given filename, so we can write to the file
+ *
+ * @param	Package			The package we are saving
  * @param	Filename		The filename we are saving too
  */
-UE_DEPRECATED(4.25, "No longer used; use EnsureLoadingComplete followed by ResetLoaders")
-COREUOBJECT_API void ResetLoadersForSave(UObject* InOuter, const TCHAR *Filename);
+COREUOBJECT_API void ResetLoadersForSave(UPackage* Package, const TCHAR* Filename);
 
 /*
  * Ensure all data that can be loaded from the linker (thumbnails, bulk data) is loaded, in preparation for saving out the given package
