AsyncLoading2: Fix broken NameToLoad support by adding new PackageIds for /Temp/ project names, so that the same disc package can be loaded into multiple runtime packages again.
Add support for running without the async loading thread (for the server) and respect the s.AsyncLoadingThreadEnabled setting
Move FAsyncLoadingThreadSettings to AsyncPackageLoader to share it between AsyncLoading and AsyncLoading2.
Add some user friendly logging on startup.
Respect the s.AsyncPostLoadEnabled setting.
There is no proper timeslicing when running on the game thread.
Requires DEFAULT_NO_THREADING=0 for IoDispatcher.

diff --git a/Engine/Source/Runtime/CoreUObject/Private/Serialization/AsyncLoading.cpp b/Engine/Source/Runtime/CoreUObject/Private/Serialization/AsyncLoading.cpp
--- a/Engine/Source/Runtime/CoreUObject/Private/Serialization/AsyncLoading.cpp
+++ b/Engine/Source/Runtime/CoreUObject/Private/Serialization/AsyncLoading.cpp
@@ -4632,7 +4632,7 @@ FAsyncLoadingThread::FAsyncLoadingThread(int32 InThreadIndex, IEDLBootNotificati
 	CancelLoadingEvent = FPlatformProcess::GetSynchEventFromPool();
 	ThreadSuspendedEvent = FPlatformProcess::GetSynchEventFromPool();
 	ThreadResumedEvent = FPlatformProcess::GetSynchEventFromPool();
-	if ((!GEventDrivenLoaderEnabled || !USE_EVENT_DRIVEN_ASYNC_LOAD_AT_BOOT_TIME) && FAsyncLoadingThread::GetAsyncLoadingThreadSettings().bAsyncLoadingThreadEnabled)
+	if ((!GEventDrivenLoaderEnabled || !USE_EVENT_DRIVEN_ASYNC_LOAD_AT_BOOT_TIME) && FAsyncLoadingThreadSettings::Get().bAsyncLoadingThreadEnabled)
 	{
 		StartThread();
 	}
@@ -4640,8 +4640,8 @@ FAsyncLoadingThread::FAsyncLoadingThread(int32 InThreadIndex, IEDLBootNotificati
 #if !IS_PROGRAM && !WITH_EDITORONLY_DATA
 	UE_LOG(LogStreaming, Display, TEXT("Async Loading initialized: Event Driven Loader: %s, Async Loading Thread: %s, Async Post Load: %s"),
 		GEventDrivenLoaderEnabled ? TEXT("true") : TEXT("false"),
-		FAsyncLoadingThread::GetAsyncLoadingThreadSettings().bAsyncLoadingThreadEnabled ? TEXT("true") : TEXT("false"),
-		FAsyncLoadingThread::GetAsyncLoadingThreadSettings().bAsyncPostLoadEnabled ? TEXT("true") : TEXT("false"));
+		FAsyncLoadingThreadSettings::Get().bAsyncLoadingThreadEnabled ? TEXT("true") : TEXT("false"),
+		FAsyncLoadingThreadSettings::Get().bAsyncPostLoadEnabled ? TEXT("true") : TEXT("false"));
 #endif
 }
 
@@ -4677,45 +4677,12 @@ void FAsyncLoadingThread::ShutdownLoading()
 	ThreadResumedEvent = nullptr;	
 }
 
-FAsyncLoadingThread::FAsyncLoadingThreadSettings::FAsyncLoadingThreadSettings()
-{
-#if THREADSAFE_UOBJECTS
-	if (FPlatformProperties::RequiresCookedData())
-	{
-		check(GConfig);
-
-		bool bConfigValue = true;
-		GConfig->GetBool(TEXT("/Script/Engine.StreamingSettings"), TEXT("s.AsyncLoadingThreadEnabled"), bConfigValue, GEngineIni);
-		bool bCommandLineDisable = FParse::Param(FCommandLine::Get(), TEXT("NoAsyncLoadingThread"));
-		bool bCommandLineEnable = FParse::Param(FCommandLine::Get(), TEXT("AsyncLoadingThread"));
-		bAsyncLoadingThreadEnabled = bCommandLineEnable || (bConfigValue && FApp::ShouldUseThreadingForPerformance() && !bCommandLineDisable);
-
-		bConfigValue = true;
-		GConfig->GetBool(TEXT("/Script/Engine.StreamingSettings"), TEXT("s.AsyncPostLoadEnabled"), bConfigValue, GEngineIni);
-		bCommandLineDisable = FParse::Param(FCommandLine::Get(), TEXT("NoAsyncPostLoad"));
-		bCommandLineEnable = FParse::Param(FCommandLine::Get(), TEXT("AsyncPostLoad"));
-		bAsyncPostLoadEnabled = bCommandLineEnable || (bConfigValue && FApp::ShouldUseThreadingForPerformance() && !bCommandLineDisable);
-	}
-	else
-#endif
-	{
-		bAsyncLoadingThreadEnabled = false;
-		bAsyncPostLoadEnabled = false;
-	}
-}
-
-FAsyncLoadingThread::FAsyncLoadingThreadSettings& FAsyncLoadingThread::GetAsyncLoadingThreadSettings()
-{
-	static FAsyncLoadingThreadSettings Settings;
-	return Settings;
-}
-
 void FAsyncLoadingThread::StartThread()
 {
 	// Make sure the GC sync object is created before we start the thread (apparently this can happen before we call InitUObject())
 	FGCCSyncObject::Create();
 
-	if (!Thread && FAsyncLoadingThread::GetAsyncLoadingThreadSettings().bAsyncLoadingThreadEnabled)
+	if (!Thread && FAsyncLoadingThreadSettings::Get().bAsyncLoadingThreadEnabled)
 	{
 		UE_LOG(LogStreaming, Log, TEXT("Starting Async Loading Thread."));
 		bThreadStarted = true;
@@ -6328,7 +6295,7 @@ EAsyncPackageState::Type FAsyncPackage::PostLoadObjects()
 		PreLoadIndex = PackageObjLoaded.Num(); // we did preloading in a different way and never incremented this
 	}
 
-	const bool bAsyncPostLoadEnabled = FAsyncLoadingThread::GetAsyncLoadingThreadSettings().bAsyncPostLoadEnabled;
+	const bool bAsyncPostLoadEnabled = FAsyncLoadingThreadSettings::Get().bAsyncPostLoadEnabled;
 	const bool bIsMultithreaded = AsyncLoadingThread.IsMultithreaded();
 
 	// PostLoad objects.
diff --git a/Engine/Source/Runtime/CoreUObject/Private/Serialization/AsyncLoading2.cpp b/Engine/Source/Runtime/CoreUObject/Private/Serialization/AsyncLoading2.cpp
--- a/Engine/Source/Runtime/CoreUObject/Private/Serialization/AsyncLoading2.cpp
+++ b/Engine/Source/Runtime/CoreUObject/Private/Serialization/AsyncLoading2.cpp
@@ -332,7 +332,9 @@ struct FGlobalImportStore
 class FPackageStore
 {
 public:
+	FCriticalSection PackageNameToPackageIdCritical;
 	TMap<FName, FPackageId> PackageNameToPackageId;
+
 	FGlobalImportStore ImportStore;
 	FIoBuffer InitialLoadIoBuffer;
 	FPackageStoreEntry* StoreEntries = nullptr;
@@ -391,7 +393,8 @@ public:
 			// Global package index
 			{
 				TRACE_CPUPROFILER_EVENT_SCOPE(LoadPackageStoreGlobalIds);
-				PackageNameToPackageId.Reserve(PackageCount);
+				// add 10% slack for temp package names
+				PackageNameToPackageId.Reserve(PackageCount + PackageCount / 10);
 				for (int I = 0; I < PackageCount; ++I)
 				{
 					PackageNameToPackageId.Add(StoreEntries[I].Name, FPackageId::FromIndex(I));
@@ -488,6 +491,9 @@ public:
 		InitialLoadStoreEntries = nullptr;
 		ScriptArcs = nullptr;
 		InitialLoadIoBuffer = FIoBuffer();
+
+		UE_LOG(LogStreaming, Display, TEXT("AsyncLoading2 - InitialLoad Finalized: Script Imports: %d"),
+			ImportStore.ScriptImportCount);
 	}
 
 	inline FGlobalImportStore& GetGlobalImportStore()
@@ -495,10 +501,24 @@ public:
 		return ImportStore;
 	}
 
-	inline FPackageId* GetPackageId(FName Name)
+	inline FPackageId FindPackageId(FName Name)
 	{
+		FScopeLock Lock(&PackageNameToPackageIdCritical);
 		FPackageId* Id = PackageNameToPackageId.Find(Name);
-		return Id;
+		return Id ? *Id : FPackageId();
+	}
+
+	inline FPackageId FindOrAddPackageId(FName Name)
+	{
+		FScopeLock Lock(&PackageNameToPackageIdCritical);
+		if (FPackageId* Id = PackageNameToPackageId.Find(Name))
+		{
+			return *Id;
+		}
+
+		FPackageId NewId = FPackageId::FromIndex(PackageNameToPackageId.Num());
+		PackageNameToPackageId.Add(Name, NewId);
+		return NewId;
 	}
 
 	inline FPackageIndex* GetGlobalImportScriptOuters(int32& OutCount) const
@@ -1023,7 +1043,8 @@ uint32 FAsyncLoadingThreadState2::TlsSlot;
 struct FAsyncPackageDesc2
 {
 	int32 RequestID;
-	FPackageId PackageId; // for NameToLoad
+	FPackageId PackageId;
+	FPackageId PackageIdToLoad;
 	FName Name;
 	FName NameToLoad;
 	/** Delegate called on completion of loading. This delegate can only be created and consumed on the game thread */
@@ -1032,11 +1053,13 @@ struct FAsyncPackageDesc2
 	FAsyncPackageDesc2(
 		int32 InRequestID,
 		FPackageId InPackageId,
+		FPackageId InPackageIdToLoad,
 		const FName& InName,
 		const FName& InNameToLoad,
 		TUniquePtr<FLoadPackageAsyncDelegate>&& InCompletionDelegate = TUniquePtr<FLoadPackageAsyncDelegate>())
 		: RequestID(InRequestID)
 		, PackageId(InPackageId)
+		, PackageIdToLoad(InPackageIdToLoad)
 		, Name(InName)
 		, NameToLoad(InNameToLoad)
 		, PackageLoadedDelegate(MoveTemp(InCompletionDelegate))
@@ -1047,6 +1070,7 @@ struct FAsyncPackageDesc2
 	FAsyncPackageDesc2(const FAsyncPackageDesc2& OldPackage)
 		: RequestID(OldPackage.RequestID)
 		, PackageId(OldPackage.PackageId)
+		, PackageIdToLoad(OldPackage.PackageIdToLoad)
 		, Name(OldPackage.Name)
 		, NameToLoad(OldPackage.NameToLoad)
 	{
@@ -1665,12 +1689,12 @@ public:
 	FORCEINLINE FAsyncPackage2* FindAsyncPackage(const FName& PackageName)
 	{
 		TRACE_CPUPROFILER_EVENT_SCOPE(FindAsyncPackage);
-		FPackageId* PackageId = GlobalPackageStore.GetPackageId(PackageName);
-		if (PackageId)
+		FPackageId PackageId = GlobalPackageStore.FindPackageId(PackageName);
+		if (PackageId.IsValid())
 		{
 			FScopeLock LockAsyncPackages(&AsyncPackagesCritical);
 			//checkSlow(IsInAsyncLoadThread());
-			return AsyncPackageLookup.FindRef(*PackageId);
+			return AsyncPackageLookup.FindRef(PackageId);
 		}
 		return nullptr;
 	}
@@ -1899,6 +1923,11 @@ void FAsyncLoadingThread2Impl::InitializeLoading()
 	}
 
 	AsyncThreadReady.Increment();
+
+	UE_LOG(LogStreaming, Display, TEXT("AsyncLoading2 - Initialized: Packages: %d, GlobalImports: %d, FNames: %d"),
+		GlobalPackageStore.PackageCount,
+		GlobalPackageStore.ImportStore.Count,
+		GlobalNameMap.GetNameEntries().Num());
 }
 
 void FAsyncLoadingThread2Impl::QueuePackage(FAsyncPackageDesc2& Package)
@@ -2014,7 +2043,7 @@ void FAsyncLoadingThread2Impl::StartBundleIoRequests()
 		PreviousPackage = Package;
 
 		FIoReadOptions ReadOptions;
-		IoDispatcher.ReadWithCallback(CreateIoChunkId(Package->Desc.PackageId.ToIndex(), 0, EIoChunkType::ExportBundleData),
+		IoDispatcher.ReadWithCallback(CreateIoChunkId(Package->Desc.PackageIdToLoad.ToIndex(), 0, EIoChunkType::ExportBundleData),
 			ReadOptions,
 			[Package](TIoStatusOr<FIoBuffer> Result)
 		{
@@ -2428,12 +2457,12 @@ void FGlobalImportStore::FindAllScriptImports()
 			FPackageIndex Outer = ScriptImportOuters[GlobalImportIndex];
 			if (Outer.IsNull())
 			{
-				UE_LOG(LogStreaming, Warning, TEXT("Failed to find import script package after initial load: %s"),
+				UE_LOG(LogStreaming, Warning, TEXT("AsyncLoading2 - Failed to find import script package after initial load: %s"),
 					*Names[GlobalImportIndex].ToString());
 			}
 			else
 			{
-				UE_LOG(LogStreaming, Warning, TEXT("Failed to find import script object after initial load: %s - %s"),
+				UE_LOG(LogStreaming, Warning, TEXT("AsyncLoading2 - Failed to find import script object after initial load: %s - %s"),
 					*Names[Outer.ToImport()].ToString(),
 					*Names[GlobalImportIndex].ToString());
 			}
@@ -2453,7 +2482,7 @@ void FAsyncPackage2::ImportPackagesRecursive()
 
 	FPackageStore& GlobalPackageStore = AsyncLoadingThread.GlobalPackageStore;
 	int32 ImportedPackageCount = 0;
-	int32* Imports = GlobalPackageStore.GetPackageImportedPackages(Desc.PackageId, ImportedPackageCount);
+	int32* Imports = GlobalPackageStore.GetPackageImportedPackages(Desc.PackageIdToLoad, ImportedPackageCount);
 	if (!Imports)
 	{
 		return;
@@ -2479,7 +2508,8 @@ void FAsyncPackage2::ImportPackagesRecursive()
 
 		if (bNeedToLoadPackage)
 		{
-			FAsyncPackageDesc2 Info(INDEX_NONE, FPackageId::FromIndex(GlobalPackageIndex), Entry.Name, Entry.Name);
+			FPackageId PackageId = FPackageId::FromIndex(GlobalPackageIndex);
+			FAsyncPackageDesc2 Info(INDEX_NONE, PackageId, PackageId, Entry.Name, Entry.Name);
 			bool bInserted;
 			FAsyncPackage2* ImportedAsyncPackage = AsyncLoadingThread.FindOrInsertPackage(&Info, bInserted);
 			if (ImportedAsyncPackage)
@@ -2518,7 +2548,7 @@ void FAsyncPackage2::SetupScriptArcs()
 
 	FPackageStore& GlobalPackageStore = AsyncLoadingThread.GlobalPackageStore;
 	int32 ScriptArcsCount = 0;
-	int32* ScriptArcs = GlobalPackageStore.GetPackageScriptArcs(Desc.PackageId, ScriptArcsCount);
+	int32* ScriptArcs = GlobalPackageStore.GetPackageScriptArcs(Desc.PackageIdToLoad, ScriptArcsCount);
 	for (int32 ScriptArcIndex = 0; ScriptArcIndex < ScriptArcsCount * 2;)
 	{
 		int32 GlobalImportIndex = ScriptArcs[ScriptArcIndex++];
@@ -3460,6 +3490,11 @@ FAsyncLoadingThread2Impl::FAsyncLoadingThread2Impl(FIoDispatcher& InIoDispatcher
 
 	FAsyncLoadingThreadState2::TlsSlot = FPlatformTLS::AllocTlsSlot();
 	FAsyncLoadingThreadState2::Create(GraphAllocator, IoDispatcher);
+
+	UE_LOG(LogStreaming, Display, TEXT("AsyncLoading2 - Created: Event Driven Loader: %s, Async Loading Thread: %s, Async Post Load: %s"),
+		GEventDrivenLoaderEnabled ? TEXT("true") : TEXT("false"),
+		FAsyncLoadingThreadSettings::Get().bAsyncLoadingThreadEnabled ? TEXT("true") : TEXT("false"),
+		FAsyncLoadingThreadSettings::Get().bAsyncPostLoadEnabled ? TEXT("true") : TEXT("false"));
 }
 
 FAsyncLoadingThread2Impl::~FAsyncLoadingThread2Impl()
@@ -3491,12 +3526,14 @@ void FAsyncLoadingThread2Impl::ShutdownLoading()
 
 void FAsyncLoadingThread2Impl::StartThread()
 {
-
-
 	// Make sure the GC sync object is created before we start the thread (apparently this can happen before we call InitUObject())
 	FGCCSyncObject::Create();
 
-	if (!Thread)
+	if (!FAsyncLoadingThreadSettings::Get().bAsyncLoadingThreadEnabled)
+	{
+		GlobalPackageStore.FinalizeInitialLoad();
+	}
+	else if (!Thread)
 	{
 		UE_LOG(LogStreaming, Log, TEXT("Starting Async Loading Thread."));
 		bThreadStarted = true;
@@ -3536,6 +3573,10 @@ void FAsyncLoadingThread2Impl::StartThread()
 			TRACE_SET_THREAD_GROUP(Thread->GetThreadID(), "AsyncLoading");
 		}
 	}
+
+	UE_LOG(LogStreaming, Display, TEXT("AsyncLoading2 - Thread Started: %s, IsInitialLoad: %s"),
+		FAsyncLoadingThreadSettings::Get().bAsyncLoadingThreadEnabled ? TEXT("true") : TEXT("false"),
+		GIsInitialLoad ? TEXT("true") : TEXT("false"));
 }
 
 bool FAsyncLoadingThread2Impl::Init()
@@ -3928,7 +3969,7 @@ FAsyncPackage2::FAsyncPackage2(
 , EDLBootNotificationManager(InEDLBootNotificationManager)
 , GraphAllocator(InGraphAllocator)
 // Begin EDL specific properties
-, ImportStore(AsyncLoadingThread.GlobalPackageStore, Desc.PackageId)
+, ImportStore(AsyncLoadingThread.GlobalPackageStore, Desc.PackageIdToLoad)
 , AsyncPackageLoadingState(EAsyncPackageLoadingState2::NewPackage)
 , bAllExportsSerialized(false)
 {
@@ -3936,10 +3977,10 @@ FAsyncPackage2::FAsyncPackage2(
 	TRACE_LOADTIME_NEW_ASYNC_PACKAGE(this, InDesc.NameToLoad);
 	AddRequestID(InDesc.RequestID);
 
-	TTuple<FExportBundleMetaEntry*, uint32> PackageExportBundles = AsyncLoadingThread.GlobalPackageStore.GetPackageExportBundleMetaEntries(Desc.PackageId);
+	TTuple<FExportBundleMetaEntry*, uint32> PackageExportBundles = AsyncLoadingThread.GlobalPackageStore.GetPackageExportBundleMetaEntries(Desc.PackageIdToLoad);
 	ExportBundleMetaEntries = PackageExportBundles.Get<0>();
 	ExportBundleCount = PackageExportBundles.Get<1>();
-	ExportCount = AsyncLoadingThread.GlobalPackageStore.GetPackageExportCount(Desc.PackageId);
+	ExportCount = AsyncLoadingThread.GlobalPackageStore.GetPackageExportCount(Desc.PackageIdToLoad);
 	Exports.AddZeroed(ExportCount);
 	OwnedObjects.Reserve(ExportCount + 1); // +1 for UPackage
 
@@ -4138,7 +4179,7 @@ void FAsyncPackage2::CreateUPackage(const FPackageSummary* PackageSummary)
 		TRACE_CPUPROFILER_EVENT_SCOPE(UPackageCreate);
 		LinkerRoot = NewObject<UPackage>(/*Outer*/nullptr, Desc.Name, RF_Public);
 		// LinkerRoot->FileName = Desc.Name;
-		LinkerRoot->SetPackageId(Desc.PackageId);
+		LinkerRoot->SetPackageId(Desc.PackageIdToLoad);
 		LinkerRoot->SetPackageFlagsTo(PackageSummary->PackageFlags);
 		LinkerRoot->LinkerPackageVersion = GPackageFileUE4Version;
 		LinkerRoot->LinkerLicenseeVersion = GPackageFileLicenseeUE4Version;
@@ -4148,7 +4189,7 @@ void FAsyncPackage2::CreateUPackage(const FPackageSummary* PackageSummary)
 	}
 	else
 	{
-		check(LinkerRoot->GetPackageId() == Desc.PackageId);
+		check(LinkerRoot->GetPackageId() == Desc.PackageIdToLoad);
 		check(LinkerRoot->GetPackageFlags() == PackageSummary->PackageFlags);
 		check(LinkerRoot->LinkerPackageVersion == GPackageFileUE4Version);
 		check(LinkerRoot->LinkerLicenseeVersion == GPackageFileLicenseeUE4Version);
@@ -4188,7 +4229,7 @@ EAsyncPackageState::Type FAsyncPackage2::PostLoadObjects()
 
 	FUObjectSerializeContext* LoadContext = GetSerializeContext();
 
-	const bool bAsyncPostLoadEnabled = true;
+	const bool bAsyncPostLoadEnabled = FAsyncLoadingThreadSettings::Get().bAsyncPostLoadEnabled;
 	const bool bIsMultithreaded = AsyncLoadingThread.IsMultithreaded();
 
 	// PostLoad objects.
@@ -4460,9 +4501,9 @@ int32 FAsyncLoadingThread2Impl::LoadPackage(const FString& InName, const FGuid*
 	FName PackageName = FName(*InName);
 	FName PackageNameToLoad = InPackageToLoadFrom ? FName(InPackageToLoadFrom) : PackageName;
 
-	FPackageId* PackageId = GlobalPackageStore.GetPackageId(PackageNameToLoad);
+	FPackageId PackageIdToLoad = GlobalPackageStore.FindPackageId(PackageNameToLoad);
 
-	if (PackageId)
+	if (PackageIdToLoad.IsValid())
 	{
 		if (FCoreDelegates::OnAsyncLoadPackage.IsBound())
 		{
@@ -4482,13 +4523,18 @@ int32 FAsyncLoadingThread2Impl::LoadPackage(const FString& InName, const FGuid*
 			CompletionDelegatePtr.Reset(new FLoadPackageAsyncDelegate(InCompletionDelegate));
 		}
 
+		FPackageId PackageId =
+			PackageName != PackageNameToLoad ?
+			GlobalPackageStore.FindOrAddPackageId(PackageName) :
+			PackageIdToLoad;
+
 		// Add new package request
-		FAsyncPackageDesc2 PackageDesc(RequestID, *PackageId, PackageName, PackageNameToLoad, MoveTemp(CompletionDelegatePtr));
+		FAsyncPackageDesc2 PackageDesc(RequestID, PackageId, PackageIdToLoad, PackageName, PackageNameToLoad, MoveTemp(CompletionDelegatePtr));
 		QueuePackage(PackageDesc);
 	}
 	else
 	{
-		UE_LOG(LogStreaming, Warning, TEXT("FAsyncLoadingThread2Impl::LoadPackage - Skipping package '%s'. Name to load is unknown: '%s')"),
+		UE_LOG(LogStreaming, Warning, TEXT("AsyncLoading2 - LoadPackage: Skipping package: '%s'. Name to load is unknown: '%s')"),
 			*PackageName.ToString(), *PackageNameToLoad.ToString());
 		InCompletionDelegate.ExecuteIfBound(PackageName, nullptr, EAsyncLoadingResult::Failed);
 	}
@@ -4502,14 +4548,14 @@ EAsyncPackageState::Type FAsyncLoadingThread2Impl::ProcessLoadingFromGameThread(
 	return IsAsyncLoading() ? EAsyncPackageState::TimeOut : EAsyncPackageState::Complete;
 }
 
-void FAsyncLoadingThread2Impl::FlushLoading(int32 PackageID)
+void FAsyncLoadingThread2Impl::FlushLoading(int32 RequestId)
 {
 	if (IsAsyncLoading())
 	{
 		// Flushing async loading while loading is suspend will result in infinite stall
 		UE_CLOG(bSuspendRequested, LogStreaming, Fatal, TEXT("Cannot Flush Async Loading while async loading is suspended"));
 
-		if (PackageID != INDEX_NONE && !ContainsRequestID(PackageID))
+		if (RequestId != INDEX_NONE && !ContainsRequestID(RequestId))
 		{
 			return;
 		}
@@ -4538,8 +4584,8 @@ void FAsyncLoadingThread2Impl::FlushLoading(int32 PackageID)
 		{
 			while (IsAsyncLoading())
 			{
-				EAsyncPackageState::Type Result = TickAsyncLoadingFromGameThread(false, false, 0, PackageID);
-				if (PackageID != INDEX_NONE && !ContainsRequestID(PackageID))
+				EAsyncPackageState::Type Result = TickAsyncLoadingFromGameThread(false, false, 0, RequestId);
+				if (RequestId != INDEX_NONE && !ContainsRequestID(RequestId))
 				{
 					break;
 				}
@@ -4559,7 +4605,7 @@ void FAsyncLoadingThread2Impl::FlushLoading(int32 PackageID)
 		double EndTime = FPlatformTime::Seconds();
 		double ElapsedTime = EndTime - StartTime;
 
-		check(PackageID != INDEX_NONE || !IsAsyncLoading());
+		check(RequestId != INDEX_NONE || !IsAsyncLoading());
 	}
 }
 
@@ -4675,9 +4721,9 @@ void FAsyncLoadingThread2::ResumeLoading()
 	Impl->ResumeLoading();
 }
 
-void FAsyncLoadingThread2::FlushLoading(int32 PackageId)
+void FAsyncLoadingThread2::FlushLoading(int32 RequestId)
 {
-	Impl->FlushLoading(PackageId);
+	Impl->FlushLoading(RequestId);
 }
 
 int32 FAsyncLoadingThread2::GetNumAsyncPackages()
diff --git a/Engine/Source/Runtime/CoreUObject/Private/Serialization/AsyncLoadingThread.h b/Engine/Source/Runtime/CoreUObject/Private/Serialization/AsyncLoadingThread.h
--- a/Engine/Source/Runtime/CoreUObject/Private/Serialization/AsyncLoadingThread.h
+++ b/Engine/Source/Runtime/CoreUObject/Private/Serialization/AsyncLoadingThread.h
@@ -149,15 +149,6 @@ struct FMaxPackageSummarySize
  */
 class FAsyncLoadingThread final : public FRunnable, public IAsyncPackageLoader
 {
-	/** Structure that holds the async loading thread ini settings */
-	struct FAsyncLoadingThreadSettings
-	{
-		bool bAsyncLoadingThreadEnabled;
-		bool bAsyncPostLoadEnabled;
-
-		FAsyncLoadingThreadSettings();
-	};
-
 	IEDLBootNotificationManager& EDLBootNotificationManager;
 
 	/** Thread to run the worker FRunnable on */
@@ -329,9 +320,6 @@ public:
 	/** [EDL] Queues StartPostLoad event */
 	void QueueEvent_StartPostLoad(FAsyncPackage* Pkg, int32 EventSystemPriority = 0);
 
-	/** Gets the ALT settigns from ini (or command line). */
-	static FAsyncLoadingThreadSettings& GetAsyncLoadingThreadSettings();
-
 	/** True if multithreaded async loading is currently being used. */
 	FORCEINLINE bool IsMultithreaded() override
 	{
diff --git a/Engine/Source/Runtime/CoreUObject/Private/Serialization/AsyncPackageLoader.cpp b/Engine/Source/Runtime/CoreUObject/Private/Serialization/AsyncPackageLoader.cpp
--- a/Engine/Source/Runtime/CoreUObject/Private/Serialization/AsyncPackageLoader.cpp
+++ b/Engine/Source/Runtime/CoreUObject/Private/Serialization/AsyncPackageLoader.cpp
@@ -391,6 +391,39 @@ static FEDLBootNotificationManager& GetGEDLBootNotificationManager()
 	return Singleton;
 }
 
+FAsyncLoadingThreadSettings::FAsyncLoadingThreadSettings()
+{
+#if THREADSAFE_UOBJECTS
+	if (FPlatformProperties::RequiresCookedData())
+	{
+		check(GConfig);
+
+		bool bConfigValue = true;
+		GConfig->GetBool(TEXT("/Script/Engine.StreamingSettings"), TEXT("s.AsyncLoadingThreadEnabled"), bConfigValue, GEngineIni);
+		bool bCommandLineDisable = FParse::Param(FCommandLine::Get(), TEXT("NoAsyncLoadingThread"));
+		bool bCommandLineEnable = FParse::Param(FCommandLine::Get(), TEXT("AsyncLoadingThread"));
+		bAsyncLoadingThreadEnabled = bCommandLineEnable || (bConfigValue && FApp::ShouldUseThreadingForPerformance() && !bCommandLineDisable);
+
+		bConfigValue = true;
+		GConfig->GetBool(TEXT("/Script/Engine.StreamingSettings"), TEXT("s.AsyncPostLoadEnabled"), bConfigValue, GEngineIni);
+		bCommandLineDisable = FParse::Param(FCommandLine::Get(), TEXT("NoAsyncPostLoad"));
+		bCommandLineEnable = FParse::Param(FCommandLine::Get(), TEXT("AsyncPostLoad"));
+		bAsyncPostLoadEnabled = bCommandLineEnable || (bConfigValue && FApp::ShouldUseThreadingForPerformance() && !bCommandLineDisable);
+	}
+	else
+#endif
+	{
+		bAsyncLoadingThreadEnabled = false;
+		bAsyncPostLoadEnabled = false;
+	}
+}
+
+FAsyncLoadingThreadSettings& FAsyncLoadingThreadSettings::Get()
+{
+	static FAsyncLoadingThreadSettings Settings;
+	return Settings;
+}
+
 bool IsEventDrivenLoaderEnabledInCookedBuilds()
 {
 	static struct FEventDrivenLoaderEnabledInCookedBuildsInit
diff --git a/Engine/Source/Runtime/CoreUObject/Public/Serialization/AsyncPackageLoader.h b/Engine/Source/Runtime/CoreUObject/Public/Serialization/AsyncPackageLoader.h
--- a/Engine/Source/Runtime/CoreUObject/Public/Serialization/AsyncPackageLoader.h
+++ b/Engine/Source/Runtime/CoreUObject/Public/Serialization/AsyncPackageLoader.h
@@ -52,6 +52,18 @@ public:
 	virtual bool IsWaitingForSomething() = 0;
 };
 
+/** Structure that holds the async loading thread ini settings */
+struct FAsyncLoadingThreadSettings
+{
+	bool bAsyncLoadingThreadEnabled;
+	bool bAsyncPostLoadEnabled;
+
+	FAsyncLoadingThreadSettings();
+
+	/** Gets the ALT settigns from ini (or command line). */
+	static FAsyncLoadingThreadSettings& Get();
+};
+
 /**
  * Asynchronous package loader interface.
  */
diff --git a/Engine/Source/Runtime/CoreUObject/Public/UObject/PackageId.h b/Engine/Source/Runtime/CoreUObject/Public/UObject/PackageId.h
--- a/Engine/Source/Runtime/CoreUObject/Public/UObject/PackageId.h
+++ b/Engine/Source/Runtime/CoreUObject/Public/UObject/PackageId.h
@@ -22,6 +22,11 @@ public:
 		return FPackageId(Index);
 	}
 
+	inline bool IsValid() const
+	{
+		return Id != InvalidId;
+	}
+
 	inline uint32 ToIndex() const
 	{
 		check(Id != InvalidId);
