Fix FPropertyMap loading of weak and lazy object pointer keys and removed temporary warning

diff --git a/Engine/Source/Runtime/CoreUObject/Private/UObject/PropertyMap.cpp b/Engine/Source/Runtime/CoreUObject/Private/UObject/PropertyMap.cpp
--- a/Engine/Source/Runtime/CoreUObject/Private/UObject/PropertyMap.cpp
+++ b/Engine/Source/Runtime/CoreUObject/Private/UObject/PropertyMap.cpp
@@ -364,14 +364,6 @@ void FMapProperty::SerializeItem(FStructuredArchive::FSlot Slot, void* Value, co
 			KeyProp->InitializeValue(TempKeyStorage);
 		}
 
-		if (KeyProp->IsA<FLazyObjectProperty>())
-		{
-			UE_LOG(LogProperty, Warning, 
-				TEXT("Loading map properties with lazy object pointer keys has a known bug in UE4.24.1 (UE-85796). "
-					 "Resaving this package may cause data loss. '%s' in package '%s'. "),
-				*GetFullName(), *UnderlyingArchive.GetArchiveName());
-		}
-
 		// Read remaining items into container
 		for (; NumEntries; --NumEntries)
 		{
diff --git a/Engine/Source/Runtime/CoreUObject/Public/UObject/UnrealType.h b/Engine/Source/Runtime/CoreUObject/Public/UObject/UnrealType.h
--- a/Engine/Source/Runtime/CoreUObject/Public/UObject/UnrealType.h
+++ b/Engine/Source/Runtime/CoreUObject/Public/UObject/UnrealType.h
@@ -3725,7 +3725,7 @@ public:
 					LocalKeyPropForCapture->InitializeValue(NewElementKey);
 				}
 
-				LocalKeyPropForCapture->CopySingleValueToScriptVM(NewElementKey, KeyPtr);
+				LocalKeyPropForCapture->CopySingleValue(NewElementKey, KeyPtr);
 
 				if (LocalValuePropForCapture->PropertyFlags & CPF_ZeroConstructor)
 				{
