Fix a rare threading issue when running async BulkData loading via the IoDispatcher.

We should mark the job as completed BEFORE we invoke the callback as this will
a) Let the callback use the results of the job if we wanted to
b) Fix potential threading issues caused by the callback invoking new jobs that will attempt to use the results of this one, which in theory could be executed before the callback returns and the job is marked as completed.

#rnx

diff --git a/Engine/Source/Runtime/CoreUObject/Private/Serialization/BulkData2.cpp b/Engine/Source/Runtime/CoreUObject/Private/Serialization/BulkData2.cpp
--- a/Engine/Source/Runtime/CoreUObject/Private/Serialization/BulkData2.cpp
+++ b/Engine/Source/Runtime/CoreUObject/Private/Serialization/BulkData2.cpp
@@ -388,13 +388,13 @@ private:
 		
 		FBulkDataBase::IoDispatcher->FreeBatch(NewBatch);
 
+		bIsCompleted = true;
+		FPlatformMisc::MemoryBarrier();
+
 		if (CompleteCallback)
 		{
 			CompleteCallback(bIsCanceled, this);
-		}
-
-		FPlatformMisc::MemoryBarrier();
-		bIsCompleted = true;
+		}	
 	}
 
 	class FBulkDataIoDispatcherRequestWorker : public FNonAbandonableTask
