Fix inline bulkdata objects not releasing their internal allocations when run with the FIoDisaptcher.

- Factored out the FIoDispatcher enabled check to a common function so that other methods can use it
- The problem is that inline data cannot be referenced currently by a FIoChunkId so ::CanLoadFromDisk would return false, this meant that ::Unlock and ::GetCopy would ignored discard requests as the code will not discard internal memory if it cannot reload it again in the future as a safety feature.
 - Ideally we would be able to reference inline data via FIoChunkId OR we would make it a hard requirement that it cannot be reloaded.
  - I need to see how much additional memory adding a FIoChunkId for each unique inline data entry would be before we can decide if it's reasonable.
  - Disabling the reloading of inline data might be okay as in theory we really shouldn't be reloadeding it as it is a waste of CPU cycles
- Added a todo as the fix is temp until we decide which way to go.

#rnx

diff --git a/Engine/Source/Runtime/CoreUObject/Private/Serialization/BulkData2.cpp b/Engine/Source/Runtime/CoreUObject/Private/Serialization/BulkData2.cpp
--- a/Engine/Source/Runtime/CoreUObject/Private/Serialization/BulkData2.cpp
+++ b/Engine/Source/Runtime/CoreUObject/Private/Serialization/BulkData2.cpp
@@ -30,6 +30,15 @@ namespace
 	static constexpr const TCHAR* OptionalExt = TEXT(".uptnl");			// Stored in a separate file that is optional
 
 	const uint16 InvalidBulkDataIndex = ~uint16(0);
+
+	FORCEINLINE bool IsIoDispatcherEnabled()
+	{
+#if ENABLE_IO_DISPATCHER
+		return FIoDispatcher::IsInitialized();
+#else
+		return false;
+#endif
+	}
 }
 
 // TODO: The code in the FileTokenSystem namespace is a temporary system so that FBulkDataBase can hold
@@ -556,11 +565,8 @@ void FBulkDataBase::Serialize(FArchive& Ar, UObject* Owner, int32 /*Index*/, boo
 
 		Ar << BulkDataOffsetInFile;
 
-#if ENABLE_IO_DISPATCHER
-		const bool bUseIoDispatcher = FIoDispatcher::IsInitialized();
-#else
-		const bool bUseIoDispatcher = false;
-#endif
+		const bool bUseIoDispatcher = IsIoDispatcherEnabled();
+
 		if ((BulkDataFlags & BULKDATA_BadDataVersion) != 0)
 		{
 			uint16 DummyValue;
@@ -837,6 +843,35 @@ int64 FBulkDataBase::GetBulkDataSize() const
 	}	
 }
 
+bool FBulkDataBase::CanLoadFromDisk() const 
+{ 
+	// If this BulkData is using the IoDispatcher then it can load from disk 
+	if (IsUsingIODispatcher())
+	{
+		return true;
+	}
+
+	// If the IoDispatcher is enabled then technically inlined data cannot load
+	// from disk, but it should act as thought it can up until the point that it
+	// actually tries to load. At which point it will give a log error.
+	//
+	// This is a temp hack while we work out if we want to allow additional inline
+	// loading at runtime (and pay the memory increase in the ToC) or prevent it entirely.
+	// TODO: Resolve this!
+	if (IsInlined() && IsIoDispatcherEnabled())
+	{
+		return true;
+	}
+
+	// If this BulkData has a fallback token then it can find it's filepath and load from disk
+	if (Fallback.Token != InvalidToken)
+	{
+		return true;
+	}
+
+	return false;
+}
+
 bool FBulkDataBase::DoesExist() const
 {
 #if ALLOW_OPTIONAL_DATA
diff --git a/Engine/Source/Runtime/CoreUObject/Public/Serialization/BulkData2.h b/Engine/Source/Runtime/CoreUObject/Public/Serialization/BulkData2.h
--- a/Engine/Source/Runtime/CoreUObject/Public/Serialization/BulkData2.h
+++ b/Engine/Source/Runtime/CoreUObject/Public/Serialization/BulkData2.h
@@ -94,7 +94,7 @@ public:
 	void ClearBulkDataFlags(uint32 BulkDataFlagsToClear);
 	uint32 GetBulkDataFlags() const { return BulkDataFlags; }
 
-	bool CanLoadFromDisk() const { return IsUsingIODispatcher() || Fallback.Token != InvalidToken; }
+	bool CanLoadFromDisk() const;
 	
 	/**
 	 * Returns true if the data references a file that currently exists and can be referenced by the file system.
