Remove EnsureBulkDataLoaded from LoadAndDetachAllBulkData; it is not required when saving a package because FUntypedBulkData::Serialize calls it later on when it is less likely to block.
#rnx

diff --git a/Engine/Source/Runtime/CoreUObject/Private/UObject/LinkerManager.cpp b/Engine/Source/Runtime/CoreUObject/Private/UObject/LinkerManager.cpp
--- a/Engine/Source/Runtime/CoreUObject/Private/UObject/LinkerManager.cpp
+++ b/Engine/Source/Runtime/CoreUObject/Private/UObject/LinkerManager.cpp
@@ -168,8 +168,6 @@ void FLinkerManager::EnsureLoadingComplete(UPackage* Package)
 	{
 		Linker->SerializeThumbnails();
 	}
-
-	Linker->LoadAndDetachAllBulkData();
 }
 
 
