Manual merge of CL: 10759497
Move ResetLoaders to after Presave, to handle Presaves that cause a circular reference that loads the package being saved again.  Since it is possible to do so, we move it to immediately before it is needed, when the file on disk is written out (or copied from temp).

diff --git a/Engine/Source/Runtime/CoreUObject/Private/UObject/SavePackage.cpp b/Engine/Source/Runtime/CoreUObject/Private/UObject/SavePackage.cpp
--- a/Engine/Source/Runtime/CoreUObject/Private/UObject/SavePackage.cpp
+++ b/Engine/Source/Runtime/CoreUObject/Private/UObject/SavePackage.cpp
@@ -5709,6 +5709,7 @@ FSavePackageResultStruct UPackage::Save(UPackage* InOuter, UObject* Base, EObjec
 						ResetLoadersForSave(InOuter, Filename);
 					}
 
+					// Compress the temporarily file to destination.
 					if (bSaveAsync)
 					{						
 						FString NewPathToSave = NewPath;
