AsyncLoading2: Add ExternalReadQueue and improve AsyncLoadingTime cpu scopes to show how much time is spent waiting for io.
The ExternalReadQueue gives slightly better performance while adding the guarantee that all external reads are always completed before ConditionalPostLoad.

diff --git a/Engine/Source/Runtime/CoreUObject/Private/Serialization/AsyncLoading2.cpp b/Engine/Source/Runtime/CoreUObject/Private/Serialization/AsyncLoading2.cpp
--- a/Engine/Source/Runtime/CoreUObject/Private/Serialization/AsyncLoading2.cpp
+++ b/Engine/Source/Runtime/CoreUObject/Private/Serialization/AsyncLoading2.cpp
@@ -1181,12 +1181,6 @@ struct FAsyncPackage2
 	*/
 	void Cancel();
 
-	/** Returns true if this package is already being loaded in the current callstack */
-	bool IsBeingProcessedRecursively() const
-	{
-		return ReentryCount > 1;
-	}
-
 	void AddOwnedObjectFromCallback(UObject* Object, bool bSubObject)
 	{
 		if (bSubObject)
@@ -1263,6 +1257,8 @@ private:
 	/** Call backs called when we finished loading this package											*/
 	using FCompletionCallback = TUniquePtr<FLoadPackageAsyncDelegate>;
 	TArray<FCompletionCallback, TInlineAllocator<2>> CompletionCallbacks;
+	/** Current index into ExternalReadDependencies array used to spread wating for external reads over several frames			*/
+	int32						ExternalReadIndex = 0;
 	/** Current index into ObjLoaded array used to spread routing PostLoad over several frames			*/
 	int32							PostLoadIndex;
 	/** Current index into DeferredPostLoadObjects array used to spread routing PostLoad over several frames			*/
@@ -1418,7 +1414,8 @@ private:
 	 *
 	 * @return Complete if all dependencies are finished, TimeOut otherwise
 	 */
-	EAsyncPackageState::Type FinishExternalReadDependencies();
+	enum EExternalReadAction { ExternalReadAction_Poll, ExternalReadAction_Wait };
+	EAsyncPackageState::Type ProcessExternalReads(EExternalReadAction Action);
 
 	/**
 	* Updates load percentage stat
@@ -1540,6 +1537,10 @@ private:
 
 	IEDLBootNotificationManager& EDLBootNotificationManager;
 
+	TQueue<FAsyncPackage2*, EQueueMode::Mpsc> ExternalReadQueue;
+	FThreadSafeCounter WaitingForIoBundleCounter;
+	FThreadSafeCounter WaitingForPostLoadCounter;
+
 	/** List of all pending package requests */
 	TSet<int32> PendingRequests;
 	/** Synchronization object for PendingRequests list */
@@ -2006,6 +2007,7 @@ bool FAsyncLoadingThread2Impl::CreateAsyncPackagesFromQueue()
 
 void FAsyncLoadingThread2Impl::AddBundleIoRequest(FAsyncPackage2* Package, const FExportBundleMetaEntry& BundleMetaEntry)
 {
+	WaitingForIoBundleCounter.Increment();
 	WaitingIoRequests.HeapPush({ Package, BundleMetaEntry.LoadOrder, BundleMetaEntry.PayloadSize });
 }
 
@@ -2049,6 +2051,7 @@ void FAsyncLoadingThread2Impl::StartBundleIoRequests()
 		{
 			Package->IoBuffer = Result.ConsumeValueOrDie();
 			Package->GetExportBundleNode(EEventLoadNode2::ExportBundle_Process, 0)->ReleaseBarrier();
+			Package->AsyncLoadingThread.WaitingForIoBundleCounter.Decrement();
 		});
 	}
 }
@@ -2686,7 +2689,21 @@ EAsyncPackageState::Type FAsyncPackage2::Event_ProcessExportBundle(FAsyncPackage
 	}
 	else
 	{
-		Package->GetNode(Package_ExportsSerialized)->ReleaseBarrier();
+		check(Package->AsyncPackageLoadingState == EAsyncPackageLoadingState2::ProcessNewImportsAndExports);
+		Package->ImportStore.ImportMap = nullptr;
+		Package->ImportStore.ImportMapCount = 0;
+		Package->bAllExportsSerialized = true;
+		Package->IoBuffer = FIoBuffer();
+		Package->AsyncPackageLoadingState = EAsyncPackageLoadingState2::PostLoad_Etc;
+
+		if (Package->ExternalReadDependencies.Num() == 0)
+		{
+			Package->GetNode(Package_ExportsSerialized)->ReleaseBarrier();
+		}
+		else
+		{
+			Package->AsyncLoadingThread.ExternalReadQueue.Enqueue(Package);
+		}
 	}
 
 	if (ExportBundleIndex == 0)
@@ -2971,13 +2988,6 @@ EAsyncPackageState::Type FAsyncPackage2::Event_ExportsDone(FAsyncPackage2* Packa
 {
 	TRACE_CPUPROFILER_EVENT_SCOPE(Event_ExportsDone);
 
-	check(Package->AsyncPackageLoadingState == EAsyncPackageLoadingState2::ProcessNewImportsAndExports);
-	Package->ImportStore.ImportMap = nullptr;
-	Package->ImportStore.ImportMapCount = 0;
-	Package->bAllExportsSerialized = true;
-	Package->AsyncPackageLoadingState = EAsyncPackageLoadingState2::PostLoad_Etc;
-	Package->IoBuffer = FIoBuffer();
-
 	Package->GetNode(EEventLoadNode2::Package_StartPostLoad)->ReleaseBarrier();
 	return EAsyncPackageState::Complete;
 }
@@ -2994,53 +3004,29 @@ EAsyncPackageState::Type FAsyncPackage2::Event_Tick(FAsyncPackage2* Package, int
 	TRACE_CPUPROFILER_EVENT_SCOPE(Event_Tick);
 
 	check(!Package->HasFinishedLoading());
+	check(Package->ExternalReadDependencies.Num() == 0);
 	
-	check(!Package->ReentryCount);
-	Package->ReentryCount++;
-
-	// Keep track of time when we start loading.
-	check(Package->LoadStartTime > 0.0f);
-
 	FAsyncPackageScope2 PackageScope(Package);
 
 	// Make sure we finish our work if there's no time limit. The loop is required as PostLoad
 	// might cause more objects to be loaded in which case we need to Preload them again.
 	EAsyncPackageState::Type LoadingState = EAsyncPackageState::Complete;
-	do
-	{
-		// Reset value to true at beginning of loop.
-		LoadingState = EAsyncPackageState::Complete;
 
-		// Begin async loading, simulates BeginLoad
-		Package->BeginAsyncLoad();
-
-		if (LoadingState == EAsyncPackageState::Complete &&
-			Package->ExternalReadDependencies.Num() > 0 &&
-			!Package->bLoadHasFailed)
-		{
-			SCOPED_LOADTIMER(Package_ExternalReadDependencies);
-			LoadingState = Package->FinishExternalReadDependencies();
-		}
+	// Begin async loading, simulates BeginLoad
+	Package->BeginAsyncLoad();
 
-		// Call PostLoad on objects, this could cause new objects to be loaded that require
-		// another iteration of the PreLoad loop.
-		if (LoadingState == EAsyncPackageState::Complete && !Package->bLoadHasFailed)
-		{
-			SCOPED_LOADTIMER(Package_PostLoadObjects);
-			LoadingState = Package->PostLoadObjects();
-		}
+	LoadingState = Package->PostLoadObjects();
 
-		// End async loading, simulates EndLoad
-		Package->EndAsyncLoad();
+	// End async loading, simulates EndLoad
+	Package->EndAsyncLoad();
 
-		// Finish objects (removing EInternalObjectFlags::AsyncLoading, dissociate imports and forced exports, 
-		// call completion callback, ...
-		// If the load has failed, perform completion callbacks and then quit
-		if (LoadingState == EAsyncPackageState::Complete || Package->bLoadHasFailed)
-		{
-			LoadingState = Package->FinishObjects();
-		}
-	} while (!FAsyncLoadingThreadState2::Get()->IsTimeLimitExceeded() && LoadingState == EAsyncPackageState::TimeOut);
+	// Finish objects (removing EInternalObjectFlags::AsyncLoading, dissociate imports and forced exports, 
+	// call completion callback, ...
+	// If the load has failed, perform completion callbacks and then quit
+	if (LoadingState == EAsyncPackageState::Complete || Package->bLoadHasFailed)
+	{
+		LoadingState = Package->FinishObjects();
+	}
 
 	// Mark this package as loaded if everything completed.
 	Package->bLoadHasFinished = (LoadingState == EAsyncPackageState::Complete);
@@ -3051,9 +3037,6 @@ EAsyncPackageState::Type FAsyncPackage2::Event_Tick(FAsyncPackage2* Package, int
 		Package->AsyncPackageLoadingState = EAsyncPackageLoadingState2::PackageComplete;
 	}
 
-	Package->ReentryCount--;
-	check(Package->ReentryCount >= 0);
-
 	if (LoadingState == EAsyncPackageState::TimeOut)
 	{
 		return EAsyncPackageState::TimeOut;
@@ -3093,6 +3076,7 @@ FEventLoadNode2* FAsyncPackage2::GetNode(int32 NodeIndex)
 
 void FAsyncLoadingThread2Impl::AddToLoadedPackages(FAsyncPackage2* Package)
 {
+	WaitingForPostLoadCounter.Increment();
 	FScopeLock LoadedLock(&LoadedPackagesCritical);
 	check(!LoadedPackages.Contains(Package));
 	LoadedPackages.Add(Package);
@@ -3135,6 +3119,20 @@ EAsyncPackageState::Type FAsyncLoadingThread2Impl::ProcessAsyncLoadingFromGameTh
 				return EAsyncPackageState::TimeOut;
 			}
 
+			if (!ExternalReadQueue.IsEmpty())
+			{
+				TRACE_CPUPROFILER_EVENT_SCOPE(ProcessExternalReads);
+
+				FAsyncPackage2* Package = nullptr;
+				ExternalReadQueue.Dequeue(Package);
+
+				EAsyncPackageState::Type Result = Package->ProcessExternalReads(FAsyncPackage2::ExternalReadAction_Wait);
+				check(Result == EAsyncPackageState::Complete);
+
+				OutPackagesProcessed++;
+				break;
+			}
+
 			if (QueuedPackagesCounter)
 			{
 				CreateAsyncPackagesFromQueue();
@@ -3313,7 +3311,6 @@ EAsyncPackageState::Type FAsyncLoadingThread2Impl::ProcessLoadedPackagesFromGame
 		for (int32 PackageIndex = 0; PackageIndex < PackagesToDelete.Num(); ++PackageIndex)
 		{
 			FAsyncPackage2* Package = PackagesToDelete[PackageIndex];
-			if (!Package->IsBeingProcessedRecursively())
 			{
 				bool bSafeToDelete = false;
 				if (Package->HasClusterObjects())
@@ -3345,6 +3342,7 @@ EAsyncPackageState::Type FAsyncLoadingThread2Impl::ProcessLoadedPackagesFromGame
 					PackagesToDelete.RemoveAtSwap(PackageIndex--);
 					Package->ClearImportedPackages();
 					Package->ReleaseRef();
+					Package->AsyncLoadingThread.WaitingForPostLoadCounter.Decrement();
 				}
 			}
 
@@ -3477,7 +3475,7 @@ FAsyncLoadingThread2Impl::FAsyncLoadingThread2Impl(FIoDispatcher& InIoDispatcher
 
 	EventSpecs.AddDefaulted(EEventLoadNode2::Package_NumPhases + EEventLoadNode2::ExportBundle_NumPhases);
 	EventSpecs[EEventLoadNode2::Package_ExportsSerialized] = { &FAsyncPackage2::Event_ExportsDone, &AsyncEventQueue, true };
-	EventSpecs[EEventLoadNode2::Package_StartPostLoad] = { &FAsyncPackage2::Event_StartPostLoad, &AsyncEventQueue, false };
+	EventSpecs[EEventLoadNode2::Package_StartPostLoad] = { &FAsyncPackage2::Event_StartPostLoad, &AsyncEventQueue, true };
 	EventSpecs[EEventLoadNode2::Package_Tick] = { &FAsyncPackage2::Event_Tick, &EventQueue, false };
 	EventSpecs[EEventLoadNode2::Package_Delete] = { &FAsyncPackage2::Event_Delete, &AsyncEventQueue, false };
 
@@ -3659,8 +3657,33 @@ uint32 FAsyncLoadingThread2Impl::Run()
 				TRACE_CPUPROFILER_EVENT_SCOPE(AsyncLoadingTime);
 				do
 				{
+					const bool bPreviousDidSomething = bDidSomething;
 					bDidSomething = false;
 
+					bool bDidExternalRead = false;
+					do
+					{
+						bDidExternalRead = false;
+						FAsyncPackage2* Package = nullptr;
+						if (ExternalReadQueue.Peek(Package))
+						{
+							TRACE_CPUPROFILER_EVENT_SCOPE(ProcessExternalReads);
+
+							FAsyncPackage2::EExternalReadAction Action =
+								bPreviousDidSomething ?
+								FAsyncPackage2::ExternalReadAction_Poll :
+								FAsyncPackage2::ExternalReadAction_Wait;
+
+							EAsyncPackageState::Type Result = Package->ProcessExternalReads(Action);
+							if (Result == EAsyncPackageState::Complete)
+							{
+								ExternalReadQueue.Pop();
+								bDidExternalRead = true;
+								bDidSomething = true;
+							}
+						}
+					} while (bDidExternalRead);
+
 					if (QueuedPackagesCounter)
 					{
 						if (CreateAsyncPackagesFromQueue())
@@ -3701,10 +3724,32 @@ uint32 FAsyncLoadingThread2Impl::Run()
 					}
 				} while (bDidSomething);
 			}
+
+			const bool bWaitingForIo = WaitingForIoBundleCounter.GetValue() > 0;
+			const bool bWaitingForPostLoad = WaitingForPostLoadCounter.GetValue() > 0;
+			const bool bIsLoadingAndWaiting = bWaitingForIo || bWaitingForPostLoad;
 			if (!bDidSomething)
 			{
-				ThreadState.ProcessDeferredFrees();
-				Waiter.Wait();
+				if (bIsLoadingAndWaiting)
+				{
+					TRACE_CPUPROFILER_EVENT_SCOPE(AsyncLoadingTime);
+					ThreadState.ProcessDeferredFrees();
+
+					if (bWaitingForIo)
+					{
+						TRACE_CPUPROFILER_EVENT_SCOPE(WaitingForIo);
+						Waiter.Wait();
+					}
+					else// if (bWaitingForPostLoad)
+					{
+						TRACE_CPUPROFILER_EVENT_SCOPE(WaitingForPostLoad);
+						Waiter.Wait();
+					}
+				}
+				else
+				{
+					Waiter.Wait();
+				}
 			}
 		}
 	}
@@ -3964,7 +4009,6 @@ FAsyncPackage2::FAsyncPackage2(
 , bCreatedLinkerRoot(false)
 , LoadStartTime(0.0)
 , LoadPercentage(0)
-, ReentryCount(0)
 , AsyncLoadingThread(InAsyncLoadingThread)
 , EDLBootNotificationManager(InEDLBootNotificationManager)
 , GraphAllocator(InGraphAllocator)
@@ -4203,13 +4247,30 @@ void FAsyncPackage2::CreateUPackage(const FPackageSummary* PackageSummary)
 	UE_LOG(LogStreaming, Verbose, TEXT("FAsyncPackage::CreateUPackage for %s finished."), *Desc.Name.ToString());
 }
 
-EAsyncPackageState::Type FAsyncPackage2::FinishExternalReadDependencies()
+EAsyncPackageState::Type FAsyncPackage2::ProcessExternalReads(EExternalReadAction Action)
 {
-	for (FExternalReadCallback& ReadCallback : ExternalReadDependencies)
+	double WaitTime;
+	if (Action == ExternalReadAction_Poll)
+	{
+		WaitTime = -1.f;
+	}
+	else// if (Action == ExternalReadAction_Wait)
+	{
+		WaitTime = 0.f;
+	}
+
+	while (ExternalReadIndex < ExternalReadDependencies.Num())
 	{
-		ReadCallback(0.0);
+		FExternalReadCallback& ReadCallback = ExternalReadDependencies[ExternalReadIndex];
+		if (!ReadCallback(WaitTime))
+		{
+			return EAsyncPackageState::TimeOut;
+		}
+		++ExternalReadIndex;
 	}
+
 	ExternalReadDependencies.Empty();
+	GetNode(Package_ExportsSerialized)->ReleaseBarrier();
 	return EAsyncPackageState::Complete;
 }
 
