AsyncLoading2: Remove zenaphore WaitForEvents scope now that the loader has more specific scopes

diff --git a/Engine/Source/Runtime/CoreUObject/Private/Serialization/Zenaphore.cpp b/Engine/Source/Runtime/CoreUObject/Private/Serialization/Zenaphore.cpp
--- a/Engine/Source/Runtime/CoreUObject/Private/Serialization/Zenaphore.cpp
+++ b/Engine/Source/Runtime/CoreUObject/Private/Serialization/Zenaphore.cpp
@@ -65,9 +65,6 @@ void FZenaphoreWaiter::Wait()
 	}
 	else
 	{
-#if CPUPROFILERTRACE_ENABLED
-		FCpuProfilerTrace::OutputBeginEvent(WaitCpuScopeId);
-#endif
 		for (;;)
 		{
 			Outer.Event->Wait(INT32_MAX, true);
@@ -79,8 +76,5 @@ void FZenaphoreWaiter::Wait()
 			}
 		}
 		SpinCount = 0;
-#if CPUPROFILERTRACE_ENABLED
-		FCpuProfilerTrace::OutputEndEvent();
-#endif
 	}
 }
diff --git a/Engine/Source/Runtime/CoreUObject/Private/Serialization/Zenaphore.h b/Engine/Source/Runtime/CoreUObject/Private/Serialization/Zenaphore.h
--- a/Engine/Source/Runtime/CoreUObject/Private/Serialization/Zenaphore.h
+++ b/Engine/Source/Runtime/CoreUObject/Private/Serialization/Zenaphore.h
@@ -18,9 +18,6 @@ public:
 	FZenaphoreWaiter(FZenaphore& Outer, const TCHAR* WaitCpuScopeName)
 		: Outer(Outer)
 	{
-#if CPUPROFILERTRACE_ENABLED
-		WaitCpuScopeId = FCpuProfilerTrace::OutputEventType(WaitCpuScopeName, CpuProfilerGroup_Default);
-#endif
 	}
 
 	void Wait();
@@ -31,9 +28,6 @@ private:
 	FZenaphore& Outer;
 	FZenaphoreWaiterNode WaiterNode;
 	int32 SpinCount = 0;
-#if CPUPROFILERTRACE_ENABLED
-	uint16 WaitCpuScopeId;
-#endif
 };
 
 class FZenaphore
