AsyncLoading2: Make sure to handle GC also when running without the async loading thread.

diff --git a/Engine/Source/Runtime/CoreUObject/Private/Serialization/AsyncLoading2.cpp b/Engine/Source/Runtime/CoreUObject/Private/Serialization/AsyncLoading2.cpp
--- a/Engine/Source/Runtime/CoreUObject/Private/Serialization/AsyncLoading2.cpp
+++ b/Engine/Source/Runtime/CoreUObject/Private/Serialization/AsyncLoading2.cpp
@@ -321,7 +321,7 @@ struct FGlobalImportStore
 	FGlobalImport* Imports = nullptr;
 	/** Reference tracking for GC management */
 	TArray<UObject*> KeepAliveObjects;
-	bool bNeedToRunGarbageCollect = false;
+	bool bNeedToHandleGarbageCollect = false;
 
 	void OnPreGarbageCollect(bool bInIsLoadingPackages);
 	void OnPostGarbageCollect();
@@ -3092,6 +3092,8 @@ EAsyncPackageState::Type FAsyncLoadingThread2Impl::ProcessAsyncLoadingFromGameTh
 	{
 		do 
 		{
+			GlobalPackageStore.GetGlobalImportStore().bNeedToHandleGarbageCollect |= IsAsyncLoadingPackages();
+
 			ThreadState.ProcessDeferredFrees();
 
 			if (bNeedsHeartbeatTick && (++LoopIterations) % 32 == 31)
@@ -3517,6 +3519,9 @@ void FAsyncLoadingThread2Impl::StartThread()
 	// Make sure the GC sync object is created before we start the thread (apparently this can happen before we call InitUObject())
 	FGCCSyncObject::Create();
 
+	FCoreUObjectDelegates::GetPreGarbageCollectDelegate().AddRaw(this, &FAsyncLoadingThread2Impl::OnPreGarbageCollect);
+	FCoreUObjectDelegates::GetPostGarbageCollect().AddRaw(this, &FAsyncLoadingThread2Impl::OnPostGarbageCollect);
+
 	if (!FAsyncLoadingThreadSettings::Get().bAsyncLoadingThreadEnabled)
 	{
 		GlobalPackageStore.FinalizeInitialLoad();
@@ -3552,9 +3557,6 @@ void FAsyncLoadingThread2Impl::StartThread()
 			}
 		}
 
-		FCoreUObjectDelegates::GetPreGarbageCollectDelegate().AddRaw(this, &FAsyncLoadingThread2Impl::OnPreGarbageCollect);
-		FCoreUObjectDelegates::GetPostGarbageCollect().AddRaw(this, &FAsyncLoadingThread2Impl::OnPostGarbageCollect);
-
 		Thread = FRunnableThread::Create(this, TEXT("FAsyncLoadingThread"), 0, TPri_Normal);
 		if (Thread)
 		{
@@ -3640,7 +3642,7 @@ uint32 FAsyncLoadingThread2Impl::Run()
 		}
 		else
 		{
-			GlobalPackageStore.GetGlobalImportStore().bNeedToRunGarbageCollect |= IsAsyncLoadingPackages();
+			GlobalPackageStore.GetGlobalImportStore().bNeedToHandleGarbageCollect |= IsAsyncLoadingPackages();
 			bool bDidSomething = false;
 			{
 				FGCScopeGuard GCGuard;
@@ -3862,11 +3864,11 @@ static void VerifyLoadFlagsWhenFinishedLoading()
 
 void FGlobalImportStore::OnPreGarbageCollect(bool bInIsLoadingPackages)
 {
-	if (!bNeedToRunGarbageCollect && !bInIsLoadingPackages)
+	if (!bNeedToHandleGarbageCollect && !bInIsLoadingPackages)
 	{
 		return;
 	}
-	bNeedToRunGarbageCollect = bInIsLoadingPackages;
+	bNeedToHandleGarbageCollect = bInIsLoadingPackages;
 
 	int32 NumWeak = 0;
 	for (int32 GlobalImportIndex = ScriptImportCount; GlobalImportIndex < Count; ++GlobalImportIndex)
@@ -3925,7 +3927,7 @@ void FGlobalImportStore::OnPostGarbageCollect()
 	{
 		return;
 	}
-	check(bNeedToRunGarbageCollect);
+	check(bNeedToHandleGarbageCollect);
 
 	for (UObject* Object : KeepAliveObjects)
 	{
