AsyncLoading2: Runtime filtering of client/server only exports
- needs gc.UseDisregardForGCOnDedicatedServers=True

diff --git a/Engine/Source/Developer/IoStoreUtilities/Private/IoStoreUtilities.cpp b/Engine/Source/Developer/IoStoreUtilities/Private/IoStoreUtilities.cpp
--- a/Engine/Source/Developer/IoStoreUtilities/Private/IoStoreUtilities.cpp
+++ b/Engine/Source/Developer/IoStoreUtilities/Private/IoStoreUtilities.cpp
@@ -225,6 +225,13 @@ enum EEventLoadNode2 : uint8
 	ExportBundle_NumPhases,
 };
 
+enum class EExportFilterFlags : uint8
+{
+	None,
+	NotForClient,
+	NotForServer
+};
+
 struct FArc
 {
 	uint32 FromNodeIndex;
@@ -1148,6 +1155,18 @@ void SerializePackageData(
 			ZenExportMapArchive << GlobalImportIndex;
 			uint32 ObjectFlags = ObjectExport.ObjectFlags;
 			ZenExportMapArchive << ObjectFlags;
+			uint8 FilterFlags = uint8(EExportFilterFlags::None);
+			if (ObjectExport.bNotForClient)
+			{
+				FilterFlags = uint8(EExportFilterFlags::NotForClient);
+			}
+			else if (ObjectExport.bNotForServer)
+			{
+				FilterFlags = uint8(EExportFilterFlags::NotForServer);
+			}
+			ZenExportMapArchive << FilterFlags;
+			uint8 Pad = 0;
+			ZenExportMapArchive.Serialize(&Pad, 7);
 		}
 		Package->ExportMapSize = ZenExportMapArchive.Tell();
 
diff --git a/Engine/Source/Runtime/CoreUObject/Private/Serialization/AsyncLoading2.cpp b/Engine/Source/Runtime/CoreUObject/Private/Serialization/AsyncLoading2.cpp
--- a/Engine/Source/Runtime/CoreUObject/Private/Serialization/AsyncLoading2.cpp
+++ b/Engine/Source/Runtime/CoreUObject/Private/Serialization/AsyncLoading2.cpp
@@ -110,6 +110,13 @@ public:
 	}
 };
 
+enum class EExportFilterFlags : uint8
+{
+	None,
+	NotForClient,
+	NotForServer
+};
+
 struct FExportMapEntry
 {
 	uint64 SerialSize;
@@ -120,6 +127,8 @@ struct FExportMapEntry
 	FPackageIndex TemplateIndex;
 	int32 GlobalImportIndex;
 	EObjectFlags ObjectFlags;
+	EExportFilterFlags FilterFlags;
+	uint8 Pad[7];
 };
 
 struct FExportBundleHeader
@@ -139,6 +148,14 @@ struct FExportBundleEntry
 	uint32 LocalExportIndex;
 };
 
+struct FExportObject
+{
+	UObject* Object = nullptr;
+	bool bFiltered = false;
+};
+
+using FExportObjects = TArray<FExportObject>;
+
 class FGlobalImport
 {
 private:
@@ -744,7 +761,7 @@ public:
 		}
 		else if (Index.IsExport())
 		{
-			Object = (*Exports)[Index.ToExport()];
+			Object = (*Exports)[Index.ToExport()].Object;
 		}
 		else
 		{
@@ -812,7 +829,7 @@ private:
 	FPackageImportStore* ImportStore = nullptr;
 	const int32* PackageNameMap = nullptr;
 	const TArray<FNameEntryId>* GlobalNameMap = nullptr;
-	const TArray<UObject*>* Exports = nullptr;
+	const FExportObjects* Exports = nullptr;
 	TArray<FExternalReadCallback>* ExternalReadDependencies;
 };
 
@@ -1309,7 +1326,7 @@ private:
 	int32 ExportCount = 0;
 	const FExportMapEntry* ExportMap = nullptr;
 	const int32* PackageNameMap = nullptr;
-	TArray<UObject*> Exports;
+	FExportObjects Exports;
 	FPackageImportStore ImportStore;
 
 	FExportBundleMetaEntry* ExportBundleMetaEntries = nullptr;
@@ -2589,6 +2606,30 @@ EAsyncPackageState::Type FAsyncPackage2::Event_ProcessExportBundle(FAsyncPackage
 
 	FScopedAsyncPackageEvent2 Scope(Package);
 
+	auto FilterExport = [](const EExportFilterFlags FilterFlags) -> bool
+	{
+#if UE_SERVER
+		return !!(static_cast<uint32>(FilterFlags) & static_cast<uint32>(EExportFilterFlags::NotForServer));
+#elif !WITH_SERVER_CODE
+		return !!(static_cast<uint32>(FilterFlags) & static_cast<uint32>(EExportFilterFlags::NotForClient));
+#else
+		static const bool bIsDedicatedServer = !GIsClient && GIsServer;
+		static const bool bIsClientOnly = GIsClient && !GIsServer;
+
+		if (bIsDedicatedServer && static_cast<uint32>(FilterFlags) & static_cast<uint32>(EExportFilterFlags::NotForServer))
+		{
+			return true;
+		}
+
+		if (bIsClientOnly && static_cast<uint32>(FilterFlags) & static_cast<uint32>(EExportFilterFlags::NotForClient))
+		{
+			return true;
+		}
+
+		return false;
+#endif
+	};
+
 	check(ExportBundleIndex < Package->ExportBundleCount);
 
 	if (ExportBundleIndex == 0)
@@ -2655,6 +2696,22 @@ EAsyncPackageState::Type FAsyncPackage2::Event_ProcessExportBundle(FAsyncPackage
 	check(BundleEntry <= BundleEntryEnd);
 	while (BundleEntry < BundleEntryEnd)
 	{
+		const FExportMapEntry& Export = Package->ExportMap[BundleEntry->LocalExportIndex];
+
+		if (FilterExport(Export.FilterFlags))
+		{
+			Package->Exports[BundleEntry->LocalExportIndex].bFiltered = true;
+			if (BundleEntry->CommandType == FExportBundleEntry::ExportBundleEntryType_Serialize)
+			{
+				const uint64 SerialSize = Package->ExportMap[BundleEntry->LocalExportIndex].SerialSize;
+				Package->SerialDataPtr += SerialSize;
+				Ar.Seek(Ar.Tell() + SerialSize);
+			}
+
+			++BundleEntry;
+			continue;
+		}
+
 		if (BundleEntry->CommandType == FExportBundleEntry::ExportBundleEntryType_Create)
 		{
 			Package->EventDrivenCreateExport(BundleEntry->LocalExportIndex);
@@ -2662,9 +2719,9 @@ EAsyncPackageState::Type FAsyncPackage2::Event_ProcessExportBundle(FAsyncPackage
 		else
 		{
 			check(BundleEntry->CommandType == FExportBundleEntry::ExportBundleEntryType_Serialize);
-			uint64 ExportSerialSize = Package->ExportMap[BundleEntry->LocalExportIndex].SerialSize;
+			const uint64 ExportSerialSize = Export.SerialSize;
 			check(Package->SerialDataPtr + ExportSerialSize <= Package->IoBuffer.Data() + Package->IoBuffer.DataSize());
-			UObject* Object = Package->Exports[BundleEntry->LocalExportIndex];
+			UObject* Object = Package->Exports[BundleEntry->LocalExportIndex].Object;
 			check(Object);
 			if (Object->HasAnyFlags(RF_NeedLoad))
 			{
@@ -2721,7 +2778,7 @@ UObject* FAsyncPackage2::EventDrivenIndexToObject(FPackageIndex Index, bool bChe
 	}
 	if (Index.IsExport())
 	{
-		Result = Exports[Index.ToExport()];
+		Result = Exports[Index.ToExport()].Object;
 	}
 	else if (Index.IsImport())
 	{
@@ -2760,7 +2817,7 @@ void FAsyncPackage2::EventDrivenCreateExport(int32 LocalExportIndex)
 	TRACE_CPUPROFILER_EVENT_SCOPE(CreateExport);
 
 	const FExportMapEntry& Export = ExportMap[LocalExportIndex];
-	UObject*& Object = Exports[LocalExportIndex];
+	UObject*& Object = Exports[LocalExportIndex].Object;
 	check(!Object);
 
 	FName ObjectName;
@@ -2916,7 +2973,7 @@ void FAsyncPackage2::EventDrivenSerializeExport(int32 LocalExportIndex, FSimpleE
 	TRACE_CPUPROFILER_EVENT_SCOPE(SerializeExport);
 
 	const FExportMapEntry& Export = ExportMap[LocalExportIndex];
-	UObject* Object = Exports[LocalExportIndex];
+	UObject* Object = Exports[LocalExportIndex].Object;
 	check(Object);
 
 	LLM_SCOPE(ELLMTag::UObject);
@@ -3180,11 +3237,11 @@ bool FAsyncPackage2::AreAllDependenciesFullyLoadedInternal(FAsyncPackage2* Packa
 			}
 			for (int32 LocalExportIndex = 0; LocalExportIndex < AsyncRoot->ExportCount; ++LocalExportIndex)
 			{
-				UObject* Object = AsyncRoot->Exports[LocalExportIndex];
-				if (Object->HasAnyFlags(RF_NeedPostLoad|RF_NeedLoad))
+				const FExportObject& Export = AsyncRoot->Exports[LocalExportIndex];
+				if (!Export.bFiltered && Export.Object->HasAnyFlags(RF_NeedPostLoad|RF_NeedLoad))
 				{
-					OutError = FString::Printf(TEXT("%s has not been %s"), *Object->GetFullName(), 
-						Object->HasAnyFlags(RF_NeedLoad) ? TEXT("Serialized") : TEXT("PostLoaded"));
+					OutError = FString::Printf(TEXT("%s has not been %s"), *Export.Object->GetFullName(), 
+						Export.Object->HasAnyFlags(RF_NeedLoad) ? TEXT("Serialized") : TEXT("PostLoaded"));
 					return false;
 				}
 			}
@@ -4019,7 +4076,7 @@ FAsyncPackage2::FAsyncPackage2(
 	ExportBundleMetaEntries = PackageExportBundles.Get<0>();
 	ExportBundleCount = PackageExportBundles.Get<1>();
 	ExportCount = AsyncLoadingThread.GlobalPackageStore.GetPackageExportCount(Desc.PackageIdToLoad);
-	Exports.AddZeroed(ExportCount);
+	Exports.AddDefaulted(ExportCount);
 	OwnedObjects.Reserve(ExportCount + 1); // +1 for UPackage
 
 	CreateNodes(EventSpecs);
@@ -4288,10 +4345,13 @@ EAsyncPackageState::Type FAsyncPackage2::PostLoadObjects()
 	// PostLoad objects.
 	while (PostLoadIndex < ExportCount && !FAsyncLoadingThreadState2::Get()->IsTimeLimitExceeded())
 	{
-		int32 LocalExportIndex = PostLoadIndex++;
-		const FExportMapEntry& Export = ExportMap[LocalExportIndex];
-		UObject* Object = Exports[LocalExportIndex];
+		const FExportObject& Export = Exports[PostLoadIndex++];
+		if (Export.bFiltered)
+		{
+			continue;
+		}
 
+		UObject* Object = Export.Object;
 		check(Object);
 		check(!Object->HasAnyFlags(RF_NeedLoad));
 		if (!Object->HasAnyFlags(RF_NeedPostLoad))
@@ -4333,8 +4393,13 @@ EAsyncPackageState::Type FAsyncPackage2::PostLoadDeferredObjects()
 		!AsyncLoadingThread.IsAsyncLoadingSuspended() &&
 		!FAsyncLoadingThreadState2::Get()->IsTimeLimitExceeded())
 	{
-		UObject* Object = Exports[DeferredPostLoadIndex++];
+		const FExportObject& Export = Exports[DeferredPostLoadIndex++];
+		if (Export.bFiltered)
+		{
+			continue;
+		}
 
+		UObject* Object = Export.Object;
 		check(Object);
 		check(!Object->HasAnyFlags(RF_NeedLoad));
 		if (Object->HasAnyFlags(RF_NeedPostLoad))
@@ -4359,7 +4424,13 @@ EAsyncPackageState::Type FAsyncPackage2::PostLoadDeferredObjects()
 			!AsyncLoadingThread.IsAsyncLoadingSuspended() &&
 			!FAsyncLoadingThreadState2::Get()->IsTimeLimitExceeded())
 		{
-			UObject* Object = Exports[DeferredFinalizeIndex++];
+			const FExportObject& Export = Exports[DeferredFinalizeIndex++];
+			if (Export.bFiltered)
+			{
+				continue;
+			}
+
+			UObject* Object = Export.Object;
 
 			// CDO need special handling, no matter if it's listed in DeferredFinalizeObjects or created here for DynamicClass
 			UObject* CDOToHandle = nullptr;
