Optimize FPackageName::TryConvertFilenameToLongPackageName by using FPathViews

This reduced the cost of this function by about 53%, which saved 13 seconds of wall time during cooking of a large title.

diff --git a/Engine/Source/Runtime/CoreUObject/Private/Misc/PackageName.cpp b/Engine/Source/Runtime/CoreUObject/Private/Misc/PackageName.cpp
--- a/Engine/Source/Runtime/CoreUObject/Private/Misc/PackageName.cpp
+++ b/Engine/Source/Runtime/CoreUObject/Private/Misc/PackageName.cpp
@@ -5,6 +5,7 @@
 =============================================================================*/
 
 #include "Misc/PackageName.h"
+#include "Containers/StringView.h"
 #include "GenericPlatform/GenericPlatformFile.h"
 #include "HAL/FileManager.h"
 #include "Misc/Paths.h"
@@ -20,7 +21,9 @@
 #include "HAL/CriticalSection.h"
 #include "HAL/ThreadHeartBeat.h"
 #include "Misc/AutomationTest.h"
+#include "Misc/PathViews.h"
 #include "Misc/ScopeLock.h"
+#include "Misc/StringBuilder.h"
 
 DEFINE_LOG_CATEGORY_STATIC(LogPackageName, Log, All);
 
@@ -316,10 +319,11 @@ private:
 	}
 };
 
-FString FPackageName::InternalFilenameToLongPackageName(const FString& InFilename)
+void FPackageName::InternalFilenameToLongPackageName(const FStringView& InFilename, FStringBuilderBase& OutPackageName)
 {
 	const FLongPackagePathsSingleton& Paths = FLongPackagePathsSingleton::Get();
-	FString Filename = InFilename.Replace(TEXT("\\"), TEXT("/"), ESearchCase::CaseSensitive);
+	FString Filename(InFilename);
+	FPaths::NormalizeFilename(Filename);
 
 	// Convert to relative path if it's not already a long package name
 	bool bIsValidLongPackageName = false;
@@ -350,10 +354,7 @@ FString FPackageName::InternalFilenameToLongPackageName(const FString& InFilenam
 		}
 	}
 
-	FString PackageName         = FPaths::GetBaseFilename(Filename);
-	int32   PackageNameStartsAt = Filename.Len() - FPaths::GetCleanFilename(Filename).Len();
-	FString Result              = Filename.Mid(0, PackageNameStartsAt + PackageName.Len());
-	Result.ReplaceInline(TEXT("\\"), TEXT("/"), ESearchCase::CaseSensitive);
+	FStringView Result = FPathViews::GetBaseFilenameWithPath(Filename);
 
 	{
 		FScopeLock ScopeLock(&ContentMountPointCriticalSection);
@@ -361,18 +362,20 @@ FString FPackageName::InternalFilenameToLongPackageName(const FString& InFilenam
 		{
 			if (Result.StartsWith(Pair.ContentPath))
 			{
-				Result = Pair.RootPath + Result.Mid(Pair.ContentPath.Len());
-				break;
+				OutPackageName << Pair.RootPath << Result.RightChop(Pair.ContentPath.Len());
+				return;
 			}
 		}
 	}
 
-	return Result;
+	OutPackageName << Result;
 }
 
 bool FPackageName::TryConvertFilenameToLongPackageName(const FString& InFilename, FString& OutPackageName, FString* OutFailureReason)
 {
-	FString LongPackageName = InternalFilenameToLongPackageName(InFilename);
+	TStringBuilder<256> LongPackageNameBuilder;
+	InternalFilenameToLongPackageName(InFilename, LongPackageNameBuilder);
+	FStringView LongPackageName = LongPackageNameBuilder.ToString();
 
 	// we don't support loading packages from outside of well defined places
 	int32 CharacterIndex;
@@ -382,7 +385,7 @@ bool FPackageName::TryConvertFilenameToLongPackageName(const FString& InFilename
 
 	if (!(bContainsDot || bContainsBackslash || bContainsColon))
 	{
-		OutPackageName = MoveTemp(LongPackageName);
+		OutPackageName = LongPackageName;
 		return true;
 	}
 
@@ -415,7 +418,7 @@ bool FPackageName::TryConvertFilenameToLongPackageName(const FString& InFilename
 		{
 			InvalidChars += TEXT(":");
 		}
-		*OutFailureReason = FString::Printf(TEXT("FilenameToLongPackageName failed to convert '%s'. Attempt result was '%s', but the path contains illegal characters '%s'"), *InFilename, *LongPackageName, *InvalidChars);
+		*OutFailureReason = FString::Printf(TEXT("FilenameToLongPackageName failed to convert '%s'. Attempt result was '%.*s', but the path contains illegal characters '%s'"), *InFilename, LongPackageName.Len(), LongPackageName.GetData(), *InvalidChars);
 	}
 
 	return false;
@@ -850,7 +853,7 @@ bool FPackageName::FindPackageFileWithoutExtension(const FString& InPackageFilen
 	return false;
 }
 
-bool FPackageName::FixPackageNameCase(FString& LongPackageName, const FString& Extension)
+bool FPackageName::FixPackageNameCase(FString& LongPackageName, const FStringView& Extension)
 {
 	// Find the matching long package root
 	const FLongPackagePathsSingleton& Paths = FLongPackagePathsSingleton::Get();
diff --git a/Engine/Source/Runtime/CoreUObject/Public/Misc/PackageName.h b/Engine/Source/Runtime/CoreUObject/Public/Misc/PackageName.h
--- a/Engine/Source/Runtime/CoreUObject/Public/Misc/PackageName.h
+++ b/Engine/Source/Runtime/CoreUObject/Public/Misc/PackageName.h
@@ -8,6 +8,8 @@
 
 #include "CoreMinimal.h"
 
+class FStringBuilderBase;
+class FStringView;
 struct FFileStatData;
 
 class COREUOBJECT_API FPackageName
@@ -453,7 +455,7 @@ public:
 	 * @param Extension The extension for this package
 	 * @return True if the long package name was fixed up, false otherwise
 	 */
-	static bool FixPackageNameCase(FString& LongPackageName, const FString& Extension);
+	static bool FixPackageNameCase(FString& LongPackageName, const FStringView& Extension);
 
 	UE_DEPRECATED(4.17, "Deprecated. Call TryConvertLongPackageNameToFilename instead, which also works on nested paths")
 	static bool ConvertRootPathToContentPath(const FString& RootPath, FString& OutContentPath);
@@ -472,9 +474,9 @@ private:
 	 * Internal function used to rename filename to long package name.
 	 *
 	 * @param InFilename
-	 * @return Long package name.
+	 * @param OutPackageName Long package name.
 	 */
-	static FString InternalFilenameToLongPackageName(const FString& InFilename);
+	static void InternalFilenameToLongPackageName(const FStringView& InFilename, FStringBuilderBase& OutPackageName);
 
 	/** Event that is triggered when a new content path is mounted */
 	static FOnContentPathMountedEvent OnContentPathMountedEvent;
