The virtual texturing system now works when running with the IoDispatcher.

The VT system already has its own very specialised file handling so rather than try to refactor that it seemed safer to implement a version of IAsyncReadFileHandle that can handle reading via the IoDispatcher.
At some point we we might want to move this to a header file so other system can use it but for now it is local to the new BulkData api.

Add an accessor FBulkDataBase::GettIODispatcher so that FBulkDataIoDispatcherRequest no longer needs to be a friend class.

diff --git a/Engine/Source/Runtime/CoreUObject/Private/Serialization/AsyncLoading2.cpp b/Engine/Source/Runtime/CoreUObject/Private/Serialization/AsyncLoading2.cpp
--- a/Engine/Source/Runtime/CoreUObject/Private/Serialization/AsyncLoading2.cpp
+++ b/Engine/Source/Runtime/CoreUObject/Private/Serialization/AsyncLoading2.cpp
@@ -1925,7 +1925,7 @@ struct FAsyncLoadingTickScope2
 void FAsyncLoadingThread2Impl::InitializeLoading()
 {
 #if USE_NEW_BULKDATA
-	FBulkDataBase::SetIODispatcher(&IoDispatcher);
+	FBulkDataBase::SetIoDispatcher(&IoDispatcher);
 #endif
 
 	{
@@ -3552,7 +3552,7 @@ FAsyncLoadingThread2Impl::~FAsyncLoadingThread2Impl()
 	}
 
 #if USE_NEW_BULKDATA
-	FBulkDataBase::SetIODispatcher(nullptr);
+	FBulkDataBase::SetIoDispatcher(nullptr);
 #endif
 }
 
diff --git a/Engine/Source/Runtime/CoreUObject/Private/Serialization/BulkData.cpp b/Engine/Source/Runtime/CoreUObject/Private/Serialization/BulkData.cpp
--- a/Engine/Source/Runtime/CoreUObject/Private/Serialization/BulkData.cpp
+++ b/Engine/Source/Runtime/CoreUObject/Private/Serialization/BulkData.cpp
@@ -1594,6 +1594,11 @@ void FUntypedBulkData::SerializeBulkData( FArchive& Ar, void* Data )
 	}
 }
 
+IAsyncReadFileHandle* FUntypedBulkData::OpenAsyncReadHandle() const
+{
+	return FPlatformFileManager::Get().GetPlatformFile().OpenAsyncRead(*GetFilename());
+}
+
 IBulkDataIORequest* FUntypedBulkData::CreateStreamingRequest(EAsyncIOPriorityAndFlags Priority, FBulkDataIORequestCallBack* CompleteCallback, uint8* UserSuppliedMemory) const
 {
 	const int64 DataSize = GetBulkDataSize();
diff --git a/Engine/Source/Runtime/CoreUObject/Private/Serialization/BulkData2.cpp b/Engine/Source/Runtime/CoreUObject/Private/Serialization/BulkData2.cpp
--- a/Engine/Source/Runtime/CoreUObject/Private/Serialization/BulkData2.cpp
+++ b/Engine/Source/Runtime/CoreUObject/Private/Serialization/BulkData2.cpp
@@ -222,6 +222,143 @@ namespace FileTokenSystem
 
 FIoDispatcher* FBulkDataBase::IoDispatcher = nullptr;
 
+class FSizeChunkIdRequest : public IAsyncReadRequest
+{
+public:
+	FSizeChunkIdRequest(const FIoChunkId& ChunkId, FAsyncFileCallBack* Callback)
+		: IAsyncReadRequest(Callback, true, nullptr)
+	{
+		TIoStatusOr<uint64> Result = FBulkDataBase::GetIoDispatcher()->GetSizeForChunk(ChunkId);
+		if (Result.IsOk())
+		{
+			Size = Result.ValueOrDie();
+		}
+
+		SetComplete();
+	}
+
+	virtual ~FSizeChunkIdRequest() = default;
+
+private:
+
+	virtual void WaitCompletionImpl(float TimeLimitSeconds) override
+	{
+		// No need to wait as the work is done in the constructor
+	}
+
+	virtual void CancelImpl() override
+	{
+		// No point canceling as the work is done in the constructor
+	}
+};
+
+class FReadChunkIdRequest : public IAsyncReadRequest
+{
+public:
+	FReadChunkIdRequest(const FIoChunkId& InChunkId, FAsyncFileCallBack* InCallback, uint8* InUserSuppliedMemory, int64 InOffset, int64 InBytesToRead)
+		: IAsyncReadRequest(InCallback, false, InUserSuppliedMemory)
+	{
+		// Because IAsyncReadRequest can return ownership of the target memory buffer in the form
+		// of a raw pointer we must pass our own memory buffer to the FIoDispatcher otherwise the 
+		// buffer that will be returned cannot have it's lifetime managed correctly.
+		if (InUserSuppliedMemory == nullptr)
+		{
+			Memory = (uint8*)FMemory::Malloc(InBytesToRead);
+		}
+
+		DoneEvent = FPlatformProcess::GetSynchEventFromPool(true);
+		DoneEvent->Reset();
+		
+		FIoReadOptions Options(InOffset, InBytesToRead);
+		Options.SetTargetVa(Memory);
+
+		auto OnRequestLoaded = [this](TIoStatusOr<FIoBuffer> Result)
+		{
+			SetComplete();
+			ReleaseEvent(true);
+		};
+
+		FBulkDataBase::GetIoDispatcher()->ReadWithCallback(InChunkId, Options, OnRequestLoaded);
+	}
+
+	virtual ~FReadChunkIdRequest()
+	{
+		ReleaseEvent(false);
+
+		// Free memory if the request allocated it (although if the user accessed the memory after
+		// reading then they will have taken ownership of it anyway, and if they didn't access the
+		// memory then why did we read it in the first place?)
+		if (Memory != nullptr && !bUserSuppliedMemory)
+		{
+			FMemory::Free(Memory);
+		}
+
+		// ~IAsyncReadRequest expects Memory to be nullptr, even if the memory was user supplied
+		Memory = nullptr;
+	}
+
+protected:
+
+	virtual void WaitCompletionImpl(float TimeLimitSeconds) override
+	{
+		if (!PollCompletion())
+		{
+			uint32 TimeLimitMilliseconds = TimeLimitSeconds <= 0.0f ? (uint32)(TimeLimitSeconds * 1000.0f) : MAX_uint32;
+			DoneEvent->Wait(TimeLimitMilliseconds);
+		}
+	}
+
+	virtual void CancelImpl() override
+	{
+		bCanceled = true;
+		SetComplete();
+
+		DoneEvent->Trigger();
+		ReleaseEvent(true);
+	}
+
+	void ReleaseEvent(bool bShouldTrigger)
+	{
+		if (DoneEvent != nullptr)
+		{
+			if (bShouldTrigger)
+			{
+				DoneEvent->Trigger();
+			}
+
+			FPlatformProcess::ReturnSynchEventToPool(DoneEvent);
+			DoneEvent = nullptr;
+		}
+	}
+
+	FIoChunkId ChunkId;
+	FEvent* DoneEvent;
+};
+
+class FAsyncReadChunkIdHandle : public IAsyncReadFileHandle
+{
+public:
+	FAsyncReadChunkIdHandle(const FIoChunkId& InChunkID) 
+		: ChunkID(InChunkID)
+	{
+
+	}
+
+	virtual ~FAsyncReadChunkIdHandle() = default;
+
+	virtual IAsyncReadRequest* SizeRequest(FAsyncFileCallBack* CompleteCallback = nullptr) override
+	{
+		return new FSizeChunkIdRequest(ChunkID, CompleteCallback);
+	}
+
+	virtual IAsyncReadRequest* ReadRequest(int64 Offset, int64 BytesToRead, EAsyncIOPriorityAndFlags PriorityAndFlags = AIOP_Normal, FAsyncFileCallBack* CompleteCallback = nullptr, uint8* UserSuppliedMemory = nullptr) override
+	{
+		return new FReadChunkIdRequest(ChunkID, CompleteCallback, UserSuppliedMemory, Offset, BytesToRead);
+	}
+private:
+	FIoChunkId ChunkID;
+};
+
 class FBulkDataIoDispatcherRequest : public IBulkDataIORequest
 {
 public:
@@ -365,7 +502,7 @@ private:
 		//IoOptions.SetRange(OffsetInBulkData, BytesToRead);
 		//IoOptions.SetTargetVa((uint64)UserSuppliedMemory);
 
-		FIoBatch NewBatch = FBulkDataBase::IoDispatcher->NewBatch();
+		FIoBatch NewBatch = FBulkDataBase::GetIoDispatcher()->NewBatch();
 		for (Request& Request : RequestArray)
 		{
 			Request.IoRequest = NewBatch.Read(Request.ChunkId, IoOptions);
@@ -395,7 +532,7 @@ private:
 			DstBuffer += BytesToRead;
 		}
 		
-		FBulkDataBase::IoDispatcher->FreeBatch(NewBatch);
+		FBulkDataBase::GetIoDispatcher()->FreeBatch(NewBatch);
 
 		bIsCompleted = true;
 		FPlatformMisc::MemoryBarrier();
@@ -941,6 +1078,18 @@ bool FBulkDataBase::IsUsingIODispatcher() const
 	return (BulkDataFlags & BULKDATA_UsesIoDispatcher) != 0;
 }
 
+IAsyncReadFileHandle* FBulkDataBase::OpenAsyncReadHandle() const
+{
+	if (IsUsingIODispatcher())
+	{
+		return new FAsyncReadChunkIdHandle(ChunkID);
+	}
+	else
+	{
+		return FPlatformFileManager::Get().GetPlatformFile().OpenAsyncRead(*GetFilename());
+	}
+}
+
 IBulkDataIORequest* FBulkDataBase::CreateStreamingRequest(EAsyncIOPriorityAndFlags Priority, FBulkDataIORequestCallBack* CompleteCallback, uint8* UserSuppliedMemory) const
 {
 	const int64 DataSize = GetBulkDataSize();
@@ -1072,7 +1221,8 @@ int64 FBulkDataBase::GetBulkDataOffsetInFile() const
 	}
 	else
 	{
-		BULKDATA_NOT_IMPLEMENTED_FOR_RUNTIME;
+		// When using the IODispatcher the BulkData object will point directly to the correct data
+		// so we don't need to consider the offset at all.
 		return 0;
 	}
 }
diff --git a/Engine/Source/Runtime/CoreUObject/Public/Serialization/BulkData.h b/Engine/Source/Runtime/CoreUObject/Public/Serialization/BulkData.h
--- a/Engine/Source/Runtime/CoreUObject/Public/Serialization/BulkData.h
+++ b/Engine/Source/Runtime/CoreUObject/Public/Serialization/BulkData.h
@@ -680,53 +680,60 @@ public:
 	/*-----------------------------------------------------------------------------
 		Async Streaming Interface.
 	-----------------------------------------------------------------------------*/
-	
+
 	/**
-	* Create an async read request for the bulk data.
-	* This version will load the entire data range that the FUntypedBulkData represents.
-	*
-	* @param Priority			Priority and flags of the request. If this includes AIOP_FLAG_PRECACHE, then memory will never be returned. The request should always be canceled and waited for, even for a precache request.
-	* @param CompleteCallback	Called from an arbitrary thread when the request is complete. Can be nullptr, if non-null, must remain valid until it is called. It will always be called.
-	* @param UserSuppliedMemory A pointer to memory for the IO request to be written to, it is up to the caller to make sure that it is large enough. If the pointer is null then the system will allocate memory instead.
-	* @return					A request for the read. This is owned by the caller and must be deleted by the caller.
-	**/
+	 * Opens a new IAsyncReadFileHandle that references the file that the BulkData object represents.
+	 *
+	 * @return A valid handle if the file can be accessed, if it cannot then nullptr.
+	 */
+	IAsyncReadFileHandle* OpenAsyncReadHandle() const;
+
+	/**
+	 * Create an async read request for the bulk data.
+	 * This version will load the entire data range that the FUntypedBulkData represents.
+	 *
+	 * @param Priority				Priority and flags of the request. If this includes AIOP_FLAG_PRECACHE, then memory will never be returned. The request should always be canceled and waited for, even for a precache request.
+	 * @param CompleteCallback		Called from an arbitrary thread when the request is complete. Can be nullptr, if non-null, must remain valid until it is called. It will always be called.
+	 * @param UserSuppliedMemory	A pointer to memory for the IO request to be written to, it is up to the caller to make sure that it is large enough. If the pointer is null then the system will allocate memory instead.
+	 * @return						A request for the read. This is owned by the caller and must be deleted by the caller.
+	 */
 	IBulkDataIORequest* CreateStreamingRequest(EAsyncIOPriorityAndFlags Priority, FBulkDataIORequestCallBack* CompleteCallback, uint8* UserSuppliedMemory) const;
 
 	/**
-	* Create an async read request for the bulk data.
-	* This version allows the user to request a subset of the data that the FUntypedBulkData represents.
-	*
-	* @param OffsetInBulkData	Offset into the bulk data to start reading from.
-	* @param BytesToRead		The number of bytes to read. If this request is AIOP_Preache, the size can be anything, even MAX_int64, otherwise the size and offset must be fully contained in the file.
-	* @param Priority			Priority and flags of the request. If this includes AIOP_FLAG_PRECACHE, then memory will never be returned. The request should always be canceled and waited for, even for a precache request.
-	* @param CompleteCallback	Called from an arbitrary thread when the request is complete. Can be nullptr, if non-null, must remain valid until it is called. It will always be called.
-	* @param UserSuppliedMemory A pointer to memory for the IO request to be written to, it is up to the caller to make sure that it is large enough. If the pointer is null then the system will allocate memory instead.
-	* @return					A request for the read. This is owned by the caller and must be deleted by the caller.
-	**/
+	 * Create an async read request for the bulk data.
+	 * This version allows the user to request a subset of the data that the FUntypedBulkData represents.
+	 *
+	 * @param OffsetInBulkData		Offset into the bulk data to start reading from.
+	 * @param BytesToRead			The number of bytes to read. If this request is AIOP_Preache, the size can be anything, even MAX_int64, otherwise the size and offset must be fully contained in the file.
+	 * @param Priority				Priority and flags of the request. If this includes AIOP_FLAG_PRECACHE, then memory will never be returned. The request should always be canceled and waited for, even for a precache request.
+	 * @param CompleteCallback		Called from an arbitrary thread when the request is complete. Can be nullptr, if non-null, must remain valid until it is called. It will always be called.
+	 * @param UserSuppliedMemory	A pointer to memory for the IO request to be written to, it is up to the caller to make sure that it is large enough. If the pointer is null then the system will allocate memory instead.
+	 * @return						A request for the read. This is owned by the caller and must be deleted by the caller.
+	 */
 	IBulkDataIORequest* CreateStreamingRequest(int64 OffsetInBulkData, int64 BytesToRead, EAsyncIOPriorityAndFlags Priority, FBulkDataIORequestCallBack* CompleteCallback, uint8* UserSuppliedMemory) const;
 
 #if USE_BULKDATA_STREAMING_TOKEN 
 
 	/**
-	* Creates a FBulkDataStreamingToken representing the area of the file that the FUntypedBulkData represents. See the declaration of 
-	* FBulkDataStreamingToken for further details.
-	* 
-	* @return	A FBulkDataStreamingToken valid for the FUntypedBulkData.
-	**/
+	 * Creates a FBulkDataStreamingToken representing the area of the file that the FUntypedBulkData represents. See the declaration of
+	 * FBulkDataStreamingToken for further details.
+	 *
+	 * @return	A FBulkDataStreamingToken valid for the FUntypedBulkData.
+	 */
 	FBulkDataStreamingToken CreateStreamingToken() const;
 
 	/**
-	* Create an async read request for a range of bulk data streaming tokens
-	* The request will read all data between the two given streaming tokens objects. They must both represent areas of data in the file!
-	* There is no way to validate this and it is up to the caller to make sure that it is correct.
-	* The memory to be read into will be automatically allocated the size of which can be retrieved by calling IBulkDataIORequest::GetSize()
-	*
-	* @param Filename			The file to read from.
-	* @param Start				The bulk data to start reading from.
-	* @param End				The bulk data to finish reading from.
-	* @param Priority			Priority and flags of the request. If this includes AIOP_FLAG_PRECACHE, then memory will never be returned. The request should always be canceled and waited for, even for a precache request.
-	* @param CompleteCallback	Called from an arbitrary thread when the request is complete. Can be nullptr, if non-null, must remain valid until it is called. It will always be called.
-	* @return					A request for the read. This is owned by the caller and must be deleted by the caller.
+	 * Create an async read request for a range of bulk data streaming tokens
+	 * The request will read all data between the two given streaming tokens objects. They must both represent areas of data in the file!
+	 * There is no way to validate this and it is up to the caller to make sure that it is correct.
+	 * The memory to be read into will be automatically allocated the size of which can be retrieved by calling IBulkDataIORequest::GetSize()
+	 *
+	 * @param Filename			The file to read from.
+	 * @param Start				The bulk data to start reading from.
+	 * @param End				The bulk data to finish reading from.
+	 * @param Priority			Priority and flags of the request. If this includes AIOP_FLAG_PRECACHE, then memory will never be returned. The request should always be canceled and waited for, even for a precache request.
+	 * @param CompleteCallback	Called from an arbitrary thread when the request is complete. Can be nullptr, if non-null, must remain valid until it is called. It will always be called.
+	 * @return					A request for the read. This is owned by the caller and must be deleted by the caller.
 	**/
 	static IBulkDataIORequest* CreateStreamingRequestForRange(const FString& Filename, const BulkDataRangeArray& RangeArray, EAsyncIOPriorityAndFlags Priority, FBulkDataIORequestCallBack* CompleteCallback);
 #endif
diff --git a/Engine/Source/Runtime/CoreUObject/Public/Serialization/BulkData2.h b/Engine/Source/Runtime/CoreUObject/Public/Serialization/BulkData2.h
--- a/Engine/Source/Runtime/CoreUObject/Public/Serialization/BulkData2.h
+++ b/Engine/Source/Runtime/CoreUObject/Public/Serialization/BulkData2.h
@@ -49,8 +49,8 @@ class COREUOBJECT_API FBulkDataBase
 public:
 	using BulkDataRangeArray = TArray<FBulkDataBase*, TInlineAllocator<8>>;
 
-	static void SetIODispatcher(FIoDispatcher* InIoDispatcher) { IoDispatcher = InIoDispatcher; }
-
+	static void				SetIoDispatcher(FIoDispatcher* InIoDispatcher) { IoDispatcher = InIoDispatcher; }
+	static FIoDispatcher*	GetIoDispatcher() { return IoDispatcher; }
 public:
 	using FileToken = int32;
 	static constexpr FileToken InvalidToken = INDEX_NONE;
@@ -118,6 +118,8 @@ public:
 	bool IsMemoryMapped() const;
 	bool IsUsingIODispatcher() const;
 
+	IAsyncReadFileHandle* OpenAsyncReadHandle() const;
+
 	IBulkDataIORequest* CreateStreamingRequest(EAsyncIOPriorityAndFlags Priority, FBulkDataIORequestCallBack* CompleteCallback, uint8* UserSuppliedMemory) const;
 	IBulkDataIORequest* CreateStreamingRequest(int64 OffsetInBulkData, int64 BytesToRead, EAsyncIOPriorityAndFlags Priority, FBulkDataIORequestCallBack* CompleteCallback, uint8* UserSuppliedMemory) const;
 	
@@ -148,7 +150,7 @@ private:
 	FString ConvertFilenameFromFlags(const FString& Filename) const;
 
 private:
-	friend class FBulkDataIoDispatcherRequest;
+
 	static FIoDispatcher* IoDispatcher;
 
 	union
diff --git a/Engine/Source/Runtime/Engine/Private/FileCache/FileCache.cpp b/Engine/Source/Runtime/Engine/Private/FileCache/FileCache.cpp
--- a/Engine/Source/Runtime/Engine/Private/FileCache/FileCache.cpp
+++ b/Engine/Source/Runtime/Engine/Private/FileCache/FileCache.cpp
@@ -641,6 +641,20 @@ IFileCacheHandle* IFileCacheHandle::CreateFileCacheHandle(const TCHAR* InFileNam
 	return new FFileCacheHandle(FileHandle);
 }
 
+IFileCacheHandle* IFileCacheHandle::CreateFileCacheHandle(IAsyncReadFileHandle* FileHandle)
+{
+	SCOPE_CYCLE_COUNTER(STAT_SFC_CreateHandle);
+
+	if (FileHandle != nullptr)
+	{
+		return new FFileCacheHandle(FileHandle);	
+	}
+	else
+	{
+		return nullptr;
+	}
+}
+
 uint32 IFileCacheHandle::GetFileCacheSize()
 {
 	return GetCache().SizeInBytes;
diff --git a/Engine/Source/Runtime/Engine/Private/FileCache/FileCache.h b/Engine/Source/Runtime/Engine/Private/FileCache/FileCache.h
--- a/Engine/Source/Runtime/Engine/Private/FileCache/FileCache.h
+++ b/Engine/Source/Runtime/Engine/Private/FileCache/FileCache.h
@@ -21,7 +21,21 @@ class IFileCacheHandle
 public:
 	static void EvictAll();
 
-	static IFileCacheHandle *CreateFileCacheHandle(const TCHAR* InFileName);
+	/** 
+	 * Create a IFileCacheHandle from a filename.
+	 * @param	InFileName			A path to a file.
+	 * @return	A IFileCacheHandle that can be used to make read requests. This will be a nullptr if the target file can not be accessed 
+	 *			for any given reason.
+	 */
+	static IFileCacheHandle* CreateFileCacheHandle(const TCHAR* InFileName);
+
+	/**
+	 * Create a IFileCacheHandle from a IAsyncReadFileHandle.
+	 * @param	FileHandle			A valid IAsyncReadFileHandle that has already been created elsewhere.	
+	 * @return	A IFileCacheHandle that can be used to make read requests. This will be a nullptr if the FileHandle was not valid.
+	 */
+	static IFileCacheHandle* CreateFileCacheHandle(IAsyncReadFileHandle* FileHandle);
+
 	virtual ~IFileCacheHandle() {};
 
 	/** Return size of underlying file cache in bytes. */
diff --git a/Engine/Source/Runtime/Engine/Private/VT/UploadingVirtualTexture.cpp b/Engine/Source/Runtime/Engine/Private/VT/UploadingVirtualTexture.cpp
--- a/Engine/Source/Runtime/Engine/Private/VT/UploadingVirtualTexture.cpp
+++ b/Engine/Source/Runtime/Engine/Private/VT/UploadingVirtualTexture.cpp
@@ -313,6 +313,7 @@ FVTDataAndStatus FUploadingVirtualTexture::ReadData(FGraphEventArray& OutComplet
 	FVirtualTextureDataChunk& Chunk = Data->Chunks[ChunkIndex];
 	FByteBulkData& BulkData = Chunk.BulkData;
 	FString ChunkFileName;
+	FIoChunkId ChunkId = FIoChunkId::InvalidChunkId;
 	int64 ChunkOffsetInFile = 0;
 
 #if WITH_EDITOR
@@ -355,8 +356,8 @@ FVTDataAndStatus FUploadingVirtualTexture::ReadData(FGraphEventArray& OutComplet
 		ChunkFileName = ChunkFileNameDCC;
 	}
 #else // WITH_EDITOR
-	ChunkFileName = BulkData.GetFilename();
 	ChunkOffsetInFile = BulkData.GetBulkDataOffsetInFile();
+	
 	if (BulkData.GetBulkDataSize() == 0)
 	{
 		if (!InvalidChunks[ChunkIndex])
@@ -371,7 +372,18 @@ FVTDataAndStatus FUploadingVirtualTexture::ReadData(FGraphEventArray& OutComplet
 	TUniquePtr<IFileCacheHandle>& Handle = HandlePerChunk[ChunkIndex];
 	if (!Handle)
 	{
-		Handle.Reset(IFileCacheHandle::CreateFileCacheHandle(*ChunkFileName));
+		// If we have a valid file name then we can create a file handle directly to it.
+		// If we do not then pass in the BulkData object which will create the IAsyncReadFileHandle for us.
+		//
+		if (!ChunkFileName.IsEmpty())
+		{
+			Handle.Reset(IFileCacheHandle::CreateFileCacheHandle(*ChunkFileName));
+		}
+		else
+		{
+			Handle.Reset(IFileCacheHandle::CreateFileCacheHandle(BulkData.OpenAsyncReadHandle()));
+		}
+
 		// Don't expect CreateFileCacheHandle() to fail, async files should never fail to open
 		if (!ensure(Handle.IsValid()))
 		{
