Optimize significant callers of FName::ToString

This reduced the cost of FName::ToString calls by the modified functions by about 70%, which saved 30 seconds of wall time during cooking of a large title.

diff --git a/Engine/Source/Runtime/CoreUObject/Private/UObject/Class.cpp b/Engine/Source/Runtime/CoreUObject/Private/UObject/Class.cpp
--- a/Engine/Source/Runtime/CoreUObject/Private/UObject/Class.cpp
+++ b/Engine/Source/Runtime/CoreUObject/Private/UObject/Class.cpp
@@ -15,6 +15,7 @@
 #include "Misc/OutputDeviceConsole.h"
 #include "Misc/AutomationTest.h"
 #include "Misc/EnumClassFlags.h"
+#include "Misc/StringBuilder.h"
 #include "UObject/ErrorException.h"
 #include "Modules/ModuleManager.h"
 #include "UObject/UObjectAllocator.h"
@@ -1497,7 +1498,9 @@ void UStruct::SerializeVersionedTaggedProperties(FStructuredArchive::FSlot Slot,
 							Tag.SetPropertyGuid(PropertyGuid);
 						}
 
-						FStructuredArchive::FSlot PropertySlot = StaticArrayContainer.IsSet() ? StaticArrayContainer->EnterElement() : PropertiesRecord.EnterField(SA_FIELD_NAME(*Tag.Name.ToString()));
+						TStringBuilder<256> TagName;
+						Tag.Name.ToString(TagName);
+						FStructuredArchive::FSlot PropertySlot = StaticArrayContainer.IsSet() ? StaticArrayContainer->EnterElement() : PropertiesRecord.EnterField(SA_FIELD_NAME(TagName.ToString()));
 
 						PropertySlot << Tag;
 
diff --git a/Engine/Source/Runtime/CoreUObject/Private/UObject/Linker.cpp b/Engine/Source/Runtime/CoreUObject/Private/UObject/Linker.cpp
--- a/Engine/Source/Runtime/CoreUObject/Private/UObject/Linker.cpp
+++ b/Engine/Source/Runtime/CoreUObject/Private/UObject/Linker.cpp
@@ -5,9 +5,11 @@
 =============================================================================*/
 
 #include "UObject/Linker.h"
+#include "Containers/StringView.h"
 #include "Misc/PackageName.h"
 #include "Misc/CommandLine.h"
 #include "Misc/Paths.h"
+#include "Misc/PathViews.h"
 #include "UObject/Package.h"
 #include "Templates/Casts.h"
 #include "UObject/UnrealType.h"
@@ -615,7 +617,7 @@ FLinkerLoad* GetPackageLinker
 		// Process any package redirects
 		{
 			const FCoreRedirectObjectName NewPackageName = FCoreRedirects::GetRedirectedName(ECoreRedirectFlags::Type_Package, FCoreRedirectObjectName(NAME_None, NAME_None, *PackageNameToCreate));
-			PackageNameToCreate = NewPackageName.PackageName.ToString();
+			NewPackageName.PackageName.ToString(PackageNameToCreate);
 		}
 
 		// The editor must not redirect packages for localization. We also shouldn't redirect script or in-memory packages.
@@ -652,7 +654,7 @@ FLinkerLoad* GetPackageLinker
 		// Process any package redirects
 		{
 			const FCoreRedirectObjectName NewPackageName = FCoreRedirects::GetRedirectedName(ECoreRedirectFlags::Type_Package, FCoreRedirectObjectName(NAME_None, NAME_None, *PackageNameToCreate));
-			PackageNameToCreate = NewPackageName.PackageName.ToString();
+			NewPackageName.PackageName.ToString(PackageNameToCreate);
 		}
 
 		// The editor must not redirect packages for localization. We also shouldn't redirect script packages.
@@ -691,7 +693,7 @@ FLinkerLoad* GetPackageLinker
 
 #if WITH_EDITORONLY_DATA
 		// Make sure the package name matches the name on disk
-		FPackageName::FixPackageNameCase(PackageNameToCreate, FPaths::GetExtension(NewFilename));
+		FPackageName::FixPackageNameCase(PackageNameToCreate, FPathViews::GetExtension(NewFilename));
 #endif
 
 		// Create the package with the provided long package name.
diff --git a/Engine/Source/Runtime/CoreUObject/Private/UObject/LinkerLoad.cpp b/Engine/Source/Runtime/CoreUObject/Private/UObject/LinkerLoad.cpp
--- a/Engine/Source/Runtime/CoreUObject/Private/UObject/LinkerLoad.cpp
+++ b/Engine/Source/Runtime/CoreUObject/Private/UObject/LinkerLoad.cpp
@@ -2795,6 +2795,16 @@ bool FLinkerLoad::VerifyImportInner(const int32 ImportIndex, FString& WarningSuf
 		return false;
 	}
 
+	TStringBuilder<256> ImportObjectName;
+	auto GetImportObjectName = [&Import, &ImportObjectName]
+	{
+		if (!ImportObjectName.Len())
+		{
+			Import.ObjectName.AppendString(ImportObjectName);
+		}
+		return ImportObjectName.ToString();
+	};
+
 	bool SafeReplace = false;
 	UObject* Pkg=NULL;
 	UPackage* TmpPkg=NULL;
@@ -2802,7 +2812,7 @@ bool FLinkerLoad::VerifyImportInner(const int32 ImportIndex, FString& WarningSuf
 	// Find or load the linker load that contains the FObjectExport for this import
 	if (Import.OuterIndex.IsNull() && Import.ClassName!=NAME_Package )
 	{
-		UE_LOG(LogLinker, Error, TEXT("%s has an inappropriate outermost, it was probably saved with a deprecated outer (file: %s)"), *Import.ObjectName.ToString(), *Filename);
+		UE_LOG(LogLinker, Error, TEXT("%s has an inappropriate outermost, it was probably saved with a deprecated outer (file: %s)"), GetImportObjectName(), *Filename);
 		Import.SourceLinker = NULL;
 		return false;
 	}
@@ -2847,7 +2857,7 @@ bool FLinkerLoad::VerifyImportInner(const int32 ImportIndex, FString& WarningSuf
 
 			// we now fully load the package that we need a single export from - however, we still use CreatePackage below as it handles all cases when the package
 			// didn't exist (native only), etc		
-			TmpPkg = LoadPackageInternal(NULL, *Import.ObjectName.ToString(), InternalLoadFlags | LOAD_IsVerifying, this, nullptr, GetSerializeContext());
+			TmpPkg = LoadPackageInternal(NULL, GetImportObjectName(), InternalLoadFlags | LOAD_IsVerifying, this, nullptr, GetSerializeContext());
 		}
 
 #if WITH_EDITOR
@@ -2861,7 +2871,7 @@ bool FLinkerLoad::VerifyImportInner(const int32 ImportIndex, FString& WarningSuf
 		// @todo linkers: This could quite possibly be cleaned up
 		if (TmpPkg == NULL)
 		{
-			TmpPkg = CreatePackage( NULL, *Import.ObjectName.ToString() );
+			TmpPkg = CreatePackage( NULL, GetImportObjectName() );
 		}
 
 		// if we couldn't create the package or it is 
@@ -2957,7 +2967,9 @@ bool FLinkerLoad::VerifyImportInner(const int32 ImportIndex, FString& WarningSuf
 			}
 
 			// Top is now pointing to the top-level UPackage for this resource
-			Pkg = CreatePackage(NULL, *Top->ObjectName.ToString() );
+			TStringBuilder<256> TopObjectName;
+			Top->ObjectName.ToString(TopObjectName);
+			Pkg = CreatePackage(NULL, TopObjectName.ToString());
 
 			// Find this import within its existing linker.
 			int32 iHash = HashNames( Import.ObjectName, Import.ClassName, Import.ClassPackage) & (ExportHashCount-1);
@@ -2970,7 +2982,7 @@ bool FLinkerLoad::VerifyImportInner(const int32 ImportIndex, FString& WarningSuf
 			{
 				if (!Import.SourceLinker->ExportMap.IsValidIndex(j))
 				{
-					UE_LOG(LogLinker, Error, TEXT("Invalid index [%d/%d] while attempting to import '%s' with LinkerRoot '%s'"), j, Import.SourceLinker->ExportMap.Num(), *Import.ObjectName.ToString(), *GetNameSafe(Import.SourceLinker->LinkerRoot));
+					UE_LOG(LogLinker, Error, TEXT("Invalid index [%d/%d] while attempting to import '%s' with LinkerRoot '%s'"), j, Import.SourceLinker->ExportMap.Num(), GetImportObjectName(), *GetNameSafe(Import.SourceLinker->LinkerRoot));
 					break;
 				}
 				else
@@ -2996,7 +3008,7 @@ bool FLinkerLoad::VerifyImportInner(const int32 ImportIndex, FString& WarningSuf
 			for( int32 j=Import.SourceLinker->ExportHash[iHash]; j!=INDEX_NONE; j=Import.SourceLinker->ExportMap[j].HashNext )
 			{
 				if (!ensureMsgf(Import.SourceLinker->ExportMap.IsValidIndex(j), TEXT("Invalid index [%d/%d] while attempting to import '%s' with LinkerRoot '%s'"),
-					j, Import.SourceLinker->ExportMap.Num(), *Import.ObjectName.ToString(), *GetNameSafe(Import.SourceLinker->LinkerRoot)))
+					j, Import.SourceLinker->ExportMap.Num(), GetImportObjectName(), *GetNameSafe(Import.SourceLinker->LinkerRoot)))
 				{
 					break;
 				}
@@ -3080,7 +3092,7 @@ bool FLinkerLoad::VerifyImportInner(const int32 ImportIndex, FString& WarningSuf
 									FObjectImport& TestImport = ImportMap[i];
 									if ( TestImport.OuterIndex == FoundIndex )
 									{
-										UE_LOG(LogLinker, Log, TEXT("Private import was referenced by import '%s' (outer)"), *Import.ObjectName.ToString());
+										UE_LOG(LogLinker, Log, TEXT("Private import was referenced by import '%s' (outer)"), GetImportObjectName());
 										SafeReplace = false;
 									}
 								}
@@ -3140,10 +3152,13 @@ bool FLinkerLoad::VerifyImportInner(const int32 ImportIndex, FString& WarningSuf
 	// If not found in file, see if it's a public native transient class or field.
 	if( Import.SourceIndex==INDEX_NONE && Pkg!=NULL )
 	{
-		UObject* ClassPackage = FindObject<UPackage>( NULL, *Import.ClassPackage.ToString() );
+		TStringBuilder<256> ImportClassTemp;
+		Import.ClassPackage.ToString(ImportClassTemp);
+		UObject* ClassPackage = FindObject<UPackage>( NULL, ImportClassTemp.ToString() );
 		if( ClassPackage )
 		{
-			UClass* FindClass = FindObject<UClass>( ClassPackage, *Import.ClassName.ToString() );
+			Import.ClassName.ToString(ImportClassTemp);
+			UClass* FindClass = FindObject<UClass>( ClassPackage, ImportClassTemp.ToString() );
 			if( FindClass )
 			{
 				UObject* FindOuter			= Pkg;
@@ -3162,13 +3177,13 @@ bool FLinkerLoad::VerifyImportInner(const int32 ImportIndex, FString& WarningSuf
 					}
 				}
 
-				UObject* FindObject = FindImport(FindClass, FindOuter, *Import.ObjectName.ToString());
+				UObject* FindObject = FindImport(FindClass, FindOuter, GetImportObjectName());
 				// Reference to in memory-only package's object, native transient class or CDO of such a class.
 				bool bIsInMemoryOnlyOrNativeTransient = bCameFromMemoryOnlyPackage || (FindObject != NULL && ((FindObject->IsNative() && FindObject->HasAllFlags(RF_Public | RF_Transient)) || (FindObject->HasAnyFlags(RF_ClassDefaultObject) && FindObject->GetClass()->IsNative() && FindObject->GetClass()->HasAllFlags(RF_Public | RF_Transient))));
 				// Check for structs which have been moved to another header (within the same class package).
 				if (!FindObject && bIsInMemoryOnlyOrNativeTransient && FindClass == UScriptStruct::StaticClass())
 				{
-					FindObject = StaticFindObject( FindClass, ANY_PACKAGE, *Import.ObjectName.ToString(), true );
+					FindObject = StaticFindObject( FindClass, ANY_PACKAGE, GetImportObjectName(), true );
 					if (FindObject && FindOuter->GetOutermost() != FindObject->GetOutermost())
 					{
 						// Limit the results to the same package.I
diff --git a/Engine/Source/Runtime/Engine/Private/AssetManager.cpp b/Engine/Source/Runtime/Engine/Private/AssetManager.cpp
--- a/Engine/Source/Runtime/Engine/Private/AssetManager.cpp
+++ b/Engine/Source/Runtime/Engine/Private/AssetManager.cpp
@@ -5,6 +5,7 @@
 #include "Engine/PrimaryAssetLabel.h"
 #include "AssetData.h"
 #include "ARFilter.h"
+#include "Containers/StringView.h"
 #include "Engine/Engine.h"
 #include "Engine/BlueprintGeneratedClass.h"
 #include "UObject/ConstructorHelpers.h"
@@ -12,6 +13,7 @@
 #include "Misc/FileHelper.h"
 #include "Misc/ScopedSlowTask.h"
 #include "Misc/Paths.h"
+#include "Misc/StringBuilder.h"
 #include "Serialization/MemoryReader.h"
 #include "AssetRegistryState.h"
 #include "HAL/PlatformFilemanager.h"
@@ -2824,10 +2826,11 @@ EAssetSetManagerResult::Type UAssetManager::ShouldSetManager(const FAssetIdentif
 		return EAssetSetManagerResult::SetButDoNotRecurse;
 	}
 
-	const FString TargetPackageString = Target.PackageName.ToString();
+	TStringBuilder<256> TargetPackageString;
+	Target.PackageName.ToString(TargetPackageString);
 
 	// Ignore script references
-	if (TargetPackageString.StartsWith(TEXT("/Script/"), ESearchCase::CaseSensitive))
+	if (FStringView(TargetPackageString).StartsWith(TEXT("/Script/"), ESearchCase::CaseSensitive))
 	{
 		return EAssetSetManagerResult::DoNotSet;
 	}
