Fix static asserts when calling some UEnum functions with a TEnumAsByte.

#rnx

diff --git a/Engine/Source/Runtime/CoreUObject/Public/UObject/Class.h b/Engine/Source/Runtime/CoreUObject/Public/UObject/Class.h
--- a/Engine/Source/Runtime/CoreUObject/Public/UObject/Class.h
+++ b/Engine/Source/Runtime/CoreUObject/Public/UObject/Class.h
@@ -2216,7 +2216,7 @@ public:
 	template<typename EnumType>
 	FORCEINLINE static FName GetValueAsName(const TEnumAsByte<EnumType> EnumeratorValue)
 	{
-		return GetValueAsName((int64)EnumeratorValue.GetValue());
+		return GetValueAsName(EnumeratorValue.GetValue());
 	}
 
 	template<typename EnumType>
@@ -2241,7 +2241,7 @@ public:
 	template<typename EnumType>
 	FORCEINLINE static FString GetValueAsString(const TEnumAsByte<EnumType> EnumeratorValue)
 	{
-		return GetValueAsString((int64)EnumeratorValue.GetValue());
+		return GetValueAsString(EnumeratorValue.GetValue());
 	}
 
 	template<typename EnumType>
@@ -2269,7 +2269,7 @@ public:
 	template<typename EnumType>
 	FORCEINLINE static FText GetDisplayValueAsText(const TEnumAsByte<EnumType> EnumeratorValue)
 	{
-		return GetDisplayValueAsText((int64)EnumeratorValue.GetValue());
+		return GetDisplayValueAsText(EnumeratorValue.GetValue());
 	}
 
 	template<typename EnumType>
