Map key allocation size fix.

diff --git a/Engine/Source/Runtime/CoreUObject/Private/UObject/PropertyMap.cpp b/Engine/Source/Runtime/CoreUObject/Private/UObject/PropertyMap.cpp
--- a/Engine/Source/Runtime/CoreUObject/Private/UObject/PropertyMap.cpp
+++ b/Engine/Source/Runtime/CoreUObject/Private/UObject/PropertyMap.cpp
@@ -337,7 +337,7 @@ void FMapProperty::SerializeItem(FStructuredArchive::FSlot Slot, void* Value, co
 		FStructuredArchive::FArray KeysToRemoveArray = Record.EnterArray(SA_FIELD_NAME(TEXT("KeysToRemove")), NumKeysToRemove);
 		if (NumKeysToRemove)
 		{
-			TempKeyStorage = (uint8*)FMemory::Malloc(MapLayout.SetLayout.Size);
+			TempKeyStorage = (uint8*)FMemory::Malloc(KeyProp->GetSize());
 			KeyProp->InitializeValue(TempKeyStorage);
 
 			FSerializedPropertyScope SerializedProperty(UnderlyingArchive, KeyProp, this);
@@ -360,7 +360,7 @@ void FMapProperty::SerializeItem(FStructuredArchive::FSlot Slot, void* Value, co
 		// Allocate temporary key space if we haven't allocated it already above
 		if (NumEntries != 0 && !TempKeyStorage)
 		{
-			TempKeyStorage = (uint8*)FMemory::Malloc(MapLayout.SetLayout.Size);
+			TempKeyStorage = (uint8*)FMemory::Malloc(KeyProp->GetSize());
 			KeyProp->InitializeValue(TempKeyStorage);
 		}
 
