Optimize FPackageName::GetShortFName()

Got about ~2x faster. Original UDN proposed doing a different and much more aggressive optimization, but that would break short package name lookup.

#udn 544531

diff --git a/Engine/Source/Runtime/CoreUObject/Private/Misc/PackageName.cpp b/Engine/Source/Runtime/CoreUObject/Private/Misc/PackageName.cpp
--- a/Engine/Source/Runtime/CoreUObject/Private/Misc/PackageName.cpp
+++ b/Engine/Source/Runtime/CoreUObject/Private/Misc/PackageName.cpp
@@ -81,22 +81,37 @@ FString FPackageName::GetShortName(const TCHAR* LongName)
 	return GetShortName(FString(LongName));
 }
 
-
 FName FPackageName::GetShortFName(const FString& LongName)
 {
-	int32 IndexOfLastSlash = INDEX_NONE;
-	LongName.FindLastChar('/', IndexOfLastSlash);
-	return FName(*LongName.Mid(IndexOfLastSlash + 1));
+	return GetShortFName(*LongName);
 }
 
 FName FPackageName::GetShortFName(const FName& LongName)
 {
-	return GetShortFName(LongName.ToString());
+	TCHAR LongNameStr[FName::StringBufferSize];
+	LongName.ToString(LongNameStr);
+
+	if (const TCHAR* Slash = FCString::Strrchr(LongNameStr, '/'))
+	{
+		return FName(Slash + 1);
+	}
+
+	return LongName;
 }
 
 FName FPackageName::GetShortFName(const TCHAR* LongName)
 {
-	return GetShortFName(FString(LongName));
+	if (LongName == nullptr)
+	{
+		return FName();
+	}
+
+	if (const TCHAR* Slash = FCString::Strrchr(LongName, '/'))
+	{
+		return FName(Slash + 1);
+	}
+
+	return FName(LongName);
 }
 
 struct FPathPair
