AsyncLoading2: Create clusters

diff --git a/Engine/Source/Runtime/CoreUObject/Private/Serialization/AsyncLoading2.cpp b/Engine/Source/Runtime/CoreUObject/Private/Serialization/AsyncLoading2.cpp
--- a/Engine/Source/Runtime/CoreUObject/Private/Serialization/AsyncLoading2.cpp
+++ b/Engine/Source/Runtime/CoreUObject/Private/Serialization/AsyncLoading2.cpp
@@ -48,6 +48,7 @@
 #include "UObject/ObjectRedirector.h"
 #include "Serialization/BulkData.h"
 #include "Serialization/LargeMemoryReader.h"
+#include "UObject/UObjectClusters.h"
 
 #if UE_BUILD_DEVELOPMENT || UE_BUILD_DEBUG
 //PRAGMA_DISABLE_OPTIMIZATION
@@ -1245,12 +1246,12 @@ struct FAsyncPackage2
 #endif
 
 	/** Checks if all dependencies (imported packages) of this package have been fully loaded */
-	bool AreAllDependenciesFullyLoaded(TSet<UPackage*>& VisitedPackages);
+	bool AreAllDependenciesFullyLoaded(TSet<FPackageId>& VisitedPackages);
 
 	/** Returs true if this package loaded objects that can create GC clusters */
 	bool HasClusterObjects() const
 	{
-		return DeferredClusterObjects.Num() > 0;
+		return bHasClusterObjects;
 	}
 
 	/** Creates GC clusters from loaded objects */
@@ -1262,7 +1263,7 @@ struct FAsyncPackage2
 private:
 
 	/** Checks if all dependencies (imported packages) of this package have been fully loaded */
-	bool AreAllDependenciesFullyLoadedInternal(FAsyncPackage2* Package, TSet<UPackage*>& VisitedPackages, FString& OutError);
+	bool AreAllDependenciesFullyLoadedInternal(FAsyncPackage2* Package, TSet<FPackageId>& VisitedPackages, FPackageId& OutPackageId);
 
 	TAtomic<int32> RefCount{ 0 };
 
@@ -1283,6 +1284,8 @@ private:
 	int32						DeferredFinalizeIndex;
 	/** Current index into DeferredClusterObjects array used to spread routing CreateClusters over several frames			*/
 	int32						DeferredClusterIndex;
+	/** True if any export can be a cluster root */
+	bool						bHasClusterObjects;
 	/** True if our load has failed */
 	bool						bLoadHasFailed;
 	/** True if our load has finished */
@@ -1293,8 +1296,6 @@ private:
 	double						LoadStartTime;
 	/** Estimated load percentage.																		*/
 	float						LoadPercentage;
-	/** Objects to create GC clusters from */
-	TArray<UObject*> DeferredClusterObjects;
 
 	/** List of all request handles */
 	TArray<int32, TInlineAllocator<2>> RequestIDs;
@@ -1311,8 +1312,6 @@ private:
 	FAsyncLoadingThread2Impl& AsyncLoadingThread;
 	IEDLBootNotificationManager& EDLBootNotificationManager;
 	FAsyncLoadEventGraphAllocator& GraphAllocator;
-	/** Packages that have been imported by this async package */
-	TSet<UPackage*> ImportedPackages;
 
 	FEventLoadNode2* PackageNodes = nullptr;
 	FEventLoadNode2* ExportBundleNodes = nullptr;
@@ -3212,43 +3211,31 @@ EAsyncPackageState::Type FAsyncLoadingThread2Impl::ProcessAsyncLoadingFromGameTh
 	return EAsyncPackageState::Complete;
 }
 
-bool FAsyncPackage2::AreAllDependenciesFullyLoadedInternal(FAsyncPackage2* Package, TSet<UPackage*>& VisitedPackages, FString& OutError)
+bool FAsyncPackage2::AreAllDependenciesFullyLoadedInternal(FAsyncPackage2* Package, TSet<FPackageId>& VisitedPackages, FPackageId& OutPackageId)
 {
-	for (UPackage* ImportPackage : Package->ImportedPackages)
+	TArrayView<const int32> ImportedPackageIndices(Package->ImportStore.ImportedPackages, Package->ImportStore.ImportedPackagesCount);
+
+	for (const int32 ImportedPackageIndex : ImportedPackageIndices)
 	{
-		if (VisitedPackages.Contains(ImportPackage))
+		const FPackageId PackageId = FPackageId::FromIndex(ImportedPackageIndex);
+
+		if (VisitedPackages.Contains(PackageId))
 		{
 			continue;
 		}
-		VisitedPackages.Add(ImportPackage);
+		VisitedPackages.Add(PackageId);
 
-		FAsyncPackage2* AsyncRoot = AsyncLoadingThread.GetAsyncPackage(ImportPackage->GetPackageId());
+		FAsyncPackage2* AsyncRoot = AsyncLoadingThread.GetAsyncPackage(PackageId);
 		if (AsyncRoot)
 		{
-			if (!AsyncRoot->bAllExportsSerialized)
-			{
-				OutError = FString::Printf(TEXT("%s Doesn't have all exports Serialized"), *Package->GetPackageName().ToString());
-				return false;
-			}
 			if (AsyncRoot->DeferredPostLoadIndex < AsyncRoot->ExportCount)
 			{
-				OutError = FString::Printf(TEXT("%s Doesn't have all objects processed by DeferredPostLoad"), *Package->GetPackageName().ToString());
+				OutPackageId = PackageId;
 				return false;
 			}
-			for (int32 LocalExportIndex = 0; LocalExportIndex < AsyncRoot->ExportCount; ++LocalExportIndex)
-			{
-				const FExportObject& Export = AsyncRoot->Exports[LocalExportIndex];
-				if (!Export.bFiltered && Export.Object->HasAnyFlags(RF_NeedPostLoad|RF_NeedLoad))
-				{
-					OutError = FString::Printf(TEXT("%s has not been %s"), *Export.Object->GetFullName(), 
-						Export.Object->HasAnyFlags(RF_NeedLoad) ? TEXT("Serialized") : TEXT("PostLoaded"));
-					return false;
-				}
-			}
 
-			if (!AreAllDependenciesFullyLoadedInternal(AsyncRoot, VisitedPackages, OutError))
+			if (!AreAllDependenciesFullyLoadedInternal(AsyncRoot, VisitedPackages, OutPackageId))
 			{
-				OutError = Package->GetPackageName().ToString() + TEXT("->") + OutError;
 				return false;
 			}
 		}
@@ -3256,15 +3243,18 @@ bool FAsyncPackage2::AreAllDependenciesFullyLoadedInternal(FAsyncPackage2* Packa
 	return true;
 }
 
-bool FAsyncPackage2::AreAllDependenciesFullyLoaded(TSet<UPackage*>& VisitedPackages)
+bool FAsyncPackage2::AreAllDependenciesFullyLoaded(TSet<FPackageId>& VisitedPackages)
 {
 	TRACE_CPUPROFILER_EVENT_SCOPE(AreAllDependenciesFullyLoaded);
 	VisitedPackages.Reset();
-	FString Error;
-	bool bLoaded = AreAllDependenciesFullyLoadedInternal(this, VisitedPackages, Error);
+	FPackageId PackageId;
+	const bool bLoaded = AreAllDependenciesFullyLoadedInternal(this, VisitedPackages, PackageId);
 	if (!bLoaded)
 	{
-		UE_LOG(LogStreaming, Verbose, TEXT("AreAllDependenciesFullyLoaded: %s"), *Error);
+		FAsyncPackage2* AsyncRoot = AsyncLoadingThread.GetAsyncPackage(PackageId);
+		TCHAR PackageName[FName::StringBufferSize];
+		AsyncRoot->GetPackageName().ToString(PackageName);
+		UE_LOG(LogStreaming, Verbose, TEXT("AreAllDependenciesFullyLoaded: '%s' doesn't have all exports processed by DeferredPostLoad"), PackageName);
 	}
 	return bLoaded;
 }
@@ -3356,7 +3346,7 @@ EAsyncPackageState::Type FAsyncLoadingThread2Impl::ProcessLoadedPackagesFromGame
 	if (Result != EAsyncPackageState::TimeOut)
 	{
 		// For performance reasons this set is created here and reset inside of AreAllDependenciesFullyLoaded
-		TSet<UPackage*> VisitedPackages;
+		TSet<FPackageId> VisistedPackages;
 
 		for (int32 PackageIndex = 0; PackageIndex < PackagesToDelete.Num(); ++PackageIndex)
 		{
@@ -3366,7 +3356,7 @@ EAsyncPackageState::Type FAsyncLoadingThread2Impl::ProcessLoadedPackagesFromGame
 				if (Package->HasClusterObjects())
 				{
 					// This package will create GC clusters but first check if all dependencies of this package have been fully loaded
-					if (Package->AreAllDependenciesFullyLoaded(VisitedPackages))
+					if (Package->AreAllDependenciesFullyLoaded(VisistedPackages))
 					{
 						if (Package->CreateClusters() == EAsyncPackageState::Complete)
 						{
@@ -4055,6 +4045,7 @@ FAsyncPackage2::FAsyncPackage2(
 , DeferredPostLoadIndex(0)
 , DeferredFinalizeIndex(0)
 , DeferredClusterIndex(0)
+, bHasClusterObjects(false)
 , bLoadHasFailed(false)
 , bLoadHasFinished(false)
 , bCreatedLinkerRoot(false)
@@ -4375,8 +4366,6 @@ EAsyncPackageState::Type FAsyncPackage2::PostLoadObjects()
 	return (PostLoadIndex == ExportCount) ? EAsyncPackageState::Complete : EAsyncPackageState::TimeOut;
 }
 
-// void CreateClustersFromPackage(FLinkerLoad* PackageLinker, TArray<UObject*>& OutClusterObjects);
-
 EAsyncPackageState::Type FAsyncPackage2::PostLoadDeferredObjects()
 {
 	SCOPED_LOADTIMER(PostLoadDeferredObjectsTime);
@@ -4484,10 +4473,17 @@ EAsyncPackageState::Type FAsyncPackage2::PostLoadDeferredObjects()
 			LinkerRoot->MarkAsFullyLoaded();			
 			LinkerRoot->SetLoadTime(FPlatformTime::Seconds() - LoadStartTime);
 
-			// if (Linker)
-			// {
-			// 	CreateClustersFromPackage(Linker, DeferredClusterObjects);
-			// }
+			if (CanCreateObjectClusters())
+			{
+				for (const FExportObject& Export : Exports)
+				{
+					if (!Export.bFiltered && Export.Object->CanBeClusterRoot())
+					{
+						bHasClusterObjects = true;
+						break;
+					}
+				}
+			}
 		}
 
 		FSoftObjectPath::InvalidateTag();
@@ -4499,27 +4495,19 @@ EAsyncPackageState::Type FAsyncPackage2::PostLoadDeferredObjects()
 
 EAsyncPackageState::Type FAsyncPackage2::CreateClusters()
 {
-	while (DeferredClusterIndex < DeferredClusterObjects.Num() &&
+	while (DeferredClusterIndex < ExportCount &&
 			!AsyncLoadingThread.IsAsyncLoadingSuspended() &&
 			!FAsyncLoadingThreadState2::Get()->IsTimeLimitExceeded())
 	{
-		UObject* ClusterRootObject = DeferredClusterObjects[DeferredClusterIndex++];
-		ClusterRootObject->CreateCluster();
-	}
+		const FExportObject& Export = Exports[DeferredClusterIndex++];
 
-	EAsyncPackageState::Type Result;
-	if (DeferredClusterIndex == DeferredClusterObjects.Num())
-	{
-		DeferredClusterIndex = 0;
-		DeferredClusterObjects.Reset();
-		Result = EAsyncPackageState::Complete;
-	}
-	else
-	{
-		Result = EAsyncPackageState::TimeOut;
+		if (!Export.bFiltered && Export.Object->CanBeClusterRoot())
+		{
+			Export.Object->CreateCluster();
+		}
 	}
 
-	return Result;
+	return DeferredClusterIndex == ExportCount ? EAsyncPackageState::Complete : EAsyncPackageState::TimeOut;
 }
 
 EAsyncPackageState::Type FAsyncPackage2::FinishObjects()
diff --git a/Engine/Source/Runtime/CoreUObject/Private/UObject/UObjectClusters.cpp b/Engine/Source/Runtime/CoreUObject/Private/UObject/UObjectClusters.cpp
--- a/Engine/Source/Runtime/CoreUObject/Private/UObject/UObjectClusters.cpp
+++ b/Engine/Source/Runtime/CoreUObject/Private/UObject/UObjectClusters.cpp
@@ -718,10 +718,15 @@ public:
 	}
 };
 
+bool CanCreateObjectClusters()
+{
+	return FPlatformProperties::RequiresCookedData() && !GIsInitialLoad && GCreateGCClusters && GAssetClustreringEnabled && !GUObjectArray.IsOpenForDisregardForGC() && GUObjectArray.DisregardForGCEnabled();
+}
+
 /** Looks through objects loaded with a package and creates clusters from them */
 void CreateClustersFromPackage(FLinkerLoad* PackageLinker, TArray<UObject*>& OutClusterObjects)
 {	
-	if (FPlatformProperties::RequiresCookedData() && !GIsInitialLoad && GCreateGCClusters && GAssetClustreringEnabled && !GUObjectArray.IsOpenForDisregardForGC() && GUObjectArray.DisregardForGCEnabled() )
+	if (CanCreateObjectClusters())
 	{
 		check(PackageLinker);
 
diff --git a/Engine/Source/Runtime/CoreUObject/Public/UObject/UObjectClusters.h b/Engine/Source/Runtime/CoreUObject/Public/UObject/UObjectClusters.h
--- a/Engine/Source/Runtime/CoreUObject/Public/UObject/UObjectClusters.h
+++ b/Engine/Source/Runtime/CoreUObject/Public/UObject/UObjectClusters.h
@@ -25,3 +25,5 @@ COREUOBJECT_API void FindStaleClusters(const TArray<FString>& Args);
 
 #endif // !UE_BUILD_SHIPPING
 
+// Whether object clusters can be created or not.
+bool CanCreateObjectClusters();
