AsyncLoading2: Common header for loader and IoStoreUtils

diff --git a/Engine/Source/Developer/IoStoreUtilities/Private/IoStoreUtilities.cpp b/Engine/Source/Developer/IoStoreUtilities/Private/IoStoreUtilities.cpp
--- a/Engine/Source/Developer/IoStoreUtilities/Private/IoStoreUtilities.cpp
+++ b/Engine/Source/Developer/IoStoreUtilities/Private/IoStoreUtilities.cpp
@@ -20,6 +20,7 @@
 #include "Serialization/JsonSerializer.h"
 #include "Serialization/BufferWriter.h"
 #include "Serialization/LargeMemoryWriter.h"
+#include "Serialization/AsyncLoading2.h"
 #include "UObject/NameBatchSerialization.h"
 #include "UObject/PackageFileSummary.h"
 #include "UObject/ObjectResource.h"
@@ -194,44 +195,12 @@ static FIoChunkId CreateChunkIdForBulkData(int32 GlobalPackageId, uint64 Bulkdat
 	return ChunkId;
 }
 
-struct FZenPackageSummary
-{
-	uint32 PackageFlags;
-	int32 NameMapOffset;
-	int32 ImportMapOffset;
-	int32 ExportMapOffset;
-	int32 ExportBundlesOffset;
-	int32 GraphDataOffset;
-	int32 GraphDataSize;
-	int32 BulkDataStartOffset;
-	int32 GlobalImportIndex;
-	int32 Pad;
-};
-
 enum EPreloadDependencyType
 {
 	PreloadDependencyType_Create,
 	PreloadDependencyType_Serialize,
 };
 
-enum EEventLoadNode2 : uint8
-{
-	Package_ExportsSerialized,
-	Package_PostLoad,
-	Package_Delete,
-	Package_NumPhases,
-
-	ExportBundle_Process = 0,
-	ExportBundle_NumPhases,
-};
-
-enum class EExportFilterFlags : uint8
-{
-	None,
-	NotForClient,
-	NotForServer
-};
-
 struct FArc
 {
 	uint32 FromNodeIndex;
@@ -243,17 +212,6 @@ struct FArc
 	}
 };
 
-struct FExportBundleEntry
-{
-	enum EExportCommandType
-	{
-		ExportCommandType_Create,
-		ExportCommandType_Serialize
-	};
-	uint32 LocalExportIndex;
-	EExportCommandType CommandType;
-};
-
 struct FExportGraphNode;
 
 struct FExportBundle
@@ -934,12 +892,6 @@ struct FExportData
 	FExportGraphNode* SerializeNode = nullptr;
 };
 
-struct FExportBundleMeta
-{
-	uint32 LoadOrder = -1;
-	uint32 PayloadSize = -1;
-};
-
 static void FindImport(TArray<FImportData>& GlobalImports, TMap<FString, int32>& GlobalImportsByFullName, TArray<FString>& TempFullNames, FObjectImport* ImportMap, int32 LocalImportIndex)
 {
 	FObjectImport* Import = &ImportMap[LocalImportIndex];
@@ -1088,7 +1040,7 @@ void SerializePackageData(
 	const TArray<FObjectExport>& ObjectExports,
 	const TArray<FExportData>& GlobalExports,
 	const TMap<FString, int32>& GlobalImportsByFullName,
-	FExportBundleMeta* ExportBundleMetaEntries,
+	FExportBundleMetaEntry* ExportBundleMetaEntries,
 	const FPackageStoreBulkDataManifest& BulkDataManifest,
 	bool bWithBulkDataManifest)
 {
@@ -1105,56 +1057,56 @@ void SerializePackageData(
 		Package->ImportMapSize = ImportMapArchive.Tell();
 
 		// Temporary Archive for serializing EDL graph data
-		FBufferWriter ZenGraphArchive(nullptr, 0, EBufferWriterFlags::AllowResize | EBufferWriterFlags::TakeOwnership);
+		FBufferWriter GraphArchive(nullptr, 0, EBufferWriterFlags::AllowResize | EBufferWriterFlags::TakeOwnership);
 
 		int32 InternalArcCount = Package->InternalArcs.Num();
-		ZenGraphArchive << InternalArcCount;
+		GraphArchive << InternalArcCount;
 		for (FArc& InternalArc : Package->InternalArcs)
 		{
-			ZenGraphArchive << InternalArc.FromNodeIndex;
-			ZenGraphArchive << InternalArc.ToNodeIndex;
+			GraphArchive << InternalArc.FromNodeIndex;
+			GraphArchive << InternalArc.ToNodeIndex;
 		}
 
 		int32 ReferencedPackagesCount = Package->ExternalArcs.Num();
-		ZenGraphArchive << ReferencedPackagesCount;
+		GraphArchive << ReferencedPackagesCount;
 		for (auto& KV : Package->ExternalArcs)
 		{
 			FPackage* ImportedPackage = KV.Key;
 			TArray<FArc>& Arcs = KV.Value;
 			int32 ExternalArcCount = Arcs.Num();
 
-			ZenGraphArchive << ImportedPackage->GlobalPackageId;
-			ZenGraphArchive << ExternalArcCount;
+			GraphArchive << ImportedPackage->GlobalPackageId;
+			GraphArchive << ExternalArcCount;
 			for (FArc& ExternalArc : Arcs)
 			{
-				ZenGraphArchive << ExternalArc.FromNodeIndex;
-				ZenGraphArchive << ExternalArc.ToNodeIndex;
+				GraphArchive << ExternalArc.FromNodeIndex;
+				GraphArchive << ExternalArc.ToNodeIndex;
 			}
 		}
-		Package->UGraphSize = ZenGraphArchive.Tell();
+		Package->UGraphSize = GraphArchive.Tell();
 
 		// Temporary Archive for serializing export map data
-		FBufferWriter ZenExportMapArchive(nullptr, 0, EBufferWriterFlags::AllowResize | EBufferWriterFlags::TakeOwnership);
+		FBufferWriter ExportMapArchive(nullptr, 0, EBufferWriterFlags::AllowResize | EBufferWriterFlags::TakeOwnership);
 		for (int32 I = 0; I < Package->ExportCount; ++I)
 		{
 			const FObjectExport& ObjectExport = ObjectExports[Package->ExportIndexOffset + I];
 			const FExportData& ExportData = GlobalExports[Package->Exports[I]];
 
 			int64 SerialSize = ObjectExport.SerialSize;
-			ZenExportMapArchive << SerialSize;
-			NameMapBuilder.SerializeName(ZenExportMapArchive, ObjectExport.ObjectName);
+			ExportMapArchive << SerialSize;
+			NameMapBuilder.SerializeName(ExportMapArchive, ObjectExport.ObjectName);
 			FPackageIndex OuterIndex = ExportData.OuterIndex;
-			ZenExportMapArchive << OuterIndex;
+			ExportMapArchive << OuterIndex;
 			FPackageIndex ClassIndex = ExportData.ClassIndex;
-			ZenExportMapArchive << ClassIndex;
+			ExportMapArchive << ClassIndex;
 			FPackageIndex SuperIndex = ExportData.SuperIndex;
-			ZenExportMapArchive << SuperIndex;
+			ExportMapArchive << SuperIndex;
 			FPackageIndex TemplateIndex = ExportData.TemplateIndex;
-			ZenExportMapArchive << TemplateIndex;
+			ExportMapArchive << TemplateIndex;
 			int32 GlobalImportIndex = ExportData.GlobalImportIndex;
-			ZenExportMapArchive << GlobalImportIndex;
+			ExportMapArchive << GlobalImportIndex;
 			uint32 ObjectFlags = ObjectExport.ObjectFlags;
-			ZenExportMapArchive << ObjectFlags;
+			ExportMapArchive << ObjectFlags;
 			uint8 FilterFlags = uint8(EExportFilterFlags::None);
 			if (ObjectExport.bNotForClient)
 			{
@@ -1164,20 +1116,20 @@ void SerializePackageData(
 			{
 				FilterFlags = uint8(EExportFilterFlags::NotForServer);
 			}
-			ZenExportMapArchive << FilterFlags;
+			ExportMapArchive << FilterFlags;
 			uint8 Pad = 0;
-			ZenExportMapArchive.Serialize(&Pad, 7);
+			ExportMapArchive.Serialize(&Pad, 7);
 		}
-		Package->ExportMapSize = ZenExportMapArchive.Tell();
+		Package->ExportMapSize = ExportMapArchive.Tell();
 
 		// Temporary archive for serializing export bundle data
-		FBufferWriter ZenExportBundlesArchive(nullptr, 0, EBufferWriterFlags::AllowResize | EBufferWriterFlags::TakeOwnership);
+		FBufferWriter ExportBundlesArchive(nullptr, 0, EBufferWriterFlags::AllowResize | EBufferWriterFlags::TakeOwnership);
 		int32 ExportBundleEntryIndex = 0;
 		for (FExportBundle& ExportBundle : Package->ExportBundles)
 		{
-			ZenExportBundlesArchive << ExportBundleEntryIndex;
+			ExportBundlesArchive << ExportBundleEntryIndex;
 			int32 EntryCount = ExportBundle.Nodes.Num();
-			ZenExportBundlesArchive << EntryCount;
+			ExportBundlesArchive << EntryCount;
 			ExportBundleEntryIndex += ExportBundle.Nodes.Num();
 		}
 		for (FExportBundle& ExportBundle : Package->ExportBundles)
@@ -1185,67 +1137,67 @@ void SerializePackageData(
 			for (FExportGraphNode* ExportNode : ExportBundle.Nodes)
 			{
 				uint32 CommandType = uint32(ExportNode->BundleEntry.CommandType);
-				ZenExportBundlesArchive << CommandType;
-				ZenExportBundlesArchive << ExportNode->BundleEntry.LocalExportIndex;
+				ExportBundlesArchive << ExportNode->BundleEntry.LocalExportIndex;
+				ExportBundlesArchive << CommandType;
 			}
 		}
-		Package->ExportBundlesSize = ZenExportBundlesArchive.Tell();
+		Package->ExportBundlesSize = ExportBundlesArchive.Tell();
 
 		Package->NameMapSize = Package->NameIndices.Num() * Package->NameIndices.GetTypeSize();
 
 		{
-			const uint64 ZenSummarySize =
-				sizeof(FZenPackageSummary)
+			const uint64 PackageSummarySize =
+				sizeof(FPackageSummary)
 				+ Package->NameMapSize
 				+ Package->ImportMapSize
 				+ Package->ExportMapSize
 				+ Package->ExportBundlesSize
 				+ Package->UGraphSize;
 
-			uint8* ZenSummaryBuffer = static_cast<uint8*>(FMemory::Malloc(ZenSummarySize));
-			FZenPackageSummary* ZenSummary = reinterpret_cast<FZenPackageSummary*>(ZenSummaryBuffer);
+			uint8* PackageSummaryBuffer = static_cast<uint8*>(FMemory::Malloc(PackageSummarySize));
+			FPackageSummary* PackageSummary = reinterpret_cast<FPackageSummary*>(PackageSummaryBuffer);
 
-			ZenSummary->PackageFlags = Package->PackageFlags;
-			ZenSummary->GraphDataSize = Package->UGraphSize;
-			ZenSummary->BulkDataStartOffset = Package->BulkDataStartOffset;
+			PackageSummary->PackageFlags = Package->PackageFlags;
+			PackageSummary->GraphDataSize = Package->UGraphSize;
+			PackageSummary->BulkDataStartOffset = Package->BulkDataStartOffset;
 			const int32* FindGlobalImportIndexForPackage = GlobalImportsByFullName.Find(Package->Name.ToString());
-			ZenSummary->GlobalImportIndex = FindGlobalImportIndexForPackage ? *FindGlobalImportIndexForPackage : -1;
+			PackageSummary->GlobalImportIndex = FindGlobalImportIndexForPackage ? *FindGlobalImportIndexForPackage : -1;
 
-			FBufferWriter ZenAr(ZenSummaryBuffer, ZenSummarySize);
-			ZenAr.Seek(sizeof(FZenPackageSummary));
+			FBufferWriter SummaryArchive(PackageSummaryBuffer, PackageSummarySize);
+			SummaryArchive.Seek(sizeof(FPackageSummary));
 
 			// NameMap data
 			{
-				ZenSummary->NameMapOffset = ZenAr.Tell();
-				ZenAr.Serialize(Package->NameIndices.GetData(), Package->NameMapSize);
+				PackageSummary->NameMapOffset = SummaryArchive.Tell();
+				SummaryArchive.Serialize(Package->NameIndices.GetData(), Package->NameMapSize);
 			}
 
 			// ImportMap data
 			{
 				check(ImportMapArchive.Tell() == Package->ImportMapSize);
-				ZenSummary->ImportMapOffset = ZenAr.Tell();
-				ZenAr.Serialize(ImportMapArchive.GetWriterData(), ImportMapArchive.Tell());
+				PackageSummary->ImportMapOffset = SummaryArchive.Tell();
+				SummaryArchive.Serialize(ImportMapArchive.GetWriterData(), ImportMapArchive.Tell());
 			}
 
 			// ExportMap data
 			{
-				check(ZenExportMapArchive.Tell() == Package->ExportMapSize);
-				ZenSummary->ExportMapOffset = ZenAr.Tell();
-				ZenAr.Serialize(ZenExportMapArchive.GetWriterData(), ZenExportMapArchive.Tell());
+				check(ExportMapArchive.Tell() == Package->ExportMapSize);
+				PackageSummary->ExportMapOffset = SummaryArchive.Tell();
+				SummaryArchive.Serialize(ExportMapArchive.GetWriterData(), ExportMapArchive.Tell());
 			}
 
 			// ExportBundle data
 			{
-				check(ZenExportBundlesArchive.Tell() == Package->ExportBundlesSize);
-				ZenSummary->ExportBundlesOffset = ZenAr.Tell();
-				ZenAr.Serialize(ZenExportBundlesArchive.GetWriterData(), ZenExportBundlesArchive.Tell());
+				check(ExportBundlesArchive.Tell() == Package->ExportBundlesSize);
+				PackageSummary->ExportBundlesOffset = SummaryArchive.Tell();
+				SummaryArchive.Serialize(ExportBundlesArchive.GetWriterData(), ExportBundlesArchive.Tell());
 			}
 
 			// Graph data
 			{
-				check(ZenGraphArchive.Tell() == Package->UGraphSize);
-				ZenSummary->GraphDataOffset = ZenAr.Tell();
-				ZenAr.Serialize(ZenGraphArchive.GetWriterData(), ZenGraphArchive.Tell());
+				check(GraphArchive.Tell() == Package->UGraphSize);
+				PackageSummary->GraphDataOffset = SummaryArchive.Tell();
+				SummaryArchive.Serialize(GraphArchive.GetWriterData(), GraphArchive.Tell());
 			}
 
 			// Export bundle chunks
@@ -1261,11 +1213,11 @@ void SerializePackageData(
 #endif
 				ExpAr->Close();
 
-				uint64 BundleBufferSize = ZenSummarySize;
+				uint64 BundleBufferSize = PackageSummarySize;
 				for (int32 ExportBundleIndex = 0; ExportBundleIndex < Package->ExportBundles.Num(); ++ExportBundleIndex)
 				{
 					FExportBundle& ExportBundle = Package->ExportBundles[ExportBundleIndex];
-					FExportBundleMeta* ExportBundleMetaEntry = ExportBundleMetaEntries + Package->FirstExportBundleMetaEntry + ExportBundleIndex;
+					FExportBundleMetaEntry* ExportBundleMetaEntry = ExportBundleMetaEntries + Package->FirstExportBundleMetaEntry + ExportBundleIndex;
 					ExportBundleMetaEntry->LoadOrder = ExportBundle.LoadOrder;
 					for (FExportGraphNode* Node : ExportBundle.Nodes)
 					{
@@ -1283,10 +1235,10 @@ void SerializePackageData(
 
 #if !SKIP_WRITE_CONTAINER
 				uint8* BundleBuffer = static_cast<uint8*>(FMemory::Malloc(BundleBufferSize));
-				FMemory::Memcpy(BundleBuffer, ZenSummaryBuffer, ZenSummarySize);
+				FMemory::Memcpy(BundleBuffer, PackageSummaryBuffer, PackageSummarySize);
 #endif
-				FMemory::Free(ZenSummaryBuffer);
-				uint64 BundleBufferOffset = ZenSummarySize;
+				FMemory::Free(PackageSummaryBuffer);
+				uint64 BundleBufferOffset = PackageSummarySize;
 				for (FExportBundle& ExportBundle : Package->ExportBundles)
 				{
 					for (FExportGraphNode* Node : ExportBundle.Nodes)
@@ -1385,7 +1337,7 @@ int32 CreateTarget(const FContainerTarget& Target)
 	uint64 ImportPreloadCount = 0;
 	uint64 ExportPreloadCount = 0;
 	FExportGraph ExportGraph;
-	TArray<FExportBundleMeta> ExportBundleMetaEntries;
+	TArray<FExportBundleMetaEntry> ExportBundleMetaEntries;
 
 	TArray<FInstallChunk> InstallChunks;
 	TArray<FPackage*> Packages;
@@ -2043,7 +1995,7 @@ int32 CreateTarget(const FContainerTarget& Target)
 	int32 StoreTocByteCount = StoreTocArchive->TotalSize();
 	int32 ImportedPackagesByteCount = ImportedPackagesArchive->TotalSize();
 	int32 GlobalImportNamesByteCount = GlobalImportNamesArchive->TotalSize();
-	int32 ExportBundleMetaByteCount = ExportBundleMetaEntries.Num() * sizeof(FExportBundleMeta);
+	int32 ExportBundleMetaByteCount = ExportBundleMetaEntries.Num() * sizeof(FExportBundleMetaEntry);
 	{
 		UE_LOG(LogIoStore, Display, TEXT("Saving global meta data to container file"));
 		FLargeMemoryWriter GlobalMetaArchive(0, true);
@@ -2118,7 +2070,7 @@ int32 CreateTarget(const FContainerTarget& Target)
 	uint64 ExportMapSize = 0;
 	uint64 NameMapSize = 0;
 	uint64 NameMapCount = 0;
-	uint64 ZenPackageSummarySize = Packages.Num() * sizeof(FZenPackageSummary);
+	uint64 PackageSummarySize = Packages.Num() * sizeof(FPackageSummary);
 	uint64 ImportedPackagesCount = 0;
 	uint64 InitialLoadSize = InitialLoadArchive->Tell();
 	uint64 ScriptArcsCount = 0;
@@ -2132,7 +2084,7 @@ int32 CreateTarget(const FContainerTarget& Target)
 	uint64 BundleCount = 0;
 	uint64 BundleEntryCount = 0;
 
-	uint64 ZenPackageHeaderSize = 0;
+	uint64 PackageHeaderSize = 0;
 
 	uint64 UniqueImportPackages = 0;
 	for (const FImportData& ImportData : GlobalImports)
@@ -2173,7 +2125,7 @@ int32 CreateTarget(const FContainerTarget& Target)
 		}
 	}
 
-	ZenPackageHeaderSize = ZenPackageSummarySize + NameMapSize + ImportMapSize + ExportMapSize + UGraphSize;
+	PackageHeaderSize = PackageSummarySize + NameMapSize + ImportMapSize + ExportMapSize + UGraphSize;
 
 	UE_LOG(LogIoStore, Display, TEXT("-------------------- IoStore Summary: %s --------------------"), *Target.TargetPlatform->PlatformName());
 	UE_LOG(LogIoStore, Display, TEXT("Packages: %8d total, %d circular dependencies, %d no preload dependencies, %d no import dependencies"),
@@ -2188,8 +2140,8 @@ int32 CreateTarget(const FContainerTarget& Target)
 	UE_LOG(LogIoStore, Display, TEXT("IoStore: %8.2f MB GlobalImportNames, %d total imports, %d script imports, %d UPackage imports"),
 		(double)GlobalImportNamesByteCount / 1024.0 / 1024.0, GlobalImportsByFullName.Num(), NumScriptImports, UniqueImportPackages);
 	UE_LOG(LogIoStore, Display, TEXT("IoStore: %8.2f MB InitialLoadData, %d script arcs, %d script outers, %d packages"), (double)InitialLoadSize / 1024.0 / 1024.0, ScriptArcsCount, NumScriptImports, Packages.Num());
-	UE_LOG(LogIoStore, Display, TEXT("IoStore: %8.2f MB PackageHeader, %d packages"), (double)ZenPackageHeaderSize / 1024.0 / 1024.0, Packages.Num());
-	UE_LOG(LogIoStore, Display, TEXT("IoStore: %8.2f MB PackageSummary"), (double)ZenPackageSummarySize / 1024.0 / 1024.0);
+	UE_LOG(LogIoStore, Display, TEXT("IoStore: %8.2f MB PackageHeader, %d packages"), (double)PackageHeaderSize / 1024.0 / 1024.0, Packages.Num());
+	UE_LOG(LogIoStore, Display, TEXT("IoStore: %8.2f MB PackageSummary"), (double)PackageSummarySize / 1024.0 / 1024.0);
 	UE_LOG(LogIoStore, Display, TEXT("IoStore: %8.2f MB PackageNameMap, %d indices"), (double)NameMapSize / 1024.0 / 1024.0, NameMapCount);
 	UE_LOG(LogIoStore, Display, TEXT("IoStore: %8.2f MB PackageImportMap"), (double)ImportMapSize / 1024.0 / 1024.0);
 	UE_LOG(LogIoStore, Display, TEXT("IoStore: %8.2f MB PackageExportMap"), (double)ExportMapSize / 1024.0 / 1024.0);
diff --git a/Engine/Source/Runtime/CoreUObject/Private/Serialization/AsyncLoading2.cpp b/Engine/Source/Runtime/CoreUObject/Private/Serialization/AsyncLoading2.cpp
--- a/Engine/Source/Runtime/CoreUObject/Private/Serialization/AsyncLoading2.cpp
+++ b/Engine/Source/Runtime/CoreUObject/Private/Serialization/AsyncLoading2.cpp
@@ -5,6 +5,7 @@
 =============================================================================*/
 
 #include "Serialization/AsyncLoading2.h"
+#include "Serialization/AsyncPackageLoader.h"
 #include "HAL/PlatformFilemanager.h"
 #include "HAL/FileManager.h"
 #include "HAL/Event.h"
@@ -56,7 +57,7 @@
 #endif
 
 struct FAsyncPackage2;
-class FAsyncLoadingThread2Impl;
+class FAsyncLoadingThread2;
 
 class FSimpleArchive
 	: public FArchive
@@ -111,13 +112,6 @@ public:
 	}
 };
 
-enum class EExportFilterFlags : uint8
-{
-	None,
-	NotForClient,
-	NotForServer
-};
-
 struct FExportMapEntry
 {
 	uint64 SerialSize;
@@ -138,17 +132,6 @@ struct FExportBundleHeader
 	uint32 EntryCount;
 };
 
-struct FExportBundleEntry
-{
-	enum EExportBundleEntryType
-	{
-		ExportBundleEntryType_Create,
-		ExportBundleEntryType_Serialize
-	};
-	uint32 CommandType;
-	uint32 LocalExportIndex;
-};
-
 struct FExportObject
 {
 	UObject* Object = nullptr;
@@ -250,32 +233,12 @@ struct FPackageStoreEntry
 	int32 ImportedPackagesOffset;
 };
 
-struct FPackageSummary
-{
-	uint32 PackageFlags;
-	int32 NameMapOffset;
-	int32 ImportMapOffset;
-	int32 ExportMapOffset;
-	int32 ExportBundlesOffset;
-	int32 GraphDataOffset;
-	int32 GraphDataSize;
-	int32 BulkDataStartOffset;
-	int32 GlobalImportIndex;
-	int32 Pad;
-};
-
 struct FPackageStoreInitialLoadEntry
 {
 	int32 ScriptArcsOffset;
 	int32 ScriptArcsCount;
 };
 
-struct FExportBundleMetaEntry
-{
-	uint32 LoadOrder;
-	uint32 PayloadSize;
-};
-
 class FGlobalNameMap
 {
 public:
@@ -843,17 +806,6 @@ enum class EAsyncPackageLoadingState2 : uint8
 	PackageComplete,
 };
 
-enum EEventLoadNode2 : uint8
-{
-	Package_ExportsSerialized,
-	Package_PostLoad,
-	Package_Delete,
-	Package_NumPhases,
-
-	ExportBundle_Process = 0,
-	ExportBundle_NumPhases,
-};
-
 class FEventLoadGraphAllocator;
 struct FAsyncLoadEventSpec;
 struct FAsyncLoadingThreadState2;
@@ -1114,10 +1066,10 @@ struct FAsyncPackageDesc2
 struct FAsyncPackage2
 {
 	friend struct FScopedAsyncPackageEvent2;
-	friend class FAsyncLoadingThread2Impl;
+	friend class FAsyncLoadingThread2;
 
 	FAsyncPackage2(const FAsyncPackageDesc2& InDesc,
-		FAsyncLoadingThread2Impl& InAsyncLoadingThread,
+		FAsyncLoadingThread2& InAsyncLoadingThread,
 		IEDLBootNotificationManager& InEDLBootNotificationManager,
 		FAsyncLoadEventGraphAllocator& InGraphAllocator,
 		const FAsyncLoadEventSpec* EventSpecs);
@@ -1309,7 +1261,7 @@ private:
 	/** List of OwnedObjects = Exports + UPackage + ObjectsCreatedFromExports */
 	TArray<UObject*> OwnedObjects;
 	/** Cached async loading thread object this package was created by */
-	FAsyncLoadingThread2Impl& AsyncLoadingThread;
+	FAsyncLoadingThread2& AsyncLoadingThread;
 	IEDLBootNotificationManager& EDLBootNotificationManager;
 	FAsyncLoadEventGraphAllocator& GraphAllocator;
 
@@ -1334,7 +1286,7 @@ private:
 	const FExportBundleEntry* ExportBundleEntries = nullptr;
 public:
 
-	FAsyncLoadingThread2Impl& GetAsyncLoadingThread()
+	FAsyncLoadingThread2& GetAsyncLoadingThread()
 	{
 		return AsyncLoadingThread;
 	}
@@ -1506,13 +1458,14 @@ private:
 	int32 ThreadId = 0;
 };
 
-class FAsyncLoadingThread2Impl
+class FAsyncLoadingThread2 final
 	: public FRunnable
+	, public IAsyncPackageLoader
 {
 	friend struct FAsyncPackage2;
 public:
-	FAsyncLoadingThread2Impl(FIoDispatcher& IoDispatcher, IEDLBootNotificationManager& InEDLBootNotificationManager);
-	virtual ~FAsyncLoadingThread2Impl();
+	FAsyncLoadingThread2(FIoDispatcher& IoDispatcher, IEDLBootNotificationManager& InEDLBootNotificationManager);
+	virtual ~FAsyncLoadingThread2();
 
 private:
 	/** Thread to run the worker FRunnable on */
@@ -1600,13 +1553,13 @@ private:
 public:
 
 	//~ Begin FRunnable Interface.
-	virtual bool Init();
-	virtual uint32 Run();
-	virtual void Stop();
+	virtual bool Init() override;
+	virtual uint32 Run() override;
+	virtual void Stop() override;
 	//~ End FRunnable Interface
 
 	/** Start the async loading thread */
-	void StartThread();
+	virtual void StartThread() override;
 
 	/** [GC] Management of global import objects */
 	void OnPreGarbageCollect();
@@ -1623,7 +1576,7 @@ public:
 	TArray<FAsyncLoadEventSpec> EventSpecs;
 
 	/** True if multithreaded async loading is currently being used. */
-	bool IsMultithreaded()
+	inline virtual bool IsMultithreaded() override
 	{
 		return bThreadStarted;
 	}
@@ -1647,14 +1600,14 @@ public:
 	}
 
 	/** Returns true if packages are currently being loaded on the async thread */
-	bool IsAsyncLoadingPackages()
+	inline virtual bool IsAsyncLoadingPackages() override
 	{
 		FPlatformMisc::MemoryBarrier();
 		return QueuedPackagesCounter != 0 || ExistingAsyncPackagesCounter.GetValue() != 0;
 	}
 
 	/** Returns true this codes runs on the async loading thread */
-	bool IsInAsyncLoadThread()
+	virtual bool IsInAsyncLoadThread() override
 	{
 		if (IsMultithreaded())
 		{
@@ -1686,14 +1639,14 @@ public:
 	}
 
 	/** Returns true if async loading is suspended */
-	bool IsAsyncLoadingSuspended()
+	inline virtual bool IsAsyncLoadingSuspended() override
 	{
 		return bSuspendRequested;
 	}
 
-	void NotifyConstructedDuringAsyncLoading(UObject* Object, bool bSubObject);
+	virtual void NotifyConstructedDuringAsyncLoading(UObject* Object, bool bSubObject) override;
 
-	void FireCompletedCompiledInImport(void* AsyncPackage, FPackageIndex Import);
+	virtual void FireCompletedCompiledInImport(void* AsyncPackage, FPackageIndex Import) override;
 
 	/**
 	* [ASYNC THREAD] Finds an existing async package in the AsyncPackages by its name.
@@ -1769,32 +1722,42 @@ public:
 	EAsyncPackageState::Type TickAsyncThreadFromGameThread(bool& bDidSomething);
 
 	/** Initializes async loading thread */
-	void InitializeLoading();
+	virtual void InitializeLoading() override;
 
-	void ShutdownLoading();
+	virtual void ShutdownLoading() override;
 
-	int32 LoadPackage(
+	virtual int32 LoadPackage(
 		const FString& InPackageName,
 		const FGuid* InGuid,
 		const TCHAR* InPackageToLoadFrom,
 		FLoadPackageAsyncDelegate InCompletionDelegate,
 		EPackageFlags InPackageFlags,
 		int32 InPIEInstanceID,
-		int32 InPackagePriority);
+		int32 InPackagePriority) override;
 
 	EAsyncPackageState::Type ProcessLoadingFromGameThread(bool bUseTimeLimit, bool bUseFullTimeLimit, float TimeLimit);
 
+	inline virtual EAsyncPackageState::Type ProcessLoading(bool bUseTimeLimit, bool bUseFullTimeLimit, float TimeLimit) override
+	{
+		return ProcessLoadingFromGameThread(bUseTimeLimit, bUseFullTimeLimit, TimeLimit);
+	}
+
 	EAsyncPackageState::Type ProcessLoadingUntilCompleteFromGameThread(TFunctionRef<bool()> CompletionPredicate, float TimeLimit);
 
-	void CancelLoading();
+	inline virtual EAsyncPackageState::Type ProcessLoadingUntilComplete(TFunctionRef<bool()> CompletionPredicate, float TimeLimit) override
+	{
+		return ProcessLoadingUntilCompleteFromGameThread(CompletionPredicate, TimeLimit);
+	}
+
+	virtual void CancelLoading() override;
 
-	void SuspendLoading();
+	virtual void SuspendLoading() override;
 
-	void ResumeLoading();
+	virtual void ResumeLoading() override;
 
-	void FlushLoading(int32 PackageId);
+	virtual void FlushLoading(int32 PackageId) override;
 
-	int32 GetNumAsyncPackages()
+	virtual int32 GetNumAsyncPackages() override
 	{
 		FPlatformMisc::MemoryBarrier();
 		return ExistingAsyncPackagesCounter.GetValue();
@@ -1805,7 +1768,7 @@ public:
 	 * @param PackageName Name of the package to return async load percentage for
 	 * @return Percentage (0-100) of the async package load or -1 of package has not been found
 	 */
-	float GetAsyncLoadPercentage(const FName& PackageName);
+	virtual float GetAsyncLoadPercentage(const FName& PackageName) override;
 
 	/**
 	 * [ASYNC/GAME THREAD] Checks if a request ID already is added to the loading queue
@@ -1899,10 +1862,10 @@ struct FAsyncPackageScope2
 /** Just like TGuardValue for FAsyncLoadingThread::AsyncLoadingTickCounter but only works for the game thread */
 struct FAsyncLoadingTickScope2
 {
-	FAsyncLoadingThread2Impl& AsyncLoadingThread;
+	FAsyncLoadingThread2& AsyncLoadingThread;
 	bool bNeedsToLeaveAsyncTick;
 
-	FAsyncLoadingTickScope2(FAsyncLoadingThread2Impl& InAsyncLoadingThread)
+	FAsyncLoadingTickScope2(FAsyncLoadingThread2& InAsyncLoadingThread)
 		: AsyncLoadingThread(InAsyncLoadingThread)
 		, bNeedsToLeaveAsyncTick(false)
 	{
@@ -1921,7 +1884,7 @@ struct FAsyncLoadingTickScope2
 	}
 };
 
-void FAsyncLoadingThread2Impl::InitializeLoading()
+void FAsyncLoadingThread2::InitializeLoading()
 {
 #if USE_NEW_BULKDATA
 	FBulkDataBase::SetIoDispatcher(&IoDispatcher);
@@ -1945,7 +1908,7 @@ void FAsyncLoadingThread2Impl::InitializeLoading()
 		GlobalNameMap.GetNameEntries().Num());
 }
 
-void FAsyncLoadingThread2Impl::QueuePackage(FAsyncPackageDesc2& Package)
+void FAsyncLoadingThread2::QueuePackage(FAsyncPackageDesc2& Package)
 {
 	//TRACE_CPUPROFILER_EVENT_SCOPE(QueuePackage);
 	{
@@ -1956,7 +1919,7 @@ void FAsyncLoadingThread2Impl::QueuePackage(FAsyncPackageDesc2& Package)
 	AltZenaphore.NotifyOne();
 }
 
-FAsyncPackage2* FAsyncLoadingThread2Impl::FindOrInsertPackage(FAsyncPackageDesc2* Desc, bool& bInserted)
+FAsyncPackage2* FAsyncLoadingThread2::FindOrInsertPackage(FAsyncPackageDesc2* Desc, bool& bInserted)
 {
 	TRACE_CPUPROFILER_EVENT_SCOPE(FindOrInsertPackage);
 	FAsyncPackage2* Package = nullptr;
@@ -1988,7 +1951,7 @@ FAsyncPackage2* FAsyncLoadingThread2Impl::FindOrInsertPackage(FAsyncPackageDesc2
 	return Package;
 }
 
-bool FAsyncLoadingThread2Impl::CreateAsyncPackagesFromQueue()
+bool FAsyncLoadingThread2::CreateAsyncPackagesFromQueue()
 {
 	TRACE_CPUPROFILER_EVENT_SCOPE(CreateAsyncPackagesFromQueue);
 	TArray<FAsyncPackageDesc2*> QueueCopy;
@@ -2019,13 +1982,13 @@ bool FAsyncLoadingThread2Impl::CreateAsyncPackagesFromQueue()
 	return QueueCopy.Num() > 0;
 }
 
-void FAsyncLoadingThread2Impl::AddBundleIoRequest(FAsyncPackage2* Package, const FExportBundleMetaEntry& BundleMetaEntry)
+void FAsyncLoadingThread2::AddBundleIoRequest(FAsyncPackage2* Package, const FExportBundleMetaEntry& BundleMetaEntry)
 {
 	WaitingForIoBundleCounter.Increment();
 	WaitingIoRequests.HeapPush({ Package, BundleMetaEntry.LoadOrder, BundleMetaEntry.PayloadSize });
 }
 
-void FAsyncLoadingThread2Impl::BundleIoRequestCompleted(const FExportBundleMetaEntry& BundleMetaEntry)
+void FAsyncLoadingThread2::BundleIoRequestCompleted(const FExportBundleMetaEntry& BundleMetaEntry)
 {
 	check(PendingBundleIoRequestsTotalSize >= BundleMetaEntry.PayloadSize)
 	PendingBundleIoRequestsTotalSize -= BundleMetaEntry.PayloadSize;
@@ -2035,7 +1998,7 @@ void FAsyncLoadingThread2Impl::BundleIoRequestCompleted(const FExportBundleMetaE
 	}
 }
 
-void FAsyncLoadingThread2Impl::StartBundleIoRequests()
+void FAsyncLoadingThread2::StartBundleIoRequests()
 {
 	TRACE_CPUPROFILER_EVENT_SCOPE(StartBundleIoRequests);
 	constexpr uint64 MaxPendingRequestsSize = 256 << 20;
@@ -2700,7 +2663,7 @@ EAsyncPackageState::Type FAsyncPackage2::Event_ProcessExportBundle(FAsyncPackage
 		if (FilterExport(Export.FilterFlags))
 		{
 			Package->Exports[BundleEntry->LocalExportIndex].bFiltered = true;
-			if (BundleEntry->CommandType == FExportBundleEntry::ExportBundleEntryType_Serialize)
+			if (BundleEntry->CommandType == FExportBundleEntry::ExportCommandType_Serialize)
 			{
 				const uint64 SerialSize = Package->ExportMap[BundleEntry->LocalExportIndex].SerialSize;
 				Package->SerialDataPtr += SerialSize;
@@ -2711,13 +2674,13 @@ EAsyncPackageState::Type FAsyncPackage2::Event_ProcessExportBundle(FAsyncPackage
 			continue;
 		}
 
-		if (BundleEntry->CommandType == FExportBundleEntry::ExportBundleEntryType_Create)
+		if (BundleEntry->CommandType == FExportBundleEntry::ExportCommandType_Create)
 		{
 			Package->EventDrivenCreateExport(BundleEntry->LocalExportIndex);
 		}
 		else
 		{
-			check(BundleEntry->CommandType == FExportBundleEntry::ExportBundleEntryType_Serialize);
+			check(BundleEntry->CommandType == FExportBundleEntry::ExportCommandType_Serialize);
 			const uint64 ExportSerialSize = Export.SerialSize;
 			check(Package->SerialDataPtr + ExportSerialSize <= Package->IoBuffer.Data() + Package->IoBuffer.DataSize());
 			UObject* Object = Package->Exports[BundleEntry->LocalExportIndex].Object;
@@ -3121,7 +3084,7 @@ FEventLoadNode2* FAsyncPackage2::GetNode(int32 NodeIndex)
 	return &PackageNodes[NodeIndex];
 }
 
-void FAsyncLoadingThread2Impl::AddToLoadedPackages(FAsyncPackage2* Package)
+void FAsyncLoadingThread2::AddToLoadedPackages(FAsyncPackage2* Package)
 {
 	WaitingForPostLoadCounter.Increment();
 	FScopeLock LoadedLock(&LoadedPackagesCritical);
@@ -3129,14 +3092,14 @@ void FAsyncLoadingThread2Impl::AddToLoadedPackages(FAsyncPackage2* Package)
 	LoadedPackages.Add(Package);
 }
 
-EAsyncPackageState::Type FAsyncLoadingThread2Impl::ProcessAsyncLoadingFromGameThread(int32& OutPackagesProcessed)
+EAsyncPackageState::Type FAsyncLoadingThread2::ProcessAsyncLoadingFromGameThread(int32& OutPackagesProcessed)
 {
 	SCOPED_LOADTIMER(AsyncLoadingTime);
 
 	check(IsInGameThread());
 
 	// If we're not multithreaded and flushing async loading, update the thread heartbeat
-	const bool bNeedsHeartbeatTick = !FAsyncLoadingThread2Impl::IsMultithreaded();
+	const bool bNeedsHeartbeatTick = !FAsyncLoadingThread2::IsMultithreaded();
 	OutPackagesProcessed = 0;
 
 	FAsyncLoadingTickScope2 InAsyncLoadingTick(*this);
@@ -3259,7 +3222,7 @@ bool FAsyncPackage2::AreAllDependenciesFullyLoaded(TSet<FPackageId>& VisitedPack
 	return bLoaded;
 }
 
-EAsyncPackageState::Type FAsyncLoadingThread2Impl::ProcessLoadedPackagesFromGameThread(bool& bDidSomething, int32 FlushRequestID)
+EAsyncPackageState::Type FAsyncLoadingThread2::ProcessLoadedPackagesFromGameThread(bool& bDidSomething, int32 FlushRequestID)
 {
 	EAsyncPackageState::Type Result = EAsyncPackageState::Complete;
 
@@ -3421,7 +3384,7 @@ EAsyncPackageState::Type FAsyncLoadingThread2Impl::ProcessLoadedPackagesFromGame
 	return Result;
 }
 
-EAsyncPackageState::Type FAsyncLoadingThread2Impl::TickAsyncLoadingFromGameThread(bool bUseTimeLimit, bool bUseFullTimeLimit, float TimeLimit, int32 FlushRequestID)
+EAsyncPackageState::Type FAsyncLoadingThread2::TickAsyncLoadingFromGameThread(bool bUseTimeLimit, bool bUseFullTimeLimit, float TimeLimit, int32 FlushRequestID)
 {
 	LLM_SCOPE(ELLMTag::AsyncLoading);
 
@@ -3449,7 +3412,7 @@ EAsyncPackageState::Type FAsyncLoadingThread2Impl::TickAsyncLoadingFromGameThrea
 			UnhashUnreachableObjects(false);
 		}
 
-		const bool bIsMultithreaded = FAsyncLoadingThread2Impl::IsMultithreaded();
+		const bool bIsMultithreaded = FAsyncLoadingThread2::IsMultithreaded();
 		double TickStartTime = FPlatformTime::Seconds();
 
 		bool bDidSomething = false;
@@ -3494,7 +3457,7 @@ EAsyncPackageState::Type FAsyncLoadingThread2Impl::TickAsyncLoadingFromGameThrea
 	return Result;
 }
 
-FAsyncLoadingThread2Impl::FAsyncLoadingThread2Impl(FIoDispatcher& InIoDispatcher, IEDLBootNotificationManager& InEDLBootNotificationManager)
+FAsyncLoadingThread2::FAsyncLoadingThread2(FIoDispatcher& InIoDispatcher, IEDLBootNotificationManager& InEDLBootNotificationManager)
 	: Thread(nullptr)
 	, EDLBootNotificationManager(InEDLBootNotificationManager)
 	, IoDispatcher(InIoDispatcher)
@@ -3534,7 +3497,7 @@ FAsyncLoadingThread2Impl::FAsyncLoadingThread2Impl(FIoDispatcher& InIoDispatcher
 		FAsyncLoadingThreadSettings::Get().bAsyncPostLoadEnabled ? TEXT("true") : TEXT("false"));
 }
 
-FAsyncLoadingThread2Impl::~FAsyncLoadingThread2Impl()
+FAsyncLoadingThread2::~FAsyncLoadingThread2()
 {
 	if (Thread)
 	{
@@ -3546,7 +3509,7 @@ FAsyncLoadingThread2Impl::~FAsyncLoadingThread2Impl()
 #endif
 }
 
-void FAsyncLoadingThread2Impl::ShutdownLoading()
+void FAsyncLoadingThread2::ShutdownLoading()
 {
 	FCoreUObjectDelegates::GetPreGarbageCollectDelegate().RemoveAll(this);
 	FCoreUObjectDelegates::GetPostGarbageCollect().RemoveAll(this);
@@ -3561,13 +3524,13 @@ void FAsyncLoadingThread2Impl::ShutdownLoading()
 	ThreadResumedEvent = nullptr;
 }
 
-void FAsyncLoadingThread2Impl::StartThread()
+void FAsyncLoadingThread2::StartThread()
 {
 	// Make sure the GC sync object is created before we start the thread (apparently this can happen before we call InitUObject())
 	FGCCSyncObject::Create();
 
-	FCoreUObjectDelegates::GetPreGarbageCollectDelegate().AddRaw(this, &FAsyncLoadingThread2Impl::OnPreGarbageCollect);
-	FCoreUObjectDelegates::GetPostGarbageCollect().AddRaw(this, &FAsyncLoadingThread2Impl::OnPostGarbageCollect);
+	FCoreUObjectDelegates::GetPreGarbageCollectDelegate().AddRaw(this, &FAsyncLoadingThread2::OnPreGarbageCollect);
+	FCoreUObjectDelegates::GetPostGarbageCollect().AddRaw(this, &FAsyncLoadingThread2::OnPostGarbageCollect);
 
 	if (!FAsyncLoadingThreadSettings::Get().bAsyncLoadingThreadEnabled)
 	{
@@ -3616,12 +3579,12 @@ void FAsyncLoadingThread2Impl::StartThread()
 		GIsInitialLoad ? TEXT("true") : TEXT("false"));
 }
 
-bool FAsyncLoadingThread2Impl::Init()
+bool FAsyncLoadingThread2::Init()
 {
 	return true;
 }
 
-void FAsyncLoadingThread2Impl::SuspendWorkers()
+void FAsyncLoadingThread2::SuspendWorkers()
 {
 	if (bWorkersSuspended)
 	{
@@ -3639,7 +3602,7 @@ void FAsyncLoadingThread2Impl::SuspendWorkers()
 	bWorkersSuspended = true;
 }
 
-void FAsyncLoadingThread2Impl::ResumeWorkers()
+void FAsyncLoadingThread2::ResumeWorkers()
 {
 	if (!bWorkersSuspended)
 	{
@@ -3653,7 +3616,7 @@ void FAsyncLoadingThread2Impl::ResumeWorkers()
 	bWorkersSuspended = false;
 }
 
-uint32 FAsyncLoadingThread2Impl::Run()
+uint32 FAsyncLoadingThread2::Run()
 {
 	LLM_SCOPE(ELLMTag::AsyncLoading);
 
@@ -3797,7 +3760,7 @@ uint32 FAsyncLoadingThread2Impl::Run()
 	return 0;
 }
 
-EAsyncPackageState::Type FAsyncLoadingThread2Impl::TickAsyncThreadFromGameThread(bool& bDidSomething)
+EAsyncPackageState::Type FAsyncLoadingThread2::TickAsyncThreadFromGameThread(bool& bDidSomething)
 {
 	check(IsInGameThread());
 	EAsyncPackageState::Type Result = EAsyncPackageState::Complete;
@@ -3820,7 +3783,7 @@ EAsyncPackageState::Type FAsyncLoadingThread2Impl::TickAsyncThreadFromGameThread
 	return Result;
 }
 
-void FAsyncLoadingThread2Impl::Stop()
+void FAsyncLoadingThread2::Stop()
 {
 	for (FAsyncLoadingThreadWorker& Worker : Workers)
 	{
@@ -3831,13 +3794,13 @@ void FAsyncLoadingThread2Impl::Stop()
 	AltZenaphore.NotifyAll();
 }
 
-void FAsyncLoadingThread2Impl::CancelLoading()
+void FAsyncLoadingThread2::CancelLoading()
 {
 	check(false);
 	// TODO
 }
 
-void FAsyncLoadingThread2Impl::SuspendLoading()
+void FAsyncLoadingThread2::SuspendLoading()
 {
 	UE_CLOG(!IsInGameThread() || IsInSlateThread(), LogStreaming, Fatal, TEXT("Async loading can only be suspended from the main thread"));
 	if (!bSuspendRequested)
@@ -3852,7 +3815,7 @@ void FAsyncLoadingThread2Impl::SuspendLoading()
 	}
 }
 
-void FAsyncLoadingThread2Impl::ResumeLoading()
+void FAsyncLoadingThread2::ResumeLoading()
 {
 	check(IsInGameThread() && !IsInSlateThread());
 	if (bSuspendRequested)
@@ -3866,7 +3829,7 @@ void FAsyncLoadingThread2Impl::ResumeLoading()
 	}
 }
 
-float FAsyncLoadingThread2Impl::GetAsyncLoadPercentage(const FName& PackageName)
+float FAsyncLoadingThread2::GetAsyncLoadPercentage(const FName& PackageName)
 {
 	float LoadPercentage = -1.0f;
 	FAsyncPackage2* Package = FindAsyncPackage(PackageName);
@@ -3987,14 +3950,14 @@ void FGlobalImportStore::OnPostGarbageCollect()
 		UnmarkedCount);
 }
 
-void FAsyncLoadingThread2Impl::OnPreGarbageCollect()
+void FAsyncLoadingThread2::OnPreGarbageCollect()
 {
 	TRACE_CPUPROFILER_EVENT_SCOPE(AltPreGC);
 	const bool bIsAsyncLoadingPackages = IsAsyncLoadingPackages();
 	GlobalPackageStore.GetGlobalImportStore().OnPreGarbageCollect(bIsAsyncLoadingPackages);
 }
 
-void FAsyncLoadingThread2Impl::OnPostGarbageCollect()
+void FAsyncLoadingThread2::OnPostGarbageCollect()
 {
 	TRACE_CPUPROFILER_EVENT_SCOPE(AltPostGC);
 	GlobalPackageStore.GetGlobalImportStore().OnPostGarbageCollect();
@@ -4005,7 +3968,7 @@ void FAsyncLoadingThread2Impl::OnPostGarbageCollect()
  * @param Object		Object created
  * @param bSubObject	Object created as a sub-object of a loaded object
  */
-void FAsyncLoadingThread2Impl::NotifyConstructedDuringAsyncLoading(UObject* Object, bool bSubObject)
+void FAsyncLoadingThread2::NotifyConstructedDuringAsyncLoading(UObject* Object, bool bSubObject)
 {
 	// Mark objects created during async loading process (e.g. from within PostLoad or CreateExport) as async loaded so they 
 	// cannot be found. This requires also keeping track of them so we can remove the async loading flag later one when we 
@@ -4020,7 +3983,7 @@ void FAsyncLoadingThread2Impl::NotifyConstructedDuringAsyncLoading(UObject* Obje
 	AsyncPackage2->AddOwnedObjectFromCallback(Object, /*bForceAdd*/ bSubObject);
 }
 
-void FAsyncLoadingThread2Impl::FireCompletedCompiledInImport(void* AsyncPackage, FPackageIndex Import)
+void FAsyncLoadingThread2::FireCompletedCompiledInImport(void* AsyncPackage, FPackageIndex Import)
 {
 	int32 ExportNodeIndex = Import.ToImport();
 	static_cast<FAsyncPackage2*>(AsyncPackage)->GetNode(ExportNodeIndex)->ReleaseBarrier();
@@ -4035,7 +3998,7 @@ void FAsyncLoadingThread2Impl::FireCompletedCompiledInImport(void* AsyncPackage,
 */
 FAsyncPackage2::FAsyncPackage2(
 	const FAsyncPackageDesc2& InDesc,
-	FAsyncLoadingThread2Impl& InAsyncLoadingThread,
+	FAsyncLoadingThread2& InAsyncLoadingThread,
 	IEDLBootNotificationManager& InEDLBootNotificationManager,
 	FAsyncLoadEventGraphAllocator& InGraphAllocator,
 	const FAsyncLoadEventSpec* EventSpecs)
@@ -4597,7 +4560,7 @@ void FAsyncPackage2::UpdateLoadPercentage()
 	LoadPercentage = FMath::Max(NewLoadPercentage, LoadPercentage);
 }
 
-int32 FAsyncLoadingThread2Impl::LoadPackage(const FString& InName, const FGuid* InGuid, const TCHAR* InPackageToLoadFrom, FLoadPackageAsyncDelegate InCompletionDelegate, EPackageFlags InPackageFlags, int32 InPIEInstanceID, int32 InPackagePriority)
+int32 FAsyncLoadingThread2::LoadPackage(const FString& InName, const FGuid* InGuid, const TCHAR* InPackageToLoadFrom, FLoadPackageAsyncDelegate InCompletionDelegate, EPackageFlags InPackageFlags, int32 InPIEInstanceID, int32 InPackagePriority)
 {
 	TRACE_CPUPROFILER_EVENT_SCOPE(LoadPackage);
 
@@ -4654,13 +4617,13 @@ int32 FAsyncLoadingThread2Impl::LoadPackage(const FString& InName, const FGuid*
 	return RequestID;
 }
 
-EAsyncPackageState::Type FAsyncLoadingThread2Impl::ProcessLoadingFromGameThread(bool bUseTimeLimit, bool bUseFullTimeLimit, float TimeLimit)
+EAsyncPackageState::Type FAsyncLoadingThread2::ProcessLoadingFromGameThread(bool bUseTimeLimit, bool bUseFullTimeLimit, float TimeLimit)
 {
 	TickAsyncLoadingFromGameThread(bUseTimeLimit, bUseFullTimeLimit, TimeLimit);
 	return IsAsyncLoading() ? EAsyncPackageState::TimeOut : EAsyncPackageState::Complete;
 }
 
-void FAsyncLoadingThread2Impl::FlushLoading(int32 RequestId)
+void FAsyncLoadingThread2::FlushLoading(int32 RequestId)
 {
 	if (IsAsyncLoading())
 	{
@@ -4721,7 +4684,7 @@ void FAsyncLoadingThread2Impl::FlushLoading(int32 RequestId)
 	}
 }
 
-EAsyncPackageState::Type FAsyncLoadingThread2Impl::ProcessLoadingUntilCompleteFromGameThread(TFunctionRef<bool()> CompletionPredicate, float TimeLimit)
+EAsyncPackageState::Type FAsyncLoadingThread2::ProcessLoadingUntilCompleteFromGameThread(TFunctionRef<bool()> CompletionPredicate, float TimeLimit)
 {
 	if (!IsAsyncLoading())
 	{
@@ -4758,102 +4721,7 @@ EAsyncPackageState::Type FAsyncLoadingThread2Impl::ProcessLoadingUntilCompleteFr
 	return TimeLimit <= 0 ? EAsyncPackageState::TimeOut : EAsyncPackageState::Complete;
 }
 
-FAsyncLoadingThread2::FAsyncLoadingThread2(FIoDispatcher& InIoDispatcher, IEDLBootNotificationManager& InEDLBootNotificationManager)
-{
-	Impl = new FAsyncLoadingThread2Impl(InIoDispatcher, InEDLBootNotificationManager);
-}
-
-FAsyncLoadingThread2::~FAsyncLoadingThread2()
-{
-	delete Impl;
-}
-
-void FAsyncLoadingThread2::InitializeLoading()
-{
-	Impl->InitializeLoading();
-}
-
-void FAsyncLoadingThread2::ShutdownLoading()
-{
-	Impl->ShutdownLoading();
-}
-
-void FAsyncLoadingThread2::StartThread()
-{
-	Impl->StartThread();
-}
-
-bool FAsyncLoadingThread2::IsMultithreaded()
-{
-	return Impl->IsMultithreaded();
-}
-
-bool FAsyncLoadingThread2::IsInAsyncLoadThread()
-{
-	return Impl->IsInAsyncLoadThread();
-}
-
-void FAsyncLoadingThread2::NotifyConstructedDuringAsyncLoading(UObject* Object, bool bSubObject)
-{
-	Impl->NotifyConstructedDuringAsyncLoading(Object, bSubObject);
-}
-
-void FAsyncLoadingThread2::FireCompletedCompiledInImport(void* AsyncPackage, FPackageIndex Import)
-{
-	Impl->FireCompletedCompiledInImport(AsyncPackage, Import);
-}
-
-int32 FAsyncLoadingThread2::LoadPackage(const FString& InPackageName, const FGuid* InGuid, const TCHAR* InPackageToLoadFrom, FLoadPackageAsyncDelegate InCompletionDelegate, EPackageFlags InPackageFlags, int32 InPIEInstanceID, int32 InPackagePriority)
-{
-	return Impl->LoadPackage(InPackageName, InGuid, InPackageToLoadFrom, InCompletionDelegate, InPackageFlags, InPIEInstanceID, InPackagePriority);
-}
-
-EAsyncPackageState::Type FAsyncLoadingThread2::ProcessLoading(bool bUseTimeLimit, bool bUseFullTimeLimit, float TimeLimit)
-{
-	return Impl->ProcessLoadingFromGameThread(bUseTimeLimit, bUseFullTimeLimit, TimeLimit);
-}
-
-EAsyncPackageState::Type FAsyncLoadingThread2::ProcessLoadingUntilComplete(TFunctionRef<bool()> CompletionPredicate, float TimeLimit)
-{
-	return Impl->ProcessLoadingUntilCompleteFromGameThread(CompletionPredicate, TimeLimit);
-}
-
-void FAsyncLoadingThread2::CancelLoading()
-{
-	Impl->CancelLoading();
-}
-
-void FAsyncLoadingThread2::SuspendLoading()
-{
-	Impl->SuspendLoading();
-}
-
-void FAsyncLoadingThread2::ResumeLoading()
-{
-	Impl->ResumeLoading();
-}
-
-void FAsyncLoadingThread2::FlushLoading(int32 RequestId)
-{
-	Impl->FlushLoading(RequestId);
-}
-
-int32 FAsyncLoadingThread2::GetNumAsyncPackages()
-{
-	return Impl->GetNumAsyncPackages();
-}
-
-float FAsyncLoadingThread2::GetAsyncLoadPercentage(const FName& PackageName)
-{
-	return Impl->GetAsyncLoadPercentage(PackageName);
-}
-
-bool FAsyncLoadingThread2::IsAsyncLoadingSuspended()
-{
-	return Impl->IsAsyncLoadingSuspended();
-}
-
-bool FAsyncLoadingThread2::IsAsyncLoadingPackages()
+IAsyncPackageLoader* MakeAsyncPackageLoader2(FIoDispatcher& InIoDispatcher, IEDLBootNotificationManager& InEDLBootNotificationManager)
 {
-	return Impl->IsAsyncLoadingPackages();
+	return new FAsyncLoadingThread2(InIoDispatcher, InEDLBootNotificationManager);
 }
diff --git a/Engine/Source/Runtime/CoreUObject/Private/Serialization/AsyncLoading2.h b/Engine/Source/Runtime/CoreUObject/Private/Serialization/AsyncLoading2.h
deleted file mode 100644
diff --git a/Engine/Source/Runtime/CoreUObject/Private/Serialization/AsyncPackageLoader.cpp b/Engine/Source/Runtime/CoreUObject/Private/Serialization/AsyncPackageLoader.cpp
--- a/Engine/Source/Runtime/CoreUObject/Private/Serialization/AsyncPackageLoader.cpp
+++ b/Engine/Source/Runtime/CoreUObject/Private/Serialization/AsyncPackageLoader.cpp
@@ -505,7 +505,7 @@ void InitAsyncThread()
 {
 	if (FIoDispatcher::IsInitialized())
 	{
-		GPackageLoader = MakeUnique<FAsyncLoadingThread2>(FIoDispatcher::Get(), GetGEDLBootNotificationManager());
+		GPackageLoader.Reset(MakeAsyncPackageLoader2(FIoDispatcher::Get(), GetGEDLBootNotificationManager()));
 	}
 	else
 	{
diff --git a/Engine/Source/Runtime/CoreUObject/Public/Serialization/AsyncLoading2.h b/Engine/Source/Runtime/CoreUObject/Public/Serialization/AsyncLoading2.h
new file mode 100644
--- /dev/null
+++ b/Engine/Source/Runtime/CoreUObject/Public/Serialization/AsyncLoading2.h
@@ -0,0 +1,87 @@
+// Copyright Epic Games, Inc. All Rights Reserved.
+
+/*=============================================================================
+	AsyncLoading2.h: Unreal async loading #2 definitions.
+=============================================================================*/
+
+#pragma once
+
+#include "CoreMinimal.h"
+
+class IAsyncPackageLoader;
+class FIoDispatcher;
+class IEDLBootNotificationManager;
+
+/**
+ * Event node.
+ */
+enum EEventLoadNode2 : uint8
+{
+	Package_ExportsSerialized,
+	Package_PostLoad,
+	Package_Delete,
+	Package_NumPhases,
+
+	ExportBundle_Process = 0,
+	ExportBundle_NumPhases,
+};
+
+/**
+ * Export filter flags.
+ */
+enum class EExportFilterFlags : uint8
+{
+	None,
+	NotForClient,
+	NotForServer
+};
+
+/**
+ * Package summary.
+ */
+struct FPackageSummary
+{
+	uint32 PackageFlags;
+	int32 NameMapOffset;
+	int32 ImportMapOffset;
+	int32 ExportMapOffset;
+	int32 ExportBundlesOffset;
+	int32 GraphDataOffset;
+	int32 GraphDataSize;
+	int32 BulkDataStartOffset;
+	int32 GlobalImportIndex;
+	int32 Pad;
+};
+
+/**
+ * Export bundle entry.
+ */
+struct FExportBundleEntry
+{
+	enum EExportCommandType
+	{
+		ExportCommandType_Create,
+		ExportCommandType_Serialize
+	};
+	uint32 LocalExportIndex;
+	uint32 CommandType;
+};
+
+/**
+ * Export bundle meta entry.
+ */
+struct FExportBundleMetaEntry
+{
+	uint32 LoadOrder = ~uint32(0);
+	uint32 PayloadSize = ~uint32(0);
+};
+
+/**
+ * Creates a new instance of the AsyncPackageLoader #2.
+ *
+ * @param InIoDispatcher				The I/O dispatcher.
+ * @param InEDLBootNotificationManager	The EDL boot notification manager.
+ *
+ * @return The async package loader.
+ */
+IAsyncPackageLoader* MakeAsyncPackageLoader2(FIoDispatcher& InIoDispatcher, IEDLBootNotificationManager& InEDLBootNotificationManager);
