Move FPackageStoreBulkDataManifest to its own code files to help cut down on the complexity of SavePackage.cpp

#rnx

diff --git a/Engine/Source/Developer/IoStoreUtilities/Private/IoStoreUtilities.cpp b/Engine/Source/Developer/IoStoreUtilities/Private/IoStoreUtilities.cpp
--- a/Engine/Source/Developer/IoStoreUtilities/Private/IoStoreUtilities.cpp
+++ b/Engine/Source/Developer/IoStoreUtilities/Private/IoStoreUtilities.cpp
@@ -16,6 +16,7 @@
 #include "Misc/Base64.h"
 #include "Modules/ModuleManager.h"
 #include "Serialization/Archive.h"
+#include "Serialization/BulkDataManifest.h"
 #include "Serialization/JsonReader.h"
 #include "Serialization/JsonSerializer.h"
 #include "Serialization/BufferWriter.h"
@@ -25,7 +26,6 @@
 #include "UObject/PackageFileSummary.h"
 #include "UObject/ObjectResource.h"
 #include "UObject/Package.h"
-#include "UObject/SavePackage.h"
 #include "Algo/Find.h"
 #include "Misc/FileHelper.h"
 
diff --git a/Engine/Source/Editor/UnrealEd/Private/CookOnTheFlyServer.cpp b/Engine/Source/Editor/UnrealEd/Private/CookOnTheFlyServer.cpp
--- a/Engine/Source/Editor/UnrealEd/Private/CookOnTheFlyServer.cpp
+++ b/Engine/Source/Editor/UnrealEd/Private/CookOnTheFlyServer.cpp
@@ -69,6 +69,7 @@
 #include "Serialization/ArchiveStackTrace.h"
 #include "DistanceFieldAtlas.h"
 #include "Cooker/AsyncIODelete.h"
+#include "Serialization/BulkDataManifest.h"
 
 #include "AssetRegistryModule.h"
 #include "AssetRegistryState.h"
diff --git a/Engine/Source/Runtime/CoreUObject/Private/Serialization/BulkDataManifest.cpp b/Engine/Source/Runtime/CoreUObject/Private/Serialization/BulkDataManifest.cpp
new file mode 100644
--- /dev/null
+++ b/Engine/Source/Runtime/CoreUObject/Private/Serialization/BulkDataManifest.cpp
@@ -0,0 +1,144 @@
+// Copyright Epic Games, Inc. All Rights Reserved.
+
+#include "Serialization/BulkDataManifest.h"
+
+#include "CoreGlobals.h"
+#include "HAL/FileManager.h"
+#include "Misc/Paths.h"
+
+void FPackageStoreBulkDataManifest::PackageDesc::AddData(EIoChunkType InType, uint64 InChunkId, uint64 InOffset, uint64 InSize, const FString& DebugFilename)
+{
+	// The ChunkId is supposed to be unique for each type
+	auto func = [=](const BulkDataDesc& Entry) { return Entry.ChunkId == InChunkId && Entry.Type == InType; };
+	if (Data.FindByPredicate(func) != nullptr)
+	{
+		FString Message;
+		for (int i = 0; i < Data.Num(); ++i)
+		{
+			const BulkDataDesc& Entry = Data[i];
+			Message.Appendf(TEXT("[%3d] ID: %20llu Offset: %8llu Size: %8llu Type: %d\n"),
+				i, Entry.ChunkId, Entry.Offset, Entry.Size, Entry.Type);
+		}
+
+		Message.Appendf(TEXT("[New] ID: %20llu Offset: %8llu Size: %8llu Type: %d\n"),
+			InChunkId, InOffset, InSize, InType);
+
+		UE_LOG(LogSerialization, Warning, TEXT("Duplicate BulkData description found in Package '%s', this will cause issues trying to run IoStore!\n%s"), *DebugFilename, *Message);
+	}
+	else
+	{
+		BulkDataDesc& Entry = Data.Emplace_GetRef();
+
+		Entry.ChunkId = InChunkId;
+		Entry.Offset = InOffset;
+		Entry.Size = InSize;
+		Entry.Type = InType;
+	}
+}
+
+void FPackageStoreBulkDataManifest::PackageDesc::AddZeroByteData(EIoChunkType InType)
+{
+	// Make sure we only add one empty read per Bulkdata type!
+	auto func = [=](const BulkDataDesc& Entry) { return Entry.Type == InType && Entry.Size == 0; };
+	if (Data.FindByPredicate(func) == nullptr)
+	{
+		BulkDataDesc& Entry = Data.Emplace_GetRef();
+
+		Entry.ChunkId = TNumericLimits<uint64>::Max();
+		Entry.Offset = 0;
+		Entry.Size = 0;
+		Entry.Type = InType;
+	}
+}
+
+FPackageStoreBulkDataManifest::FPackageStoreBulkDataManifest(const FString& ProjectPath)
+{
+	Filename = ProjectPath / TEXT("Metadata/BulkDataInfo.ubulkmanifest");
+}
+
+FArchive& operator<<(FArchive& Ar, FPackageStoreBulkDataManifest::PackageDesc::BulkDataDesc& Entry)
+{
+	Ar << Entry.ChunkId;
+	Ar << Entry.Offset;
+	Ar << Entry.Size;
+	Ar << Entry.Type;
+
+	return Ar;
+}
+
+FArchive& operator<<(FArchive& Ar, FPackageStoreBulkDataManifest::PackageDesc& Entry)
+{
+	Ar << Entry.Data;
+
+	return Ar;
+}
+
+FPackageStoreBulkDataManifest::~FPackageStoreBulkDataManifest()
+{
+}
+
+bool FPackageStoreBulkDataManifest::Load()
+{
+	Data.Empty();
+
+	TUniquePtr<FArchive> BinArchive(IFileManager::Get().CreateFileReader(*Filename));
+	if (BinArchive != nullptr)
+	{
+		*BinArchive << Data;
+		return true;
+	}
+	else
+	{
+		return false;
+	}
+}
+
+void FPackageStoreBulkDataManifest::Save()
+{
+	TUniquePtr<FArchive> BinArchive(IFileManager::Get().CreateFileWriter(*Filename));
+	*BinArchive << Data;
+}
+
+const FPackageStoreBulkDataManifest::PackageDesc* FPackageStoreBulkDataManifest::Find(const FString& PackageFilename) const
+{
+	const  FString NormalizedFilename = FixFilename(PackageFilename);
+	return Data.Find(NormalizedFilename);
+}
+
+void FPackageStoreBulkDataManifest::AddFileAccess(const FString& PackageFilename, EIoChunkType InType, uint64 InChunkId, uint64 InOffset, uint64 InSize)
+{
+	const FString NormalizedFilename = FixFilename(PackageFilename);
+
+	PackageDesc& Entry = GetOrCreateFileAccess(NormalizedFilename);
+
+	if (InSize > 0)
+	{
+		Entry.AddData(InType, InChunkId, InOffset, InSize, NormalizedFilename);
+	}
+	else
+	{
+		Entry.AddZeroByteData(InType);
+	}
+
+}
+
+FPackageStoreBulkDataManifest::PackageDesc& FPackageStoreBulkDataManifest::GetOrCreateFileAccess(const FString& PackageFilename)
+{
+	if (PackageDesc* Entry = Data.Find(PackageFilename))
+	{
+		return *Entry;
+	}
+	else
+	{
+		return Data.Add(PackageFilename);
+	}
+}
+
+FString FPackageStoreBulkDataManifest::FixFilename(const FString& InFilename) const
+{
+	FString OutFilename = InFilename;
+	FPaths::NormalizeFilename(OutFilename);
+	FPaths::MakePathRelativeTo(OutFilename, *RootPath);
+
+	return OutFilename;
+}
diff --git a/Engine/Source/Runtime/CoreUObject/Private/UObject/SavePackage.cpp b/Engine/Source/Runtime/CoreUObject/Private/UObject/SavePackage.cpp
--- a/Engine/Source/Runtime/CoreUObject/Private/UObject/SavePackage.cpp
+++ b/Engine/Source/Runtime/CoreUObject/Private/UObject/SavePackage.cpp
@@ -59,6 +59,7 @@
 #include "Serialization/Formatters/JsonArchiveOutputFormatter.h"
 #include "Serialization/ArchiveUObjectFromStructuredArchive.h"
 #include "UObject/AsyncWorkSequence.h"
+#include "Serialization/BulkDataManifest.h"
 
 DEFINE_LOG_CATEGORY_STATIC(LogSavePackage, Log, All);
 
@@ -449,143 +450,6 @@ private:
 	TSet<FNameEntryId> ReferencedNames;
 };
 
-void FPackageStoreBulkDataManifest::PackageDesc::AddData(EIoChunkType InType, uint64 InChunkId, uint64 InOffset, uint64 InSize, const FString& DebugFilename)
-{
-	// The ChunkId is supposed to be unique for each type
-	auto func = [=](const BulkDataDesc& Entry) { return Entry.ChunkId == InChunkId && Entry.Type == InType; };
-	if (Data.FindByPredicate(func) != nullptr)
-	{
-		FString Message;
-		for (int i = 0; i < Data.Num(); ++i)
-		{
-			const BulkDataDesc& Entry = Data[i];
-			Message.Appendf(TEXT("[%3d] ID: %20llu Offset: %8llu Size: %8llu Type: %d\n"),
-								 i, Entry.ChunkId, Entry.Offset, Entry.Size, Entry.Type);
-		}
-
-		Message.Appendf(TEXT(	"[New] ID: %20llu Offset: %8llu Size: %8llu Type: %d\n"),
-								InChunkId, InOffset, InSize, InType);
-
-		UE_LOG(	LogSavePackage, Warning, TEXT("Duplicate BulkData description found in Package '%s', this will cause issues trying to run IoStore!\n%s"),*DebugFilename , *Message);
-	}
-	else
-	{
-		BulkDataDesc& Entry = Data.Emplace_GetRef();
-
-		Entry.ChunkId = InChunkId;
-		Entry.Offset = InOffset;
-		Entry.Size = InSize;
-		Entry.Type = InType;
-	}
-}
-
-void FPackageStoreBulkDataManifest::PackageDesc::AddZeroByteData(EIoChunkType InType)
-{
-	// Make sure we only add one empty read per Bulkdata type!
-	auto func = [=](const BulkDataDesc& Entry) { return Entry.Type == InType && Entry.Size == 0; };
-	if (Data.FindByPredicate(func) == nullptr)
-	{
-		BulkDataDesc& Entry = Data.Emplace_GetRef();
-
-		Entry.ChunkId = TNumericLimits<uint64>::Max();
-		Entry.Offset = 0;
-		Entry.Size = 0;
-		Entry.Type = InType;
-	}
-}
-
-FPackageStoreBulkDataManifest::FPackageStoreBulkDataManifest(const FString& ProjectPath)
-{
-	Filename = ProjectPath / TEXT("Metadata/BulkDataInfo.ubulkmanifest");
-}
-
-FArchive& operator<<(FArchive& Ar, FPackageStoreBulkDataManifest::PackageDesc::BulkDataDesc& Entry)
-{
-	Ar << Entry.ChunkId;
-	Ar << Entry.Offset;
-	Ar << Entry.Size; 
-	Ar << Entry.Type;
-
-	return Ar;
-}
-
-FArchive& operator<<(FArchive& Ar, FPackageStoreBulkDataManifest::PackageDesc& Entry)
-{
-	Ar << Entry.Data;
-
-	return Ar;
-}
-
-FPackageStoreBulkDataManifest::~FPackageStoreBulkDataManifest()
-{
-}
-
-bool FPackageStoreBulkDataManifest::Load()
-{
-	Data.Empty();
-
-	TUniquePtr<FArchive> BinArchive(IFileManager::Get().CreateFileReader(*Filename));
-	if (BinArchive != nullptr)
-	{
-		*BinArchive << Data;
-		return true;
-	}
-	else
-	{
-		return false;
-	}
-}
-
-void FPackageStoreBulkDataManifest::Save()
-{
-	TUniquePtr<FArchive> BinArchive(IFileManager::Get().CreateFileWriter(*Filename));
-	*BinArchive << Data;
-}
-
-const FPackageStoreBulkDataManifest::PackageDesc* FPackageStoreBulkDataManifest::Find(const FString& PackageFilename) const
-{
-	const  FString NormalizedFilename = FixFilename(PackageFilename);
-	return Data.Find(NormalizedFilename);
-}
-
-void FPackageStoreBulkDataManifest::AddFileAccess(const FString& PackageFilename, EIoChunkType InType, uint64 InChunkId, uint64 InOffset, uint64 InSize)
-{
-	const FString NormalizedFilename = FixFilename(PackageFilename);
-	
-	PackageDesc& Entry = GetOrCreateFileAccess(NormalizedFilename);
-
-	if (InSize > 0)
-	{
-		Entry.AddData(InType, InChunkId, InOffset, InSize, NormalizedFilename);
-	}
-	else
-	{
-		Entry.AddZeroByteData(InType);
-	}
-	
-}
-
-FPackageStoreBulkDataManifest::PackageDesc& FPackageStoreBulkDataManifest::GetOrCreateFileAccess(const FString& PackageFilename)
-{
-	if (PackageDesc* Entry = Data.Find(PackageFilename))
-	{
-		return *Entry;
-	}
-	else
-	{
-		return Data.Add(PackageFilename);
-	}
-}
-
-FString FPackageStoreBulkDataManifest::FixFilename(const FString& InFilename) const
-{
-	FString OutFilename = InFilename;
-	FPaths::NormalizeFilename(OutFilename);
-	FPaths::MakePathRelativeTo(OutFilename, *RootPath);
-
-	return OutFilename;
-}
-
 #if WITH_EDITOR
 static void AddReplacementsNames(FPackageNameMapSaver& NameMapSaver, UObject* Obj, const ITargetPlatform* TargetPlatform)
 {
@@ -6553,3 +6417,9 @@ void FLooseFileWriter::WriteBulkdata(const FLooseFileWriter::FBulkDataInfo& Info
 
 	WriteToFile(ArchiveFilename, BulkData.Data(), BulkData.DataSize());
 }
+
+FSavePackageContext::~FSavePackageContext()
+{
+	delete PackageStoreWriter;
+	delete BulkDataManifest;
+}
diff --git a/Engine/Source/Runtime/CoreUObject/Public/Serialization/BulkDataManifest.h b/Engine/Source/Runtime/CoreUObject/Public/Serialization/BulkDataManifest.h
new file mode 100644
--- /dev/null
+++ b/Engine/Source/Runtime/CoreUObject/Public/Serialization/BulkDataManifest.h
@@ -0,0 +1,52 @@
+// Copyright Epic Games, Inc. All Rights Reserved.
+
+#pragma once
+
+#include "Containers/Map.h"
+#include "IO/IoDispatcher.h"
+
+class FPackageStoreBulkDataManifest
+{
+public:
+	COREUOBJECT_API FPackageStoreBulkDataManifest(const FString& ProjectPath);
+	COREUOBJECT_API ~FPackageStoreBulkDataManifest();
+
+	COREUOBJECT_API bool Load();
+	COREUOBJECT_API void Save();
+
+	void AddFileAccess(const FString& PackageFilename, EIoChunkType InType, uint64 InChunkId, uint64 InOffset, uint64 InSize);
+
+	const FString& GetFilename() const { return Filename; }
+
+	class PackageDesc
+	{
+	public:
+		struct BulkDataDesc
+		{
+			uint64 ChunkId;	// Note this is the Offset before the linker BulkDataStartOffset is
+							// applied, to make it easier to compute at runtime.
+			uint64 Offset;
+			uint64 Size;
+			EIoChunkType Type;
+		};
+
+		void AddData(EIoChunkType InType, uint64 InChunkId, uint64 InOffset, uint64 InSize, const FString& DebugFilename);
+		void AddZeroByteData(EIoChunkType InType);
+
+		const TArray<BulkDataDesc>& GetDataArray() const { return Data; }
+	private:
+		friend FArchive& operator<<(FArchive& Ar, PackageDesc& Entry);
+		TArray<BulkDataDesc> Data;
+	};
+
+	COREUOBJECT_API const PackageDesc* Find(const FString& PackageName) const;
+
+private:
+	PackageDesc& GetOrCreateFileAccess(const FString& PackageFilename);
+
+	FString FixFilename(const FString& InFileName) const;
+
+	FString RootPath;
+	FString Filename;
+	TMap<FString, PackageDesc> Data;
+};
diff --git a/Engine/Source/Runtime/CoreUObject/Public/UObject/SavePackage.h b/Engine/Source/Runtime/CoreUObject/Public/UObject/SavePackage.h
--- a/Engine/Source/Runtime/CoreUObject/Public/UObject/SavePackage.h
+++ b/Engine/Source/Runtime/CoreUObject/Public/UObject/SavePackage.h
@@ -5,58 +5,13 @@
 #include "Containers/Array.h"
 #include "Containers/Set.h"
 #include "Containers/Map.h"
-#include "IO/IoDispatcher.h"
 #include "UObject/NameTypes.h"
 
 class FArchive;
+class FIoBuffer;
 class FLinkerLoad;
 class FLinkerSave;
-
-class FPackageStoreBulkDataManifest
-{
-public:
-	COREUOBJECT_API FPackageStoreBulkDataManifest(const FString& ProjectPath);
-	COREUOBJECT_API ~FPackageStoreBulkDataManifest();
-
-	COREUOBJECT_API bool Load();
-	COREUOBJECT_API void Save() ;
-
-	void AddFileAccess(const FString& PackageFilename, EIoChunkType InType, uint64 InChunkId, uint64 InOffset, uint64 InSize);
-	
-	const FString& GetFilename() const { return Filename; }
-
-	class PackageDesc
-	{
-	public:
-		struct BulkDataDesc
-		{
-			uint64 ChunkId;	// Note this is the Offset before the linker BulkDataStartOffset is
-							// applied, to make it easier to compute at runtime.
-			uint64 Offset;
-			uint64 Size;
-			EIoChunkType Type;
-		};
-
-		void AddData(EIoChunkType InType, uint64 InChunkId, uint64 InOffset, uint64 InSize, const FString& DebugFilename);
-		void AddZeroByteData(EIoChunkType InType);
-
-		const TArray<BulkDataDesc>& GetDataArray() const { return Data; }
-	private:
-		friend FArchive& operator<<(FArchive& Ar, PackageDesc& Entry);
-		TArray<BulkDataDesc> Data;
-	};
-
-	COREUOBJECT_API const PackageDesc* Find(const FString& PackageName) const;
-
-private:
-	PackageDesc& GetOrCreateFileAccess(const FString& PackageFilename);
-
-	FString FixFilename(const FString& InFileName) const;
-
-	FString RootPath;
-	FString Filename;
-	TMap<FString, PackageDesc> Data;
-};
+class FPackageStoreBulkDataManifest;
 
 class FPackageStoreWriter
 {
@@ -127,11 +82,7 @@ public:
 	{
 	}
 
-	~FSavePackageContext()
-	{
-		delete PackageStoreWriter;
-		delete BulkDataManifest;
-	}
+	COREUOBJECT_API ~FSavePackageContext();
 
 	FPackageStoreWriter* PackageStoreWriter;
 	FPackageStoreBulkDataManifest* BulkDataManifest;
