Fix some coding standards mistakes I made when originally submitted FPackageStoreBulkDataManifest.

#rnx

diff --git a/Engine/Source/Developer/IoStoreUtilities/Private/IoStoreUtilities.cpp b/Engine/Source/Developer/IoStoreUtilities/Private/IoStoreUtilities.cpp
--- a/Engine/Source/Developer/IoStoreUtilities/Private/IoStoreUtilities.cpp
+++ b/Engine/Source/Developer/IoStoreUtilities/Private/IoStoreUtilities.cpp
@@ -799,7 +799,7 @@ static bool WriteBulkData(	const FString& Filename, EIoChunkType Type, const FPa
 							FIoStoreWriter* IoStoreWriter)
 {
 
-	const FPackageStoreBulkDataManifest::PackageDesc* PackageDesc = BulkDataManifest.Find(Package.FileName);
+	const FPackageStoreBulkDataManifest::FPackageDesc* PackageDesc = BulkDataManifest.Find(Package.FileName);
 	if (PackageDesc != nullptr)
 	{
 		const FIoChunkId BulkDataChunkId = CreateChunkIdForBulkData(Package.GlobalPackageId, TNumericLimits<uint64>::Max()-1, Type, *Package.FileName);
@@ -826,7 +826,7 @@ static bool WriteBulkData(	const FString& Filename, EIoChunkType Type, const FPa
 #endif
 
 		// Create additional mapping chunks as needed
-		for (const FPackageStoreBulkDataManifest::PackageDesc::BulkDataDesc& BulkDataDesc : PackageDesc->GetDataArray())
+		for (const FPackageStoreBulkDataManifest::FPackageDesc::FBulkDataDesc& BulkDataDesc : PackageDesc->GetDataArray())
 		{
 			if (BulkDataDesc.Type == Type)
 			{
diff --git a/Engine/Source/Runtime/CoreUObject/Private/Serialization/BulkDataManifest.cpp b/Engine/Source/Runtime/CoreUObject/Private/Serialization/BulkDataManifest.cpp
--- a/Engine/Source/Runtime/CoreUObject/Private/Serialization/BulkDataManifest.cpp
+++ b/Engine/Source/Runtime/CoreUObject/Private/Serialization/BulkDataManifest.cpp
@@ -6,16 +6,16 @@
 #include "HAL/FileManager.h"
 #include "Misc/Paths.h"
 
-void FPackageStoreBulkDataManifest::PackageDesc::AddData(EIoChunkType InType, uint64 InChunkId, uint64 InOffset, uint64 InSize, const FString& DebugFilename)
+void FPackageStoreBulkDataManifest::FPackageDesc::AddData(EIoChunkType InType, uint64 InChunkId, uint64 InOffset, uint64 InSize, const FString& DebugFilename)
 {
 	// The ChunkId is supposed to be unique for each type
-	auto func = [=](const BulkDataDesc& Entry) { return Entry.ChunkId == InChunkId && Entry.Type == InType; };
+	auto func = [=](const FBulkDataDesc& Entry) { return Entry.ChunkId == InChunkId && Entry.Type == InType; };
 	if (Data.FindByPredicate(func) != nullptr)
 	{
 		FString Message;
 		for (int i = 0; i < Data.Num(); ++i)
 		{
-			const BulkDataDesc& Entry = Data[i];
+			const FBulkDataDesc& Entry = Data[i];
 			Message.Appendf(TEXT("[%3d] ID: %20llu Offset: %8llu Size: %8llu Type: %d\n"),
 				i, Entry.ChunkId, Entry.Offset, Entry.Size, Entry.Type);
 		}
@@ -27,7 +27,7 @@ void FPackageStoreBulkDataManifest::PackageDesc::AddData(EIoChunkType InType, ui
 	}
 	else
 	{
-		BulkDataDesc& Entry = Data.Emplace_GetRef();
+		FBulkDataDesc& Entry = Data.Emplace_GetRef();
 
 		Entry.ChunkId = InChunkId;
 		Entry.Offset = InOffset;
@@ -36,13 +36,13 @@ void FPackageStoreBulkDataManifest::PackageDesc::AddData(EIoChunkType InType, ui
 	}
 }
 
-void FPackageStoreBulkDataManifest::PackageDesc::AddZeroByteData(EIoChunkType InType)
+void FPackageStoreBulkDataManifest::FPackageDesc::AddZeroByteData(EIoChunkType InType)
 {
 	// Make sure we only add one empty read per Bulkdata type!
-	auto func = [=](const BulkDataDesc& Entry) { return Entry.Type == InType && Entry.Size == 0; };
+	auto func = [=](const FBulkDataDesc& Entry) { return Entry.Type == InType && Entry.Size == 0; };
 	if (Data.FindByPredicate(func) == nullptr)
 	{
-		BulkDataDesc& Entry = Data.Emplace_GetRef();
+		FBulkDataDesc& Entry = Data.Emplace_GetRef();
 
 		Entry.ChunkId = TNumericLimits<uint64>::Max();
 		Entry.Offset = 0;
@@ -56,7 +56,7 @@ FPackageStoreBulkDataManifest::FPackageStoreBulkDataManifest(const FString& Proj
 	Filename = ProjectPath / TEXT("Metadata/BulkDataInfo.ubulkmanifest");
 }
 
-FArchive& operator<<(FArchive& Ar, FPackageStoreBulkDataManifest::PackageDesc::BulkDataDesc& Entry)
+FArchive& operator<<(FArchive& Ar, FPackageStoreBulkDataManifest::FPackageDesc::FBulkDataDesc& Entry)
 {
 	Ar << Entry.ChunkId;
 	Ar << Entry.Offset;
@@ -66,7 +66,7 @@ FArchive& operator<<(FArchive& Ar, FPackageStoreBulkDataManifest::PackageDesc::B
 	return Ar;
 }
 
-FArchive& operator<<(FArchive& Ar, FPackageStoreBulkDataManifest::PackageDesc& Entry)
+FArchive& operator<<(FArchive& Ar, FPackageStoreBulkDataManifest::FPackageDesc& Entry)
 {
 	Ar << Entry.Data;
 
@@ -99,7 +99,7 @@ void FPackageStoreBulkDataManifest::Save()
 	*BinArchive << Data;
 }
 
-const FPackageStoreBulkDataManifest::PackageDesc* FPackageStoreBulkDataManifest::Find(const FString& PackageFilename) const
+const FPackageStoreBulkDataManifest::FPackageDesc* FPackageStoreBulkDataManifest::Find(const FString& PackageFilename) const
 {
 	const  FString NormalizedFilename = FixFilename(PackageFilename);
 	return Data.Find(NormalizedFilename);
@@ -109,7 +109,7 @@ void FPackageStoreBulkDataManifest::AddFileAccess(const FString& PackageFilename
 {
 	const FString NormalizedFilename = FixFilename(PackageFilename);
 
-	PackageDesc& Entry = GetOrCreateFileAccess(NormalizedFilename);
+	FPackageDesc& Entry = GetOrCreateFileAccess(NormalizedFilename);
 
 	if (InSize > 0)
 	{
@@ -122,9 +122,9 @@ void FPackageStoreBulkDataManifest::AddFileAccess(const FString& PackageFilename
 
 }
 
-FPackageStoreBulkDataManifest::PackageDesc& FPackageStoreBulkDataManifest::GetOrCreateFileAccess(const FString& PackageFilename)
+FPackageStoreBulkDataManifest::FPackageDesc& FPackageStoreBulkDataManifest::GetOrCreateFileAccess(const FString& PackageFilename)
 {
-	if (PackageDesc* Entry = Data.Find(PackageFilename))
+	if (FPackageDesc* Entry = Data.Find(PackageFilename))
 	{
 		return *Entry;
 	}
diff --git a/Engine/Source/Runtime/CoreUObject/Public/Serialization/BulkDataManifest.h b/Engine/Source/Runtime/CoreUObject/Public/Serialization/BulkDataManifest.h
--- a/Engine/Source/Runtime/CoreUObject/Public/Serialization/BulkDataManifest.h
+++ b/Engine/Source/Runtime/CoreUObject/Public/Serialization/BulkDataManifest.h
@@ -18,10 +18,10 @@ public:
 
 	const FString& GetFilename() const { return Filename; }
 
-	class PackageDesc
+	class FPackageDesc
 	{
 	public:
-		struct BulkDataDesc
+		struct FBulkDataDesc
 		{
 			uint64 ChunkId;	// Note this is the Offset before the linker BulkDataStartOffset is
 							// applied, to make it easier to compute at runtime.
@@ -33,20 +33,20 @@ public:
 		void AddData(EIoChunkType InType, uint64 InChunkId, uint64 InOffset, uint64 InSize, const FString& DebugFilename);
 		void AddZeroByteData(EIoChunkType InType);
 
-		const TArray<BulkDataDesc>& GetDataArray() const { return Data; }
+		const TArray<FBulkDataDesc>& GetDataArray() const { return Data; }
 	private:
-		friend FArchive& operator<<(FArchive& Ar, PackageDesc& Entry);
-		TArray<BulkDataDesc> Data;
+		friend FArchive& operator<<(FArchive& Ar, FPackageDesc& Entry);
+		TArray<FBulkDataDesc> Data;
 	};
 
-	COREUOBJECT_API const PackageDesc* Find(const FString& PackageName) const;
+	COREUOBJECT_API const FPackageDesc* Find(const FString& PackageName) const;
 
 private:
-	PackageDesc& GetOrCreateFileAccess(const FString& PackageFilename);
+	FPackageDesc& GetOrCreateFileAccess(const FString& PackageFilename);
 
 	FString FixFilename(const FString& InFileName) const;
 
 	FString RootPath;
 	FString Filename;
-	TMap<FString, PackageDesc> Data;
+	TMap<FString, FPackageDesc> Data;
 };
