Optimized unversioned array property loading

Nothing major, makes UPS loading ~2% faster and exclusive array loading speed about ~30% faster for the data set used to profile this.

diff --git a/Engine/Source/Runtime/CoreUObject/Private/UObject/PropertyArray.cpp b/Engine/Source/Runtime/CoreUObject/Private/UObject/PropertyArray.cpp
--- a/Engine/Source/Runtime/CoreUObject/Private/UObject/PropertyArray.cpp
+++ b/Engine/Source/Runtime/CoreUObject/Private/UObject/PropertyArray.cpp
@@ -93,13 +93,29 @@ bool FArrayProperty::Identical( const void* A, const void* B, uint32 PortFlags )
 	return true;
 }
 
+static bool CanBulkSerialize(FProperty* Property)
+{
+#if PLATFORM_LITTLE_ENDIAN
+	// All numeric properties except TEnumAsByte
+	EClassCastFlags CastFlags = Property->GetClass()->ClassCastFlags;
+	if (!!(CastFlags & CASTCLASS_UNumericProperty))
+	{
+		bool bEnumAsByte = (CastFlags & CASTCLASS_UByteProperty) != 0 && static_cast<FByteProperty*>(Property)->Enum;
+		return !bEnumAsByte;
+	}
+#endif
+
+	return false;
+}
+
 void FArrayProperty::SerializeItem(FStructuredArchive::FSlot Slot, void* Value, void const* Defaults) const
 {
 	checkSlow(Inner);
 	FArchive& UnderlyingArchive = Slot.GetUnderlyingArchive();
+	const bool bUPS = Slot.GetArchiveState().UseUnversionedPropertySerialization();
 	TOptional<FPropertyTag> MaybeInnerTag;
 
-	if (UnderlyingArchive.IsTextFormat() && !Slot.GetArchiveState().UseUnversionedPropertySerialization()
+	if (UnderlyingArchive.IsTextFormat() && !bUPS
 		&& Inner && Inner->IsA<FStructProperty>())
 	{
 		MaybeInnerTag.Emplace(UnderlyingArchive, Inner, 0, (uint8*)Value, (uint8*)Defaults);	
@@ -130,6 +146,26 @@ void FArrayProperty::SerializeItem(FStructuredArchive::FSlot Slot, void* Value,
 				ArrayHelper.RemoveValues(n, OldNum - n);
 			}
 		}
+		else if (bUPS)
+		{
+			if (CanBulkSerialize(Inner))
+			{
+				ArrayHelper.EmptyAndAddUninitializedValues(n);
+
+				UnderlyingArchive.Serialize(ArrayHelper.GetRawPtr(), n * Inner->ElementSize);
+			}
+			else
+			{
+				ArrayHelper.EmptyAndAddValues(n);
+
+				for (int32 i = 0; i < n; ++i)
+				{
+					Inner->SerializeItem(Array.EnterElement(), ArrayHelper.GetRawPtr(i));
+				}
+			}
+			
+			return;
+		}
 		else
 		{
 			ArrayHelper.EmptyAndAddValues(n);
@@ -138,8 +174,7 @@ void FArrayProperty::SerializeItem(FStructuredArchive::FSlot Slot, void* Value,
 	ArrayHelper.CountBytes( UnderlyingArchive );
 
 	// Serialize a PropertyTag for the inner property of this array, allows us to validate the inner struct to see if it has changed
-	if (!Slot.GetArchiveState().UseUnversionedPropertySerialization() && 
-		UnderlyingArchive.UE4Ver() >= VER_UE4_INNER_ARRAY_TAG_INFO &&
+	if (!bUPS && UnderlyingArchive.UE4Ver() >= VER_UE4_INNER_ARRAY_TAG_INFO &&
 		Inner && Inner->IsA<FStructProperty>())
 	{
 		if (!MaybeInnerTag)
