FProperty-related compile fixes

diff --git a/Engine/Source/Runtime/CoreUObject/Private/UObject/PropertyArray.cpp b/Engine/Source/Runtime/CoreUObject/Private/UObject/PropertyArray.cpp
--- a/Engine/Source/Runtime/CoreUObject/Private/UObject/PropertyArray.cpp
+++ b/Engine/Source/Runtime/CoreUObject/Private/UObject/PropertyArray.cpp
@@ -97,10 +97,10 @@ static bool CanBulkSerialize(FProperty* Property)
 {
 #if PLATFORM_LITTLE_ENDIAN
 	// All numeric properties except TEnumAsByte
-	EClassCastFlags CastFlags = Property->GetClass()->ClassCastFlags;
-	if (!!(CastFlags & CASTCLASS_UNumericProperty))
+	uint64 CastFlags = Property->GetClass()->GetCastFlags();
+	if (!!(CastFlags & CASTCLASS_FNumericProperty))
 	{
-		bool bEnumAsByte = (CastFlags & CASTCLASS_UByteProperty) != 0 && static_cast<FByteProperty*>(Property)->Enum;
+		bool bEnumAsByte = (CastFlags & CASTCLASS_FByteProperty) != 0 && static_cast<FByteProperty*>(Property)->Enum;
 		return !bEnumAsByte;
 	}
 #endif
