The BulkData manifest can no longer be broken by changing the values of EIoChunkType.

Note that this will invalidating existing manifests and require a new cook!

- We have had several cases of EIoChunkType having new entries inserted at the start causing the values for bulkdata to change which in turn would invalidate the contents of the BulkData manifest, so to avoid issues in the future the manifest will now save out its own values and we only convert them to EIoChunkType at the last possible moment when building the ucas files.
- Added a custom version to the manifests file archive, this will make it safer to edit the format in the future as we will be able to avoid invalidating the entire file.
 - Unfortunatly adding the version does not help with existing files which are now invalid.

#rnx

diff --git a/Engine/Source/Developer/IoStoreUtilities/Private/IoStoreUtilities.cpp b/Engine/Source/Developer/IoStoreUtilities/Private/IoStoreUtilities.cpp
--- a/Engine/Source/Developer/IoStoreUtilities/Private/IoStoreUtilities.cpp
+++ b/Engine/Source/Developer/IoStoreUtilities/Private/IoStoreUtilities.cpp
@@ -795,14 +795,30 @@ static void BuildBundles(FExportGraph& ExportGraph, const TArray<FPackage*>& Pac
 	}
 }
 
-static bool WriteBulkData(	const FString& Filename, EIoChunkType Type, const FPackage& Package, const FPackageStoreBulkDataManifest& BulkDataManifest,
-							FIoStoreWriter* IoStoreWriter)
+static EIoChunkType BulkdataTypeToChunkIdType(FPackageStoreBulkDataManifest::EBulkdataType Type)
+{
+	switch (Type)
+	{
+	case FPackageStoreBulkDataManifest::EBulkdataType::Normal:
+		return EIoChunkType::BulkData;
+	case FPackageStoreBulkDataManifest::EBulkdataType::Optional:
+		return EIoChunkType::OptionalBulkData;
+	default:
+		UE_LOG(LogIoStore, Error, TEXT("Invalid EBulkdataType (%d) found!"), Type);
+		return EIoChunkType::Invalid;
+	}
+}
+
+static bool WriteBulkData(	const FString& Filename, FPackageStoreBulkDataManifest::EBulkdataType Type, const FPackage& Package, 
+							const FPackageStoreBulkDataManifest& BulkDataManifest, FIoStoreWriter* IoStoreWriter)
 {
 
 	const FPackageStoreBulkDataManifest::FPackageDesc* PackageDesc = BulkDataManifest.Find(Package.FileName);
 	if (PackageDesc != nullptr)
 	{
-		const FIoChunkId BulkDataChunkId = CreateChunkIdForBulkData(Package.GlobalPackageId, TNumericLimits<uint64>::Max()-1, Type, *Package.FileName);
+		const EIoChunkType ChunkIdType = BulkdataTypeToChunkIdType(Type);
+
+		const FIoChunkId BulkDataChunkId = CreateChunkIdForBulkData(Package.GlobalPackageId, TNumericLimits<uint64>::Max()-1, ChunkIdType, *Package.FileName);
 
 #if !SKIP_WRITE_CONTAINER		
 		FIoBuffer IoBuffer;
@@ -830,7 +846,7 @@ static bool WriteBulkData(	const FString& Filename, EIoChunkType Type, const FPa
 		{
 			if (BulkDataDesc.Type == Type)
 			{
-				const FIoChunkId AccessChunkId = CreateChunkIdForBulkData(Package.GlobalPackageId, BulkDataDesc.ChunkId, Type, *Package.FileName);
+				const FIoChunkId AccessChunkId = CreateChunkIdForBulkData(Package.GlobalPackageId, BulkDataDesc.ChunkId, ChunkIdType, *Package.FileName);
 #if !SKIP_WRITE_CONTAINER	
 				const FIoStatus PartialResult = IoStoreWriter->MapPartialRange(BulkDataChunkId, BulkDataDesc.Offset, BulkDataDesc.Size, AccessChunkId);
 				if (!PartialResult.IsOk())
@@ -1840,11 +1856,11 @@ int32 CreateTarget(const TCHAR* OutputDir, const TCHAR* CookedDir, const TArray<
 					FString Extension = FPaths::GetExtension(NormalizedSourcePath);
 					if (Extension == TEXT("uptnl"))
 					{
-						WriteBulkData(NormalizedSourcePath, EIoChunkType::OptionalBulkData, *TargetFile.Package, BulkDataManifest, ContainerTarget.IoStoreWriter);
+						WriteBulkData(NormalizedSourcePath, FPackageStoreBulkDataManifest::EBulkdataType::Optional, *TargetFile.Package, BulkDataManifest, ContainerTarget.IoStoreWriter);
 					}
 					else
 					{
-						WriteBulkData(NormalizedSourcePath, EIoChunkType::BulkData, *TargetFile.Package, BulkDataManifest, ContainerTarget.IoStoreWriter);
+						WriteBulkData(NormalizedSourcePath, FPackageStoreBulkDataManifest::EBulkdataType::Normal, *TargetFile.Package, BulkDataManifest, ContainerTarget.IoStoreWriter);
 					}
 				}
 			}
diff --git a/Engine/Source/Runtime/CoreUObject/Private/Serialization/BulkDataManifest.cpp b/Engine/Source/Runtime/CoreUObject/Private/Serialization/BulkDataManifest.cpp
--- a/Engine/Source/Runtime/CoreUObject/Private/Serialization/BulkDataManifest.cpp
+++ b/Engine/Source/Runtime/CoreUObject/Private/Serialization/BulkDataManifest.cpp
@@ -4,9 +4,32 @@
 
 #include "CoreGlobals.h"
 #include "HAL/FileManager.h"
+#include "Misc/Guid.h"
 #include "Misc/Paths.h"
+#include "Serialization/CustomVersion.h"
 
-void FPackageStoreBulkDataManifest::FPackageDesc::AddData(EIoChunkType InType, uint64 InChunkId, uint64 InOffset, uint64 InSize, const FString& DebugFilename)
+struct FBulkDataManifestVersion
+{
+	FBulkDataManifestVersion() = delete;
+
+	enum Type
+	{
+		// Before any version changes were made
+		BeforeCustomVersionWasAdded = 0,
+
+		// -----<new versions can be added above this line>-------------------------------------------------
+		VersionPlusOne,
+		LatestVersion = VersionPlusOne - 1
+	};
+
+	// The GUID for this custom version number
+	const static FGuid GUID;
+};
+
+const FGuid FBulkDataManifestVersion::GUID(0x54683250, 0x809948AF, 0x8BC89896, 0xFBADF9B7);
+FCustomVersionRegistration GRegisterBulkDataManifestVersion(FBulkDataManifestVersion::GUID, FBulkDataManifestVersion::LatestVersion, TEXT("BulkDataManifestVersion"));
+
+void FPackageStoreBulkDataManifest::FPackageDesc::AddData(EBulkdataType InType, uint64 InChunkId, uint64 InOffset, uint64 InSize, const FString& DebugFilename)
 {
 	// The ChunkId is supposed to be unique for each type
 	auto func = [=](const FBulkDataDesc& Entry) { return Entry.ChunkId == InChunkId && Entry.Type == InType; };
@@ -36,7 +59,7 @@ void FPackageStoreBulkDataManifest::FPackageDesc::AddData(EIoChunkType InType, u
 	}
 }
 
-void FPackageStoreBulkDataManifest::FPackageDesc::AddZeroByteData(EIoChunkType InType)
+void FPackageStoreBulkDataManifest::FPackageDesc::AddZeroByteData(EBulkdataType InType)
 {
 	// Make sure we only add one empty read per Bulkdata type!
 	auto func = [=](const FBulkDataDesc& Entry) { return Entry.Type == InType && Entry.Size == 0; };
@@ -73,10 +96,6 @@ FArchive& operator<<(FArchive& Ar, FPackageStoreBulkDataManifest::FPackageDesc&
 	return Ar;
 }
 
-FPackageStoreBulkDataManifest::~FPackageStoreBulkDataManifest()
-{
-}
-
 bool FPackageStoreBulkDataManifest::Load()
 {
 	Data.Empty();
@@ -84,6 +103,11 @@ bool FPackageStoreBulkDataManifest::Load()
 	TUniquePtr<FArchive> BinArchive(IFileManager::Get().CreateFileReader(*Filename));
 	if (BinArchive != nullptr)
 	{
+		// Load and apply the custom versions before we load any data
+		FCustomVersionContainer CustomVersions;
+		CustomVersions.Serialize(*BinArchive);
+		BinArchive->SetCustomVersions(CustomVersions);
+		
 		*BinArchive << Data;
 		return true;
 	}
@@ -96,6 +120,14 @@ bool FPackageStoreBulkDataManifest::Load()
 void FPackageStoreBulkDataManifest::Save()
 {
 	TUniquePtr<FArchive> BinArchive(IFileManager::Get().CreateFileWriter(*Filename));
+	BinArchive->UsingCustomVersion(FBulkDataManifestVersion::GUID);
+
+	// Take the versions from the archive and serialize them out
+	// NOTE: Serializing out now means we cannot add additional versions while serializing
+	// it is assumed that we are only using FBulkDataManifestVersion::GUID
+	FCustomVersionContainer CustomVersions = BinArchive->GetCustomVersions();
+	CustomVersions.Serialize(*BinArchive);
+
 	*BinArchive << Data;
 }
 
@@ -105,7 +137,7 @@ const FPackageStoreBulkDataManifest::FPackageDesc* FPackageStoreBulkDataManifest
 	return Data.Find(NormalizedFilename);
 }
 
-void FPackageStoreBulkDataManifest::AddFileAccess(const FString& PackageFilename, EIoChunkType InType, uint64 InChunkId, uint64 InOffset, uint64 InSize)
+void FPackageStoreBulkDataManifest::AddFileAccess(const FString& PackageFilename, EBulkdataType InType, uint64 InChunkId, uint64 InOffset, uint64 InSize)
 {
 	const FString NormalizedFilename = FixFilename(PackageFilename);
 
@@ -119,7 +151,6 @@ void FPackageStoreBulkDataManifest::AddFileAccess(const FString& PackageFilename
 	{
 		Entry.AddZeroByteData(InType);
 	}
-
 }
 
 FPackageStoreBulkDataManifest::FPackageDesc& FPackageStoreBulkDataManifest::GetOrCreateFileAccess(const FString& PackageFilename)
diff --git a/Engine/Source/Runtime/CoreUObject/Private/UObject/SavePackage.cpp b/Engine/Source/Runtime/CoreUObject/Private/UObject/SavePackage.cpp
--- a/Engine/Source/Runtime/CoreUObject/Private/UObject/SavePackage.cpp
+++ b/Engine/Source/Runtime/CoreUObject/Private/UObject/SavePackage.cpp
@@ -6228,7 +6228,8 @@ void SaveBulkData(FLinkerSave* Linker, const UPackage* InOuter, const TCHAR* Fil
 			if (SavePackageContext != nullptr && SavePackageContext->BulkDataManifest != nullptr)
 			{
 				const bool bIsOptional = (BulkDataStorageInfo.BulkDataFlags & BULKDATA_OptionalPayload) != 0;
-				const EIoChunkType Type = bIsOptional ? EIoChunkType::OptionalBulkData : EIoChunkType::BulkData;
+				const FPackageStoreBulkDataManifest::EBulkdataType Type = !bIsOptional ? FPackageStoreBulkDataManifest::EBulkdataType::Normal :
+																						 FPackageStoreBulkDataManifest::EBulkdataType::Optional;
 
 				SavePackageContext->BulkDataManifest->AddFileAccess(Filename, Type, StoredBulkStartOffset, BulkStartOffset, SizeOnDisk);
 			}
diff --git a/Engine/Source/Runtime/CoreUObject/Public/Serialization/BulkDataManifest.h b/Engine/Source/Runtime/CoreUObject/Public/Serialization/BulkDataManifest.h
--- a/Engine/Source/Runtime/CoreUObject/Public/Serialization/BulkDataManifest.h
+++ b/Engine/Source/Runtime/CoreUObject/Public/Serialization/BulkDataManifest.h
@@ -8,30 +8,37 @@
 class FPackageStoreBulkDataManifest
 {
 public:
+	enum class EBulkdataType : uint8
+	{
+		Normal = 0,
+		Optional
+	};
+
 	COREUOBJECT_API FPackageStoreBulkDataManifest(const FString& ProjectPath);
-	COREUOBJECT_API ~FPackageStoreBulkDataManifest();
+	~FPackageStoreBulkDataManifest() = default;
 
 	COREUOBJECT_API bool Load();
 	COREUOBJECT_API void Save();
 
-	void AddFileAccess(const FString& PackageFilename, EIoChunkType InType, uint64 InChunkId, uint64 InOffset, uint64 InSize);
+	void AddFileAccess(const FString& PackageFilename, EBulkdataType InType, uint64 InChunkId, uint64 InOffset, uint64 InSize);
 
 	const FString& GetFilename() const { return Filename; }
-
 	class FPackageDesc
+
 	{
 	public:
+
 		struct FBulkDataDesc
 		{
-			uint64 ChunkId;	// Note this is the Offset before the linker BulkDataStartOffset is
-							// applied, to make it easier to compute at runtime.
-			uint64 Offset;
-			uint64 Size;
-			EIoChunkType Type;
+			uint64			ChunkId;	// Note this is the Offset before the linker BulkDataStartOffset is
+										// applied, to make it easier to compute at runtime.
+			uint64			Offset;
+			uint64			Size;
+			EBulkdataType	Type;
 		};
 
-		void AddData(EIoChunkType InType, uint64 InChunkId, uint64 InOffset, uint64 InSize, const FString& DebugFilename);
-		void AddZeroByteData(EIoChunkType InType);
+		void AddData(EBulkdataType InType, uint64 InChunkId, uint64 InOffset, uint64 InSize, const FString& DebugFilename);
+		void AddZeroByteData(EBulkdataType InType);
 
 		const TArray<FBulkDataDesc>& GetDataArray() const { return Data; }
 	private:
