Speculative static analysis warning fix

diff --git a/Engine/Source/Runtime/CoreUObject/Private/UObject/PropertyArray.cpp b/Engine/Source/Runtime/CoreUObject/Private/UObject/PropertyArray.cpp
--- a/Engine/Source/Runtime/CoreUObject/Private/UObject/PropertyArray.cpp
+++ b/Engine/Source/Runtime/CoreUObject/Private/UObject/PropertyArray.cpp
@@ -110,13 +110,12 @@ static bool CanBulkSerialize(FProperty* Property)
 
 void FArrayProperty::SerializeItem(FStructuredArchive::FSlot Slot, void* Value, void const* Defaults) const
 {
-	checkSlow(Inner);
+	check(Inner);
 	FArchive& UnderlyingArchive = Slot.GetUnderlyingArchive();
 	const bool bUPS = Slot.GetArchiveState().UseUnversionedPropertySerialization();
 	TOptional<FPropertyTag> MaybeInnerTag;
 
-	if (UnderlyingArchive.IsTextFormat() && !bUPS
-		&& Inner && Inner->IsA<FStructProperty>())
+	if (UnderlyingArchive.IsTextFormat() && !bUPS && Inner->IsA<FStructProperty>())
 	{
 		MaybeInnerTag.Emplace(UnderlyingArchive, Inner, 0, (uint8*)Value, (uint8*)Defaults);	
 		Slot << SA_ATTRIBUTE(TEXT("InnerStructName"), MaybeInnerTag.GetValue().StructName);
@@ -174,8 +173,7 @@ void FArrayProperty::SerializeItem(FStructuredArchive::FSlot Slot, void* Value,
 	ArrayHelper.CountBytes( UnderlyingArchive );
 
 	// Serialize a PropertyTag for the inner property of this array, allows us to validate the inner struct to see if it has changed
-	if (!bUPS && UnderlyingArchive.UE4Ver() >= VER_UE4_INNER_ARRAY_TAG_INFO &&
-		Inner && Inner->IsA<FStructProperty>())
+	if (!bUPS && UnderlyingArchive.UE4Ver() >= VER_UE4_INNER_ARRAY_TAG_INFO && Inner->IsA<FStructProperty>())
 	{
 		if (!MaybeInnerTag)
 		{
