Fix an assert when serializing packages that contain UObjects that do not have the RF_NeedLoad flag set.

If we do not serialize the object we need to make sure we seek past the data in the archive.

Added more info to the assert.

#rnx

diff --git a/Engine/Source/Runtime/CoreUObject/Private/Serialization/AsyncLoading2.cpp b/Engine/Source/Runtime/CoreUObject/Private/Serialization/AsyncLoading2.cpp
--- a/Engine/Source/Runtime/CoreUObject/Private/Serialization/AsyncLoading2.cpp
+++ b/Engine/Source/Runtime/CoreUObject/Private/Serialization/AsyncLoading2.cpp
@@ -2700,9 +2700,16 @@ EAsyncPackageState::Type FAsyncPackage2::Event_ProcessExportBundle(FAsyncPackage
 				TRACE_LOADTIME_SERIALIZE_EXPORT_SCOPE(Object, ExportSerialSize);
 				const int64 Pos = Ar.Tell();
 				check(ExportSerialSize <= uint64(Ar.TotalSize() - Pos));
+
 				Package->EventDrivenSerializeExport(BundleEntry->LocalExportIndex, Ar);
-				check(ExportSerialSize == uint64(Ar.Tell() - Pos));
+				checkf(ExportSerialSize == uint64(Ar.Tell() - Pos), TEXT("Expect read size: %llu - Actual read size: %llu"), ExportSerialSize, uint64(Ar.Tell() - Pos));
 			}
+			else
+			{
+				// If we are not going to serialize the data we need to skip it!
+				Ar.Seek(Ar.Tell() + ExportSerialSize);
+			}
+
 			check(!Object->HasAnyFlags(RF_NeedLoad));
 
 			Package->SerialDataPtr += ExportSerialSize;
