AsyncLoading2: Improve logging.
Log and Display is quite silent.
Verbose in Test gives information about when all packages are requested and completely loaded.
Verbose in Development gives some additional information.
VeryVerbose in Development adds even more information.

diff --git a/Engine/Source/Runtime/CoreUObject/Private/Serialization/AsyncLoading2.cpp b/Engine/Source/Runtime/CoreUObject/Private/Serialization/AsyncLoading2.cpp
--- a/Engine/Source/Runtime/CoreUObject/Private/Serialization/AsyncLoading2.cpp
+++ b/Engine/Source/Runtime/CoreUObject/Private/Serialization/AsyncLoading2.cpp
@@ -53,7 +53,30 @@
 
 #if UE_BUILD_DEVELOPMENT || UE_BUILD_DEBUG
 //PRAGMA_DISABLE_OPTIMIZATION
-#define ALT2_VERIFY_ASYNC_FLAGS
+#endif
+
+#define ALT2_VERIFY_ASYNC_FLAGS DO_CHECK
+
+#define UE_ASYNC_PACKAGE_LOG(Verbosity, Desc, Function, Text) \
+if (Desc.Name != Desc.NameToLoad) \
+{ \
+	UE_LOG(LogStreaming, Verbosity, Function TEXT(": ") Text TEXT(": %s (%d) %s (%d)."), \
+		*Desc.Name.ToString(), \
+		Desc.PackageId.ToIndexForDebugging(), \
+		*Desc.NameToLoad.ToString(), \
+		Desc.PackageIdToLoad.ToIndexForDebugging()); \
+} \
+else \
+{ \
+	UE_LOG(LogStreaming, Verbosity, Function TEXT(": ") Text TEXT(": %s (%d)."), \
+		*Desc.Name.ToString(), \
+		Desc.PackageId.ToIndexForDebugging()); \
+}
+
+#if DO_CHECK
+#define UE_ASYNC_PACKAGE_LOG_VERBOSE(Verbosity, Desc, Function, Text) UE_ASYNC_PACKAGE_LOG(Verbosity, Desc, Function, Text)
+#else
+#define UE_ASYNC_PACKAGE_LOG_VERBOSE(Verbosity, Desc, Function, Text)
 #endif
 
 struct FAsyncPackage2;
@@ -1975,6 +1998,16 @@ bool FAsyncLoadingThread2::CreateAsyncPackagesFromQueue()
 	{
 		bool bInserted;
 		FAsyncPackage2* Package = FindOrInsertPackage(PackageRequest, bInserted);
+
+		if (bInserted)
+		{
+			UE_ASYNC_PACKAGE_LOG(Verbose, Package->Desc, TEXT("CreateAsyncPackages"), TEXT("AddPackage"));
+		}
+		else
+		{
+			UE_ASYNC_PACKAGE_LOG_VERBOSE(Verbose, Package->Desc, TEXT("CreateAsyncPackages"), TEXT("UpdPackage"));
+		}
+
 		--QueuedPackagesCounter;
 		{
 			TRACE_CPUPROFILER_EVENT_SCOPE(ImportPackages);
@@ -2501,21 +2534,30 @@ void FAsyncPackage2::ImportPackagesRecursive()
 			FPackageId PackageId = FPackageId::FromIndex(GlobalPackageIndex);
 			FAsyncPackageDesc2 Info(INDEX_NONE, PackageId, PackageId, Entry.Name, Entry.Name);
 			bool bInserted;
-			FAsyncPackage2* ImportedAsyncPackage = AsyncLoadingThread.FindOrInsertPackage(&Info, bInserted);
-			if (ImportedAsyncPackage)
+			FAsyncPackage2* ImportedPackage = AsyncLoadingThread.FindOrInsertPackage(&Info, bInserted);
+			if (ImportedPackage)
 			{
-				TRACE_LOADTIME_ASYNC_PACKAGE_IMPORT_DEPENDENCY(this, ImportedAsyncPackage);
-				ImportedAsyncPackage->AddRef();
+				TRACE_LOADTIME_ASYNC_PACKAGE_IMPORT_DEPENDENCY(this, ImportedPackage);
+				if (bInserted)
+				{
+					UE_ASYNC_PACKAGE_LOG(Verbose, ImportedPackage->Desc, TEXT("ImportPackages"), TEXT("AddPackage"));
+				}
+				else
+				{
+					UE_ASYNC_PACKAGE_LOG_VERBOSE(VeryVerbose, ImportedPackage->Desc, TEXT("ImportPackages"), TEXT("UpdPackage"));
+				}
+				ImportedPackage->AddRef();
 				ImportedAsyncPackages.Reserve(ImportedPackageCount);
-				ImportedAsyncPackages.Add(ImportedAsyncPackage);
+				ImportedAsyncPackages.Add(ImportedPackage);
 				if (bInserted)
 				{
-					ImportedAsyncPackage->ImportPackagesRecursive();
-					ImportedAsyncPackage->StartLoading();
+					ImportedPackage->ImportPackagesRecursive();
+					ImportedPackage->StartLoading();
 				}
 			}
 		}
 	}
+	UE_ASYNC_PACKAGE_LOG_VERBOSE(VeryVerbose, Desc, TEXT("ImportPackages"), TEXT("ImportsDone"));
 }
 
 void FAsyncPackage2::StartLoading()
@@ -3064,6 +3106,8 @@ EAsyncPackageState::Type FAsyncPackage2::Event_PostLoad(FAsyncPackage2* Package,
 		Package->AsyncPackageLoadingState = EAsyncPackageLoadingState2::PackageComplete;
 	}
 
+	UE_ASYNC_PACKAGE_LOG(Verbose, Package->Desc, TEXT("AsyncPackage"), TEXT("LoadFinished"));
+
 	if (LoadingState == EAsyncPackageState::TimeOut)
 	{
 		return EAsyncPackageState::TimeOut;
@@ -3309,6 +3353,8 @@ EAsyncPackageState::Type FAsyncLoadingThread2::ProcessLoadedPackagesFromGameThre
 			Package->bAddedForDelete = true;
 			Package->MarkRequestIDsAsComplete();
 
+			UE_ASYNC_PACKAGE_LOG(Verbose, Package->Desc, TEXT("GameThread"), TEXT("LoadCompleted"));
+
 			if (FlushRequestID != INDEX_NONE && !ContainsRequestID(FlushRequestID))
 			{
 				// The only package we care about has finished loading, so we're good to exit
@@ -4264,7 +4310,14 @@ void FAsyncPackage2::CreateUPackage(const FPackageSummary* PackageSummary)
 	AddOwnedObjectWithAsyncFlag(LinkerRoot, /*bForceAdd*/ !bCreatedLinkerRoot);
 	check(LinkerRoot->HasAnyInternalFlags(EInternalObjectFlags::Async));
 
-	UE_LOG(LogStreaming, Verbose, TEXT("FAsyncPackage::CreateUPackage for %s finished."), *Desc.Name.ToString());
+	if (bCreatedLinkerRoot)
+	{
+		UE_ASYNC_PACKAGE_LOG_VERBOSE(VeryVerbose, Desc, TEXT("CreateUPackage"), TEXT("AddPackage"));
+	}
+	else
+	{
+		UE_ASYNC_PACKAGE_LOG_VERBOSE(VeryVerbose, Desc, TEXT("CreateUPackage"), TEXT("UpdPackage"));
+	}
 }
 
 EAsyncPackageState::Type FAsyncPackage2::ProcessExternalReads(EExternalReadAction Action)
@@ -4623,11 +4676,14 @@ int32 FAsyncLoadingThread2::LoadPackage(const FString& InName, const FGuid* InGu
 		// Add new package request
 		FAsyncPackageDesc2 PackageDesc(RequestID, PackageId, PackageIdToLoad, PackageName, PackageNameToLoad, MoveTemp(CompletionDelegatePtr));
 		QueuePackage(PackageDesc);
+
+		UE_ASYNC_PACKAGE_LOG(Verbose, PackageDesc, TEXT("LoadPackage"), TEXT("QueuePackage"));
 	}
 	else
 	{
-		UE_LOG(LogStreaming, Warning, TEXT("AsyncLoading2 - LoadPackage: Skipping package: '%s'. Name to load is unknown: '%s')"),
-			*PackageName.ToString(), *PackageNameToLoad.ToString());
+		FPackageId PackageId = PackageName != PackageNameToLoad ? GlobalPackageStore.FindPackageId(PackageName) : PackageIdToLoad;
+		FAsyncPackageDesc2 PackageDesc(RequestID, PackageId, PackageIdToLoad, PackageName, PackageNameToLoad);
+		UE_ASYNC_PACKAGE_LOG(Warning, PackageDesc, TEXT("LoadPackage"), TEXT("SkipUnknownPackage"));
 		InCompletionDelegate.ExecuteIfBound(PackageName, nullptr, EAsyncLoadingResult::Failed);
 	}
 
diff --git a/Engine/Source/Runtime/CoreUObject/Public/UObject/PackageId.h b/Engine/Source/Runtime/CoreUObject/Public/UObject/PackageId.h
--- a/Engine/Source/Runtime/CoreUObject/Public/UObject/PackageId.h
+++ b/Engine/Source/Runtime/CoreUObject/Public/UObject/PackageId.h
@@ -33,6 +33,11 @@ public:
 		return Id;
 	}
 
+	inline uint32 ToIndexForDebugging() const
+	{
+		return Id;
+	}
+
 	inline bool operator==(FPackageId Other) const
 	{
 		return Id == Other.Id;
