AsyncLoading2: Replay bug fixes.
Properly handle load requests for temp packages that are already completed with a simple warning log.
Add an extra call toMarkAsFullyLoaded after finish objects (just as in the old loader), not actually required but probably nice to keep that behavior if some other code relies on it.

diff --git a/Engine/Source/Runtime/CoreUObject/Private/Serialization/AsyncLoading2.cpp b/Engine/Source/Runtime/CoreUObject/Private/Serialization/AsyncLoading2.cpp
--- a/Engine/Source/Runtime/CoreUObject/Private/Serialization/AsyncLoading2.cpp
+++ b/Engine/Source/Runtime/CoreUObject/Private/Serialization/AsyncLoading2.cpp
@@ -58,19 +58,19 @@
 #define ALT2_VERIFY_ASYNC_FLAGS DO_CHECK
 
 #define UE_ASYNC_PACKAGE_LOG(Verbosity, Desc, Function, Text) \
-if (Desc.Name != Desc.NameToLoad) \
+if ((Desc).Name != (Desc).NameToLoad) \
 { \
 	UE_LOG(LogStreaming, Verbosity, Function TEXT(": ") Text TEXT(": %s (%d) %s (%d)."), \
-		*Desc.Name.ToString(), \
-		Desc.PackageId.ToIndexForDebugging(), \
-		*Desc.NameToLoad.ToString(), \
-		Desc.PackageIdToLoad.ToIndexForDebugging()); \
+		*(Desc).Name.ToString(), \
+		(Desc).PackageId.ToIndexForDebugging(), \
+		*(Desc).NameToLoad.ToString(), \
+		(Desc).PackageIdToLoad.ToIndexForDebugging()); \
 } \
 else \
 { \
 	UE_LOG(LogStreaming, Verbosity, Function TEXT(": ") Text TEXT(": %s (%d)."), \
-		*Desc.Name.ToString(), \
-		Desc.PackageId.ToIndexForDebugging()); \
+		*(Desc).Name.ToString(), \
+		(Desc).PackageId.ToIndexForDebugging()); \
 }
 
 #if DO_CHECK
@@ -1852,7 +1852,14 @@ private:
 
 	FAsyncPackage2* CreateAsyncPackage(const FAsyncPackageDesc2& Desc)
 	{
-		return new FAsyncPackage2(Desc, *this, EDLBootNotificationManager, GraphAllocator, EventSpecs.GetData());
+		if (Desc.PackageIdToLoad.ToIndex() < (uint32)GlobalPackageStore.PackageCount)
+		{
+			return new FAsyncPackage2(Desc, *this, EDLBootNotificationManager, GraphAllocator, EventSpecs.GetData());
+		}
+		else
+		{
+			return nullptr;
+		}
 	}
 
 	/**
@@ -1994,21 +2001,27 @@ bool FAsyncLoadingThread2::CreateAsyncPackagesFromQueue()
 		QueuedPackages.Reset();
 	}
 
-	for (FAsyncPackageDesc2* PackageRequest : QueueCopy)
+	for (FAsyncPackageDesc2* PackageDesc : QueueCopy)
 	{
 		bool bInserted;
-		FAsyncPackage2* Package = FindOrInsertPackage(PackageRequest, bInserted);
+		FAsyncPackage2* Package = FindOrInsertPackage(PackageDesc, bInserted);
 
 		if (bInserted)
 		{
-			UE_ASYNC_PACKAGE_LOG(Verbose, Package->Desc, TEXT("CreateAsyncPackages"), TEXT("AddPackage"));
+			UE_ASYNC_PACKAGE_LOG(Verbose, *PackageDesc, TEXT("CreateAsyncPackages"), TEXT("AddPackage"));
+		}
+		else if (!Package)
+		{
+			UE_ASYNC_PACKAGE_LOG(Warning, *PackageDesc, TEXT("CreateAsyncPackages"),
+				TEXT("Skipping unknown package, probably a temp package that has already been completely loaded"));
 		}
 		else
 		{
-			UE_ASYNC_PACKAGE_LOG_VERBOSE(Verbose, Package->Desc, TEXT("CreateAsyncPackages"), TEXT("UpdPackage"));
+			UE_ASYNC_PACKAGE_LOG_VERBOSE(Verbose, *PackageDesc, TEXT("CreateAsyncPackages"), TEXT("UpdPackage"));
 		}
 
 		--QueuedPackagesCounter;
+		if (Package)
 		{
 			TRACE_CPUPROFILER_EVENT_SCOPE(ImportPackages);
 			Package->ImportPackagesRecursive();
@@ -2017,9 +2030,12 @@ bool FAsyncLoadingThread2::CreateAsyncPackagesFromQueue()
 		{
 			Package->StartLoading();
 		}
-		delete PackageRequest;
+		delete PackageDesc;
 
-		StartBundleIoRequests();
+		if (Package)
+		{
+			StartBundleIoRequests();
+		}
 	}
 
 	return QueueCopy.Num() > 0;
@@ -2532,19 +2548,19 @@ void FAsyncPackage2::ImportPackagesRecursive()
 		if (bNeedToLoadPackage)
 		{
 			FPackageId PackageId = FPackageId::FromIndex(GlobalPackageIndex);
-			FAsyncPackageDesc2 Info(INDEX_NONE, PackageId, PackageId, Entry.Name, Entry.Name);
+			FAsyncPackageDesc2 PackageDesc(INDEX_NONE, PackageId, PackageId, Entry.Name, Entry.Name);
 			bool bInserted;
-			FAsyncPackage2* ImportedPackage = AsyncLoadingThread.FindOrInsertPackage(&Info, bInserted);
+			FAsyncPackage2* ImportedPackage = AsyncLoadingThread.FindOrInsertPackage(&PackageDesc, bInserted);
 			if (ImportedPackage)
 			{
 				TRACE_LOADTIME_ASYNC_PACKAGE_IMPORT_DEPENDENCY(this, ImportedPackage);
 				if (bInserted)
 				{
-					UE_ASYNC_PACKAGE_LOG(Verbose, ImportedPackage->Desc, TEXT("ImportPackages"), TEXT("AddPackage"));
+					UE_ASYNC_PACKAGE_LOG(Verbose, PackageDesc, TEXT("ImportPackages"), TEXT("AddPackage"));
 				}
 				else
 				{
-					UE_ASYNC_PACKAGE_LOG_VERBOSE(VeryVerbose, ImportedPackage->Desc, TEXT("ImportPackages"), TEXT("UpdPackage"));
+					UE_ASYNC_PACKAGE_LOG_VERBOSE(VeryVerbose, PackageDesc, TEXT("ImportPackages"), TEXT("UpdPackage"));
 				}
 				ImportedPackage->AddRef();
 				ImportedAsyncPackages.Reserve(ImportedPackageCount);
@@ -2555,6 +2571,10 @@ void FAsyncPackage2::ImportPackagesRecursive()
 					ImportedPackage->StartLoading();
 				}
 			}
+			else
+			{
+				UE_ASYNC_PACKAGE_LOG(Error, PackageDesc, TEXT("ImportPackages"), TEXT("Skipping unknown package, but this should never happen here"));
+			}
 		}
 	}
 	UE_ASYNC_PACKAGE_LOG_VERBOSE(VeryVerbose, Desc, TEXT("ImportPackages"), TEXT("ImportsDone"));
@@ -3106,7 +3126,11 @@ EAsyncPackageState::Type FAsyncPackage2::Event_PostLoad(FAsyncPackage2* Package,
 		Package->AsyncPackageLoadingState = EAsyncPackageLoadingState2::PackageComplete;
 	}
 
-	UE_ASYNC_PACKAGE_LOG(Verbose, Package->Desc, TEXT("AsyncPackage"), TEXT("LoadFinished"));
+	if (Package->LinkerRoot && LoadingState == EAsyncPackageState::Complete)
+	{
+		UE_ASYNC_PACKAGE_LOG(Verbose, Package->Desc, TEXT("AsyncPackage"), TEXT("FullyLoaded"));
+		Package->LinkerRoot->MarkAsFullyLoaded();
+	}
 
 	if (LoadingState == EAsyncPackageState::TimeOut)
 	{
@@ -4683,7 +4707,7 @@ int32 FAsyncLoadingThread2::LoadPackage(const FString& InName, const FGuid* InGu
 	{
 		FPackageId PackageId = PackageName != PackageNameToLoad ? GlobalPackageStore.FindPackageId(PackageName) : PackageIdToLoad;
 		FAsyncPackageDesc2 PackageDesc(RequestID, PackageId, PackageIdToLoad, PackageName, PackageNameToLoad);
-		UE_ASYNC_PACKAGE_LOG(Warning, PackageDesc, TEXT("LoadPackage"), TEXT("SkipUnknownPackage"));
+		UE_ASYNC_PACKAGE_LOG(Warning, PackageDesc, TEXT("LoadPackage"), TEXT("Skipping unknown package, the provided package does not exist"));
 		InCompletionDelegate.ExecuteIfBound(PackageName, nullptr, EAsyncLoadingResult::Failed);
 	}
 
