AsyncLoading2: Fix replay crash caused by keeping stale UPackage pointer due to a missing weak ptr check for IsUnreachable.

diff --git a/Engine/Source/Runtime/CoreUObject/Private/Serialization/AsyncLoading2.cpp b/Engine/Source/Runtime/CoreUObject/Private/Serialization/AsyncLoading2.cpp
--- a/Engine/Source/Runtime/CoreUObject/Private/Serialization/AsyncLoading2.cpp
+++ b/Engine/Source/Runtime/CoreUObject/Private/Serialization/AsyncLoading2.cpp
@@ -237,9 +237,11 @@ private:
 	void InternalMakeRaw()
 	{
 		check(bIsWeakPointer);
-		FUObjectItem* ObjectItem = GUObjectArray.IndexToObject(WeakPointer.ObjectIndex);
+		const bool bEvenIfPendingKill = false;
+		FUObjectItem* ObjectItem = GUObjectArray.IndexToValidObject(WeakPointer.ObjectIndex, bEvenIfPendingKill);
 		int32 ActualSerialNumber = GUObjectArray.GetSerialNumber(WeakPointer.ObjectIndex);
 		Object = (ObjectItem && ActualSerialNumber == WeakPointer.SerialNumber) ? (UObject*)ObjectItem->Object : nullptr;
+		check(!Object || (!Object->IsUnreachable() && !Object->IsPendingKill()));
 		bIsWeakPointer = false;
 	}
 };
