Fix a crash when running EngineTest with the IoDispatcher.

- LoadDataDirectly now branches to the older file system if the IoDispatcher is disabled
- Implement a direct loading path from the IoDispatcher for non-inline data
- Log an error if someone tries to directly load inline data via the IoDipatcher as we do not and will not support that (although I do have a code shelf ready for it just in case)
- Re-enabled the perf warning about reloading inline bulkdata as a start towards removing this functionality as we cannot support it in the IoDisptcher without adding memory bloat and it is not desirable in the first place.
- Removed the malloc call in GetCopy before LoadDataDirectly since that method will allocate the buffer if needed anyway, this ends up being a small optimization by reducing the number of calls to ::GetBulkDataSize

#rnx

diff --git a/Engine/Source/Runtime/CoreUObject/Private/Serialization/BulkData2.cpp b/Engine/Source/Runtime/CoreUObject/Private/Serialization/BulkData2.cpp
--- a/Engine/Source/Runtime/CoreUObject/Private/Serialization/BulkData2.cpp
+++ b/Engine/Source/Runtime/CoreUObject/Private/Serialization/BulkData2.cpp
@@ -937,7 +937,6 @@ void FBulkDataBase::GetCopy(void** DstBuffer, bool bDiscardInternalCopy)
 		}
 		else
 		{
-			*DstBuffer = FMemory::Malloc(GetBulkDataSize(), 0);
 			LoadDataDirectly(DstBuffer);
 		}
 	}
@@ -1250,7 +1249,7 @@ void FBulkDataBase::LoadDataDirectly(void** DstBuffer)
 		return; // Early out if there is nothing to load anyway
 	}
 
-	if (!IsUsingIODispatcher())
+	if (!IsIoDispatcherEnabled())
 	{
 		const int64 BulkDataSize = GetBulkDataSize();
 		FileTokenSystem::Data FileData = FileTokenSystem::GetFileData(Fallback.Token);
@@ -1271,7 +1270,7 @@ void FBulkDataBase::LoadDataDirectly(void** DstBuffer)
 
 		// If the data is inlined then we already loaded is during ::Serialize, this warning should help track cases where data is being discarded then re-requested.
 		// Disabled at the moment
-		//UE_CLOG(IsInlined(), LogSerialization, Warning, TEXT("Reloading inlined bulk data directly from disk, this is detrimental to loading performance. Filename: '%s'."), *Filename);
+		UE_CLOG(IsInlined(), LogSerialization, Warning, TEXT("Reloading inlined bulk data directly from disk, this is detrimental to loading performance. Filename: '%s'."), *Filename);
 
 		FArchive* Ar = IFileManager::Get().CreateFileReader(*Filename, FILEREAD_Silent);
 		checkf(Ar != nullptr, TEXT("Failed to open the file to load bulk data from. Filename: '%s'."), *Filename);
@@ -1288,9 +1287,31 @@ void FBulkDataBase::LoadDataDirectly(void** DstBuffer)
 
 		delete Ar;
 	}
+	else if (IsUsingIODispatcher())
+	{
+		// Allocate the buffer if needed
+		if (*DstBuffer == nullptr)
+		{
+			*DstBuffer = FMemory::Malloc(GetBulkDataSize(), 0);
+		}
+
+		// Set up our options (we only need to set the target)
+		FIoReadOptions Options;
+		Options.SetTargetVa(*DstBuffer);
+		
+		FIoBatch NewBatch = GetIoDispatcher()->NewBatch();
+		FIoRequest Request = NewBatch.Read(ChunkID, Options);
+
+		NewBatch.Issue();
+		NewBatch.Wait(); // Blocking wait until all requests in the batch are done
+
+		check(Request.IsOk());
+
+		FBulkDataBase::GetIoDispatcher()->FreeBatch(NewBatch);
+	}
 	else
 	{
-		BULKDATA_NOT_IMPLEMENTED_FOR_RUNTIME;
+		UE_LOG(LogSerialization, Error, TEXT("Attempting to reload inline BulkData when the IoDispatcher is enabled, this operation is not supported! (%d)"), IsInlined());
 	}
 }
 
