Fixed an issue where runtime created BulkData objects would incorrectly return ::CanLoadFromDisk as true.

- Was caused by a special case check to allow inline data to unload when using the IoDispatcher since technically it cannot be loaded from disk which was blocking it from being discarded. However runtime created  bulkdata would often class as 'inline' which although having no file reference, started returning true when ::CanLoadFromDisk was called
- ::GetCopy now uses ::CanDiscardInternalData to determine if we can discard the internal data (we should probably change ::Unlock to do this as well but that is a different work item)
- Moved the special case code for inline data when running with the IoDispatcher from ::CanLoadFromDisk to ::CanDiscardInternalData.

#rnx

diff --git a/Engine/Source/Runtime/CoreUObject/Private/Serialization/BulkData2.cpp b/Engine/Source/Runtime/CoreUObject/Private/Serialization/BulkData2.cpp
--- a/Engine/Source/Runtime/CoreUObject/Private/Serialization/BulkData2.cpp
+++ b/Engine/Source/Runtime/CoreUObject/Private/Serialization/BulkData2.cpp
@@ -907,7 +907,7 @@ void FBulkDataBase::GetCopy(void** DstBuffer, bool bDiscardInternalCopy)
 		{
 			FMemory::Memcpy(*DstBuffer, DataBuffer, GetBulkDataSize());
 
-			if (bDiscardInternalCopy && (CanLoadFromDisk() || IsSingleUse()))
+			if (bDiscardInternalCopy && CanDiscardInternalData())
 			{
 				UE_LOG(LogSerialization, Warning, TEXT("FBulkDataBase::GetCopy both copied and discarded it's data, passing in an empty pointer would avoid an extra allocate and memcpy!"));
 				FreeData();
@@ -922,7 +922,7 @@ void FBulkDataBase::GetCopy(void** DstBuffer, bool bDiscardInternalCopy)
 	{
 		if (IsBulkDataLoaded())
 		{
-			if (bDiscardInternalCopy && (CanLoadFromDisk() || IsSingleUse()))
+			if (bDiscardInternalCopy && CanDiscardInternalData())
 			{
 				// Since we were going to discard the data anyway we can just hand over ownership to the caller
 				::Swap(*DstBuffer, DataBuffer);
@@ -987,18 +987,6 @@ bool FBulkDataBase::CanLoadFromDisk() const
 		return true;
 	}
 
-	// If the IoDispatcher is enabled then technically inlined data cannot load
-	// from disk, but it should act as thought it can up until the point that it
-	// actually tries to load. At which point it will give a log error.
-	//
-	// This is a temp hack while we work out if we want to allow additional inline
-	// loading at runtime (and pay the memory increase in the ToC) or prevent it entirely.
-	// TODO: Resolve this!
-	if (IsInlined() && IsIoDispatcherEnabled())
-	{
-		return true;
-	}
-
 	// If this BulkData has a fallback token then it can find it's filepath and load from disk
 	if (Fallback.Token != InvalidToken)
 	{
@@ -1240,6 +1228,21 @@ FString FBulkDataBase::GetFilename() const
 	}
 }
 
+bool FBulkDataBase::CanDiscardInternalData() const
+{
+	// We can discard the data if:
+	// -	We can reload the Bulkdata from disk
+	// -	If the Bulkdata object has been marked as single use which shows 
+	//		that there is no intent to access the data again)
+	// -	If we are using the IoDispatcher and the data is currently inlined
+	//		since we will not be able to reload inline data when the IoStore is
+	//		active.
+
+	// TODO: This is currently called from ::GetCopy but not ::Unlock, we should investe unifying the
+	// rules for discarding data
+	return CanLoadFromDisk() || IsSingleUse() || (IsInlined() && IsIoDispatcherEnabled());
+}
+
 void FBulkDataBase::LoadDataDirectly(void** DstBuffer)
 {
 	DECLARE_SCOPE_CYCLE_COUNTER(TEXT("FBulkDataBase::LoadDataDirectly"), STAT_UBD_LoadDataDirectly, STATGROUP_Memory);
diff --git a/Engine/Source/Runtime/CoreUObject/Public/Serialization/BulkData2.h b/Engine/Source/Runtime/CoreUObject/Public/Serialization/BulkData2.h
--- a/Engine/Source/Runtime/CoreUObject/Public/Serialization/BulkData2.h
+++ b/Engine/Source/Runtime/CoreUObject/Public/Serialization/BulkData2.h
@@ -139,6 +139,13 @@ public:
 	FOwnedBulkDataPtr* StealFileMapping();
 
 private:
+	/**
+	 * Poll to see if it is safe to discard the data owned by the Bulkdata object
+	 *
+	 * @return True if we are allowed to discard the existing data in the Bulkdata object.
+	 */
+	bool CanDiscardInternalData() const;
+
 	void LoadDataDirectly(void** DstBuffer);
 
 	void SerializeDuplicateData(FArchive& Ar, uint32& OutBulkDataFlags, int64& OutBulkDataSizeOnDisk, int64& OutBulkDataOffsetInFile);
