Allow Bulkdata2 to load data from archives that do not allow lazy loading.

- The use case is that the enterprise team want to be able to load .udsmesh files at runtime.
- If the archive has lazy loading disabled and the Bulkdata is inline OR in a different file then we don't need to treat the serialization any differently than when lazy loading is enabled.
- If lazy loading is disabled but the data is at the end of the current archive then we need to serialize the data immediately by seeking, serializing and then seeking back to the original position.
 - This is because we will not get another chance to load the data (if lazy loading is disabled then the life time of the archive is going to be limited)
- Unlike the old Bulkdata we avoid branching on the archive type and only implement the very specific behaviour we need
- This means that a nullptr linker is now a valid value so the assert has changed to an if statement
- Likewise, the filename being nulllptr is also expected so the assert has changed to an if statement.

#rnx

diff --git a/Engine/Source/Runtime/CoreUObject/Private/Serialization/BulkData2.cpp b/Engine/Source/Runtime/CoreUObject/Private/Serialization/BulkData2.cpp
--- a/Engine/Source/Runtime/CoreUObject/Private/Serialization/BulkData2.cpp
+++ b/Engine/Source/Runtime/CoreUObject/Private/Serialization/BulkData2.cpp
@@ -740,87 +740,93 @@ void FBulkDataBase::Serialize(FArchive& Ar, UObject* Owner, int32 /*Index*/, boo
 		if (bUseIoDispatcher == false)
 		{
 			Linker = FLinkerLoad::FindExistingLinkerForPackage(Package);
-			check(Linker != nullptr);
-
-			Filename = &Linker->Filename;
-			PackageName = Package->FileName;
+			
+			if (Linker != nullptr)
+			{
+				Filename = &Linker->Filename;
+				PackageName = Package->FileName;
+			}
 		}
 
-		FArchive* CacheableArchive = Ar.GetCacheableArchive();
-		if (Ar.IsAllowingLazyLoading() && CacheableArchive != nullptr)
+		if (IsInlined())
 		{
-			if (IsInlined())
-			{
-				// Inline data is already in the archive so serialize it immediately
-				AllocateData(BulkDataSize);
-				SerializeBulkData(Ar, DataBuffer, BulkDataSize);
-			}
-			else
+			// Inline data is already in the archive so serialize it immediately
+			AllocateData(BulkDataSize);
+			SerializeBulkData(Ar, DataBuffer, BulkDataSize);
+		}
+		else
+		{
+			if (IsDuplicateNonOptional())
 			{
-				if (IsDuplicateNonOptional())
+				auto DoesOptionalDataExist = [this](const FString* PackageFilename)->bool
 				{
-					auto DoesOptionalDataExist = [this](const FString* PackageFilename)->bool
-					{
 #if ALLOW_OPTIONAL_DATA
-						if (!IsUsingIODispatcher())
-						{
-							check(PackageFilename != nullptr);
-							FString OptionalDataFilename = ConvertFilenameFromFlags(*PackageFilename);
-							return IFileManager::Get().FileExists(*OptionalDataFilename);
-						}
-						else
-						{
-							return IoDispatcher->DoesChunkExist(ChunkID);
-						}
+					if (!IsUsingIODispatcher())
+					{
+						check(PackageFilename != nullptr);
+						FString OptionalDataFilename = ConvertFilenameFromFlags(*PackageFilename);
+						return IFileManager::Get().FileExists(*OptionalDataFilename);
+					}
+					else
+					{
+						return IoDispatcher->DoesChunkExist(ChunkID);
+					}
 #else
-						return false;
+					return false;
 #endif
-					};
+				};
+
+				if (DoesOptionalDataExist(Filename))
+				{
+					SerializeDuplicateData(Ar, BulkDataFlags, BulkDataSizeOnDisk, BulkDataOffsetInFile);
 
-					if (DoesOptionalDataExist(Filename))
+					if (!IsInlined() && bUseIoDispatcher)
 					{
-						SerializeDuplicateData(Ar, BulkDataFlags, BulkDataSizeOnDisk, BulkDataOffsetInFile);
-
-						if (!IsInlined() && bUseIoDispatcher)
-						{
-							// Regenerate the FIoChunkId to find the optional BulkData instead!
-							const int64 BulkDataID = BulkDataSize > 0 ? BulkDataOffsetInFile : TNumericLimits<uint64>::Max();
-							ChunkID = CreateBulkdataChunkId(Package->GetPackageId().ToIndex(), BulkDataID, EIoChunkType::OptionalBulkData);
-							BulkDataFlags |= BULKDATA_UsesIoDispatcher; // Indicates that this BulkData should use the FIoChunkId rather than a filename
-						}
-						else
-						{
-							Fallback.Token = InvalidToken;
-							Fallback.BulkDataSize = BulkDataSize;
-						}
+						// Regenerate the FIoChunkId to find the optional BulkData instead!
+						const int64 BulkDataID = BulkDataSize > 0 ? BulkDataOffsetInFile : TNumericLimits<uint64>::Max();
+						ChunkID = CreateBulkdataChunkId(Package->GetPackageId().ToIndex(), BulkDataID, EIoChunkType::OptionalBulkData);
+						BulkDataFlags |= BULKDATA_UsesIoDispatcher; // Indicates that this BulkData should use the FIoChunkId rather than a filename
 					}
 					else
 					{
-						// Skip over the optional data info (can't do a seek because we need to load the flags to work
-						// out if we read things as 32bit or 64bit)
-						uint32 DummyValue32;
-						int64 DummyValue64;
-						SerializeDuplicateData(Ar, DummyValue32, DummyValue64, DummyValue64);
+						Fallback.Token = InvalidToken;
+						Fallback.BulkDataSize = BulkDataSize;
 					}
 				}
-
-				// Fix up the file offset if we have a linker (if we do not then we will be loading via FIoDispatcher anyway)
-				if (Linker != nullptr)
+				else
 				{
-					BulkDataOffsetInFile += Linker->Summary.BulkDataStartOffset;
+					// Skip over the optional data info (can't do a seek because we need to load the flags to work
+					// out if we read things as 32bit or 64bit)
+					uint32 DummyValue32;
+					int64 DummyValue64;
+					SerializeDuplicateData(Ar, DummyValue32, DummyValue64, DummyValue64);
 				}
 			}
 
-			// If we are not using the FIoDispatcher then we need to make sure we can retrieve the filename later
-			if (bUseIoDispatcher == false)
+			// Fix up the file offset if we have a linker (if we do not then we will be loading via FIoDispatcher anyway)
+			if (Linker != nullptr)
 			{
-				check(Filename != nullptr);
-				Fallback.Token = FileTokenSystem::RegisterFileToken( PackageName, *Filename, BulkDataOffsetInFile);
+				BulkDataOffsetInFile += Linker->Summary.BulkDataStartOffset;
+			}
+
+			// If the archive does not support lazy loading and the data is not in a different file then we have to load 
+			// the data from the archive immediately as we won't get another chance.
+			if (!Ar.IsAllowingLazyLoading() && !IsInSeperateFile())
+			{
+				const int64 CurrentArchiveOffset = Ar.Tell();
+				Ar.Seek(BulkDataOffsetInFile);
+
+				AllocateData(BulkDataSize);
+				SerializeBulkData(Ar, DataBuffer, BulkDataSize);
+
+				Ar.Seek(CurrentArchiveOffset); // Return back to the original point in the archive so future serialization can continue
 			}
 		}
-		else
+
+		// If we are not using the FIoDispatcher and we have a filename then we need to make sure we can retrieve it later!
+		if (bUseIoDispatcher == false && Filename != nullptr)
 		{
-			BULKDATA_NOT_IMPLEMENTED_FOR_RUNTIME;
+			Fallback.Token = FileTokenSystem::RegisterFileToken(PackageName, *Filename, BulkDataOffsetInFile);
 		}
 	}
 #else
