Fix up some code formatting that got messed up during a merge

#rnx

diff --git a/Engine/Source/Runtime/CoreUObject/Private/UObject/Class.cpp b/Engine/Source/Runtime/CoreUObject/Private/UObject/Class.cpp
--- a/Engine/Source/Runtime/CoreUObject/Private/UObject/Class.cpp
+++ b/Engine/Source/Runtime/CoreUObject/Private/UObject/Class.cpp
@@ -626,16 +626,16 @@ public:
 // Constructors.
 //
 UStruct::UStruct( EStaticConstructor, int32 InSize, int32 InMinAlignment, EObjectFlags InFlags )
-:	UField			( EC_StaticConstructor, InFlags )
-,	SuperStruct		( nullptr )
+	: UField(EC_StaticConstructor, InFlags)
+	, SuperStruct(nullptr)
 	, Children(nullptr)
 	, ChildProperties(nullptr)
-,	PropertiesSize	( InSize )
-,	MinAlignment	( InMinAlignment )
-,	PropertyLink	( NULL )
-,	RefLink			( NULL )
-,	DestructorLink	( NULL )
-, PostConstructLink( NULL )
+	, PropertiesSize(InSize)
+	, MinAlignment(InMinAlignment)
+	, PropertyLink(NULL)
+	, RefLink(NULL)
+	, DestructorLink(NULL)
+	, PostConstructLink(NULL)
 {
 }
 
@@ -1237,195 +1237,195 @@ void UStruct::SerializeVersionedTaggedProperties(FStructuredArchive::FSlot Slot,
 		else
 #endif // WITH_TEXT_ARCHIVE_SUPPORT
 		{
-		// Load tagged properties.
-		FStructuredArchive::FStream PropertiesStream = Slot.EnterStream();
+			// Load tagged properties.
+			FStructuredArchive::FStream PropertiesStream = Slot.EnterStream();
 
-		// This code assumes that properties are loaded in the same order they are saved in. This removes a n^2 search 
-		// and makes it an O(n) when properties are saved in the same order as they are loaded (default case). In the 
-		// case that a property was reordered the code falls back to a slower search.
+			// This code assumes that properties are loaded in the same order they are saved in. This removes a n^2 search 
+			// and makes it an O(n) when properties are saved in the same order as they are loaded (default case). In the 
+			// case that a property was reordered the code falls back to a slower search.
 			FProperty*	Property = PropertyLink;
-		bool		bAdvanceProperty	= false;
-		int32		RemainingArrayDim	= Property ? Property->ArrayDim : 0;
+			bool		bAdvanceProperty	= false;
+			int32		RemainingArrayDim	= Property ? Property->ArrayDim : 0;
 
-		// Load all stored properties, potentially skipping unknown ones.
-		while (true)
-		{
-			FStructuredArchive::FRecord PropertyRecord = PropertiesStream.EnterElement().EnterRecord();
+			// Load all stored properties, potentially skipping unknown ones.
+			while (true)
+			{
+				FStructuredArchive::FRecord PropertyRecord = PropertiesStream.EnterElement().EnterRecord();
 
-			FPropertyTag Tag;
-			PropertyRecord << SA_VALUE(TEXT("Tag"), Tag);
+				FPropertyTag Tag;
+				PropertyRecord << SA_VALUE(TEXT("Tag"), Tag);
 
-			if (Tag.Name.IsNone())
-			{
-				break;
-			}
+				if (Tag.Name.IsNone())
+				{
+					break;
+				}
 
-			// Move to the next property to be serialized
-			if( bAdvanceProperty && --RemainingArrayDim <= 0 )
-			{
-				Property = Property->PropertyLinkNext;
-				// Skip over properties that don't need to be serialized.
-				while( Property && !Property->ShouldSerializeValue( UnderlyingArchive ) )
+				// Move to the next property to be serialized
+				if( bAdvanceProperty && --RemainingArrayDim <= 0 )
 				{
 					Property = Property->PropertyLinkNext;
+					// Skip over properties that don't need to be serialized.
+					while( Property && !Property->ShouldSerializeValue( UnderlyingArchive ) )
+					{
+						Property = Property->PropertyLinkNext;
+					}
+					RemainingArrayDim = Property ? Property->ArrayDim : 0;
 				}
-				RemainingArrayDim = Property ? Property->ArrayDim : 0;
-			}
-			bAdvanceProperty = false;
+				bAdvanceProperty = false;
 			
-			// Optionally resolve properties using Guid Property tags in non cooked builds that support it.
-			if (bArePropertyGuidsAvailable && Tag.HasPropertyGuid)
-			{
-				// Use property guids from blueprint generated classes to redirect serialised data.
-				FName Result = FindPropertyNameFromGuid(Tag.PropertyGuid);
-				if (Result != NAME_None && Tag.Name != Result)
+				// Optionally resolve properties using Guid Property tags in non cooked builds that support it.
+				if (bArePropertyGuidsAvailable && Tag.HasPropertyGuid)
 				{
-					Tag.Name = Result;
+					// Use property guids from blueprint generated classes to redirect serialised data.
+					FName Result = FindPropertyNameFromGuid(Tag.PropertyGuid);
+					if (Result != NAME_None && Tag.Name != Result)
+					{
+						Tag.Name = Result;
+					}
 				}
-			}
-			// If this property is not the one we expect (e.g. skipped as it matches the default value), do the brute force search.
-			if( Property == nullptr || Property->GetFName() != Tag.Name )
-			{
-				// No need to check redirects on platforms where everything is cooked. Always check for save games
-				if (bUseRedirects && !UnderlyingArchive.HasAnyPortFlags(PPF_DuplicateForPIE|PPF_Duplicate))
+				// If this property is not the one we expect (e.g. skipped as it matches the default value), do the brute force search.
+				if( Property == nullptr || Property->GetFName() != Tag.Name )
 				{
-					for (UStruct* CheckStruct = GetOwnerStruct(); CheckStruct; CheckStruct = CheckStruct->GetSuperStruct())
+					// No need to check redirects on platforms where everything is cooked. Always check for save games
+					if (bUseRedirects && !UnderlyingArchive.HasAnyPortFlags(PPF_DuplicateForPIE|PPF_Duplicate))
 					{
-							FName NewTagName = FProperty::FindRedirectedPropertyName(CheckStruct, Tag.Name);
-						if (!NewTagName.IsNone())
+						for (UStruct* CheckStruct = GetOwnerStruct(); CheckStruct; CheckStruct = CheckStruct->GetSuperStruct())
 						{
-							Tag.Name = NewTagName;
-							break;
+							FName NewTagName = FProperty::FindRedirectedPropertyName(CheckStruct, Tag.Name);
+							if (!NewTagName.IsNone())
+							{
+								Tag.Name = NewTagName;
+								break;
+							}
 						}
 					}
-				}
 
 					FProperty* CurrentProperty = Property;
-				// Search forward...
-				for ( ; Property; Property=Property->PropertyLinkNext )
-				{
-					if( Property->GetFName() == Tag.Name )
-					{
-						break;
-					}
-				}
-				// ... and then search from the beginning till we reach the current property if it's not found.
-				if( Property == nullptr )
-				{
-					for( Property = PropertyLink; Property && Property != CurrentProperty; Property = Property->PropertyLinkNext )
+					// Search forward...
+					for ( ; Property; Property=Property->PropertyLinkNext )
 					{
 						if( Property->GetFName() == Tag.Name )
 						{
 							break;
 						}
 					}
-
-					if( Property == CurrentProperty )
-					{
-						// Property wasn't found.
-						Property = nullptr;
-					}
-				}
-
-				RemainingArrayDim = Property ? Property->ArrayDim : 0;
-			}
-
-			const int64 StartOfProperty = UnderlyingArchive.Tell();
-
-			if (!Property)
-			{
-				Property = CustomFindProperty(Tag.Name);
-			}
-
-			if (Property)
-			{
-				FName PropID = Property->GetID();
-
-				// Check if this is a struct property and we have a redirector
-				// No need to check redirects on platforms where everything is cooked. Always check for save games
-				if (bUseRedirects)
-				{
-					if (Tag.Type == NAME_StructProperty && PropID == NAME_StructProperty)
+					// ... and then search from the beginning till we reach the current property if it's not found.
+					if( Property == nullptr )
 					{
-						const FName NewName = FLinkerLoad::FindNewNameForStruct(Tag.StructName);
-							const FName StructName = CastFieldChecked<FStructProperty>(Property)->Struct->GetFName();
-						if (NewName == StructName)
+						for( Property = PropertyLink; Property && Property != CurrentProperty; Property = Property->PropertyLinkNext )
 						{
-							Tag.StructName = NewName;
+							if( Property->GetFName() == Tag.Name )
+							{
+								break;
+							}
 						}
-					}
-					else if ((PropID == NAME_EnumProperty) && ((Tag.Type == NAME_EnumProperty) || (Tag.Type == NAME_ByteProperty)))
-					{
-						const FName NewName = FLinkerLoad::FindNewNameForEnum(Tag.EnumName);
-						if (!NewName.IsNone())
+
+						if( Property == CurrentProperty )
 						{
-							Tag.EnumName = NewName;
+							// Property wasn't found.
+							Property = nullptr;
 						}
 					}
-				}
 
-#if WITH_EDITOR
-				if (BreakRecursionIfFullyLoad && BreakRecursionIfFullyLoad->HasAllFlags(RF_LoadCompleted))
-				{
-				}
-				else
-#endif // WITH_EDITOR
-				// editoronly properties should be skipped if we are NOT the editor, or we are 
-				// the editor but are cooking for console (editoronly implies notforconsole)
-				if ((Property->PropertyFlags & CPF_EditorOnly) && !FPlatformProperties::HasEditorOnlyData() && !GForceLoadEditorOnly)
-				{
+					RemainingArrayDim = Property ? Property->ArrayDim : 0;
 				}
-				// check for valid array index
-				else if (Tag.ArrayIndex >= Property->ArrayDim || Tag.ArrayIndex < 0)
-				{
-					UE_LOG(LogClass, Warning, TEXT("Array bound exceeded (var %s=%d, exceeds %s [0-%d] in package:  %s"),
-						*Tag.Name.ToString(), Tag.ArrayIndex, *GetName(), Property->ArrayDim - 1, *UnderlyingArchive.GetArchiveName());
-				}
-				else if (!Property->ShouldSerializeValue(UnderlyingArchive))
+
+				const int64 StartOfProperty = UnderlyingArchive.Tell();
+
+				if (!Property)
 				{
-					UE_CLOG((UnderlyingArchive.IsPersistent() && FPlatformProperties::RequiresCookedData()), LogClass, Warning, TEXT("Skipping saved property %s of %s since it is no longer serializable for asset:  %s. (Maybe resave asset?)"), *Tag.Name.ToString(), *GetName(), *UnderlyingArchive.GetArchiveName());
+					Property = CustomFindProperty(Tag.Name);
 				}
-				else
+
+				if (Property)
 				{
-					FStructuredArchive::FSlot ValueSlot = PropertyRecord.EnterField(SA_FIELD_NAME(TEXT("Value")));
+					FName PropID = Property->GetID();
 
-					switch (Property->ConvertFromType(Tag, ValueSlot, Data, DefaultsStruct))
+					// Check if this is a struct property and we have a redirector
+					// No need to check redirects on platforms where everything is cooked. Always check for save games
+					if (bUseRedirects)
 					{
-						case EConvertFromTypeResult::Converted:
-							bAdvanceProperty = true;
-							break;
-
-						case EConvertFromTypeResult::UseSerializeItem:
-							if (Tag.Type != PropID)
+						if (Tag.Type == NAME_StructProperty && PropID == NAME_StructProperty)
+						{
+							const FName NewName = FLinkerLoad::FindNewNameForStruct(Tag.StructName);
+							const FName StructName = CastFieldChecked<FStructProperty>(Property)->Struct->GetFName();
+							if (NewName == StructName)
 							{
-								UE_LOG(LogClass, Warning, TEXT("Type mismatch in %s of %s - Previous (%s) Current(%s) for package:  %s"), *Tag.Name.ToString(), *GetName(), *Tag.Type.ToString(), *PropID.ToString(), *UnderlyingArchive.GetArchiveName());
+								Tag.StructName = NewName;
 							}
-							else
+						}
+						else if ((PropID == NAME_EnumProperty) && ((Tag.Type == NAME_EnumProperty) || (Tag.Type == NAME_ByteProperty)))
+						{
+							const FName NewName = FLinkerLoad::FindNewNameForEnum(Tag.EnumName);
+							if (!NewName.IsNone())
 							{
-								uint8* DestAddress = Property->ContainerPtrToValuePtr<uint8>(Data, Tag.ArrayIndex);
-								uint8* DefaultsFromParent = Property->ContainerPtrToValuePtrForDefaults<uint8>(DefaultsStruct, Defaults, Tag.ArrayIndex);
-
-								// This property is ok.
-								Tag.SerializeTaggedProperty(ValueSlot, Property, DestAddress, DefaultsFromParent);
-								bAdvanceProperty = !UnderlyingArchive.IsCriticalError();
+								Tag.EnumName = NewName;
 							}
-							break;
+						}
+					}
 
-						case EConvertFromTypeResult::CannotConvert:
-							break;
+#if WITH_EDITOR
+					if (BreakRecursionIfFullyLoad && BreakRecursionIfFullyLoad->HasAllFlags(RF_LoadCompleted))
+					{
+					}
+					else
+#endif // WITH_EDITOR
+					// editoronly properties should be skipped if we are NOT the editor, or we are 
+					// the editor but are cooking for console (editoronly implies notforconsole)
+					if ((Property->PropertyFlags & CPF_EditorOnly) && !FPlatformProperties::HasEditorOnlyData() && !GForceLoadEditorOnly)
+					{
+					}
+					// check for valid array index
+					else if (Tag.ArrayIndex >= Property->ArrayDim || Tag.ArrayIndex < 0)
+					{
+						UE_LOG(LogClass, Warning, TEXT("Array bound exceeded (var %s=%d, exceeds %s [0-%d] in package:  %s"),
+							*Tag.Name.ToString(), Tag.ArrayIndex, *GetName(), Property->ArrayDim - 1, *UnderlyingArchive.GetArchiveName());
+					}
+					else if (!Property->ShouldSerializeValue(UnderlyingArchive))
+					{
+						UE_CLOG((UnderlyingArchive.IsPersistent() && FPlatformProperties::RequiresCookedData()), LogClass, Warning, TEXT("Skipping saved property %s of %s since it is no longer serializable for asset:  %s. (Maybe resave asset?)"), *Tag.Name.ToString(), *GetName(), *UnderlyingArchive.GetArchiveName());
+					}
+					else
+					{
+						FStructuredArchive::FSlot ValueSlot = PropertyRecord.EnterField(SA_FIELD_NAME(TEXT("Value")));
 
-						default:
-							check(false);
+						switch (Property->ConvertFromType(Tag, ValueSlot, Data, DefaultsStruct))
+						{
+							case EConvertFromTypeResult::Converted:
+								bAdvanceProperty = true;
+								break;
+
+							case EConvertFromTypeResult::UseSerializeItem:
+								if (Tag.Type != PropID)
+								{
+									UE_LOG(LogClass, Warning, TEXT("Type mismatch in %s of %s - Previous (%s) Current(%s) for package:  %s"), *Tag.Name.ToString(), *GetName(), *Tag.Type.ToString(), *PropID.ToString(), *UnderlyingArchive.GetArchiveName());
+								}
+								else
+								{
+									uint8* DestAddress = Property->ContainerPtrToValuePtr<uint8>(Data, Tag.ArrayIndex);
+									uint8* DefaultsFromParent = Property->ContainerPtrToValuePtrForDefaults<uint8>(DefaultsStruct, Defaults, Tag.ArrayIndex);
+
+									// This property is ok.
+									Tag.SerializeTaggedProperty(ValueSlot, Property, DestAddress, DefaultsFromParent);
+									bAdvanceProperty = !UnderlyingArchive.IsCriticalError();
+								}
+								break;
+
+							case EConvertFromTypeResult::CannotConvert:
+								break;
+
+							default:
+								check(false);
+						}
 					}
 				}
-			}
 
 				int64 Loaded = UnderlyingArchive.Tell() - StartOfProperty;
 
-			if (!bAdvanceProperty)
-			{
-				UnderlyingArchive.Seek(StartOfProperty + Tag.Size);
-			}
+				if (!bAdvanceProperty)
+				{
+					UnderlyingArchive.Seek(StartOfProperty + Tag.Size);
+				}
 				else
 				{
 					check(Tag.Size == Loaded);
@@ -1727,11 +1727,9 @@ void UStruct::SerializeProperties(FArchive& Ar)
 	}
 }
 
-
-
-void UStruct::Serialize( FArchive& Ar )
+void UStruct::Serialize(FArchive& Ar)
 {
-	Super::Serialize( Ar );
+	Super::Serialize(Ar);
 
 #if USTRUCT_FAST_ISCHILDOF_IMPL == USTRUCT_ISCHILDOF_STRUCTARRAY
 	UStruct* SuperStructBefore = GetSuperStruct();
@@ -4041,20 +4039,20 @@ void UClass::SetUpRuntimeReplicationData()
 		const bool bIsNativeClass = HasAnyClassFlags(CLASS_Native);
 		if (!bIsNativeClass)
 		{
-		// Sort NetProperties so that their ClassReps are sorted by memory offset
+			// Sort NetProperties so that their ClassReps are sorted by memory offset
 			struct FComparePropertyOffsets
-		{
-				FORCEINLINE bool operator()(FProperty& A, FProperty& B) const
 			{
-				// Ensure stable sort
-				if ( A.GetOffset_ForGC() == B.GetOffset_ForGC() )
+				FORCEINLINE bool operator()(FProperty& A, FProperty& B) const
 				{
-					return A.GetName() < B.GetName();
-				}
+					// Ensure stable sort
+					if ( A.GetOffset_ForGC() == B.GetOffset_ForGC() )
+					{
+						return A.GetName() < B.GetName();
+					}
 
-				return A.GetOffset_ForGC() < B.GetOffset_ForGC();
-			}
-		};
+					return A.GetOffset_ForGC() < B.GetOffset_ForGC();
+				}
+			};
 
 			Sort(NetProperties.GetData(), NetProperties.Num(), FComparePropertyOffsets());
 		}
@@ -4666,7 +4664,7 @@ void UClass::PurgeClass(bool bRecompilingOnLoad)
 			PropertiesPendingDestruction = ChildProperties;
 			ChildProperties = nullptr;
 		}
-}
+	}
 	else
 #endif // WITH_EDITORONLY_DATA
 	{
@@ -4725,9 +4723,9 @@ bool UClass::HasProperty(FProperty* InProperty) const
 	{
 		UClass* PropertiesClass = InProperty->GetOwner<UClass>();
 		if (PropertiesClass)
-	{
-		return IsChildOf(PropertiesClass);
-	}
+		{
+			return IsChildOf(PropertiesClass);
+		}
 	}
 
 	return false;
