Remove BULKDATA_NOT_IMPLEMENTED_FOR_RUNTIME now that it is no longer used.

#rnx

diff --git a/Engine/Source/Runtime/CoreUObject/Public/Serialization/BulkData2.h b/Engine/Source/Runtime/CoreUObject/Public/Serialization/BulkData2.h
--- a/Engine/Source/Runtime/CoreUObject/Public/Serialization/BulkData2.h
+++ b/Engine/Source/Runtime/CoreUObject/Public/Serialization/BulkData2.h
@@ -7,13 +7,6 @@
 #include "Async/AsyncFileHandle.h"
 #include "IO/IoDispatcher.h"
 
-// Any place marked with BULKDATA_NOT_IMPLEMENTED_FOR_RUNTIME marks a method that we do not support
-// but needs to exist in order for the code to compile. This pretty much means code that is editor only
-// but is not compiled out at runtime etc.
-// It can also be found inside of methods to show code paths that the old path implements but I have
-// not been able to find any way to actually execute/test.
-#define BULKDATA_NOT_IMPLEMENTED_FOR_RUNTIME PLATFORM_BREAK()
-
 struct FOwnedBulkDataPtr;
 
 /**
