Clean up some Bulkdata2 logigng

#rnx

diff --git a/Engine/Source/Runtime/CoreUObject/Private/Serialization/BulkData2.cpp b/Engine/Source/Runtime/CoreUObject/Private/Serialization/BulkData2.cpp
--- a/Engine/Source/Runtime/CoreUObject/Private/Serialization/BulkData2.cpp
+++ b/Engine/Source/Runtime/CoreUObject/Private/Serialization/BulkData2.cpp
@@ -1255,6 +1255,7 @@ void FBulkDataBase::LoadDataDirectly(void** DstBuffer)
 	
 	if (!CanLoadFromDisk())
 	{
+		UE_LOG(LogSerialization, Warning, TEXT("Attempting to load a BulkData object that cannot be loaded from disk"));
 		return; // Early out if there is nothing to load anyway
 	}
 
@@ -1320,6 +1321,7 @@ void FBulkDataBase::LoadDataDirectly(void** DstBuffer)
 	}
 	else
 	{
+		// Note that currently this shouldn't be reachable as we should early out due to the ::CanLoadFromDisk check at the start of the method
 		UE_LOG(LogSerialization, Error, TEXT("Attempting to reload inline BulkData when the IoDispatcher is enabled, this operation is not supported! (%d)"), IsInlined());
 	}
 }
