PR #6276: Fix infinite TArray property stacking when saving TArray property into a default config (Contributed by [REDACTED_PERSON])
#rnx

diff --git a/Engine/Source/Runtime/CoreUObject/Private/UObject/Obj.cpp b/Engine/Source/Runtime/CoreUObject/Private/UObject/Obj.cpp
--- a/Engine/Source/Runtime/CoreUObject/Private/UObject/Obj.cpp
+++ b/Engine/Source/Runtime/CoreUObject/Private/UObject/Obj.cpp
@@ -2474,15 +2474,18 @@ void UObject::SaveConfig( uint64 Flags, const TCHAR* InFilename, FConfigCacheIni
 			FArrayProperty* Array   = CastField<FArrayProperty>( Property );
 			if( Array )
 			{
+				FConfigSection* Sec = Config->GetSectionPrivate(*Section, 1, 0, *PropFileName);
+				// Default ini's require the array syntax to be applied to the property name
+				FString CompleteKey = FString::Printf(TEXT("%s%s"), bIsADefaultIniWrite ? TEXT("+") : TEXT(""), *Key);
+				if (Sec)
+				{
+					// Delete the old value for the property in the ConfigCache before (conditionally) adding in the new value
+					Sec->Remove(*CompleteKey);
+				}
+
 				if (!bPropDeprecated && (!bShouldCheckIfIdenticalBeforeAdding || !Property->Identical_InContainer(this, SuperClassDefaultObject)))
 				{
-					FConfigSection* Sec = Config->GetSectionPrivate( *Section, 1, 0, *PropFileName );
 					check(Sec);
-					Sec->Remove( *Key );
-
-					// Default ini's require the array syntax to be applied to the property name
-					FString CompleteKey = FString::Printf(TEXT("%s%s"), bIsADefaultIniWrite ? TEXT("+") : TEXT(""), *Key);
-
 					FScriptArrayHelper_InContainer ArrayHelper(Array, this);
 					for( int32 i=0; i<ArrayHelper.Num(); i++ )
 					{
@@ -2491,16 +2494,6 @@ void UObject::SaveConfig( uint64 Flags, const TCHAR* InFilename, FConfigCacheIni
 						Sec->Add(*CompleteKey, *Buffer);
 					}
 				}
-				else
-				{
-					// If we are not writing it to config above, we should make sure that this property isn't stagnant in the cache.
-					FConfigSection* Sec = Config->GetSectionPrivate( *Section, 1, 0, *PropFileName );
-					if( Sec )
-					{
-						Sec->Remove( *Key );
-					}
-
-				}
 			}
 			else
 			{
