Fix the following cook performance and correctness issues in UPackage::Save:

Cook Performance:
1) Avoid performance cost of calling CreateTempFilename() for cases where we don't actually need a temp filename (cost involves a FileSize check with the filesystem)
2) Avoid performance cost of calling FileSize() to get saved size of packages and use the recorded buffer size instead (which gives uncompressed sizes).

Correctness:
1) When saving packages asynchronously, the UPackage::FileSize field was populated with the filesystem size of the package, but since that may not have been written yet by the async writer, it was often set to (uint64)-1.
2) When computing the MBWritten cook stat, make sure that bulk data is included in the accumulated total and that the stat isn't calculated incorrectly in DiffOnly modes.
3) Temporary files would be left lying around if the package was determined to have no exports or to have been completely nativized.

diff --git a/Engine/Source/Runtime/CoreUObject/Private/UObject/SavePackage.cpp b/Engine/Source/Runtime/CoreUObject/Private/UObject/SavePackage.cpp
--- a/Engine/Source/Runtime/CoreUObject/Private/UObject/SavePackage.cpp
+++ b/Engine/Source/Runtime/CoreUObject/Private/UObject/SavePackage.cpp
@@ -60,6 +60,7 @@
 #include "Serialization/ArchiveUObjectFromStructuredArchive.h"
 #include "UObject/AsyncWorkSequence.h"
 #include "Serialization/BulkDataManifest.h"
+#include "Misc/ScopeExit.h"
 
 DEFINE_LOG_CATEGORY_STATIC(LogSavePackage, Log, All);
 
@@ -73,7 +74,7 @@ static FCriticalSection InitializeCoreClassesCritSec;
 static void SaveThumbnails(UPackage* InOuter, FLinkerSave* Linker, FStructuredArchive::FSlot Slot);
 static void SaveAssetRegistryData(UPackage* InOuter, FLinkerSave* Linker, FStructuredArchive::FSlot Slot);
 static void SaveBulkData(FLinkerSave* Linker, const UPackage* InOuter, const TCHAR* Filename, const ITargetPlatform* TargetPlatform,
-						 FSavePackageContext* SavePackageContext, const bool bTextFormat, const bool bDiffing, const bool bComputeHash, TAsyncWorkSequence<FMD5>& AsyncWriteAndHashSequence, int64 TotalPackageSizeUncompressed);
+						 FSavePackageContext* SavePackageContext, const bool bTextFormat, const bool bDiffing, const bool bComputeHash, TAsyncWorkSequence<FMD5>& AsyncWriteAndHashSequence, int64& TotalPackageSizeUncompressed);
 static void SaveWorldLevelInfo(UPackage* InOuter, FLinkerSave* Linker, FStructuredArchive::FRecord Record);
 static EObjectMark GetExcludedObjectMarksForTargetPlatform(const class ITargetPlatform* TargetPlatform, const bool bIsCooking);
 
@@ -308,13 +309,10 @@ static void CheckObjectPriorToSave(FArchiveUObject& Ar, UObject* InObj, UPackage
 	}
 }
 
-static bool EndSavingIfCancelled( FLinkerSave* Linker, const FString& TempFilename ) 
+static bool EndSavingIfCancelled() 
 {
 	if ( GWarn->ReceivedUserCancel() )
 	{
-		// free the file handle and delete the temporary file
-		Linker->CloseAndDestroySaver();
-		IFileManager::Get().Delete( *TempFilename );
 		return true;
 	}
 	return false;
@@ -3469,15 +3467,6 @@ FSavePackageResultStruct UPackage::Save(UPackage* InOuter, UObject* Base, EObjec
 			bCleanupIsRequired = Base->PreSaveRoot(Filename);
 		}
 
-		const FString BaseFilename = FPaths::GetBaseFilename(Filename);
-		// Make temp file. CreateTempFilename guarantees unique, non-existing filename.
-		// The temp file will be saved in the game save folder to not have to deal with potentially too long paths.
-		// Since the temp filename may include a 32 character GUID as well, limit the user prefix to 32 characters.
-		FString TempFilename;
-		FString TextFormatTempFilename;
-		TempFilename = FPaths::CreateTempFilename(*FPaths::ProjectSavedDir(), *BaseFilename.Left(32));
-		TextFormatTempFilename = TempFilename + FPackageName::GetTextAssetPackageExtension();
-
 		// Init.
 		FString CleanFilename = FPaths::GetCleanFilename(Filename);
 
@@ -3646,6 +3635,26 @@ FSavePackageResultStruct UPackage::Save(UPackage* InOuter, UObject* Base, EObjec
 				FArchiveFormatterType* Formatter = nullptr;
 				FArchive* TextFormatArchive = nullptr;
 				const bool bTextFormat = FString(Filename).EndsWith(FPackageName::GetTextAssetPackageExtension()) || FString(Filename).EndsWith(FPackageName::GetTextMapPackageExtension());
+
+				const FString BaseFilename = FPaths::GetBaseFilename(Filename);
+				// Make temp file. CreateTempFilename guarantees unique, non-existing filename.
+				// The temp file will be saved in the game save folder to not have to deal with potentially too long paths.
+				// Since the temp filename may include a 32 character GUID as well, limit the user prefix to 32 characters.
+				TOptional<FString> TempFilename;
+				TOptional<FString> TextFormatTempFilename;
+				ON_SCOPE_EXIT
+				{
+					// free the file handle and delete the temporary file
+					Linker->CloseAndDestroySaver();
+					if (TempFilename.IsSet())
+					{
+						IFileManager::Get().Delete(*TempFilename.GetValue());
+					}
+					if (TextFormatTempFilename.IsSet())
+					{
+						IFileManager::Get().Delete(*TextFormatTempFilename.GetValue());
+					}
+				};
 	
 				{
 					TRACE_CPUPROFILER_EVENT_SCOPE(UPackage_Save_CreateLinkerSave);
@@ -3704,13 +3713,22 @@ FSavePackageResultStruct UPackage::Save(UPackage* InOuter, UObject* Base, EObjec
 						else
 						{
 							// Allocate the linker, forcing byte swapping if wanted.
-							Linker = TUniquePtr<FLinkerSave>(new FLinkerSave(InOuter, *TempFilename, bForceByteSwapping, bSaveUnversioned));
+							TempFilename = FPaths::CreateTempFilename(*FPaths::ProjectSavedDir(), *BaseFilename.Left(32));
+							Linker = TUniquePtr<FLinkerSave>(new FLinkerSave(InOuter, *TempFilename.GetValue(), bForceByteSwapping, bSaveUnversioned));
 						}
 
 #if WITH_TEXT_ARCHIVE_SUPPORT
 					if (bTextFormat)
 					{
-						TextFormatArchive = IFileManager::Get().CreateFileWriter(*TextFormatTempFilename);
+						if (TempFilename.IsSet())
+						{
+							TextFormatTempFilename = TempFilename.GetValue() + FPackageName::GetTextAssetPackageExtension();
+						}
+						else
+						{
+							TextFormatTempFilename = FPaths::CreateTempFilename(*FPaths::ProjectSavedDir(), *BaseFilename.Left(32)) + FPackageName::GetTextAssetPackageExtension();
+						}
+						TextFormatArchive = IFileManager::Get().CreateFileWriter(*TextFormatTempFilename.GetValue());
 						FJsonArchiveOutputFormatter* OutputFormatter = new FJsonArchiveOutputFormatter(*TextFormatArchive);
 						OutputFormatter->SetObjectIndicesMap(&Linker->ObjectIndicesMap);
 						Formatter = OutputFormatter;
@@ -3752,7 +3770,7 @@ FSavePackageResultStruct UPackage::Save(UPackage* InOuter, UObject* Base, EObjec
 				InOuter->LinkerLicenseeVersion = Linker->LicenseeUE4Ver();
 				InOuter->LinkerCustomVersion = Linker->GetCustomVersions();
 
-				if (EndSavingIfCancelled(Linker.Get(), TempFilename))
+				if (EndSavingIfCancelled())
 				{ 
 					return ESavePackageResult::Canceled; 
 				}
@@ -3985,13 +4003,13 @@ FSavePackageResultStruct UPackage::Save(UPackage* InOuter, UObject* Base, EObjec
 				TMap<UObject*, UObject*> DuplicateRedirects = UnmarkExportTagFromDuplicates();
 #endif // WITH_EDITOR
 
-				if ( EndSavingIfCancelled( Linker.Get(), TempFilename ) )
+				if ( EndSavingIfCancelled() )
 				{ 
 					return ESavePackageResult::Canceled;
 				}
 				SlowTask.EnterProgressFrame();
 
-				if ( EndSavingIfCancelled( Linker.Get(), TempFilename ) )
+				if ( EndSavingIfCancelled() )
 				{ 
 					return ESavePackageResult::Canceled;
 				}
@@ -4091,7 +4109,7 @@ FSavePackageResultStruct UPackage::Save(UPackage* InOuter, UObject* Base, EObjec
 					}
 				}
 
-				if ( EndSavingIfCancelled( Linker.Get(), TempFilename ) )
+				if ( EndSavingIfCancelled() )
 				{ 
 					return ESavePackageResult::Canceled;
 				}
@@ -4160,9 +4178,6 @@ FSavePackageResultStruct UPackage::Save(UPackage* InOuter, UObject* Base, EObjec
 						}
 					}
 
-					// Free the file handle and delete the temporary file
-					Linker->CloseAndDestroySaver();
-					IFileManager::Get().Delete( *TempFilename );
 					if (!(SaveFlags & SAVE_NoError))
 					{
 						Error->Logf(ELogVerbosity::Warning, TEXT("Can't save %s: Graph is linked to object %s in external map"), Filename, *CulpritString);
@@ -4212,9 +4227,6 @@ FSavePackageResultStruct UPackage::Save(UPackage* InOuter, UObject* Base, EObjec
 							(PropertyRef != nullptr) ? *PropertyRef->GetName() : TEXT("unknown property ref"));
 					}
 
-					// free the file handle and delete the temporary file
-					Linker->CloseAndDestroySaver();
-					IFileManager::Get().Delete(*TempFilename);
 					if (!(SaveFlags & SAVE_NoError))
 					{
 						Error->Logf(ELogVerbosity::Warning, TEXT("Can't save %s: Graph is linked to external private object %s"), Filename, *CulpritString);
@@ -4270,7 +4282,7 @@ FSavePackageResultStruct UPackage::Save(UPackage* InOuter, UObject* Base, EObjec
 				}
 				int32 OffsetAfterPackageFileSummary = Linker->Tell();
 		
-				if ( EndSavingIfCancelled( Linker.Get(), TempFilename ) )
+				if ( EndSavingIfCancelled() )
 				{ 
 					return ESavePackageResult::Canceled;
 				}
@@ -4299,7 +4311,7 @@ FSavePackageResultStruct UPackage::Save(UPackage* InOuter, UObject* Base, EObjec
 #endif
 					NameMapSaver.UpdateLinker(*Linker.Get(), Conform, bTextFormat ? nullptr : Linker->Saver);
 				}
-				if ( EndSavingIfCancelled( Linker.Get(), TempFilename ) )
+				if ( EndSavingIfCancelled() )
 				{ 
 					return ESavePackageResult::Canceled;
 				}
@@ -4340,7 +4352,7 @@ FSavePackageResultStruct UPackage::Save(UPackage* InOuter, UObject* Base, EObjec
 					}
 				}
 
-				if ( EndSavingIfCancelled( Linker.Get(), TempFilename ) )
+				if ( EndSavingIfCancelled() )
 				{ 
 					return ESavePackageResult::Canceled;
 				}
@@ -4413,7 +4425,7 @@ FSavePackageResultStruct UPackage::Save(UPackage* InOuter, UObject* Base, EObjec
 				}
 
 
-				if ( EndSavingIfCancelled( Linker.Get(), TempFilename ) )
+				if ( EndSavingIfCancelled() )
 				{ 
 					return ESavePackageResult::Canceled;
 				}
@@ -4427,7 +4439,7 @@ FSavePackageResultStruct UPackage::Save(UPackage* InOuter, UObject* Base, EObjec
 					Linker->Summary.ImportCount = Linker->ImportMap.Num();
 				}
 
-				if ( EndSavingIfCancelled( Linker.Get(), TempFilename ) )
+				if ( EndSavingIfCancelled() )
 				{ 
 					return ESavePackageResult::Canceled;
 				}
@@ -4466,7 +4478,7 @@ FSavePackageResultStruct UPackage::Save(UPackage* InOuter, UObject* Base, EObjec
 				}
 #endif
 
-				if ( EndSavingIfCancelled( Linker.Get(), TempFilename ) )
+				if ( EndSavingIfCancelled() )
 				{ 
 					return ESavePackageResult::Canceled;
 				}
@@ -4489,7 +4501,7 @@ FSavePackageResultStruct UPackage::Save(UPackage* InOuter, UObject* Base, EObjec
 
 				Linker->Summary.ExportCount = Linker->ExportMap.Num();
 
-				if ( EndSavingIfCancelled( Linker.Get(), TempFilename ) )
+				if ( EndSavingIfCancelled() )
 				{ 
 					return ESavePackageResult::Canceled;
 				}
@@ -4577,7 +4589,7 @@ FSavePackageResultStruct UPackage::Save(UPackage* InOuter, UObject* Base, EObjec
 				}
 
 
-				if ( EndSavingIfCancelled( Linker.Get(), TempFilename ) )
+				if ( EndSavingIfCancelled() )
 				{ 
 					return ESavePackageResult::Canceled;
 				}
@@ -4604,7 +4616,7 @@ FSavePackageResultStruct UPackage::Save(UPackage* InOuter, UObject* Base, EObjec
 					}
 				}
 
-				if ( EndSavingIfCancelled( Linker.Get(), TempFilename ) )
+				if ( EndSavingIfCancelled() )
 				{ 
 					return ESavePackageResult::Canceled;
 				}
@@ -4641,7 +4653,7 @@ FSavePackageResultStruct UPackage::Save(UPackage* InOuter, UObject* Base, EObjec
 				}
 
 
-				if ( EndSavingIfCancelled( Linker.Get(), TempFilename ) )
+				if ( EndSavingIfCancelled() )
 				{ 
 					return ESavePackageResult::Canceled;
 				}
@@ -4694,7 +4706,7 @@ FSavePackageResultStruct UPackage::Save(UPackage* InOuter, UObject* Base, EObjec
 
 				// Find components referenced by exports.
 
-				if ( EndSavingIfCancelled( Linker.Get(), TempFilename ) )
+				if ( EndSavingIfCancelled() )
 				{ 
 					return ESavePackageResult::Canceled;
 				}
@@ -4717,7 +4729,7 @@ FSavePackageResultStruct UPackage::Save(UPackage* InOuter, UObject* Base, EObjec
 				int32 OffsetAfterImportMap = Linker->Tell();
 
 
-				if (EndSavingIfCancelled(Linker.Get(), TempFilename))
+				if (EndSavingIfCancelled())
 				{
 					return ESavePackageResult::Canceled;
 				}
@@ -4740,7 +4752,7 @@ FSavePackageResultStruct UPackage::Save(UPackage* InOuter, UObject* Base, EObjec
 				int32 OffsetAfterExportMap = Linker->Tell();
 
 
-				if (EndSavingIfCancelled(Linker.Get(), TempFilename))
+				if (EndSavingIfCancelled())
 				{
 					return ESavePackageResult::Canceled;
 				}
@@ -4778,7 +4790,7 @@ FSavePackageResultStruct UPackage::Save(UPackage* InOuter, UObject* Base, EObjec
 				}
 
 
-				if (EndSavingIfCancelled(Linker.Get(), TempFilename))
+				if (EndSavingIfCancelled())
 				{ 
 					return ESavePackageResult::Canceled;
 				}
@@ -5206,7 +5218,7 @@ FSavePackageResultStruct UPackage::Save(UPackage* InOuter, UObject* Base, EObjec
 				
 				Linker->Summary.TotalHeaderSize	= Linker->Tell();
 
-				if ( EndSavingIfCancelled( Linker.Get(), TempFilename ) )
+				if ( EndSavingIfCancelled() )
 				{ 
 					return ESavePackageResult::Canceled;
 				}
@@ -5238,7 +5250,7 @@ FSavePackageResultStruct UPackage::Save(UPackage* InOuter, UObject* Base, EObjec
 					int32 LastExportSaveStep = 0;
 					for( int32 i=0; i<Linker->ExportMap.Num(); i++ )
 					{
-						if ( EndSavingIfCancelled( Linker.Get(), TempFilename ) )
+						if ( EndSavingIfCancelled() )
 						{ 
 							return ESavePackageResult::Canceled;
 						}
@@ -5331,7 +5343,7 @@ FSavePackageResultStruct UPackage::Save(UPackage* InOuter, UObject* Base, EObjec
 				}
 
 				
-				if ( EndSavingIfCancelled( Linker.Get(), TempFilename ) )
+				if ( EndSavingIfCancelled() )
 				{ 
 					return ESavePackageResult::Canceled;
 				}
@@ -5374,7 +5386,7 @@ FSavePackageResultStruct UPackage::Save(UPackage* InOuter, UObject* Base, EObjec
 				// write the package post tag
 				if (!bTextFormat)
 				{
-				uint32 Tag = PACKAGE_FILE_TAG;
+					uint32 Tag = PACKAGE_FILE_TAG;
 					StructuredArchiveRoot.GetUnderlyingArchive() << Tag;
 				}
 
@@ -5473,7 +5485,7 @@ FSavePackageResultStruct UPackage::Save(UPackage* InOuter, UObject* Base, EObjec
 					check( Linker->Tell() == OffsetAfterExportMap );
 				}
 
-				if (EndSavingIfCancelled(Linker.Get(), TempFilename))
+				if (EndSavingIfCancelled())
 				{ 
 					return ESavePackageResult::Canceled;
 				}
@@ -5535,7 +5547,7 @@ FSavePackageResultStruct UPackage::Save(UPackage* InOuter, UObject* Base, EObjec
 					check( Linker->Tell() == OffsetAfterPackageFileSummary );
 				}
 
-				if ( EndSavingIfCancelled( Linker.Get(), TempFilename ) )
+				if ( EndSavingIfCancelled() )
 				{ 
 					return ESavePackageResult::Canceled;
 				}
@@ -5552,16 +5564,15 @@ FSavePackageResultStruct UPackage::Save(UPackage* InOuter, UObject* Base, EObjec
 
 					if (!bFileWriterSuccess)
 					{
-						IFileManager::Get().Delete(*TempFilename);
 						UE_LOG(LogSavePackage, Error, TEXT("Error writing temp file '%s' for '%s'"),
-							*TempFilename, Filename);
+							TempFilename.IsSet() ? *TempFilename.GetValue() : TEXT("UNKNOWN"), Filename);
 						return ESavePackageResult::Error;
 					}
 				}
 				UNCLOCK_CYCLES(Time);
 				UE_CLOG(!bDiffing, LogSavePackage, Verbose,  TEXT("Save=%.2fms"), FPlatformTime::ToMilliseconds(Time) );
 		
-				if ( EndSavingIfCancelled( Linker.Get(), TempFilename ) )
+				if ( EndSavingIfCancelled() )
 				{ 
 					return ESavePackageResult::Canceled;
 				}
@@ -5587,7 +5598,7 @@ FSavePackageResultStruct UPackage::Save(UPackage* InOuter, UObject* Base, EObjec
 							FArchiveStackTrace* Writer = (FArchiveStackTrace*)(Linker->Saver);
 							TMap<FName, FArchiveDiffStats> PackageDiffStats;
 							Writer->CompareWith(*NewPath, IsEventDrivenLoaderEnabledInCookedBuilds() ? Linker->Summary.TotalHeaderSize : 0, CutoffString, DiffSettings.MaxDiffsToLog, PackageDiffStats);
-							TotalPackageSizeUncompressed = Writer->TotalSize();
+							TotalPackageSizeUncompressed += Writer->TotalSize();
 
 							auto MergeStats = [](TMap<FName, FArchiveDiffStats>& InOut, const TMap<FName, FArchiveDiffStats>& ToMerge)
 							{
@@ -5612,7 +5623,7 @@ FSavePackageResultStruct UPackage::Save(UPackage* InOuter, UObject* Base, EObjec
 							FArchiveStackTrace* Writer = (FArchiveStackTrace*)(Linker->Saver);
 							FArchiveDiffMap OutDiffMap;
 							bDiffOnlyIdentical = Writer->GenerateDiffMap(*NewPath, IsEventDrivenLoaderEnabledInCookedBuilds() ? Linker->Summary.TotalHeaderSize : 0, DiffSettings.MaxDiffsToLog, OutDiffMap);
-							TotalPackageSizeUncompressed = Writer->TotalSize();
+							TotalPackageSizeUncompressed += Writer->TotalSize();
 							if (InOutDiffMap)
 							{
 								*InOutDiffMap = MoveTemp(OutDiffMap);
@@ -5627,7 +5638,8 @@ FSavePackageResultStruct UPackage::Save(UPackage* InOuter, UObject* Base, EObjec
 							FLargeMemoryWriter* Writer = static_cast<FLargeMemoryWriter*>(Linker->Saver);
 							const int64 DataSize = Writer->TotalSize();
 
-							TotalPackageSizeUncompressed += DataSize;
+							if (!(SaveFlags & SAVE_DiffCallstack))  // Avoid double counting the package size if SAVE_DiffCallstack flag is set and DiffSettings.bSaveForDiff == true.
+								TotalPackageSizeUncompressed += DataSize;
 
 							if (IsEventDrivenLoaderEnabledInCookedBuilds() && Linker->IsCooking())
 							{
@@ -5697,16 +5709,22 @@ FSavePackageResultStruct UPackage::Save(UPackage* InOuter, UObject* Base, EObjec
 					// Move the temporary file.
 					else
 					{
+						check(TempFilename.IsSet());
+
 						if (bTextFormat)
 						{
-							IFileManager::Get().Delete(*TempFilename);
+							check(TextFormatTempFilename.IsSet());
+							IFileManager::Get().Delete(*TempFilename.GetValue());
 							TempFilename = TextFormatTempFilename;
+							TextFormatTempFilename.Reset();
 						}
 
-						UE_LOG(LogSavePackage, Log,  TEXT("Moving '%s' to '%s'"), *TempFilename, *NewPath );
-						TotalPackageSizeUncompressed += IFileManager::Get().FileSize(*TempFilename);
+						UE_LOG(LogSavePackage, Log,  TEXT("Moving '%s' to '%s'"), TempFilename.IsSet() ? *TempFilename.GetValue() : TEXT("UNKNOWN"), *NewPath );
+						TotalPackageSizeUncompressed += PackageSize;
+
+						Success = IFileManager::Get().Move( *NewPath, *TempFilename.GetValue());
+						TempFilename.Reset();
 
-						Success = IFileManager::Get().Move( *NewPath, *TempFilename );
 						if (FinalTimeStamp != FDateTime::MinValue())
 						{
 							IFileManager::Get().SetTimeStamp(*NewPath, FinalTimeStamp);
@@ -5760,7 +5778,7 @@ FSavePackageResultStruct UPackage::Save(UPackage* InOuter, UObject* Base, EObjec
 						}
 						
 						// Update package FileSize value
-						InOuter->FileSize = IFileManager::Get().FileSize(Filename);
+						InOuter->FileSize = PackageSize;
 
 						// Warn about long package names, which may be bad for consoles with limited filename lengths.
 						if( bWarnOfLongFilename == true )
@@ -5790,9 +5808,6 @@ FSavePackageResultStruct UPackage::Save(UPackage* InOuter, UObject* Base, EObjec
 							}
 						}
 					}
-
-					// Delete the temporary file.
-					IFileManager::Get().Delete( *TempFilename );
 				}
 				COOK_STAT(SavePackageStats::MBWritten += ((double)TotalPackageSizeUncompressed) / 1024.0 / 1024.0);
 
@@ -6068,7 +6083,7 @@ static void SaveAssetRegistryData(UPackage* InOuter, FLinkerSave* Linker, FStruc
 }
 
 void SaveBulkData(FLinkerSave* Linker, const UPackage* InOuter, const TCHAR* Filename, const ITargetPlatform* TargetPlatform,
-				  FSavePackageContext* SavePackageContext, const bool bTextFormat, const bool bDiffing, const bool bComputeHash, TAsyncWorkSequence<FMD5>& AsyncWriteAndHashSequence, int64 TotalPackageSizeUncompressed)
+				  FSavePackageContext* SavePackageContext, const bool bTextFormat, const bool bDiffing, const bool bComputeHash, TAsyncWorkSequence<FMD5>& AsyncWriteAndHashSequence, int64& TotalPackageSizeUncompressed)
 {
 	// Now we write all the bulkdata that is supposed to be at the end of the package
 	// and fix up the offset
@@ -6251,9 +6266,10 @@ void SaveBulkData(FLinkerSave* Linker, const UPackage* InOuter, const TCHAR* Fil
 
 			if (SavePackageContext != nullptr && SavePackageContext->PackageStoreWriter != nullptr && bWriteBulkToDisk)
 			{
-				auto ToIoBuffer = [](FLargeMemoryWriter* Writer)
+				auto AddSizeAndConvertToIoBuffer = [&TotalPackageSizeUncompressed](FLargeMemoryWriter* Writer)
 				{
 					const int64 TotalSize = Writer->TotalSize();
+					TotalPackageSizeUncompressed += TotalSize;
 					return FIoBuffer(FIoBuffer::AssumeOwnership, Writer->ReleaseOwnership(), TotalSize);
 				};
 
@@ -6262,13 +6278,13 @@ void SaveBulkData(FLinkerSave* Linker, const UPackage* InOuter, const TCHAR* Fil
 				BulkInfo.LooseFilePath = Filename;
 				BulkInfo.BulkdataType = FPackageStoreWriter::FBulkDataInfo::Standard;
 
-				SavePackageContext->PackageStoreWriter->WriteBulkdata(BulkInfo, ToIoBuffer(BulkArchive.Get()));
+				SavePackageContext->PackageStoreWriter->WriteBulkdata(BulkInfo, AddSizeAndConvertToIoBuffer(BulkArchive.Get()));
 
 				BulkInfo.BulkdataType = FPackageStoreWriter::FBulkDataInfo::Optional;
-				SavePackageContext->PackageStoreWriter->WriteBulkdata(BulkInfo, ToIoBuffer(OptionalBulkArchive.Get()));
+				SavePackageContext->PackageStoreWriter->WriteBulkdata(BulkInfo, AddSizeAndConvertToIoBuffer(OptionalBulkArchive.Get()));
 
 				BulkInfo.BulkdataType = FPackageStoreWriter::FBulkDataInfo::Mmap;
-				SavePackageContext->PackageStoreWriter->WriteBulkdata(BulkInfo, ToIoBuffer(MappedBulkArchive.Get()));
+				SavePackageContext->PackageStoreWriter->WriteBulkdata(BulkInfo, AddSizeAndConvertToIoBuffer(MappedBulkArchive.Get()));
 			}
 			else
 			{
