Remove needless circular dependency load deferring work during property serialization when EDL is enabled

Saves about 5% of unversioned property loading time for an internal project

diff --git a/Engine/Source/Runtime/CoreUObject/Private/UObject/LinkerPlaceholderBase.cpp b/Engine/Source/Runtime/CoreUObject/Private/UObject/LinkerPlaceholderBase.cpp
--- a/Engine/Source/Runtime/CoreUObject/Private/UObject/LinkerPlaceholderBase.cpp
+++ b/Engine/Source/Runtime/CoreUObject/Private/UObject/LinkerPlaceholderBase.cpp
@@ -245,14 +245,14 @@ void* FLinkerPlaceholderObjectImpl::FindRawPlaceholderContainer(const FLinkerPla
  ******************************************************************************/
 
 //------------------------------------------------------------------------------
-FScopedPlaceholderContainerTracker::FScopedPlaceholderContainerTracker(UObject* InPlaceholderContainerCandidate)
-	: PlaceholderReferencerCandidate(InPlaceholderContainerCandidate)
+void FScopedPlaceholderContainerTracker::Push(UObject* InPlaceholderContainerCandidate)
 {
+	PlaceholderReferencerCandidate = InPlaceholderContainerCandidate;
 	FPlaceholderContainerTracker::Get().PerspectiveReferencerStack.Add(InPlaceholderContainerCandidate);
 }
 
 //------------------------------------------------------------------------------
-FScopedPlaceholderContainerTracker::~FScopedPlaceholderContainerTracker()
+void FScopedPlaceholderContainerTracker::Pop()
 {
 	UObject* StackTop = FPlaceholderContainerTracker::Get().PerspectiveReferencerStack.Pop();
 #if USE_DEFERRED_DEPENDENCY_CHECK_VERIFICATION_TESTS
@@ -275,8 +275,7 @@ FScopedPlaceholderRawContainerTracker::~FScopedPlaceholderRawContainerTracker()
 }
 
 //------------------------------------------------------------------------------
-FScopedPlaceholderPropertyTracker::FScopedPlaceholderPropertyTracker(FFieldVariant InIntermediateProperty)
-	: IntermediateProperty() // leave null, as a sentinel value (implying that PerspectiveReferencerStack was empty)
+void FScopedPlaceholderPropertyTracker::Push(FFieldVariant InIntermediateProperty)
 {
 	FPlaceholderContainerTracker& ContainerRepository = FPlaceholderContainerTracker::Get();
 	if (ContainerRepository.PerspectiveReferencerStack.Num() > 0 || ContainerRepository.PerspectiveRootDataStack.Num() > 0)
@@ -291,7 +290,7 @@ FScopedPlaceholderPropertyTracker::FScopedPlaceholderPropertyTracker(FFieldVaria
 }
 
 //------------------------------------------------------------------------------
-FScopedPlaceholderPropertyTracker::~FScopedPlaceholderPropertyTracker()
+void FScopedPlaceholderPropertyTracker::Pop()
 {
 	FPlaceholderContainerTracker& ContainerRepository = FPlaceholderContainerTracker::Get();
 	if (IntermediateProperty.IsValid())
diff --git a/Engine/Source/Runtime/CoreUObject/Private/UObject/LinkerPlaceholderBase.h b/Engine/Source/Runtime/CoreUObject/Private/UObject/LinkerPlaceholderBase.h
--- a/Engine/Source/Runtime/CoreUObject/Private/UObject/LinkerPlaceholderBase.h
+++ b/Engine/Source/Runtime/CoreUObject/Private/UObject/LinkerPlaceholderBase.h
@@ -9,6 +9,15 @@
 #include "UObject/LinkerLoad.h"
 #include "UObject/Field.h"
 
+FORCEINLINE bool UseCircularDependencyLoadDeferring()
+{
+#if USE_CIRCULAR_DEPENDENCY_LOAD_DEFERRING
+	return !GEventDrivenLoaderEnabled;
+#else
+	return false;
+#endif
+}
+
 /*******************************************************************************
  * FPlaceholderContainerTracker / FScopedPlaceholderPropertyTracker
  ******************************************************************************/
@@ -23,11 +32,28 @@
 struct FScopedPlaceholderContainerTracker
 {
 public:
-	FScopedPlaceholderContainerTracker(UObject* PerspectivePlaceholderReferencer);
-	~FScopedPlaceholderContainerTracker();
+	FScopedPlaceholderContainerTracker(UObject* InPlaceholderContainerCandidate)
+	: bEnabled(UseCircularDependencyLoadDeferring())
+	{
+		if (bEnabled)
+		{
+			Push(InPlaceholderContainerCandidate);
+		}
+	}
+	~FScopedPlaceholderContainerTracker()
+	{
+		if (bEnabled)
+		{
+			Pop();
+		}
+	}
 
 private:
+	const bool bEnabled;
 	UObject* PlaceholderReferencerCandidate;
+
+	void Push(UObject* InIntermediateProperty);
+	void Pop();
 };
 
 /** 
@@ -40,11 +66,29 @@ private:
 struct FScopedPlaceholderPropertyTracker
 {
 public:
-	 FScopedPlaceholderPropertyTracker(FFieldVariant IntermediateProperty);
-	~FScopedPlaceholderPropertyTracker();
+	 FScopedPlaceholderPropertyTracker(FFieldVariant InIntermediateProperty)
+		 : bEnabled(UseCircularDependencyLoadDeferring())
+	 {
+		if (bEnabled)
+		{
+			Push(InIntermediateProperty);
+		}
+	 }
+	
+	~FScopedPlaceholderPropertyTracker()
+	{
+		if (bEnabled)
+		{
+			Pop();
+		}
+	}
 
 private:
+	const bool bEnabled;
 	FFieldVariant IntermediateProperty;
+	
+	void Push(FFieldVariant InIntermediateProperty);
+	void Pop();
 };
 
 /*******************************************************************************
diff --git a/Engine/Source/Runtime/CoreUObject/Private/UObject/PropertyStruct.cpp b/Engine/Source/Runtime/CoreUObject/Private/UObject/PropertyStruct.cpp
--- a/Engine/Source/Runtime/CoreUObject/Private/UObject/PropertyStruct.cpp
+++ b/Engine/Source/Runtime/CoreUObject/Private/UObject/PropertyStruct.cpp
@@ -15,24 +15,27 @@
 
 static inline void PreloadInnerStructMembers(FStructProperty* StructProperty)
 {
-#if USE_CIRCULAR_DEPENDENCY_LOAD_DEFERRING
-	uint32 PropagatedLoadFlags = 0;
-	if (FLinkerLoad* Linker = StructProperty->GetLinker())
+	if (UseCircularDependencyLoadDeferring())
 	{
-		PropagatedLoadFlags |= (Linker->LoadFlags & LOAD_DeferDependencyLoads);
-	}
+		uint32 PropagatedLoadFlags = 0;
+		if (FLinkerLoad* Linker = StructProperty->GetLinker())
+		{
+			PropagatedLoadFlags |= (Linker->LoadFlags & LOAD_DeferDependencyLoads);
+		}
 
-	if (UScriptStruct* Struct = StructProperty->Struct)
-	{
-		if (FLinkerLoad* StructLinker = Struct->GetLinker())
+		if (UScriptStruct* Struct = StructProperty->Struct)
 		{
-			TGuardValue<uint32> LoadFlagGuard(StructLinker->LoadFlags, StructLinker->LoadFlags | PropagatedLoadFlags);
-			Struct->RecursivelyPreload();
+			if (FLinkerLoad* StructLinker = Struct->GetLinker())
+			{
+				TGuardValue<uint32> LoadFlagGuard(StructLinker->LoadFlags, StructLinker->LoadFlags | PropagatedLoadFlags);
+				Struct->RecursivelyPreload();
+			}
 		}
 	}
-#else // USE_CIRCULAR_DEPENDENCY_LOAD_DEFERRING
-	StructProperty->Struct->RecursivelyPreload();
-#endif // USE_CIRCULAR_DEPENDENCY_LOAD_DEFERRING
+	else
+	{
+		StructProperty->Struct->RecursivelyPreload();
+	}
 }
 
 /*-----------------------------------------------------------------------------
@@ -144,11 +147,7 @@ uint32 FStructProperty::GetValueTypeHashInternal(const void* Src) const
 
 void FStructProperty::SerializeItem(FStructuredArchive::FSlot Slot, void* Value, void const* Defaults) const
 {
-	check(Struct);
-
-#if USE_CIRCULAR_DEPENDENCY_LOAD_DEFERRING
 	FScopedPlaceholderPropertyTracker ImportPropertyTracker(this);
-#endif
 
 	Struct->SerializeItem(Slot, Value, Defaults);
 }
