Fix non optional duplicated bulkdata serialization when running with the IoDispatcher.

I was using the original FIoChunkId (ie the non optional data) to check if the duplicate optional data existed, to fix this I now load the duplicate data values first then use them to create the potential FIoChunkId then apply them if we pass the correct existance check.
I also moved the code to it's own method to hide the implementation a bit since most project/platforms don't use this path and so people don't need to worry about it.

#rnx

diff --git a/Engine/Source/Runtime/CoreUObject/Private/Serialization/BulkData2.cpp b/Engine/Source/Runtime/CoreUObject/Private/Serialization/BulkData2.cpp
--- a/Engine/Source/Runtime/CoreUObject/Private/Serialization/BulkData2.cpp
+++ b/Engine/Source/Runtime/CoreUObject/Private/Serialization/BulkData2.cpp
@@ -733,8 +733,8 @@ void FBulkDataBase::Serialize(FArchive& Ar, UObject* Owner, int32 /*Index*/, boo
 			Fallback.BulkDataSize = BulkDataSize;
 		}
 
-		const FString* Filename = nullptr;
 		FName PackageName;
+		const FString* Filename = nullptr;
 		const FLinkerLoad* Linker = nullptr;
 
 		if (bUseIoDispatcher == false)
@@ -758,49 +758,7 @@ void FBulkDataBase::Serialize(FArchive& Ar, UObject* Owner, int32 /*Index*/, boo
 		{
 			if (IsDuplicateNonOptional())
 			{
-				auto DoesOptionalDataExist = [this](const FString* PackageFilename)->bool
-				{
-#if ALLOW_OPTIONAL_DATA
-					if (!IsUsingIODispatcher())
-					{
-						check(PackageFilename != nullptr);
-						FString OptionalDataFilename = ConvertFilenameFromFlags(*PackageFilename);
-						return IFileManager::Get().FileExists(*OptionalDataFilename);
-					}
-					else
-					{
-						return IoDispatcher->DoesChunkExist(ChunkID);
-					}
-#else
-					return false;
-#endif
-				};
-
-				if (DoesOptionalDataExist(Filename))
-				{
-					SerializeDuplicateData(Ar, BulkDataFlags, BulkDataSizeOnDisk, BulkDataOffsetInFile);
-
-					if (!IsInlined() && bUseIoDispatcher)
-					{
-						// Regenerate the FIoChunkId to find the optional BulkData instead!
-						const int64 BulkDataID = BulkDataSize > 0 ? BulkDataOffsetInFile : TNumericLimits<uint64>::Max();
-						ChunkID = CreateBulkdataChunkId(Package->GetPackageId().ToIndex(), BulkDataID, EIoChunkType::OptionalBulkData);
-						BulkDataFlags |= BULKDATA_UsesIoDispatcher; // Indicates that this BulkData should use the FIoChunkId rather than a filename
-					}
-					else
-					{
-						Fallback.Token = InvalidToken;
-						Fallback.BulkDataSize = BulkDataSize;
-					}
-				}
-				else
-				{
-					// Skip over the optional data info (can't do a seek because we need to load the flags to work
-					// out if we read things as 32bit or 64bit)
-					uint32 DummyValue32;
-					int64 DummyValue64;
-					SerializeDuplicateData(Ar, DummyValue32, DummyValue64, DummyValue64);
-				}
+				ProcessDuplicateData(Ar, Package, Filename, BulkDataSizeOnDisk, BulkDataOffsetInFile);
 			}
 
 			// Fix up the file offset if we have a linker (if we do not then we will be loading via FIoDispatcher anyway)
@@ -1326,6 +1284,49 @@ void FBulkDataBase::LoadDataDirectly(void** DstBuffer)
 	}
 }
 
+void FBulkDataBase::ProcessDuplicateData(FArchive& Ar, const UPackage* Package, const FString* Filename, int64& InOutSizeOnDisk, int64& InOutOffsetInFile)
+{
+	// We need to load the optional bulkdata info as we might need to create a FIoChunkId based on it!
+	uint32 NewFlags;
+	int64 NewSizeOnDisk;
+	int64 NewOffset;
+
+	SerializeDuplicateData(Ar, NewFlags, NewSizeOnDisk, NewOffset);
+
+#if ALLOW_OPTIONAL_DATA
+	if (IsUsingIODispatcher())
+	{
+		const int64 BulkDataID = NewSizeOnDisk > 0 ? NewOffset : TNumericLimits<uint64>::Max();
+		FIoChunkId OptionalChunkId = CreateBulkdataChunkId(Package->GetPackageId().ToIndex(), BulkDataID, EIoChunkType::OptionalBulkData);
+
+		if (IoDispatcher->DoesChunkExist(OptionalChunkId))
+		{
+			BulkDataFlags = NewFlags | BULKDATA_UsesIoDispatcher;
+			InOutSizeOnDisk = NewSizeOnDisk;
+			InOutOffsetInFile = NewOffset;
+
+			ChunkID = OptionalChunkId;
+		}
+	}
+	else
+	{
+		check(Filename != nullptr);
+		const FString OptionalDataFilename = ConvertFilenameFromFlags(*Filename);
+
+		if (IFileManager::Get().FileExists(*OptionalDataFilename))
+		{
+			BulkDataFlags = NewFlags;
+			InOutSizeOnDisk = NewSizeOnDisk;
+			InOutOffsetInFile = NewOffset;
+
+			// Note we do not override Filename with OptionalDataFilename as we are supposed to store the original!
+			Fallback.Token = InvalidToken;
+			Fallback.BulkDataSize = InOutSizeOnDisk;
+		}
+	}
+#endif
+}
+
 void FBulkDataBase::SerializeDuplicateData(FArchive& Ar, uint32& OutBulkDataFlags, int64& OutBulkDataSizeOnDisk, int64& OutBulkDataOffsetInFile)
 {
 	Ar << OutBulkDataFlags;
diff --git a/Engine/Source/Runtime/CoreUObject/Public/Serialization/BulkData2.h b/Engine/Source/Runtime/CoreUObject/Public/Serialization/BulkData2.h
--- a/Engine/Source/Runtime/CoreUObject/Public/Serialization/BulkData2.h
+++ b/Engine/Source/Runtime/CoreUObject/Public/Serialization/BulkData2.h
@@ -141,6 +141,7 @@ private:
 
 	void LoadDataDirectly(void** DstBuffer);
 
+	void ProcessDuplicateData(FArchive& Ar, const UPackage* Package, const FString* Filename, int64& InOutSizeOnDisk, int64& InOutOffsetInFile);
 	void SerializeDuplicateData(FArchive& Ar, uint32& OutBulkDataFlags, int64& OutBulkDataSizeOnDisk, int64& OutBulkDataOffsetInFile);
 	void SerializeBulkData(FArchive& Ar, void* DstBuffer, int64 DataLength);
 
