Minor optimization to reduce the number of memory allocations due to string conversions when calling FBulkDataBase::ConvertFilenameFromFlags.

#rnx

diff --git a/Engine/Source/Runtime/CoreUObject/Private/Serialization/BulkData2.cpp b/Engine/Source/Runtime/CoreUObject/Private/Serialization/BulkData2.cpp
--- a/Engine/Source/Runtime/CoreUObject/Private/Serialization/BulkData2.cpp
+++ b/Engine/Source/Runtime/CoreUObject/Private/Serialization/BulkData2.cpp
@@ -24,10 +24,10 @@ DEFINE_LOG_CATEGORY_STATIC(LogBulkDataRuntime, Log, All);
 namespace
 {
 	// TODO: Maybe expose this and start using everywhere?
-	static constexpr const TCHAR* InlinedExt = TEXT(".uexp");			// Stored in the export data
-	static constexpr const TCHAR* DefaultExt = TEXT(".ubulk");			// Stored in a separate file
-	static constexpr const TCHAR* MemoryMappedExt = TEXT(".m.ubulk");	// Stored in a separate file aligned for memory mapping
-	static constexpr const TCHAR* OptionalExt = TEXT(".uptnl");			// Stored in a separate file that is optional
+	static const FString InlinedExt = TEXT(".uexp");			// Stored in the export data
+	static const FString DefaultExt = TEXT(".ubulk");			// Stored in a separate file
+	static const FString MemoryMappedExt = TEXT(".m.ubulk");	// Stored in a separate file aligned for memory mapping
+	static const FString OptionalExt = TEXT(".uptnl");			// Stored in a separate file that is optional
 
 	const uint16 InvalidBulkDataIndex = ~uint16(0);
 
