Copying Private-LoadTimes-4.24 stream to Main. Biggest changes are in Materials/Shader memory freezing.

diff --git a/Engine/Source/Runtime/CoreUObject/Private/Serialization/AsyncLoading.cpp b/Engine/Source/Runtime/CoreUObject/Private/Serialization/AsyncLoading.cpp
--- a/Engine/Source/Runtime/CoreUObject/Private/Serialization/AsyncLoading.cpp
+++ b/Engine/Source/Runtime/CoreUObject/Private/Serialization/AsyncLoading.cpp
@@ -91,7 +91,7 @@ DECLARE_FLOAT_ACCUMULATOR_STAT(TEXT("Total PostLoadObjects time GT"), STAT_FAsyn
 
 DECLARE_FLOAT_ACCUMULATOR_STAT( TEXT( "Async loading block time" ), STAT_AsyncIO_AsyncLoadingBlockingTime, STATGROUP_AsyncIO );
 DECLARE_FLOAT_ACCUMULATOR_STAT( TEXT( "Async package precache wait time" ), STAT_AsyncIO_AsyncPackagePrecacheWaitTime, STATGROUP_AsyncIO );
-
+	
 /** Helper function for profiling load times */
 static FName StaticGetNativeClassName(UClass* InClass)
 {
@@ -2904,8 +2904,8 @@ void FAsyncPackage::EventDrivenCreateExport(int32 LocalExportIndex)
 						// Do this for all subobjects created in the native constructor.
 						if (!Export.Object->HasAnyFlags(RF_LoadCompleted))
 						{
-								UE_LOG(LogStreaming, VeryVerbose, TEXT("Note2: %s was constructed during load and is an export and so needs loading."), *Export.Object->GetFullName());
-								UE_CLOG(!Export.Object->HasAllFlags(RF_WillBeLoaded), LogStreaming, Fatal, TEXT("%s was found in memory and is an export but does not have all load flags."), *Export.Object->GetFullName());
+							UE_LOG(LogStreaming, VeryVerbose, TEXT("Note2: %s was constructed during load and is an export and so needs loading."), *Export.Object->GetFullName());
+							UE_CLOG(!Export.Object->HasAllFlags(RF_WillBeLoaded), LogStreaming, Fatal, TEXT("%s was found in memory and is an export but does not have all load flags."), *Export.Object->GetFullName());
 							if(Export.Object->HasAnyFlags(RF_ClassDefaultObject))
 							{
 								// never call PostLoadSubobjects on class default objects, this matches the behavior of the old linker where
@@ -2915,9 +2915,9 @@ void FAsyncPackage::EventDrivenCreateExport(int32 LocalExportIndex)
 							}
 							else
 							{
-							Export.Object->SetFlags(RF_NeedLoad | RF_NeedPostLoad | RF_NeedPostLoadSubobjects | RF_WasLoaded);
+								Export.Object->SetFlags(RF_NeedLoad | RF_NeedPostLoad | RF_NeedPostLoadSubobjects | RF_WasLoaded);
 							}
-								Export.Object->ClearFlags(RF_WillBeLoaded);
+							Export.Object->ClearFlags(RF_WillBeLoaded);
 						}
 					}
 				}
@@ -3162,15 +3162,19 @@ void FAsyncPackage::EventDrivenSerializeExport(int32 LocalExportIndex)
 		check(!Linker->TemplateForGetArchetypeFromLoader);
 		Linker->TemplateForGetArchetypeFromLoader = Template;
 
-		if (Object->HasAnyFlags(RF_ClassDefaultObject))
-		{
-			SCOPED_ACCUM_LOADTIME(Serialize, StaticGetNativeClassName(Object->GetClass()));
-			Object->GetClass()->SerializeDefaultObject(Object, *Linker);
-		}
-		else
 		{
+			ACCUM_LOADTIMECOUNT_STAT(StaticGetNativeClassName(Object->GetClass()).ToString());
+			SCOPED_ACCUM_LOADTIME_STAT(StaticGetNativeClassName(Object->GetClass()).ToString());
 			SCOPED_ACCUM_LOADTIME(Serialize, StaticGetNativeClassName(Object->GetClass()));
-			Object->Serialize(*Linker);
+		
+			if (Object->HasAnyFlags(RF_ClassDefaultObject))
+			{
+				Object->GetClass()->SerializeDefaultObject(Object, *Linker);
+			}
+			else
+			{
+				Object->Serialize(*Linker);
+			}
 		}
 		check(Linker->TemplateForGetArchetypeFromLoader == Template);
 		Linker->TemplateForGetArchetypeFromLoader = nullptr;
@@ -4721,6 +4725,7 @@ uint32 FAsyncLoadingThread::Run()
 	if (!IsInGameThread())
 	{
 		FPlatformProcess::SetThreadAffinityMask(FPlatformAffinity::GetAsyncLoadingThreadMask());
+		FMemory::SetupTLSCachesOnCurrentThread();
 	}
 
 	bool bWasSuspendedLastFrame = false;
@@ -5590,6 +5595,7 @@ EAsyncPackageState::Type FAsyncPackage::CreateLinker()
 		// Try to find existing package or create it if not already present.
 		UPackage* Package = nullptr;
 		{
+			SCOPED_LOADTIMER(CreateLinker_CreatePackage);
 			FGCScopeGuard GCGuard;
 			Package = CreatePackage(nullptr, *Desc.Name.ToString());
 			if (!Package)
@@ -5613,6 +5619,7 @@ EAsyncPackageState::Type FAsyncPackage::CreateLinker()
 
 		if (Package->FileName == NAME_None && !Package->bHasBeenFullyLoaded)
 		{
+			SCOPED_LOADTIMER(CreateLinker_SetFlags);
 			// We just created the package, so set ownership flag and set up package info
 			bCreatedLinkerRoot = true;
 
@@ -5631,7 +5638,10 @@ EAsyncPackageState::Type FAsyncPackage::CreateLinker()
 		LastObjectWorkWasPerformedOn = Package;
 		// if the linker already exists, we don't need to lookup the file (it may have been pre-created with
 		// a different filename)
-		Linker = FLinkerLoad::FindExistingLinkerForPackage(Package);
+		{
+			SCOPED_LOADTIMER(CreateLinker_FindLinker);
+			Linker = FLinkerLoad::FindExistingLinkerForPackage(Package);
+		}
 		if (Linker)
 		{
 			if (GEventDrivenLoaderEnabled)
@@ -5650,6 +5660,7 @@ EAsyncPackageState::Type FAsyncPackage::CreateLinker()
 			// Process any package redirects
 			FString NameToLoad;
 			{
+				SCOPED_LOADTIMER(CreateLinker_GetRedirectedName);
 				const FCoreRedirectObjectName NewPackageName = FCoreRedirects::GetRedirectedName(ECoreRedirectFlags::Type_Package, FCoreRedirectObjectName(NAME_None, NAME_None, Desc.NameToLoad));
 				NameToLoad = NewPackageName.PackageName.ToString();
 			}
@@ -5657,6 +5668,7 @@ EAsyncPackageState::Type FAsyncPackage::CreateLinker()
 			// The editor must not redirect packages for localization.
 			if (!GIsEditor)
 			{
+				SCOPED_LOADTIMER(CreateLinker_MassagePath);
 				// Allow delegates to resolve this path
 				NameToLoad = FPackageName::GetDelegateResolvedPackagePath(NameToLoad);
 				NameToLoad = FPackageName::GetLocalizedPackagePath(NameToLoad);
@@ -5665,54 +5677,62 @@ EAsyncPackageState::Type FAsyncPackage::CreateLinker()
 			const FGuid* const Guid = Desc.Guid.IsValid() ? &Desc.Guid : nullptr;
 
 			FString PackageFileName;
-			const bool DoesPackageExist = FPackageName::DoesPackageExist(NameToLoad, Guid, &PackageFileName);
+			bool DoesPackageExist = false;
+			{
+				SCOPED_LOADTIMER(CreateLinker_DoesExist);
+				DoesPackageExist = FPackageName::DoesPackageExist(NameToLoad, Guid, &PackageFileName);
+			}
 
-			if (Desc.NameToLoad == NAME_None ||
-				(!GetConvertedDynamicPackageNameToTypeName().Contains(Desc.Name) &&
-					!DoesPackageExist))
 			{
-				FName FailedLoadName = FName(*NameToLoad);
+				SCOPED_LOADTIMER(CreateLinker_MissingPackage);
 
-				if (!FLinkerLoad::IsKnownMissingPackage(FailedLoadName))
+				if (Desc.NameToLoad == NAME_None ||
+					(!GetConvertedDynamicPackageNameToTypeName().Contains(Desc.Name) &&
+						!DoesPackageExist))
 				{
-					UE_LOG(LogStreaming, Error, TEXT("Couldn't find file for package %s requested by async loading code. NameToLoad: %s"), *Desc.Name.ToString(), *Desc.NameToLoad.ToString());
+					FName FailedLoadName = FName(*NameToLoad);
+
+					if (!FLinkerLoad::IsKnownMissingPackage(FailedLoadName))
+					{
+						UE_LOG(LogStreaming, Error, TEXT("Couldn't find file for package %s requested by async loading code. NameToLoad: %s"), *Desc.Name.ToString(), *Desc.NameToLoad.ToString());
 
 #if !WITH_EDITORONLY_DATA
-					UE_CLOG(bUseTimeLimit, LogStreaming, Error, TEXT("This will hitch streaming because it ends up searching the disk instead of finding the file in the pak file."));
+						UE_CLOG(bUseTimeLimit, LogStreaming, Error, TEXT("This will hitch streaming because it ends up searching the disk instead of finding the file in the pak file."));
 #endif
 
-					if (GEventDrivenLoaderEnabled)
-					{
-						TSet<FName> DependentPackages;
-						FWeakAsyncPackagePtr WeakPtr(this);
-						TArray<FEventLoadNodePtr> AddedNodes;
-						EventNodeArray.GetAddedNodes(AddedNodes, this);
-						for (const FEventLoadNodePtr& NodePtr : AddedNodes)
+						if (GEventDrivenLoaderEnabled)
 						{
-							const FEventLoadNode& Node = EventNodeArray.GetNode(NodePtr);
-							for (const FEventLoadNodePtr& Other : Node.NodesWaitingForMe)
+							TSet<FName> DependentPackages;
+							FWeakAsyncPackagePtr WeakPtr(this);
+							TArray<FEventLoadNodePtr> AddedNodes;
+							EventNodeArray.GetAddedNodes(AddedNodes, this);
+							for (const FEventLoadNodePtr& NodePtr : AddedNodes)
 							{
-								FName DependentPackageName = Other.WaitingPackage.HumanReadableStringForDebugging();
-								if (DependentPackageName != NAME_None)
+								const FEventLoadNode& Node = EventNodeArray.GetNode(NodePtr);
+								for (const FEventLoadNodePtr& Other : Node.NodesWaitingForMe)
 								{
-									DependentPackages.Add(DependentPackageName);
+									FName DependentPackageName = Other.WaitingPackage.HumanReadableStringForDebugging();
+									if (DependentPackageName != NAME_None)
+									{
+										DependentPackages.Add(DependentPackageName);
+									}
 								}
 							}
-						}
 
-						UE_LOG(LogStreaming, Error, TEXT("Found %d dependent packages..."), DependentPackages.Num());
-						for (const FName& DependentPackageName : DependentPackages)
-						{
-							UE_LOG(LogStreaming, Error, TEXT("  %s"), *DependentPackageName.ToString());
+							UE_LOG(LogStreaming, Error, TEXT("Found %d dependent packages..."), DependentPackages.Num());
+							for (const FName& DependentPackageName : DependentPackages)
+							{
+								UE_LOG(LogStreaming, Error, TEXT("  %s"), *DependentPackageName.ToString());
+							}
 						}
+
+						// Add to known missing list so it won't error again
+						FLinkerLoad::AddKnownMissingPackage(FailedLoadName);
 					}
 
-					// Add to known missing list so it won't error again
-					FLinkerLoad::AddKnownMissingPackage(FailedLoadName);
+					bLoadHasFailed = true;
+					return EAsyncPackageState::TimeOut;
 				}
-
-				bLoadHasFailed = true;
-				return EAsyncPackageState::TimeOut;
 			}
 
 			// Create raw async linker, requiring to be ticked till finished creating.
@@ -5727,6 +5747,7 @@ EAsyncPackageState::Type FAsyncPackage::CreateLinker()
 				LinkerFlags |= LOAD_PackageForPIE;
 			}
 #endif
+			SCOPED_LOADTIMER(CreateLinker_CreateLinkerAsync);
 			FUObjectSerializeContext* LoadContext = GetSerializeContext();
 			if (GEventDrivenLoaderEnabled)
 			{
diff --git a/Engine/Source/Runtime/CoreUObject/Private/Serialization/AsyncLoading2.cpp b/Engine/Source/Runtime/CoreUObject/Private/Serialization/AsyncLoading2.cpp
--- a/Engine/Source/Runtime/CoreUObject/Private/Serialization/AsyncLoading2.cpp
+++ b/Engine/Source/Runtime/CoreUObject/Private/Serialization/AsyncLoading2.cpp
@@ -52,6 +52,7 @@
 #include "Serialization/BulkData.h"
 #include "Serialization/LargeMemoryReader.h"
 #include "UObject/UObjectClusters.h"
+#include "ProfilingDebugging/CountersTrace.h"
 
 #if UE_BUILD_DEVELOPMENT || UE_BUILD_DEBUG
 //PRAGMA_DISABLE_OPTIMIZATION
@@ -99,6 +100,8 @@ if ((Condition)) \
 #define UE_ASYNC_PACKAGE_CLOG_VERBOSE(Condition, Verbosity, PackageDesc, LogDesc, Format, ...)
 #endif
 
+TRACE_DECLARE_INT_COUNTER(PendingBundleIoRequests, TEXT("AsyncLoading/PendingBundleIoRequests"));
+
 struct FAsyncPackage2;
 class FAsyncLoadingThread2;
 
@@ -2282,6 +2285,7 @@ void FAsyncLoadingThread2::StartBundleIoRequests()
 			Package->GetExportBundleNode(EEventLoadNode2::ExportBundle_Process, 0)->ReleaseBarrier();
 			Package->AsyncLoadingThread.WaitingForIoBundleCounter.Decrement();
 		});
+		TRACE_COUNTER_DECREMENT(PendingBundleIoRequests);
 	}
 }
 
@@ -4155,16 +4159,35 @@ static void VerifyLoadFlagsWhenFinishedLoading()
 			const bool bHasAnyLoadIntermediateFlags = !!(Flags & LoadIntermediateFlags);
 			const bool bWasLoaded = !!(Flags & RF_WasLoaded);
 			const bool bLoadCompleted = !!(Flags & RF_LoadCompleted);
-			check(!bHasAnyAsyncFlags);
-			check(!bHasAnyLoadIntermediateFlags);
+
+			UE_CLOG(bHasAnyLoadIntermediateFlags, LogStreaming, Warning,
+				TEXT("Object '%s' (ObjectFlags=%d, InternalObjectFlags=%d) should not have any load flags now")
+				TEXT(", or this check is incorrectly reached during active loading."),
+				*Obj->GetFullName(),
+				Flags,
+				InternalFlags);
+
 			if (bWasLoaded)
 			{
 				const bool bIsPackage = Obj->IsA(UPackage::StaticClass());
-				check(bIsPackage || bLoadCompleted);
+
+				UE_CLOG(!bIsPackage && !bLoadCompleted, LogStreaming, Warning,
+					TEXT("Object '%s' (ObjectFlags=%d, InternalObjectFlags=%d) is a serialized object and should be completely loaded now")
+					TEXT(", or this check is incorrectly reached during active loading."),
+					*Obj->GetFullName(),
+					Flags,
+					InternalFlags);
+
+				UE_CLOG(bHasAnyAsyncFlags, LogStreaming, Warning,
+					TEXT("Object '%s' (ObjectFlags=%d, InternalObjectFlags=%d) is a serialized object and should not have any async flags now")
+					TEXT(", or this check is incorrectly reached during active loading."),
+					*Obj->GetFullName(),
+					Flags,
+					InternalFlags);
 			}
 		}
 	}
-	UE_LOG(LogStreaming, Log, TEXT("Verified load flags when finished loading"));
+	UE_LOG(LogStreaming, Log, TEXT("Verified load flags when finished active loading."));
 }
 
 void FGlobalImportStore::OnPreGarbageCollect(bool bInIsLoadingPackages)
@@ -5033,6 +5056,7 @@ void FAsyncLoadingThread2::FlushLoading(int32 RequestId)
 
 EAsyncPackageState::Type FAsyncLoadingThread2::ProcessLoadingUntilCompleteFromGameThread(TFunctionRef<bool()> CompletionPredicate, float TimeLimit)
 {
+	TRACE_CPUPROFILER_EVENT_SCOPE(ProcessLoadingUntilComplete);
 	if (!IsAsyncLoading())
 	{
 		return EAsyncPackageState::Complete;
diff --git a/Engine/Source/Runtime/CoreUObject/Private/Serialization/BulkData2.cpp b/Engine/Source/Runtime/CoreUObject/Private/Serialization/BulkData2.cpp
--- a/Engine/Source/Runtime/CoreUObject/Private/Serialization/BulkData2.cpp
+++ b/Engine/Source/Runtime/CoreUObject/Private/Serialization/BulkData2.cpp
@@ -16,6 +16,8 @@
 
 DEFINE_LOG_CATEGORY_STATIC(LogBulkDataRuntime, Log, All);
 
+IMPLEMENT_TYPE_LAYOUT(FBulkDataBase);
+
 // If set to 0 then the loose file fallback will be used even if the -ZenLoader command line flag is present.
 #define ENABLE_IO_DISPATCHER 1
 
@@ -130,7 +132,7 @@ namespace FileTokenSystem
 
 	FRWLock TokenLock;
 
-	FBulkDataBase::FileToken RegisterFileToken( const FName& PackageName, const FString& Filename, uint64 BulkDataOffsetInFile )
+	FileToken RegisterFileToken( const FName& PackageName, const FString& Filename, uint64 BulkDataOffsetInFile )
 	{
 		FWriteScopeLock LockForScope(TokenLock);
 
@@ -140,11 +142,11 @@ namespace FileTokenSystem
 		Data.PackageName = PackageName;
 		Data.BulkDataOffsetInFile = BulkDataOffsetInFile;
 
-		FBulkDataBase::FileToken FileToken  = TokenData.Add(Data);
+		FileToken FileToken  = TokenData.Add(Data);
 		return FileToken;
 	}
 
-	void UnregisterFileToken(FBulkDataBase::FileToken ID)
+	void UnregisterFileToken(FileToken ID)
 	{
 		if (ID != FBulkDataBase::InvalidToken)
 		{
@@ -157,7 +159,7 @@ namespace FileTokenSystem
 		}
 	}
 
-	FBulkDataBase::FileToken CopyFileToken(FBulkDataBase::FileToken ID)
+	FileToken CopyFileToken(FileToken ID)
 	{
 		if (ID != FBulkDataBase::InvalidToken)
 		{
@@ -181,7 +183,7 @@ namespace FileTokenSystem
 		}
 	}
 
-	Data GetFileData(FBulkDataBase::FileToken ID)
+	Data GetFileData(FileToken ID)
 	{
 		if (ID == FBulkDataBase::InvalidToken)
 		{
@@ -198,7 +200,7 @@ namespace FileTokenSystem
 		return Output;
 	}
 
-	FString GetFilename(FBulkDataBase::FileToken ID)
+	FString GetFilename(FileToken ID)
 	{
 		if (ID == FBulkDataBase::InvalidToken)
 		{
@@ -209,7 +211,7 @@ namespace FileTokenSystem
 		return StringTable.Resolve(TokenData[ID].PackageName);
 	}
 
-	uint64 GetBulkDataOffset(FBulkDataBase::FileToken ID)
+	uint64 GetBulkDataOffset(FileToken ID)
 	{
 		if (ID == FBulkDataBase::InvalidToken)
 		{
@@ -614,7 +616,7 @@ private:
 };
 
 FBulkDataBase::FBulkDataBase(FBulkDataBase&& Other)
-	: ChunkID(Other.ChunkID) // Copies the entire union
+	: Data(Other.Data) // Copies the entire union
 	, DataAllocation(Other.DataAllocation)
 	, BulkDataFlags(Other.BulkDataFlags)
 
@@ -623,7 +625,7 @@ FBulkDataBase::FBulkDataBase(FBulkDataBase&& Other)
 
 	if (!Other.IsUsingIODispatcher())
 	{
-		Other.Fallback.Token = InvalidToken; // Prevent the other object from unregistering the token
+		Other.Data.Fallback.Token = InvalidToken; // Prevent the other object from unregistering the token
 	}	
 }
 
@@ -636,7 +638,7 @@ FBulkDataBase& FBulkDataBase::operator=(const FBulkDataBase& Other)
 
 	RemoveBulkData();
 
-	Fallback.Token = FileTokenSystem::CopyFileToken(Other.Fallback.Token);
+	Data.Fallback.Token = FileTokenSystem::CopyFileToken(Other.Data.Fallback.Token);
 
 	// Copy token
 	BulkDataFlags = Other.BulkDataFlags;
@@ -651,7 +653,7 @@ FBulkDataBase& FBulkDataBase::operator=(const FBulkDataBase& Other)
 	else
 	{
 		const int64 BulkDataSize = GetBulkDataSize();
-		FileTokenSystem::Data FileData = FileTokenSystem::GetFileData(Fallback.Token);
+		FileTokenSystem::Data FileData = FileTokenSystem::GetFileData(Data.Fallback.Token);
 		const FString MemoryMappedFilename = ConvertFilenameFromFlags(FileData.PackageHeaderFilename);
 		MemoryMapBulkData(MemoryMappedFilename, FileData.BulkDataOffsetInFile, BulkDataSize);
 	}
@@ -666,7 +668,7 @@ FBulkDataBase::~FBulkDataBase()
 	FreeData();
 	if (!IsUsingIODispatcher())
 	{
-		FileTokenSystem::UnregisterFileToken(Fallback.Token);
+		FileTokenSystem::UnregisterFileToken(Data.Fallback.Token);
 	}	
 }
 
@@ -730,15 +732,15 @@ void FBulkDataBase::Serialize(FArchive& Ar, UObject* Owner, int32 /*Index*/, boo
 			const EIoChunkType Type = IsOptional() ? EIoChunkType::OptionalBulkData : EIoChunkType::BulkData;
 
 			const int64 BulkDataID = BulkDataSize > 0 ? BulkDataOffsetInFile : TNumericLimits<uint64>::Max();
-			ChunkID = CreateBulkdataChunkId(Package->GetPackageId().ToIndex(), BulkDataID, Type);
+			Data.ChunkID = CreateBulkdataChunkId(Package->GetPackageId().ToIndex(), BulkDataID, Type);
 
 			BulkDataFlags |= BULKDATA_UsesIoDispatcher; // Indicates that this BulkData should use the FIoChunkId rather than a filename
 		}
 		else
 		{
 			// Invalidate the Token and then set the BulkDataSize for fast retrieval
-			Fallback.Token = InvalidToken;
-			Fallback.BulkDataSize = BulkDataSize;
+			Data.Fallback.Token = InvalidToken;
+			Data.Fallback.BulkDataSize = BulkDataSize;
 		}
 
 		const FString* Filename = nullptr;
@@ -808,7 +810,7 @@ void FBulkDataBase::Serialize(FArchive& Ar, UObject* Owner, int32 /*Index*/, boo
 		// If we are not using the FIoDispatcher and we have a filename then we need to make sure we can retrieve it later!
 		if (bUseIoDispatcher == false && Filename != nullptr)
 		{
-			Fallback.Token = FileTokenSystem::RegisterFileToken(Package->FileName, *Filename, BulkDataOffsetInFile);
+			Data.Fallback.Token = FileTokenSystem::RegisterFileToken(Package->FileName, *Filename, BulkDataOffsetInFile);
 		}
 
 		if (bShouldForceLoad)
@@ -884,7 +886,7 @@ void* FBulkDataBase::Realloc(int64 SizeInBytes)
 									// future then this check is a reminder that we need to handle the IoDispatcher 
 									// vs fallback case.
 
-	Fallback.BulkDataSize = SizeInBytes;
+	Data.Fallback.BulkDataSize = SizeInBytes;
 
 	return GetDataBufferForWrite();
 }
@@ -980,12 +982,12 @@ int64 FBulkDataBase::GetBulkDataSize() const
 {
 	if (IsUsingIODispatcher())
 	{
-		TIoStatusOr<uint64> Result = IoDispatcher->GetSizeForChunk(ChunkID);
+		TIoStatusOr<uint64> Result = IoDispatcher->GetSizeForChunk(Data.ChunkID);
 		return Result.ValueOrDie(); // TODO: Consider logging errors instead of relying on ::ValueOrDie
 	}
 	else
 	{
-		return Fallback.BulkDataSize;
+		return Data.Fallback.BulkDataSize;
 	}	
 }
 
@@ -997,8 +999,20 @@ bool FBulkDataBase::CanLoadFromDisk() const
 		return true;
 	}
 
+	// If the IoDispatcher is enabled then technically inlined data cannot load
+	// from disk, but it should act as thought it can up until the point that it
+	// actually tries to load. At which point it will give a log error.
+	//
+	// This is a temp hack while we work out if we want to allow additional inline
+	// loading at runtime (and pay the memory increase in the ToC) or prevent it entirely.
+	// TODO: Resolve this!
+	if (IsInlined() && IsIoDispatcherEnabled())
+	{
+		return true;
+	}
+
 	// If this BulkData has a fallback token then it can find it's filepath and load from disk
-	if (Fallback.Token != InvalidToken)
+	if (Data.Fallback.Token != InvalidToken)
 	{
 		return true;
 	}
@@ -1011,14 +1025,14 @@ bool FBulkDataBase::DoesExist() const
 #if ALLOW_OPTIONAL_DATA
 	if (!IsUsingIODispatcher())
 	{
-		FString Filename = FileTokenSystem::GetFilename(Fallback.Token);
+		FString Filename = FileTokenSystem::GetFilename(Data.Fallback.Token);
 		Filename = ConvertFilenameFromFlags(Filename);
 
 		return IFileManager::Get().FileExists(*Filename);
 	}
 	else
 	{
-		return IoDispatcher->DoesChunkExist(ChunkID);
+		return IoDispatcher->DoesChunkExist(Data.ChunkID);
 	}
 #else
 	return false;
@@ -1084,7 +1098,7 @@ IAsyncReadFileHandle* FBulkDataBase::OpenAsyncReadHandle() const
 {
 	if (IsUsingIODispatcher())
 	{
-		return new FAsyncReadChunkIdHandle(ChunkID);
+		return new FAsyncReadChunkIdHandle(Data.ChunkID);
 	}
 	else
 	{
@@ -1103,7 +1117,7 @@ IBulkDataIORequest* FBulkDataBase::CreateStreamingRequest(int64 OffsetInBulkData
 {
 	if (IsUsingIODispatcher())
 	{
-		FBulkDataIoDispatcherRequest* IoRequest = new FBulkDataIoDispatcherRequest(ChunkID, OffsetInBulkData, BytesToRead, CompleteCallback, UserSuppliedMemory);	
+		FBulkDataIoDispatcherRequest* IoRequest = new FBulkDataIoDispatcherRequest(Data.ChunkID, OffsetInBulkData, BytesToRead, CompleteCallback, UserSuppliedMemory);
 		IoRequest->StartAsyncWork();
 
 		return IoRequest;
@@ -1111,7 +1125,7 @@ IBulkDataIORequest* FBulkDataBase::CreateStreamingRequest(int64 OffsetInBulkData
 	else
 	{
 		const int64 BulkDataSize = GetBulkDataSize();
-		FileTokenSystem::Data FileData = FileTokenSystem::GetFileData(Fallback.Token);
+		FileTokenSystem::Data FileData = FileTokenSystem::GetFileData(Data.Fallback.Token);
 
 		check(FileData.PackageHeaderFilename.IsEmpty() == false);
 		const FString Filename = ConvertFilenameFromFlags(FileData.PackageHeaderFilename);
@@ -1156,7 +1170,7 @@ IBulkDataIORequest* FBulkDataBase::CreateStreamingRequestForRange(const BulkData
 		FBulkDataIoDispatcherRequest::ChunkIdArray ChunkIds;
 		for (const FBulkDataBase* BulkData : RangeArray)
 		{
-			ChunkIds.Push(BulkData->ChunkID);
+			ChunkIds.Push(BulkData->Data.ChunkID);
 		}
 
 		FBulkDataIoDispatcherRequest* IoRequest = new FBulkDataIoDispatcherRequest(ChunkIds, CompleteCallback);
@@ -1205,8 +1219,8 @@ void FBulkDataBase::RemoveBulkData()
 
 	if (!IsUsingIODispatcher())
 	{ 
-		FileTokenSystem::UnregisterFileToken(Fallback.Token);
-		Fallback.Token = InvalidToken;
+		FileTokenSystem::UnregisterFileToken(Data.Fallback.Token);
+		Data.Fallback.Token = InvalidToken;
 	}
 
 	BulkDataFlags = 0;
@@ -1217,7 +1231,7 @@ int64 FBulkDataBase::GetBulkDataOffsetInFile() const
 {
 	if (!IsUsingIODispatcher())
 	{
-		return FileTokenSystem::GetBulkDataOffset(Fallback.Token);
+		return FileTokenSystem::GetBulkDataOffset(Data.Fallback.Token);
 	}
 	else
 	{
@@ -1231,7 +1245,7 @@ FString FBulkDataBase::GetFilename() const
 {
 	if (!IsUsingIODispatcher())
 	{
-		FString Filename = FileTokenSystem::GetFilename(Fallback.Token);
+		FString Filename = FileTokenSystem::GetFilename(Data.Fallback.Token);
 		return ConvertFilenameFromFlags(Filename);
 	}
 	else
@@ -1269,7 +1283,7 @@ void FBulkDataBase::LoadDataDirectly(void** DstBuffer)
 	if (!IsIoDispatcherEnabled())
 	{
 		const int64 BulkDataSize = GetBulkDataSize();
-		FileTokenSystem::Data FileData = FileTokenSystem::GetFileData(Fallback.Token);
+		FileTokenSystem::Data FileData = FileTokenSystem::GetFileData(Data.Fallback.Token);
 
 		FString Filename; 
 		int64 Offset = FileData.BulkDataOffsetInFile;
@@ -1317,7 +1331,7 @@ void FBulkDataBase::LoadDataDirectly(void** DstBuffer)
 		Options.SetTargetVa(*DstBuffer);
 		
 		FIoBatch NewBatch = GetIoDispatcher()->NewBatch();
-		FIoRequest Request = NewBatch.Read(ChunkID, Options);
+		FIoRequest Request = NewBatch.Read(Data.ChunkID, Options);
 
 		NewBatch.Issue();
 		NewBatch.Wait(); // Blocking wait until all requests in the batch are done
@@ -1354,7 +1368,7 @@ void FBulkDataBase::ProcessDuplicateData(FArchive& Ar, const UPackage* Package,
 			InOutSizeOnDisk = NewSizeOnDisk;
 			InOutOffsetInFile = NewOffset;
 
-			ChunkID = OptionalChunkId;
+			Data.ChunkID = OptionalChunkId;
 		}
 	}
 	else
@@ -1369,8 +1383,8 @@ void FBulkDataBase::ProcessDuplicateData(FArchive& Ar, const UPackage* Package,
 			InOutOffsetInFile = NewOffset;
 
 			// Note we do not override Filename with OptionalDataFilename as we are supposed to store the original!
-			Fallback.Token = InvalidToken;
-			Fallback.BulkDataSize = InOutSizeOnDisk;
+			Data.Fallback.Token = InvalidToken;
+			Data.Fallback.BulkDataSize = InOutSizeOnDisk;
 		}
 	}
 #endif
@@ -1486,7 +1500,7 @@ FString FBulkDataBase::ConvertFilenameFromFlags(const FString& Filename) const
 	}
 }
 
-void FBulkDataBase::FBulkDataAllocation::Free(FBulkDataBase* Owner)
+void FBulkDataAllocation::Free(FBulkDataBase* Owner)
 {
 	if (!Owner->IsDataMemoryMapped())
 	{
@@ -1502,7 +1516,7 @@ void FBulkDataBase::FBulkDataAllocation::Free(FBulkDataBase* Owner)
 	}
 }
 
-void* FBulkDataBase::FBulkDataAllocation::AllocateData(FBulkDataBase* Owner, SIZE_T SizeInBytes)
+void* FBulkDataAllocation::AllocateData(FBulkDataBase* Owner, SIZE_T SizeInBytes)
 {
 	checkf(Allocation == nullptr, TEXT("Trying to allocate a BulkData object without freeing it first!"));
 
@@ -1511,14 +1525,14 @@ void* FBulkDataBase::FBulkDataAllocation::AllocateData(FBulkDataBase* Owner, SIZ
 	return Allocation;
 }
 
-void FBulkDataBase::FBulkDataAllocation::SetData(FBulkDataBase* Owner, void* Buffer)
+void FBulkDataAllocation::SetData(FBulkDataBase* Owner, void* Buffer)
 {
 	checkf(Allocation == nullptr, TEXT("Trying to assign a BulkData object without freeing it first!"));
 
 	Allocation = Buffer;
 }
 
-void FBulkDataBase::FBulkDataAllocation::SetMemoryMappedData(FBulkDataBase* Owner, IMappedFileHandle* MappedHandle, IMappedFileRegion* MappedRegion)
+void FBulkDataAllocation::SetMemoryMappedData(FBulkDataBase* Owner, IMappedFileHandle* MappedHandle, IMappedFileRegion* MappedRegion)
 {
 	checkf(Allocation == nullptr, TEXT("Trying to assign a BulkData object without freeing it first!"));
 	FOwnedBulkDataPtr* Ptr = new FOwnedBulkDataPtr(MappedHandle, MappedRegion);
@@ -1528,7 +1542,7 @@ void FBulkDataBase::FBulkDataAllocation::SetMemoryMappedData(FBulkDataBase* Owne
 	Allocation = Ptr;
 }
 
-void* FBulkDataBase::FBulkDataAllocation::GetAllocationForWrite(const FBulkDataBase* Owner) const
+void* FBulkDataAllocation::GetAllocationForWrite(const FBulkDataBase* Owner) const
 {
 	if (!Owner->IsDataMemoryMapped())
 	{
@@ -1540,7 +1554,7 @@ void* FBulkDataBase::FBulkDataAllocation::GetAllocationForWrite(const FBulkDataB
 	}
 }
 
-const void* FBulkDataBase::FBulkDataAllocation::GetAllocationReadOnly(const FBulkDataBase* Owner) const
+const void* FBulkDataAllocation::GetAllocationReadOnly(const FBulkDataBase* Owner) const
 {
 	if (!Owner->IsDataMemoryMapped())
 	{
@@ -1557,7 +1571,7 @@ const void* FBulkDataBase::FBulkDataAllocation::GetAllocationReadOnly(const FBul
 	}
 }
 
-FOwnedBulkDataPtr* FBulkDataBase::FBulkDataAllocation::StealFileMapping(FBulkDataBase* Owner)
+FOwnedBulkDataPtr* FBulkDataAllocation::StealFileMapping(FBulkDataBase* Owner)
 {
 	FOwnedBulkDataPtr* Ptr;
 	if (!Owner->IsDataMemoryMapped())
@@ -1574,7 +1588,7 @@ FOwnedBulkDataPtr* FBulkDataBase::FBulkDataAllocation::StealFileMapping(FBulkDat
 	return Ptr;
 }
 	
-void FBulkDataBase::FBulkDataAllocation::Swap(FBulkDataBase* Owner, void** DstBuffer)
+void FBulkDataAllocation::Swap(FBulkDataBase* Owner, void** DstBuffer)
 {
 	if (!Owner->IsDataMemoryMapped())
 	{
diff --git a/Engine/Source/Runtime/CoreUObject/Private/Serialization/UnversionedPropertySerialization.cpp b/Engine/Source/Runtime/CoreUObject/Private/Serialization/UnversionedPropertySerialization.cpp
--- a/Engine/Source/Runtime/CoreUObject/Private/Serialization/UnversionedPropertySerialization.cpp
+++ b/Engine/Source/Runtime/CoreUObject/Private/Serialization/UnversionedPropertySerialization.cpp
@@ -2,6 +2,8 @@
 
 #include "Serialization/UnversionedPropertySerialization.h"
 #include "Serialization/UnversionedPropertySerializationTest.h"
+#include "Interfaces/ITargetPlatform.h"
+#include "Misc/ScopeRWLock.h"
 #include "UObject/UnrealType.h"
 
 // Caches a property array per UStruct to avoid link-walking and touching FProperty data.
@@ -704,6 +706,45 @@ bool CanUseUnversionedPropertySerialization()
 	return bAllow;
 }
 
+static bool CanUseUnversionedPropertySerialization(FConfigFile& TargetIni)
+{
+	bool bAllow;
+	return TargetIni.GetBool(TEXT("Core.System"), TEXT("CanUseUnversionedPropertySerialization"), /* out */ bAllow) && bAllow;
+}
+
+static struct
+{
+	FRWLock Lock;
+	TMap<uint32, bool> PlatformValues;
+}
+GUPSIniValueCache;
+
+bool CanUseUnversionedPropertySerialization(const ITargetPlatform* Target)
+{
+	if (!Target)
+	{
+		// Use current platform settings
+		return CanUseUnversionedPropertySerialization();
+	}
+
+	int32 TargetID = Target->GetPlatformOrdinal();
+
+	FWriteScopeLock Scope(GUPSIniValueCache.Lock);
+
+	if (const bool* CachedValue = GUPSIniValueCache.PlatformValues.Find(TargetID))
+	{
+		return *CachedValue;
+	}
+
+	FConfigFile TargetIni;
+	FConfigCacheIni::LoadLocalIniFile(TargetIni, TEXT("Engine"), /* base INI */ true, *Target->IniPlatformName());
+	bool bTargetValue = CanUseUnversionedPropertySerialization(TargetIni);
+			
+	GUPSIniValueCache.PlatformValues.Add(TargetID, bTargetValue);
+
+	return bTargetValue;
+}
+
 void DestroyUnversionedSchema(const UStruct* Struct)
 {
 #if CACHE_UNVERSIONED_PROPERTY_SCHEMA
@@ -714,13 +755,13 @@ void DestroyUnversionedSchema(const UStruct* Struct)
 
 void SerializeUnversionedProperties(const UStruct* Struct, FStructuredArchive::FSlot Slot, uint8* Data, UStruct* DefaultsStruct, uint8* DefaultsData)
 {
-	check(CanUseUnversionedPropertySerialization());
-
 	FArchive& UnderlyingArchive = Slot.GetUnderlyingArchive();
 	FStructuredArchive::FRecord StructRecord = Slot.EnterRecord();
 
 	if (UnderlyingArchive.IsLoading())
 	{
+		check(CanUseUnversionedPropertySerialization());
+
 		FUnversionedHeader Header;
 		Header.Load(StructRecord.EnterStream(SA_FIELD_NAME(TEXT("Header"))));
 
diff --git a/Engine/Source/Runtime/CoreUObject/Private/Serialization/UnversionedPropertySerialization.h b/Engine/Source/Runtime/CoreUObject/Private/Serialization/UnversionedPropertySerialization.h
--- a/Engine/Source/Runtime/CoreUObject/Private/Serialization/UnversionedPropertySerialization.h
+++ b/Engine/Source/Runtime/CoreUObject/Private/Serialization/UnversionedPropertySerialization.h
@@ -6,9 +6,12 @@
 
 struct FUnversionedStructSchema;
 
-// Check if unversioned property serialization is configured to be used
+// Check if unversioned property serialization is configured to be used on current platform
 bool CanUseUnversionedPropertySerialization();
 
+// Check if unversioned property serialization is configured to be used on target platform
+bool CanUseUnversionedPropertySerialization(const ITargetPlatform* Target);
+
 // Serialize sparse unversioned properties for a particular struct
 void SerializeUnversionedProperties(const UStruct* Struct, FStructuredArchive::FSlot Slot, uint8* Data, UStruct* DefaultsStruct, uint8* Defaults);
 
diff --git a/Engine/Source/Runtime/CoreUObject/Private/Serialization/UnversionedPropertySerializationTest.cpp b/Engine/Source/Runtime/CoreUObject/Private/Serialization/UnversionedPropertySerializationTest.cpp
--- a/Engine/Source/Runtime/CoreUObject/Private/Serialization/UnversionedPropertySerializationTest.cpp
+++ b/Engine/Source/Runtime/CoreUObject/Private/Serialization/UnversionedPropertySerializationTest.cpp
@@ -239,7 +239,7 @@ struct FUnversionedPropertyTest : public FUnversionedPropertyTestInput
 	static const uint8* FindPairPtr(const FScriptMapHelper& Helper, const uint8* Key)
 	{
 		const FProperty* KeyProp = Helper.GetKeyProperty();
-		int32 Index = Helper.Map->FindPairIndex(Key, Helper.MapLayout,
+		int32 Index = Helper.HeapMap->FindPairIndex(Key, Helper.MapLayout,
 												[KeyProp](const void* Key) { return KeyProp->GetValueTypeHash(Key); },
 												[KeyProp](const void* A, const void* B) { return Equals(KeyProp, A, B); });
 		return Index >= 0 ? Helper.GetPairPtr(Index) : nullptr;
diff --git a/Engine/Source/Runtime/CoreUObject/Private/UObject/GarbageCollection.cpp b/Engine/Source/Runtime/CoreUObject/Private/UObject/GarbageCollection.cpp
--- a/Engine/Source/Runtime/CoreUObject/Private/UObject/GarbageCollection.cpp
+++ b/Engine/Source/Runtime/CoreUObject/Private/UObject/GarbageCollection.cpp
@@ -2364,9 +2364,11 @@ void FArrayProperty::EmitReferenceInfo(UClass& OwnerClass, int32 BaseOffset, TAr
 {
 	if (Inner->ContainsObjectReference(EncounteredStructProps))
 	{
+		bool bUsesFreezableAllocator = EnumHasAnyFlags(ArrayFlags, EArrayPropertyFlags::UsesMemoryImageAllocator);
+
 		if( Inner->IsA(FStructProperty::StaticClass()) )
 		{
-			OwnerClass.EmitObjectReference(BaseOffset + GetOffset_ForGC(), GetFName(), GCRT_ArrayStruct);
+			OwnerClass.EmitObjectReference(BaseOffset + GetOffset_ForGC(), GetFName(), bUsesFreezableAllocator ? GCRT_ArrayStructFreezable : GCRT_ArrayStruct);
 
 			OwnerClass.ReferenceTokenStream.EmitStride(Inner->ElementSize);
 			const uint32 SkipIndexIndex = OwnerClass.ReferenceTokenStream.EmitSkipIndexPlaceholder();
@@ -2376,11 +2378,11 @@ void FArrayProperty::EmitReferenceInfo(UClass& OwnerClass, int32 BaseOffset, TAr
 		}
 		else if( Inner->IsA(FObjectProperty::StaticClass()) )
 		{
-			OwnerClass.EmitObjectReference(BaseOffset + GetOffset_ForGC(), GetFName(), GCRT_ArrayObject);
+			OwnerClass.EmitObjectReference(BaseOffset + GetOffset_ForGC(), GetFName(), bUsesFreezableAllocator ? GCRT_ArrayObjectFreezable : GCRT_ArrayObject);
 		}
 		else if( Inner->IsA(FInterfaceProperty::StaticClass()) )
 		{
-			OwnerClass.EmitObjectReference(BaseOffset + GetOffset_ForGC(), GetFName(), GCRT_ArrayStruct);
+			OwnerClass.EmitObjectReference(BaseOffset + GetOffset_ForGC(), GetFName(), bUsesFreezableAllocator ? GCRT_ArrayStructFreezable : GCRT_ArrayStruct);
 
 			OwnerClass.ReferenceTokenStream.EmitStride(Inner->ElementSize);
 			const uint32 SkipIndexIndex = OwnerClass.ReferenceTokenStream.EmitSkipIndexPlaceholder();
@@ -2682,7 +2684,8 @@ void FGCReferenceTokenStream::Fixup(void (*AddReferencedObjectsPtr)(UObject*, cl
 		switch (Token.Type)
 		{
 		case GCRT_ArrayStruct:
-			{
+		case GCRT_ArrayStructFreezable:
+		{
 				// Skip stride and move to Skip Info
 				TokenIndex += 2;
 				const FGCSkipInfo SkipInfo = ReadSkipInfo(TokenIndex);
@@ -2754,6 +2757,7 @@ void FGCReferenceTokenStream::Fixup(void (*AddReferencedObjectsPtr)(UObject*, cl
 		case GCRT_None:
 		case GCRT_Object:		
 		case GCRT_ArrayObject:
+		case GCRT_ArrayObjectFreezable:
 		case GCRT_AddFieldPathReferencedObject:
 		case GCRT_ArrayAddFieldPathReferencedObject:
 		case GCRT_EndOfPointer:
diff --git a/Engine/Source/Runtime/CoreUObject/Private/UObject/LinkerSave.cpp b/Engine/Source/Runtime/CoreUObject/Private/UObject/LinkerSave.cpp
--- a/Engine/Source/Runtime/CoreUObject/Private/UObject/LinkerSave.cpp
+++ b/Engine/Source/Runtime/CoreUObject/Private/UObject/LinkerSave.cpp
@@ -4,7 +4,6 @@
 #include "HAL/FileManager.h"
 #include "Misc/ConfigCacheIni.h"
 #include "Serialization/LargeMemoryWriter.h"
-#include "Serialization/UnversionedPropertySerialization.h"
 #include "UObject/Package.h"
 #include "UObject/Class.h"
 #include "Templates/Casts.h"
@@ -63,9 +62,6 @@ FLinkerSave::FLinkerSave(UPackage* InParent, const TCHAR* InFilename, bool bForc
 		this->SetIsSaving(true);
 		this->SetIsPersistent(true);
 		ArForceByteSwapping		= bForceByteSwapping;
-		bool bUseUnversionedProperties = bInSaveUnversioned && CanUseUnversionedPropertySerialization();
-		SetUseUnversionedPropertySerialization(bUseUnversionedProperties);
-		Saver->SetUseUnversionedPropertySerialization(bUseUnversionedProperties);
 
 #if USE_STABLE_LOCALIZATION_KEYS
 		if (GIsEditor)
@@ -119,9 +115,6 @@ FLinkerSave::FLinkerSave(UPackage* InParent, FArchive *InSaver, bool bForceByteS
 		this->SetIsSaving(true);
 		this->SetIsPersistent(true);
 		ArForceByteSwapping = bForceByteSwapping;
-		bool bUseUnversionedProperties = bInSaveUnversioned && CanUseUnversionedPropertySerialization();
-		SetUseUnversionedPropertySerialization(bUseUnversionedProperties);
-		Saver->SetUseUnversionedPropertySerialization(bUseUnversionedProperties);
 
 #if USE_STABLE_LOCALIZATION_KEYS
 		if (GIsEditor)
@@ -170,9 +163,6 @@ FLinkerSave::FLinkerSave(UPackage* InParent, bool bForceByteSwapping, bool bInSa
 		this->SetIsSaving(true);
 		this->SetIsPersistent(true);
 		ArForceByteSwapping		= bForceByteSwapping;
-		bool bUseUnversionedProperties = bInSaveUnversioned && CanUseUnversionedPropertySerialization();
-		SetUseUnversionedPropertySerialization(bUseUnversionedProperties);
-		Saver->SetUseUnversionedPropertySerialization(bUseUnversionedProperties);
 
 #if USE_STABLE_LOCALIZATION_KEYS
 		if (GIsEditor)
diff --git a/Engine/Source/Runtime/CoreUObject/Private/UObject/Obj.cpp b/Engine/Source/Runtime/CoreUObject/Private/UObject/Obj.cpp
--- a/Engine/Source/Runtime/CoreUObject/Private/UObject/Obj.cpp
+++ b/Engine/Source/Runtime/CoreUObject/Private/UObject/Obj.cpp
@@ -43,6 +43,7 @@
 #include "UObject/LinkerLoad.h"
 #include "Misc/RedirectCollector.h"
 #include "UObject/GCScopeLock.h"
+#include "ProfilingDebugging/LoadTimeTracker.h"
 
 #include "Serialization/ArchiveUObjectFromStructuredArchive.h"
 #include "Serialization/ArchiveDescribeReference.h"
@@ -1241,6 +1242,8 @@ IMPLEMENT_FARCHIVE_SERIALIZER(UObject)
 
 void UObject::Serialize(FStructuredArchive::FRecord Record)
 {
+	SCOPED_LOADTIMER(UObject_Serialize);
+
 #if WITH_EDITOR
 	bool bReportSoftObjectPathRedirects = false;
 
@@ -2013,7 +2016,7 @@ void UObject::LoadConfig( UClass* ConfigClass/*=NULL*/, const TCHAR* InFilename/
 	{
 		TFieldIterator<FProperty> It1(Struct1);
 		TFieldIterator<FProperty> It2(Struct2);
-		for (;;)
+		for (;;++It1,++It2)
 		{
 			bool bAtEnd1 = !It1;
 			bool bAtEnd2 = !It2;
diff --git a/Engine/Source/Runtime/CoreUObject/Private/UObject/PropertyArray.cpp b/Engine/Source/Runtime/CoreUObject/Private/UObject/PropertyArray.cpp
--- a/Engine/Source/Runtime/CoreUObject/Private/UObject/PropertyArray.cpp
+++ b/Engine/Source/Runtime/CoreUObject/Private/UObject/PropertyArray.cpp
@@ -20,6 +20,7 @@ IMPLEMENT_FIELD(FArrayProperty)
 #if WITH_EDITORONLY_DATA
 FArrayProperty::FArrayProperty(UField* InField)
 	: FArrayProperty_Super(InField)
+	, ArrayFlags(EArrayPropertyFlags::None)
 {
 	UArrayProperty* SourceProperty = CastChecked<UArrayProperty>(InField);
 	Inner = CastField<FProperty>(SourceProperty->Inner->GetAssociatedFField());
@@ -62,7 +63,9 @@ void FArrayProperty::LinkInternal(FArchive& Ar)
 	//}
 	//Ar.Preload(Inner);
 	Inner->Link(Ar);
-	Super::LinkInternal(Ar);
+
+	SetElementSize();
+	PropertyFlags |= CPF_HasGetValueTypeHash;
 }
 bool FArrayProperty::Identical( const void* A, const void* B, uint32 PortFlags ) const
 {
@@ -628,7 +631,7 @@ void FArrayProperty::DestroyValueInternal( void* Dest ) const
 	ArrayHelper.EmptyValues();
 
 	//@todo UE4 potential double destroy later from this...would be ok for a script array, but still
-	((FScriptArray*)Dest)->~FScriptArray();
+	ArrayHelper.DestroyContainer_Unsafe();
 }
 bool FArrayProperty::PassCPPArgsByRef() const
 {
diff --git a/Engine/Source/Runtime/CoreUObject/Private/UObject/PropertyMap.cpp b/Engine/Source/Runtime/CoreUObject/Private/UObject/PropertyMap.cpp
--- a/Engine/Source/Runtime/CoreUObject/Private/UObject/PropertyMap.cpp
+++ b/Engine/Source/Runtime/CoreUObject/Private/UObject/PropertyMap.cpp
@@ -183,25 +183,30 @@ namespace UE4MapProperty_Private
 
 IMPLEMENT_FIELD(FMapProperty)
 
-FMapProperty::FMapProperty(FFieldVariant InOwner, const FName& InName, EObjectFlags InObjectFlags)
+FMapProperty::FMapProperty(FFieldVariant InOwner, const FName& InName, EObjectFlags InObjectFlags, EMapPropertyFlags InMapFlags)
 	: FMapProperty_Super(InOwner, InName, InObjectFlags)
 {
 	// These are expected to be set post-construction by AddCppProperty
 	KeyProp = nullptr;
 	ValueProp = nullptr;
+
+	MapFlags = InMapFlags;
 }
 
-FMapProperty::FMapProperty(FFieldVariant InOwner, const FName& InName, EObjectFlags InObjectFlags, int32 InOffset, EPropertyFlags InFlags)
-: FMapProperty_Super(InOwner, InName, InObjectFlags, InOffset, InFlags)
+FMapProperty::FMapProperty(FFieldVariant InOwner, const FName& InName, EObjectFlags InObjectFlags, int32 InOffset, EPropertyFlags InFlags, EMapPropertyFlags InMapFlags)
+	: FMapProperty_Super(InOwner, InName, InObjectFlags, InOffset, InFlags)
 {
 	// These are expected to be set post-construction by AddCppProperty
 	KeyProp   = nullptr;
 	ValueProp = nullptr;
+
+	MapFlags = InMapFlags;
 }
 
 #if WITH_EDITORONLY_DATA
 FMapProperty::FMapProperty(UField* InField)
 	: FMapProperty_Super(InField)
+	, MapFlags(EMapPropertyFlags::None)
 {
 	UMapProperty* SourceProperty = CastChecked<UMapProperty>(InField);
 	MapLayout = SourceProperty->MapLayout;
@@ -1164,9 +1169,12 @@ EConvertFromTypeResult FMapProperty::ConvertFromType(const FPropertyTag& Tag, FS
 
 void FScriptMapHelper::Rehash()
 {
-	// Moved out-of-line to maybe fix a weird link error
-	Map->Rehash(MapLayout, [=](const void* Src) {
-		return KeyProp->GetValueTypeHash(Src);
+	WithScriptMap([this](auto* Map)
+	{
+		// Moved out-of-line to maybe fix a weird link error
+		Map->Rehash(MapLayout, [=](const void* Src) {
+			return KeyProp->GetValueTypeHash(Src);
+		});
 	});
 }
 
diff --git a/Engine/Source/Runtime/CoreUObject/Private/UObject/SavePackage.cpp b/Engine/Source/Runtime/CoreUObject/Private/UObject/SavePackage.cpp
--- a/Engine/Source/Runtime/CoreUObject/Private/UObject/SavePackage.cpp
+++ b/Engine/Source/Runtime/CoreUObject/Private/UObject/SavePackage.cpp
@@ -58,6 +58,7 @@
 #include "Serialization/Formatters/BinaryArchiveFormatter.h"
 #include "Serialization/Formatters/JsonArchiveOutputFormatter.h"
 #include "Serialization/ArchiveUObjectFromStructuredArchive.h"
+#include "Serialization/UnversionedPropertySerialization.h"
 #include "UObject/AsyncWorkSequence.h"
 #include "Serialization/BulkDataManifest.h"
 #include "Misc/ScopeExit.h"
@@ -3765,6 +3766,10 @@ FSavePackageResultStruct UPackage::Save(UPackage* InOuter, UObject* Base, EObjec
 				Linker->SetFilterEditorOnly( FilterEditorOnly );
 				Linker->SetCookingTarget(TargetPlatform);
 
+				bool bUseUnversionedProperties = bSaveUnversioned && CanUseUnversionedPropertySerialization(TargetPlatform);
+				Linker->SetUseUnversionedPropertySerialization(bUseUnversionedProperties);
+				Linker->Saver->SetUseUnversionedPropertySerialization(bUseUnversionedProperties);
+
 				// Make sure the package has the same version as the linker
 				InOuter->LinkerPackageVersion = Linker->UE4Ver();
 				InOuter->LinkerLicenseeVersion = Linker->LicenseeUE4Ver();
diff --git a/Engine/Source/Runtime/CoreUObject/Private/UObject/UObjectArray.cpp b/Engine/Source/Runtime/CoreUObject/Private/UObject/UObjectArray.cpp
--- a/Engine/Source/Runtime/CoreUObject/Private/UObject/UObjectArray.cpp
+++ b/Engine/Source/Runtime/CoreUObject/Private/UObject/UObjectArray.cpp
@@ -14,6 +14,69 @@ DEFINE_LOG_CATEGORY_STATIC(LogUObjectArray, Log, All);
 
 FUObjectClusterContainer GUObjectClusters;
 
+#if STATS || ENABLE_STATNAMEDEVENTS_UOBJECT
+void FUObjectItem::CreateStatID() const
+{
+	// @todo can we put this in a header?
+//		SCOPE_CYCLE_COUNTER(STAT_CreateStatID);
+
+	FString LongName;
+	LongName.Reserve(255);
+	TArray<UObjectBase const*, TInlineAllocator<24>> ClassChain;
+
+	// Build class hierarchy
+	UObjectBase const* Target = Object;
+	do
+	{
+		ClassChain.Add(Target);
+		Target = Target->GetOuter();
+	} while (Target);
+
+	// Start with class name
+	if (Object->GetClass())
+	{
+		Object->GetClass()->GetFName().GetDisplayNameEntry()->AppendNameToString(LongName);
+	}
+
+	// Now process from parent -> child so we can append strings more efficiently.
+	bool bFirstEntry = true;
+	for (int32 i = ClassChain.Num() - 1; i >= 0; i--)
+	{
+		Target = ClassChain[i];
+		const FNameEntry* NameEntry = Target->GetFName().GetDisplayNameEntry();
+		if (bFirstEntry)
+		{
+			NameEntry->AppendNameToPathString(LongName);
+		}
+		else
+		{
+			if (!LongName.IsEmpty())
+			{
+				LongName += TEXT(".");
+			}
+			NameEntry->AppendNameToString(LongName);
+		}
+		bFirstEntry = false;
+	}
+
+#if STATS
+	StatID = FDynamicStats::CreateStatId<FStatGroup_STATGROUP_UObjects>(LongName);
+#else // ENABLE_STATNAMEDEVENTS
+	const auto& ConversionData = StringCast<PROFILER_CHAR>(*LongName);
+	const int32 NumStorageChars = (ConversionData.Length() + 1);	//length doesn't include null terminator
+
+	PROFILER_CHAR* StoragePtr = new PROFILER_CHAR[NumStorageChars];
+	FMemory::Memcpy(StoragePtr, ConversionData.Get(), NumStorageChars * sizeof(PROFILER_CHAR));
+
+	if (FPlatformAtomics::InterlockedCompareExchangePointer((void**)&StatIDStringStorage, StoragePtr, nullptr) != nullptr)
+	{
+		delete[] StoragePtr;
+	}
+
+	StatID = TStatId(StatIDStringStorage);
+#endif
+}
+#endif
 FUObjectArray::FUObjectArray()
 : ObjFirstGCIndex(0)
 , ObjLastNonGCIndex(INDEX_NONE)
diff --git a/Engine/Source/Runtime/CoreUObject/Private/UObject/UObjectBase.cpp b/Engine/Source/Runtime/CoreUObject/Private/UObject/UObjectBase.cpp
--- a/Engine/Source/Runtime/CoreUObject/Private/UObject/UObjectBase.cpp
+++ b/Engine/Source/Runtime/CoreUObject/Private/UObject/UObjectBase.cpp
@@ -95,9 +95,6 @@ UObjectBase::UObjectBase( EObjectFlags InFlags )
 ,	InternalIndex		(INDEX_NONE)
 ,	ClassPrivate		(nullptr)
 ,	OuterPrivate		(nullptr)
-#if ENABLE_STATNAMEDEVENTS_UOBJECT
-, StatIDStringStorage(nullptr)
-#endif
 {}
 
 /**
@@ -113,9 +110,6 @@ UObjectBase::UObjectBase(UClass* InClass, EObjectFlags InFlags, EInternalObjectF
 ,	InternalIndex		(INDEX_NONE)
 ,	ClassPrivate		(InClass)
 ,	OuterPrivate		(InOuter)
-#if ENABLE_STATNAMEDEVENTS_UOBJECT
-, StatIDStringStorage(nullptr)
-#endif
 {
 	check(ClassPrivate);
 	// Add to global table.
@@ -136,76 +130,9 @@ UObjectBase::~UObjectBase()
 		check(GetFName() == NAME_None);
 		GUObjectArray.FreeUObjectIndex(this);
 	}
-
-#if ENABLE_STATNAMEDEVENTS_UOBJECT
-	delete[] StatIDStringStorage;
-	StatIDStringStorage = nullptr;
-#endif
 }
 
-#if STATS || ENABLE_STATNAMEDEVENTS_UOBJECT
 
-void UObjectBase::CreateStatID() const
-{
-	SCOPE_CYCLE_COUNTER(STAT_CreateStatID);
-
-	FString LongName;
-	LongName.Reserve(255);
-	TArray<UObjectBase const*, TInlineAllocator<24>> ClassChain;
-	
-	// Build class hierarchy
-	UObjectBase const* Target = this;
-	do
-	{
-		ClassChain.Add(Target);
-		Target = Target->GetOuter();
-	} while(Target);
-
-	// Start with class name
-	if (GetClass())
-	{
-		GetClass()->GetFName().GetDisplayNameEntry()->AppendNameToString(LongName);
-	}
-
-	// Now process from parent -> child so we can append strings more efficiently.
-	bool bFirstEntry = true;
-	for (int32 i = ClassChain.Num()-1; i >= 0; i--)
-	{
-		Target = ClassChain[i];
-		const FNameEntry* NameEntry = Target->GetFName().GetDisplayNameEntry();
-		if (bFirstEntry)
-		{
-			NameEntry->AppendNameToPathString(LongName);
-		}
-		else
-		{
-			if (!LongName.IsEmpty())
-			{
-				LongName += TEXT(".");
-			}
-			NameEntry->AppendNameToString(LongName);
-		}			
-		bFirstEntry = false;
-	}
-
-#if STATS
-	StatID = FDynamicStats::CreateStatId<FStatGroup_STATGROUP_UObjects>( LongName );
-#else // ENABLE_STATNAMEDEVENTS
-	const auto& ConversionData = StringCast<PROFILER_CHAR>(*LongName);
-	const int32 NumStorageChars = (ConversionData.Length() + 1);	//length doesn't include null terminator
-
-	PROFILER_CHAR* StoragePtr = new PROFILER_CHAR[NumStorageChars];
-	FMemory::Memcpy(StoragePtr, ConversionData.Get(), NumStorageChars * sizeof(PROFILER_CHAR));
-	
-	if (FPlatformAtomics::InterlockedCompareExchangePointer((void**)&StatIDStringStorage, StoragePtr, nullptr) != nullptr)
-	{
-		delete[] StoragePtr;
-	}
-	
-	StatID = TStatId(StatIDStringStorage);
-#endif
-}
-#endif
 
 
 /**
@@ -277,7 +204,7 @@ void UObjectBase::AddObject(FName InName, EInternalObjectFlags InSetInternalFlag
  */
 void UObjectBase::LowLevelRename(FName NewName,UObject *NewOuter)
 {
-	STAT(StatID = TStatId();) // reset the stat id since this thing now has a different name
+	STAT(((UObject*)this)->ResetStatID()); // reset the stat id since this thing now has a different name
 	UnhashObject(this);
 	check(InternalIndex >= 0);
 	NamePrivate = NewName;
@@ -290,7 +217,7 @@ void UObjectBase::LowLevelRename(FName NewName,UObject *NewOuter)
 
 void UObjectBase::SetClass(UClass* NewClass)
 {
-	STAT(StatID = TStatId();) // reset the stat id since this thing now has a different name
+	STAT(((UObject*)this)->ResetStatID()); // reset the stat id since this thing now has a different name
 
 	UnhashObject(this);
 #if USE_UBER_GRAPH_PERSISTENT_FRAME
diff --git a/Engine/Source/Runtime/CoreUObject/Private/UObject/UObjectGlobals.cpp b/Engine/Source/Runtime/CoreUObject/Private/UObject/UObjectGlobals.cpp
--- a/Engine/Source/Runtime/CoreUObject/Private/UObject/UObjectGlobals.cpp
+++ b/Engine/Source/Runtime/CoreUObject/Private/UObject/UObjectGlobals.cpp
@@ -2448,7 +2448,7 @@ UObject* StaticAllocateObject
 
 	if (!bSubObject)
 	{
-		FMemory::Memzero((void *)Obj, TotalSize);
+			FMemory::Memzero((void *)Obj, TotalSize);
 		new ((void *)Obj) UObjectBase(const_cast<UClass*>(InClass), InFlags|RF_NeedInitialization, InternalSetFlags, InOuter, InName);
 	}
 	else
@@ -3128,9 +3128,9 @@ UObject* StaticConstructObject_Internal
 	UObject*		InOuter								/*=GetTransientPackage()*/,
 	FName			InName								/*=NAME_None*/,
 	EObjectFlags	InFlags								/*=0*/,
-	EInternalObjectFlags InternalSetFlags /*=0*/,
+	EInternalObjectFlags InternalSetFlags				/*=0*/,
 	UObject*		InTemplate							/*=NULL*/,
-	bool bCopyTransientsFromClassDefaults	/*=false*/,
+	bool bCopyTransientsFromClassDefaults				/*=false*/,
 	FObjectInstancingGraph* InInstanceGraph				/*=NULL*/,
 	bool bAssumeTemplateIsArchetype	/*=false*/
 )
@@ -3167,7 +3167,7 @@ UObject* StaticConstructObject_Internal
 	// Don't call the constructor on recycled subobjects, they haven't been destroyed.
 	if (!bRecycledSubobject)
 	{		
-		FScopeCycleCounterUObject ConstructorScope(InClass, GET_STATID(STAT_ConstructObject));
+		STAT(FScopeCycleCounterUObject ConstructorScope(InClass, GET_STATID(STAT_ConstructObject)));
 		(*InClass->ClassConstructor)( FObjectInitializer(Result, InTemplate, bCopyTransientsFromClassDefaults, true, InInstanceGraph) );
 	}
 	
@@ -4212,7 +4212,7 @@ namespace UE4CodeGen_Private
 			case EPropertyGenFlags::Array:
 			{
 				const FArrayPropertyParams* Prop = (const FArrayPropertyParams*)PropBase;
-				NewProp = new FArrayProperty(Outer, UTF8_TO_TCHAR(Prop->NameUTF8), Prop->ObjectFlags, Prop->Offset, Prop->PropertyFlags);
+				NewProp = new FArrayProperty(Outer, UTF8_TO_TCHAR(Prop->NameUTF8), Prop->ObjectFlags, Prop->Offset, Prop->PropertyFlags, Prop->ArrayFlags);
 
 				// Next property is the array inner
 				ReadMore = 1;
@@ -4227,7 +4227,7 @@ namespace UE4CodeGen_Private
 			case EPropertyGenFlags::Map:
 			{
 				const FMapPropertyParams* Prop = (const FMapPropertyParams*)PropBase;
-				NewProp = new FMapProperty(Outer, UTF8_TO_TCHAR(Prop->NameUTF8), Prop->ObjectFlags, Prop->Offset, Prop->PropertyFlags);
+				NewProp = new FMapProperty(Outer, UTF8_TO_TCHAR(Prop->NameUTF8), Prop->ObjectFlags, Prop->Offset, Prop->PropertyFlags, Prop->MapFlags);
 
 				// Next two properties are the map key and value inners
 				ReadMore = 2;
diff --git a/Engine/Source/Runtime/CoreUObject/Public/Serialization/BulkData2.h b/Engine/Source/Runtime/CoreUObject/Public/Serialization/BulkData2.h
--- a/Engine/Source/Runtime/CoreUObject/Public/Serialization/BulkData2.h
+++ b/Engine/Source/Runtime/CoreUObject/Public/Serialization/BulkData2.h
@@ -31,6 +31,63 @@ public:
 	virtual void Cancel() = 0;
 };
 
+using FileToken = int32;
+struct FBulkDataOrId
+{
+	union
+	{
+		// Inline data or fallback path
+		struct
+		{
+			uint64 BulkDataSize;
+			FileToken Token;
+
+		} Fallback;
+
+		// For IODispatcher
+		FIoChunkId ChunkID;
+	}; // Note that the union will end up being 16 bytes with padding
+};
+DECLARE_INTRINSIC_TYPE_LAYOUT(FBulkDataOrId);
+
+/**
+ * This is a wrapper for the BulkData memory allocation so we can use a single pointer to either
+ * reference a straight memory allocation or in the case that the BulkData object represents a 
+ * memory mapped file region, a FOwnedBulkDataPtr.
+ * This makes the code more complex but it means that we do not pay any additional memory cost when
+ * memory mapping isn't being used at a small cpu cost. However the number of BulkData object usually
+ * means that the memory saving is worth it compared to how infrequently the memory accessors are 
+ * actually called.
+ *
+ * Note: We use a flag set in the owning BulkData object to tell us what storage type we are using 
+ * so all accessors require that a pointer to the parent object be passed in.
+ */
+class FBulkDataBase;
+class FBulkDataAllocation
+{
+public:
+	// Misc
+	bool IsLoaded() const { return Allocation != nullptr; }
+	void Free(FBulkDataBase* Owner);
+
+	// Set as a raw buffer
+	void* AllocateData(FBulkDataBase* Owner, SIZE_T SizeInBytes); //DataBuffer = FMemory::Realloc(DataBuffer, SizeInBytes, DEFAULT_ALIGNMENT);
+	void SetData(FBulkDataBase* Owner, void* Buffer);
+
+	// Set as memory mapped
+	void SetMemoryMappedData(FBulkDataBase* Owner, IMappedFileHandle* MappedHandle, IMappedFileRegion* MappedRegion);
+
+	// Getters		
+	void* GetAllocationForWrite(const FBulkDataBase* Owner) const;
+	const void* GetAllocationReadOnly(const FBulkDataBase* Owner) const;
+
+	FOwnedBulkDataPtr* StealFileMapping(FBulkDataBase* Owner);
+	void Swap(FBulkDataBase* Owner, void** DstBuffer);
+private:
+	void* Allocation = nullptr; // Will either be the data allocation or a FOwnedBulkDataPtr if memory mapped
+};
+DECLARE_INTRINSIC_TYPE_LAYOUT(FBulkDataAllocation);
+
 /**
  * Callback to use when making streaming requests
  */
@@ -41,13 +98,14 @@ typedef TFunction<void(bool bWasCancelled, IBulkDataIORequest*)> FBulkDataIORequ
  */
 class COREUOBJECT_API FBulkDataBase
 {
+	DECLARE_TYPE_LAYOUT(FBulkDataBase, NonVirtual);
+
 public:
 	using BulkDataRangeArray = TArray<FBulkDataBase*, TInlineAllocator<8>>;
 
 	static void				SetIoDispatcher(FIoDispatcher* InIoDispatcher) { IoDispatcher = InIoDispatcher; }
 	static FIoDispatcher* GetIoDispatcher() { return IoDispatcher; }
 public:
-	using FileToken = int32;
 	static constexpr FileToken InvalidToken = INDEX_NONE;
 
 	FBulkDataBase(const FBulkDataBase& Other) { *this = Other; }
@@ -55,9 +113,9 @@ public:
 	FBulkDataBase& operator=(const FBulkDataBase& Other);
 
 	FBulkDataBase()
-	{
-		Fallback.BulkDataSize = 0;
-		Fallback.Token = InvalidToken;
+	{ 
+		Data.Fallback.BulkDataSize = 0;
+		Data.Fallback.Token = InvalidToken;
 	}
 	~FBulkDataBase();
 
@@ -135,6 +193,7 @@ public:
 	FOwnedBulkDataPtr* StealFileMapping();
 
 private:
+	friend FBulkDataAllocation;
 
 	void SetRuntimeBulkDataFlags(uint32 BulkDataFlagsToSet);
 	void ClearRuntimeBulkDataFlags(uint32 BulkDataFlagsToClear);
@@ -163,61 +222,13 @@ private:
 	FString ConvertFilenameFromFlags(const FString& Filename) const;
 
 private:
-	/**
-	 * This is a wrapper for the BulkData memory allocation so we can use a single pointer to either
-	 * reference a straight memory allocation or in the case that the BulkData object represents a 
-	 * memory mapped file region, a FOwnedBulkDataPtr.
-	 * This makes the code more complex but it means that we do not pay any additional memory cost when
-	 * memory mapping isn't being used at a small cpu cost. However the number of BulkData object usually
-	 * means that the memory saving is worth it compared to how infrequently the memory accessors are 
-	 * actually called.
-	 *
-	 * Note: We use a flag set in the owning BulkData object to tell us what storage type we are using 
-	 * so all accessors require that a pointer to the parent object be passed in.
-	 */
-	class FBulkDataAllocation
-	{
-	public:
-		// Misc
-		bool IsLoaded() const { return Allocation != nullptr; }
-		void Free(FBulkDataBase* Owner);
-
-		// Set as a raw buffer
-		void* AllocateData(FBulkDataBase* Owner, SIZE_T SizeInBytes); //DataBuffer = FMemory::Realloc(DataBuffer, SizeInBytes, DEFAULT_ALIGNMENT);
-		void SetData(FBulkDataBase* Owner, void* Buffer);
-
-		// Set as memory mapped
-		void SetMemoryMappedData(FBulkDataBase* Owner, IMappedFileHandle* MappedHandle, IMappedFileRegion* MappedRegion);
-
-		// Getters		
-		void* GetAllocationForWrite(const FBulkDataBase* Owner) const;
-		const void* GetAllocationReadOnly(const FBulkDataBase* Owner) const;
-	
-		FOwnedBulkDataPtr* StealFileMapping(FBulkDataBase* Owner);
-		void Swap(FBulkDataBase* Owner, void** DstBuffer);
-	private:
-		void* Allocation = nullptr; // Will either be the data allocation or a FOwnedBulkDataPtr if memory mapped
-	};
 
 	static FIoDispatcher* IoDispatcher;
 
-	union
-	{
-		// Inline data or fallback path
-		struct
-		{
-			uint64 BulkDataSize;
-			FileToken Token;
-			
-		} Fallback;
-
-		// For IODispatcher
-		FIoChunkId ChunkID;	
-	}; // Note that the union will end up being 16 bytes with padding
-	
-	FBulkDataAllocation DataAllocation;
-	uint32 BulkDataFlags = 0;
-	mutable uint8 LockStatus = 0; // Mutable so that the read only lock can be const
+	LAYOUT_FIELD(FBulkDataOrId, Data);
+	LAYOUT_FIELD(FBulkDataAllocation, DataAllocation);
+	LAYOUT_FIELD_INITIALIZED(uint32, BulkDataFlags, 0);
+	LAYOUT_MUTABLE_FIELD_INITIALIZED(uint8, LockStatus, 0); // Mutable so that the read only lock can be const
 };
 
 /**
diff --git a/Engine/Source/Runtime/CoreUObject/Public/UObject/FastReferenceCollector.h b/Engine/Source/Runtime/CoreUObject/Public/UObject/FastReferenceCollector.h
--- a/Engine/Source/Runtime/CoreUObject/Public/UObject/FastReferenceCollector.h
+++ b/Engine/Source/Runtime/CoreUObject/Public/UObject/FastReferenceCollector.h
@@ -721,6 +721,17 @@ private:
 						}
 					}
 					break;
+					case GCRT_ArrayObjectFreezable:
+					{
+						// We're dealing with an array of object references.
+						TArray<UObject*, FMemoryImageAllocator>& ObjectArray = *((TArray<UObject*, FMemoryImageAllocator>*)(StackEntryData + ReferenceInfo.Offset));
+						TokenReturnCount = ReferenceInfo.ReturnCount;
+						for (int32 ObjectIndex = 0, ObjectNum = ObjectArray.Num(); ObjectIndex < ObjectNum; ++ObjectIndex)
+						{
+							ReferenceProcessor.HandleTokenStreamObjectReference(NewObjectsToSerialize, CurrentObject, ObjectArray[ObjectIndex], ReferenceTokenStreamIndex, true);
+						}
+					}
+					break;
 					case GCRT_ArrayStruct:
 					{
 						// We're dealing with a dynamic array of structs.
@@ -748,6 +759,33 @@ private:
 						}
 					}
 					break;
+					case GCRT_ArrayStructFreezable:
+					{
+						// We're dealing with a dynamic array of structs.
+						const FFreezableScriptArray& Array = *((FFreezableScriptArray*)(StackEntryData + ReferenceInfo.Offset));
+						StackEntry++;
+						StackEntryData = (uint8*)Array.GetData();
+						StackEntry->Data = StackEntryData;
+						StackEntry->Stride = TokenStream->ReadStride(TokenStreamIndex);
+						StackEntry->Count = Array.Num();
+
+						const FGCSkipInfo SkipInfo = TokenStream->ReadSkipInfo(TokenStreamIndex);
+						StackEntry->LoopStartIndex = TokenStreamIndex;
+
+						if (StackEntry->Count == 0)
+						{
+							// Skip empty array by jumping to skip index and set return count to the one about to be read in.
+							TokenStreamIndex = SkipInfo.SkipIndex;
+							TokenReturnCount = TokenStream->GetSkipReturnCount(SkipInfo);
+						}
+						else
+						{
+							// Loop again.
+							check(StackEntry->Data);
+							TokenReturnCount = 0;
+						}
+					}
+					break;
 					case GCRT_PersistentObject:
 					{
 						// We're dealing with an object reference (this code should be identical to GCRT_NoopPersistentObject if !bIgnoreNoopTokens)
diff --git a/Engine/Source/Runtime/CoreUObject/Public/UObject/GarbageCollection.h b/Engine/Source/Runtime/CoreUObject/Public/UObject/GarbageCollection.h
--- a/Engine/Source/Runtime/CoreUObject/Public/UObject/GarbageCollection.h
+++ b/Engine/Source/Runtime/CoreUObject/Public/UObject/GarbageCollection.h
@@ -55,6 +55,8 @@ enum EGCReferenceType
 	GCRT_EndOfStream,
 	GCRT_NoopPersistentObject,
 	GCRT_NoopClass,
+	GCRT_ArrayObjectFreezable,
+	GCRT_ArrayStructFreezable,
 };
 
 /** 
diff --git a/Engine/Source/Runtime/CoreUObject/Public/UObject/ObjectMacros.h b/Engine/Source/Runtime/CoreUObject/Public/UObject/ObjectMacros.h
--- a/Engine/Source/Runtime/CoreUObject/Public/UObject/ObjectMacros.h
+++ b/Engine/Source/Runtime/CoreUObject/Public/UObject/ObjectMacros.h
@@ -446,7 +446,29 @@ enum EPropertyFlags : uint64
 
 ENUM_CLASS_FLAGS(EPropertyFlags)
 
-/** 
+/**
+ * Extra flags for array properties.
+ */
+enum class EArrayPropertyFlags
+{
+	None,
+	UsesMemoryImageAllocator
+};
+
+ENUM_CLASS_FLAGS(EArrayPropertyFlags)
+
+/**
+ * Extra flags for map properties.
+ */
+enum class EMapPropertyFlags
+{
+	None,
+	UsesMemoryImageAllocator
+};
+
+ENUM_CLASS_FLAGS(EMapPropertyFlags)
+
+/**
  * Flags describing an object instance
  */
 enum EObjectFlags
diff --git a/Engine/Source/Runtime/CoreUObject/Public/UObject/UObjectArray.h b/Engine/Source/Runtime/CoreUObject/Public/UObject/UObjectArray.h
--- a/Engine/Source/Runtime/CoreUObject/Public/UObject/UObjectArray.h
+++ b/Engine/Source/Runtime/CoreUObject/Public/UObject/UObjectArray.h
@@ -34,13 +34,32 @@ struct FUObjectItem
 	// Weak Object Pointer Serial number associated with the object
 	int32 SerialNumber;
 
+#if STATS || ENABLE_STATNAMEDEVENTS_UOBJECT
+	/** Stat id of this object, 0 if nobody asked for it yet */
+	mutable TStatId StatID;
+
+#if ENABLE_STATNAMEDEVENTS_UOBJECT
+	mutable PROFILER_CHAR* StatIDStringStorage;
+#endif
+#endif // STATS || ENABLE_STATNAMEDEVENTS
+
 	FUObjectItem()
 		: Object(nullptr)
 		, Flags(0)
 		, ClusterRootIndex(0)
 		, SerialNumber(0)
+#if ENABLE_STATNAMEDEVENTS_UOBJECT
+		, StatIDStringStorage(nullptr)
+#endif
 	{
 	}
+	~FUObjectItem()
+	{
+#if ENABLE_STATNAMEDEVENTS_UOBJECT
+		delete[] StatIDStringStorage;
+		StatIDStringStorage = nullptr;
+#endif
+	}
 
 	FORCEINLINE void SetOwnerIndex(int32 OwnerIndex)
 	{
@@ -190,6 +209,10 @@ struct FUObjectItem
 		ClusterRootIndex = 0;
 		SerialNumber = 0;
 	}
+
+#if STATS || ENABLE_STATNAMEDEVENTS_UOBJECT
+	COREUOBJECT_API void CreateStatID() const;
+#endif
 };
 
 /**
diff --git a/Engine/Source/Runtime/CoreUObject/Public/UObject/UObjectBase.h b/Engine/Source/Runtime/CoreUObject/Public/UObject/UObjectBase.h
--- a/Engine/Source/Runtime/CoreUObject/Public/UObject/UObjectBase.h
+++ b/Engine/Source/Runtime/CoreUObject/Public/UObject/UObjectBase.h
@@ -34,9 +34,6 @@ class COREUOBJECT_API UObjectBase
 protected:
 	UObjectBase() :
 		 NamePrivate(NoInit)  // screwy, but the name was already set and we don't want to set it again
-#if ENABLE_STATNAMEDEVENTS_UOBJECT
-		, StatIDStringStorage(nullptr)
-#endif
 	{
 	}
 
@@ -141,40 +138,6 @@ public:
 		return NamePrivate;
 	}
 
-	/** 
-	 * Returns the stat ID of the object, used for profiling. This will create a stat ID if needed.
-	 *
-	 * @param bForDeferred If true, a stat ID will be created even if a group is disabled
-	 */
-	FORCEINLINE TStatId GetStatID(bool bForDeferredUse = false) const
-	{
-#if STATS
-		// this is done to avoid even registering stats for a disabled group (unless we plan on using it later)
-		if (bForDeferredUse || FThreadStats::IsCollectingData(GET_STATID(STAT_UObjectsStatGroupTester)))
-		{
-			if (!StatID.IsValidStat())
-			{
-				CreateStatID();
-			}
-			return StatID;
-		}
-#elif ENABLE_STATNAMEDEVENTS_UOBJECT
-		if (!StatID.IsValidStat() && (bForDeferredUse || GCycleStatsShouldEmitNamedEvents))
-		{
-			CreateStatID();
-		}
-		return StatID;
-#endif // STATS
-		return TStatId(); // not doing stats at the moment, or ever
-	}
-
-	
-private:
-#if STATS || ENABLE_STATNAMEDEVENTS_UOBJECT
-	/** Creates a stat ID for this object */
-	void CreateStatID() const;
-#endif
-
 protected:
 	/**
 	 * Set the object flags directly
@@ -249,17 +212,6 @@ private:
 
 	/** Object this object resides in. */
 	UObject*						OuterPrivate;
-
-
-#if STATS || ENABLE_STATNAMEDEVENTS_UOBJECT
-	/** Stat id of this object, 0 if nobody asked for it yet */
-	mutable TStatId				StatID;
-
-#if ENABLE_STATNAMEDEVENTS_UOBJECT
-	mutable PROFILER_CHAR* StatIDStringStorage;
-#endif
-#endif // STATS || ENABLE_STATNAMEDEVENTS
-
 	
 	friend class FBlueprintCompileReinstancer;
 
diff --git a/Engine/Source/Runtime/CoreUObject/Public/UObject/UObjectBaseUtility.h b/Engine/Source/Runtime/CoreUObject/Public/UObject/UObjectBaseUtility.h
--- a/Engine/Source/Runtime/CoreUObject/Public/UObject/UObjectBaseUtility.h
+++ b/Engine/Source/Runtime/CoreUObject/Public/UObject/UObjectBaseUtility.h
@@ -606,6 +606,57 @@ public:
 	{
 		return GetName() < Other.GetName();
 	}
+
+
+
+	/*******
+	 * Stats
+	 *******/
+#if STATS || ENABLE_STATNAMEDEVENTS_UOBJECT
+	FORCEINLINE void ResetStatID()
+	{
+		GUObjectArray.IndexToObject(InternalIndex)->StatID = TStatId();
+	}
+#endif
+	/**
+	  * Returns the stat ID of the object, used for profiling. This will create a stat ID if needed.
+	  *
+	  * @param bForDeferred If true, a stat ID will be created even if a group is disabled
+	  */
+	FORCEINLINE TStatId GetStatID(bool bForDeferredUse = false) const
+	{
+#if STATS
+		const TStatId& StatID = GUObjectArray.IndexToObject(InternalIndex)->StatID;
+
+		// this is done to avoid even registering stats for a disabled group (unless we plan on using it later)
+		if (bForDeferredUse || FThreadStats::IsCollectingData(GET_STATID(STAT_UObjectsStatGroupTester)))
+		{
+			if (!StatID.IsValidStat())
+			{
+				CreateStatID();
+			}
+			return StatID;
+		}
+#elif ENABLE_STATNAMEDEVENTS_UOBJECT
+		const TStatId& StatID = GUObjectArray.IndexToObject(InternalIndex)->StatID;
+		if (!StatID.IsValidStat() && (bForDeferredUse || GCycleStatsShouldEmitNamedEvents))
+		{
+			CreateStatID();
+		}
+		return StatID;
+#endif // STATS
+		return TStatId(); // not doing stats at the moment, or ever
+	}
+
+private:
+#if STATS || ENABLE_STATNAMEDEVENTS_UOBJECT
+	/** Creates a stat ID for this object */
+	void CreateStatID() const
+	{
+		GUObjectArray.IndexToObject(InternalIndex)->CreateStatID();
+	}
+#endif
+
 };
 
 /** Returns false if this pointer cannot be a valid pointer to a UObject */
diff --git a/Engine/Source/Runtime/CoreUObject/Public/UObject/UObjectGlobals.h b/Engine/Source/Runtime/CoreUObject/Public/UObject/UObjectGlobals.h
--- a/Engine/Source/Runtime/CoreUObject/Public/UObject/UObjectGlobals.h
+++ b/Engine/Source/Runtime/CoreUObject/Public/UObject/UObjectGlobals.h
@@ -2511,6 +2511,38 @@ namespace UE4CodeGen_Private
 #endif
 	};
 
+	struct FArrayPropertyParams // : FPropertyParamsBaseWithOffset
+	{
+		const char*         NameUTF8;
+		const char*         RepNotifyFuncUTF8;
+		EPropertyFlags      PropertyFlags;
+		EPropertyGenFlags   Flags;
+		EObjectFlags        ObjectFlags;
+		int32               ArrayDim;
+		int32               Offset;
+		EArrayPropertyFlags ArrayFlags;
+#if WITH_METADATA
+		const FMetaDataPairParam* MetaDataArray;
+		int32                     NumMetaData;
+#endif
+	};
+
+	struct FMapPropertyParams // : FPropertyParamsBaseWithOffset
+	{
+		const char*       NameUTF8;
+		const char*       RepNotifyFuncUTF8;
+		EPropertyFlags    PropertyFlags;
+		EPropertyGenFlags Flags;
+		EObjectFlags      ObjectFlags;
+		int32             ArrayDim;
+		int32             Offset;
+		EMapPropertyFlags MapFlags;
+#if WITH_METADATA
+		const FMetaDataPairParam* MetaDataArray;
+		int32                     NumMetaData;
+#endif
+	};
+
 	// These property types don't add new any construction parameters to their base property
 	typedef FGenericPropertyParams FInt8PropertyParams;
 	typedef FGenericPropertyParams FInt16PropertyParams;
@@ -2525,8 +2557,6 @@ namespace UE4CodeGen_Private
 	typedef FGenericPropertyParams FDoublePropertyParams;
 	typedef FGenericPropertyParams FNamePropertyParams;
 	typedef FGenericPropertyParams FStrPropertyParams;
-	typedef FGenericPropertyParams FArrayPropertyParams;
-	typedef FGenericPropertyParams FMapPropertyParams;
 	typedef FGenericPropertyParams FSetPropertyParams;
 	typedef FGenericPropertyParams FTextPropertyParams;
 	typedef FObjectPropertyParams  FWeakObjectPropertyParams;
