Copying Private-LoadTimes-4.24 stream to Main. Biggest changes are in Materials/Shader memory freezing.

diff --git a/Engine/Source/Runtime/CoreUObject/Public/UObject/UnrealType.h b/Engine/Source/Runtime/CoreUObject/Public/UObject/UnrealType.h
--- a/Engine/Source/Runtime/CoreUObject/Public/UObject/UnrealType.h
+++ b/Engine/Source/Runtime/CoreUObject/Public/UObject/UnrealType.h
@@ -12,6 +12,7 @@
 #include "Containers/List.h"
 #include "Containers/ArrayView.h"
 #include "Serialization/SerializedPropertyScope.h"
+#include "Serialization/MemoryImage.h"
 #include "Templates/Casts.h"
 #include "Templates/Greater.h"
 #include "Templates/IsFloatingPoint.h"
@@ -2716,6 +2717,12 @@ public:
 // Describes a dynamic array.
 //
 
+using FFreezableScriptArray = TScriptArray<TMemoryImageAllocator<DEFAULT_ALIGNMENT>>;
+
+//@todo stever
+//static_assert(sizeof(FScriptArray) == sizeof(FFreezableScriptArray) && alignof(FScriptArray) == alignof(FFreezableScriptArray), "FScriptArray and FFreezableScriptArray are expected to be layout-compatible");
+
+
 // need to break this out a different type so that the DECLARE_CASTED_CLASS_INTRINSIC macro can digest the comma
 typedef TProperty<FScriptArray, FProperty> FArrayProperty_Super;
 class FScriptArrayHelper;
@@ -2726,21 +2733,34 @@ class COREUOBJECT_API FArrayProperty : public FArrayProperty_Super
 
 	// Variables.
 	FProperty* Inner;
+	EArrayPropertyFlags ArrayFlags;
 
 public:
+	/** Type of the CPP property **/
+	enum
+	{
+		// These need to be the same as FFreezableScriptArray
+		CPPSize = sizeof(FScriptArray),
+		CPPAlignment = alignof(FScriptArray)
+	};
+
 	typedef FArrayProperty_Super::TTypeFundamentals TTypeFundamentals;
 	typedef TTypeFundamentals::TCppType TCppType;
 
-	FArrayProperty(FFieldVariant InOwner, const FName& InName, EObjectFlags InObjectFlags)
+	FArrayProperty(FFieldVariant InOwner, const FName& InName, EObjectFlags InObjectFlags, EArrayPropertyFlags InArrayPropertyFlags=EArrayPropertyFlags::None)
 		: FArrayProperty_Super(InOwner, InName, InObjectFlags)
 		, Inner(nullptr)
 	{
+		ArrayFlags = InArrayPropertyFlags;
+		SetElementSize();
 	}
 
-	FArrayProperty(FFieldVariant InOwner, const FName& InName, EObjectFlags InObjectFlags, int32 InOffset, EPropertyFlags InFlags)
+	FArrayProperty(FFieldVariant InOwner, const FName& InName, EObjectFlags InObjectFlags, int32 InOffset, EPropertyFlags InFlags, EArrayPropertyFlags InArrayPropertyFlags)
 		: FArrayProperty_Super(InOwner, InName, InObjectFlags, InOffset, InFlags)
 		, Inner(nullptr)
 	{
+		ArrayFlags = InArrayPropertyFlags;
+		SetElementSize();
 	}
 
 	virtual ~FArrayProperty();
@@ -2774,6 +2794,23 @@ public:
 	virtual bool NetSerializeItem( FArchive& Ar, UPackageMap* Map, void* Data, TArray<uint8> * MetaData = NULL ) const override;
 	virtual void ExportTextItem( FString& ValueStr, const void* PropertyValue, const void* DefaultValue, UObject* Parent, int32 PortFlags, UObject* ExportRootScope ) const override;
 	virtual const TCHAR* ImportText_Internal( const TCHAR* Buffer, void* Data, int32 PortFlags, UObject* OwnerObject, FOutputDevice* ErrorText ) const override;
+	virtual void InitializeValueInternal(void* Dest) const override
+	{
+		if (EnumHasAnyFlags(ArrayFlags, EArrayPropertyFlags::UsesMemoryImageAllocator))
+		{
+			for (int32 i = 0; i < this->ArrayDim; ++i)
+			{
+				new ((uint8*)Dest + i * this->ElementSize) FScriptArray;
+			}
+		}
+		else
+		{
+			for (int32 i = 0; i < this->ArrayDim; ++i)
+			{
+				new ((uint8*)Dest + i * this->ElementSize) FFreezableScriptArray;
+			}
+		}
+	}
 	virtual void CopyValuesInternal( void* Dest, void const* Src, int32 Count  ) const override;
 	virtual void ClearValueInternal( void* Data ) const override;
 	virtual void DestroyValueInternal( void* Dest ) const override;
@@ -2784,6 +2821,12 @@ public:
 	virtual void EmitReferenceInfo(UClass& OwnerClass, int32 BaseOffset, TArray<const FStructProperty*>& EncounteredStructProps) override;
 	virtual bool SameType(const FProperty* Other) const override;
 	virtual EConvertFromTypeResult ConvertFromType(const FPropertyTag& Tag, FStructuredArchive::FSlot Slot, uint8* Data, UStruct* DefaultsStruct) override;
+
+	virtual int32 GetMinAlignment() const override
+	{
+		// This is the same as alignof(FFreezableScriptArray)
+		return alignof(FScriptArray);
+	}
 	// End of FProperty interface
 
 	FString GetCPPTypeCustom(FString* ExtendedTypeText, uint32 CPPExportFlags, const FString& InnerTypeText, const FString& InInnerExtendedTypeText) const;
@@ -2793,8 +2836,19 @@ public:
 
 	/** Called by ImportTextItem, but can also be used by a non-ArrayProperty whose ArrayDim is > 1. ArrayHelper should be supplied by ArrayProperties and nullptr for fixed-size arrays. */
 	static const TCHAR* ImportTextInnerItem(const TCHAR* Buffer, FProperty* Inner, void* Data, int32 PortFlags, UObject* OwnerObject, FScriptArrayHelper* ArrayHelper = nullptr, FOutputDevice* ErrorText = (FOutputDevice*)GWarn);
+
+private:
+	FORCEINLINE void SetElementSize()
+	{
+		this->ElementSize = CPPSize;
+	}
 };
 
+using FFreezableScriptMap = TScriptMap<FMemoryImageSetAllocator>;
+
+//@todo stever
+//static_assert(sizeof(FScriptArray) == sizeof(FFreezableScriptArray) && alignof(FScriptArray) == alignof(FFreezableScriptArray), "FScriptArray and FFreezableScriptArray are expected to be layout-compatible");
+
 // need to break this out a different type so that the DECLARE_CASTED_CLASS_INTRINSIC macro can digest the comma
 typedef TProperty<FScriptMap, FProperty> FMapProperty_Super;
 
@@ -2806,13 +2860,14 @@ class COREUOBJECT_API FMapProperty : public FMapProperty_Super
 	FProperty*       KeyProp;
 	FProperty*       ValueProp;
 	FScriptMapLayout MapLayout;
+	EMapPropertyFlags MapFlags;
 
 public:
 	typedef FMapProperty_Super::TTypeFundamentals TTypeFundamentals;
 	typedef TTypeFundamentals::TCppType TCppType;
 
-	FMapProperty(FFieldVariant InOwner, const FName& InName, EObjectFlags InObjectFlags);
-	FMapProperty(FFieldVariant InOwner, const FName& InName, EObjectFlags InObjectFlags, int32 InOffset, EPropertyFlags InFlags);
+	FMapProperty(FFieldVariant InOwner, const FName& InName, EObjectFlags InObjectFlags, EMapPropertyFlags InMapFlags=EMapPropertyFlags::None);
+	FMapProperty(FFieldVariant InOwner, const FName& InName, EObjectFlags InObjectFlags, int32 InOffset, EPropertyFlags InFlags, EMapPropertyFlags InMapFlags);
 
 #if WITH_EDITORONLY_DATA
 	explicit FMapProperty(UField* InField);
@@ -2927,19 +2982,30 @@ public:
  **/
 class FScriptArrayHelper
 {
+	enum EInternal { Internal };
+
+	template <typename CallableType>
+	auto WithScriptArray(CallableType&& Callable) const
+	{
+		if (!!(ArrayFlags & EArrayPropertyFlags::UsesMemoryImageAllocator))
+		{
+			return Callable(FreezableArray);
+		}
+		else
+		{
+			return Callable(HeapArray);
+		}
+	}
+
 public:
 	/**
 	 *	Constructor, brings together a property and an instance of the property located in memory
 	 *	@param	InProperty: the property associated with this memory
 	 *	@param	InArray: pointer to raw memory that corresponds to this array. This can be NULL, and sometimes is, but in that case almost all operations will crash.
 	**/
-	FORCEINLINE FScriptArrayHelper(const FArrayProperty* InProperty, const void *InArray)
-		: InnerProperty(InProperty->Inner)
-		, Array((FScriptArray*)InArray)  //@todo, we are casting away the const here
-		, ElementSize(InnerProperty->ElementSize)
+	FORCEINLINE FScriptArrayHelper(const FArrayProperty* InProperty, const void* InArray)
+		: FScriptArrayHelper(Internal, InProperty->Inner, InArray, InProperty->Inner->ElementSize, InProperty->ArrayFlags)
 	{
-		check(ElementSize > 0);
-		check(InnerProperty);
 	}
 
 	/**
@@ -2957,14 +3023,16 @@ public:
 	**/
 	FORCEINLINE int32 Num() const
 	{
-		checkSlow(Array->Num() >= 0); 
-		return Array->Num();
+		int32 Result = WithScriptArray([](auto* Array) { return Array->Num(); });
+		checkSlow(Result >= 0);
+		return Result;
 	}
 	/**
 	 *	Static version of Num() used when you don't need to bother to construct a FScriptArrayHelper. Returns the number of elements in the array.
 	 *	@param	Target: pointer to the raw memory associated with a FScriptArray
 	 *	@return The number of elements in the array.
 	**/
+	UE_DEPRECATED(4.25, "This shortcut is no longer valid - the Num() should be read from a proper array helper")
 	static FORCEINLINE int32 Num(const void *Target)
 	{
 		checkSlow(((const FScriptArray*)Target)->Num() >= 0); 
@@ -2983,7 +3051,7 @@ public:
 			return NULL;
 		}
 		checkSlow(IsValidIndex(Index)); 
-		return (uint8*)Array->GetData() + Index * ElementSize;
+		return (uint8*)WithScriptArray([](auto* Array) { return Array->GetData(); }) + Index * ElementSize;
 	}
 	/**
 	*	Empty the array, then add blank, constructed values to a given size.
@@ -3075,7 +3143,7 @@ public:
 	{
 		check(Count>0);
 		checkSlow(Num() >= 0);
-		const int32 OldNum = Array->Add(Count, ElementSize);
+		const int32 OldNum = WithScriptArray([this, Count](auto* Array) { return Array->Add(Count, ElementSize); });
 		return OldNum;
 	}
 	/**
@@ -3095,7 +3163,7 @@ public:
 	{
 		check(Count>0);
 		check(Index>=0 && Index <= Num());
-		Array->Insert(Index, Count, ElementSize);
+		WithScriptArray([this, Index, Count](auto* Array) { Array->Insert(Index, Count, ElementSize); });
 		ConstructItems(Index, Count);
 	}
 	/**
@@ -3112,7 +3180,7 @@ public:
 		}
 		if (OldNum || Slack)
 		{
-			Array->Empty(Slack, ElementSize);
+			WithScriptArray([this, Slack](auto* Array) { Array->Empty(Slack, ElementSize); });
 		}
 	}
 	/**
@@ -3125,7 +3193,7 @@ public:
 		check(Count>0);
 		check(Index>=0 && Index + Count <= Num());
 		DestructItems(Index, Count);
-		Array->Remove(Index, Count, ElementSize);
+		WithScriptArray([this, Index, Count](auto* Array) { Array->Remove(Index, Count, ElementSize); });
 	}
 
 	/**
@@ -3147,7 +3215,7 @@ public:
 	**/
 	void SwapValues(int32 A, int32 B)
 	{
-		Array->SwapMemory(A, B, ElementSize);
+		WithScriptArray([this, A, B](auto* Array) { Array->SwapMemory(A, B, ElementSize); });
 	}
 
 	/**
@@ -3157,9 +3225,8 @@ public:
 	**/
 	void MoveAssign(void* InOtherArray)
 	{
-		FScriptArray* OtherArray = (FScriptArray*)InOtherArray;
-		checkSlow(OtherArray);
-		Array->MoveAssign(*OtherArray, ElementSize);
+		checkSlow(InOtherArray);
+		WithScriptArray([this, InOtherArray](auto* Array) { Array->MoveAssign(*static_cast<decltype(Array)>(InOtherArray), ElementSize); });
 	}
 
 	/**
@@ -3168,26 +3235,41 @@ public:
 	**/
 	void CountBytes( FArchive& Ar  ) const
 	{
-		Array->CountBytes(Ar, ElementSize);
-	}	
+		WithScriptArray([this, &Ar](auto* Array) { Array->CountBytes(Ar, ElementSize); });
+	}
 
-	static FScriptArrayHelper CreateHelperFormInnerProperty(const FProperty* InInnerProperty, const void *InArray)
+	/**
+	 * Destroys the container object - THERE SHOULD BE NO MORE USE OF THIS HELPER AFTER THIS FUNCTION IS CALLED!
+	 */
+	void DestroyContainer_Unsafe()
 	{
-		check(InInnerProperty);
-		FScriptArrayHelper ScriptArrayHelper;
-		ScriptArrayHelper.InnerProperty = InInnerProperty;
-		ScriptArrayHelper.Array = (FScriptArray*)InArray;
-		ScriptArrayHelper.ElementSize = InInnerProperty->ElementSize;
-		return ScriptArrayHelper;
+		WithScriptArray([](auto* Array) { DestructItem(Array); });
+	}
+
+	static FScriptArrayHelper CreateHelperFormInnerProperty(const FProperty* InInnerProperty, const void *InArray, EArrayPropertyFlags InArrayFlags = EArrayPropertyFlags::None)
+	{
+		return FScriptArrayHelper(Internal, InInnerProperty, InArray, InInnerProperty->ElementSize, InArrayFlags);
 	}
 
 private:
+	FScriptArrayHelper(EInternal, const FProperty* InInnerProperty, const void* InArray, int32 InElementSize, EArrayPropertyFlags InArrayFlags)
+		: InnerProperty(InInnerProperty)
+		, ElementSize(InElementSize)
+		, ArrayFlags(InArrayFlags)
+	{
+		//@todo, we are casting away the const here
+		if (!!(InArrayFlags & EArrayPropertyFlags::UsesMemoryImageAllocator))
+		{
+			FreezableArray = (FFreezableScriptArray*)InArray;
+		}
+		else
+		{
+			HeapArray = (FScriptArray*)InArray;
+		}
 
-	FScriptArrayHelper()
-		: InnerProperty(nullptr)
-		, Array(nullptr)
-		, ElementSize(0)
-	{}
+		check(ElementSize > 0);
+		check(InnerProperty);
+	}
 
 	/**
 	 *	Internal function to call into the property system to construct / initialize elements.
@@ -3259,8 +3341,13 @@ private:
 	}
 
 	const FProperty* InnerProperty;
-	FScriptArray* Array;
+	union
+	{
+		FScriptArray* HeapArray;
+		FFreezableScriptArray* FreezableArray;
+	};
 	int32 ElementSize;
+	EArrayPropertyFlags ArrayFlags;
 };
 
 class FScriptArrayHelper_InContainer : public FScriptArrayHelper
@@ -3283,8 +3370,23 @@ public:
  */
 class FScriptMapHelper
 {
+	enum EInternal { Internal };
+
 	friend class FMapProperty;
 
+	template <typename CallableType>
+	auto WithScriptMap(CallableType&& Callable) const
+	{
+		if (!!(MapFlags & EMapPropertyFlags::UsesMemoryImageAllocator))
+		{
+			return Callable(FreezableMap);
+		}
+		else
+		{
+			return Callable(HeapMap);
+		}
+	}
+
 public:
 	/**
 	 * Constructor, brings together a property and an instance of the property located in memory
@@ -3293,12 +3395,8 @@ public:
 	 * @param  InMap       Pointer to raw memory that corresponds to this map. This can be NULL, and sometimes is, but in that case almost all operations will crash.
 	 */
 	FORCEINLINE FScriptMapHelper(const FMapProperty* InProperty, const void* InMap)
-		: KeyProp         (InProperty->KeyProp)
-		, ValueProp       (InProperty->ValueProp)
-		, Map             ((FScriptMap*)InMap)  //@todo, we are casting away the const here
-		, MapLayout(InProperty->MapLayout)
+		: FScriptMapHelper(Internal, InProperty->KeyProp, InProperty->ValueProp, InMap, InProperty->MapLayout, InProperty->MapFlags)
 	{
-		check(KeyProp && ValueProp);
 	}
 
 	/**
@@ -3310,7 +3408,7 @@ public:
 	 */
 	FORCEINLINE bool IsValidIndex(int32 Index) const
 	{
-		return Map->IsValidIndex(Index);
+		return WithScriptMap([Index](auto* Map) { return Map->IsValidIndex(Index); });
 	}
 
 	/**
@@ -3320,7 +3418,7 @@ public:
 	 */
 	FORCEINLINE int32 Num() const
 	{
-		int32 Result = Map->Num();
+		int32 Result = WithScriptMap([](auto* Map) { return Map->Num(); });
 		checkSlow(Result >= 0); 
 		return Result;
 	}
@@ -3332,9 +3430,12 @@ public:
 	 */
 	FORCEINLINE int32 GetMaxIndex() const
 	{
-		int32 Result = Map->GetMaxIndex();
-		checkSlow(Result >= Num());
-		return Result;
+		return WithScriptMap([](auto* Map)
+		{
+			int32 Result = Map->GetMaxIndex();
+			checkSlow(Result >= Map->Num());
+			return Result;
+		});
 	}
 
 	/**
@@ -3344,6 +3445,7 @@ public:
 	 *
 	 * @return The number of elements in the map.
 	 */
+	UE_DEPRECATED(4.25, "This shortcut is no longer valid - the Num() should be read from a proper map helper")
 	static FORCEINLINE int32 Num(const void* Target)
 	{
 		int32 Result = ((const FScriptMap*)Target)->Num();
@@ -3360,14 +3462,17 @@ public:
 	 */
 	FORCEINLINE uint8* GetPairPtr(int32 Index)
 	{
-		if (Num() == 0)
+		return WithScriptMap([this, Index](auto* Map) -> uint8*
 		{
-			checkSlow(!Index);
-			return nullptr;
-		}
+			if (Map->Num() == 0)
+			{
+				checkSlow(!Index);
+				return nullptr;
+			}
 
-		checkSlow(IsValidIndex(Index));
-		return (uint8*)Map->GetData(Index, MapLayout);
+			checkSlow(Map->IsValidIndex(Index));
+			return (uint8*)Map->GetData(Index, MapLayout);
+		});
 	}
 
 	/**
@@ -3381,14 +3486,17 @@ public:
 	 */
 	FORCEINLINE uint8* GetKeyPtr(int32 Index)
 	{
-		if (Num() == 0)
+		return WithScriptMap([this, Index](auto* Map) -> uint8*
 		{
-			checkSlow(!Index);
-			return nullptr;
-		}
+			if (Map->Num() == 0)
+			{
+				checkSlow(!Index);
+				return nullptr;
+			}
 		
-		checkSlow(IsValidIndex(Index));
-		return (uint8*)Map->GetData(Index, MapLayout);
+			checkSlow(Map->IsValidIndex(Index));
+			return (uint8*)Map->GetData(Index, MapLayout);
+		});
 	}
 
 	/**
@@ -3400,14 +3508,17 @@ public:
 	 */
 	FORCEINLINE uint8* GetValuePtr(int32 Index)
 	{
-		if (Num() == 0)
+		return WithScriptMap([this, Index](auto* Map) -> uint8*
 		{
-			checkSlow(!Index);
-			return nullptr;
-		}
+			if (Map->Num() == 0)
+			{
+				checkSlow(!Index);
+				return nullptr;
+			}
 		
-		checkSlow(IsValidIndex(Index));
-		return (uint8*)Map->GetData(Index, MapLayout) + MapLayout.ValueOffset;
+			checkSlow(Map->IsValidIndex(Index));
+			return (uint8*)Map->GetData(Index, MapLayout) + MapLayout.ValueOffset;
+		});
 	}
 
 	/**
@@ -3430,9 +3541,11 @@ public:
 	 */
 	void MoveAssign(void* InOtherMap)
 	{
-		FScriptMap* OtherMap = (FScriptMap*)InOtherMap;
-		checkSlow(OtherMap);
-		Map->MoveAssign(*OtherMap, MapLayout);
+		checkSlow(InOtherMap);
+		return WithScriptMap([this, InOtherMap](auto* Map)
+		{
+			Map->MoveAssign(*(decltype(Map))InOtherMap, MapLayout);
+		});
 	}
 
 	/**
@@ -3442,9 +3555,12 @@ public:
 	 */
 	FORCEINLINE int32 AddUninitializedValue()
 	{
-		checkSlow(Num() >= 0);
+		return WithScriptMap([this](auto* Map)
+		{
+			checkSlow(Map->Num() >= 0);
 
-		return Map->AddUninitialized(MapLayout);
+			return Map->AddUninitialized(MapLayout);
+		});
 	}
 
 	/**
@@ -3462,7 +3578,10 @@ public:
 		}
 		if (OldNum || Slack)
 		{
-			Map->Empty(Slack, MapLayout);
+			return WithScriptMap([this, Slack](auto* Map)
+			{
+				Map->Empty(Slack, MapLayout);
+			});
 		}
 	}
 
@@ -3474,12 +3593,15 @@ public:
 	 **/
 	int32 AddDefaultValue_Invalid_NeedsRehash()
 	{
-		checkSlow(Num() >= 0);
+		return WithScriptMap([this](auto* Map)
+		{
+			checkSlow(Map->Num() >= 0);
 
-		int32 Result = AddUninitializedValue();
-		ConstructItem(Result);
+			int32 Result = Map->AddUninitialized(MapLayout);
+			ConstructItem(Result);
 
-		return Result;
+			return Result;
+		});
 	}
 
 	/**
@@ -3509,17 +3631,20 @@ public:
 	 */
 	void RemoveAt(int32 Index, int32 Count = 1)
 	{
-		check(IsValidIndex(Index));
-
-		DestructItems(Index, Count);
-		for (; Count; ++Index)
+		return WithScriptMap([this, Index, Count](auto* Map)
 		{
-			if (IsValidIndex(Index))
+			check(Map->IsValidIndex(Index));
+
+			DestructItems(Index, Count);
+			for (int32 LocalCount = Count, LocalIndex = Index; LocalCount; ++LocalIndex)
 			{
-				Map->RemoveAt(Index, MapLayout);
-				--Count;
+				if (Map->IsValidIndex(LocalIndex))
+				{
+					Map->RemoveAt(LocalIndex, MapLayout);
+					--LocalCount;
+				}
 			}
-		}
+		});
 	}
 
 	/**
@@ -3535,24 +3660,28 @@ public:
 	 */
 	int32 FindInternalIndex(int32 LogicalIdx) const
 	{
-		if (LogicalIdx < 0 && LogicalIdx > Num())
+		return WithScriptMap([this, LogicalIdx](auto* Map) -> int32
 		{
-			return INDEX_NONE;
-		}
+			int32 LocalLogicalIdx = LogicalIdx;
+			if (LocalLogicalIdx < 0 && LocalLogicalIdx > Map->Num())
+			{
+				return INDEX_NONE;
+			}
 
-		int32 MaxIndex = GetMaxIndex();
-		for (int32 Actual = 0; Actual < MaxIndex; ++Actual)
-		{
-			if (IsValidIndex(Actual))
+			int32 MaxIndex = Map->GetMaxIndex();
+			for (int32 Actual = 0; Actual < MaxIndex; ++Actual)
 			{
-				if (LogicalIdx == 0)
+				if (Map->IsValidIndex(Actual))
 				{
-					return Actual;
+					if (LocalLogicalIdx == 0)
+					{
+						return Actual;
+					}
+					--LocalLogicalIdx;
 				}
-				--LogicalIdx;
 			}
-		}
-		return INDEX_NONE;
+			return INDEX_NONE;
+		});
 	}
 
 	/**
@@ -3565,39 +3694,42 @@ public:
 	 */
 	int32 FindMapIndexWithKey(const void* PairWithKeyToFind, int32 IndexHint = 0) const
 	{
-		int32 MapMax = GetMaxIndex();
-		if (MapMax == 0)
+		return WithScriptMap([this, PairWithKeyToFind, IndexHint](auto* Map) -> int32
 		{
-			return INDEX_NONE;
-		}
+			int32 MapMax = Map->GetMaxIndex();
+			if (MapMax == 0)
+			{
+				return INDEX_NONE;
+			}
 
-		check(IndexHint >= 0 && IndexHint < MapMax);
+			check(IndexHint >= 0 && IndexHint < MapMax);
 
-		FProperty* LocalKeyProp = this->KeyProp; // prevent aliasing in loop below
+			FProperty* LocalKeyProp = this->KeyProp; // prevent aliasing in loop below
 
-		int32 Index = IndexHint;
-		for (;;)
-		{
-			if (IsValidIndex(Index))
+			int32 Index = IndexHint;
+			for (;;)
 			{
-				const void* PairToSearch = GetPairPtrWithoutCheck(Index);
-				if (LocalKeyProp->Identical(PairWithKeyToFind, PairToSearch))
+				if (Map->IsValidIndex(Index))
 				{
-					return Index;
+					const void* PairToSearch = Map->GetData(Index, MapLayout);
+					if (LocalKeyProp->Identical(PairWithKeyToFind, PairToSearch))
+					{
+						return Index;
+					}
 				}
-			}
 
-			++Index;
-			if (Index == MapMax)
-			{
-				Index = 0;
-			}
+				++Index;
+				if (Index == MapMax)
+				{
+					Index = 0;
+				}
 
-			if (Index == IndexHint)
-			{
-				return INDEX_NONE;
+				if (Index == IndexHint)
+				{
+					return INDEX_NONE;
+				}
 			}
-		}
+		});
 	}
 
 	/**
@@ -3618,13 +3750,15 @@ public:
 	/** Finds the associated pair from hash, rather than linearly searching */
 	uint8* FindMapPairPtrFromHash(const void* KeyPtr)
 	{
-		FProperty* LocalKeyPropForCapture = KeyProp;
-		int32 Index = Map->FindPairIndex(
-			KeyPtr,
-			MapLayout,
-			[LocalKeyPropForCapture](const void* ElementKey) { return LocalKeyPropForCapture->GetValueTypeHash(ElementKey); },
-			[LocalKeyPropForCapture](const void* A, const void* B) { return LocalKeyPropForCapture->Identical(A, B); }
+		int32 Index = WithScriptMap([this, KeyPtr, LocalKeyPropForCapture = this->KeyProp](auto* Map)
+		{
+			return Map->FindPairIndex(
+				KeyPtr,
+				MapLayout,
+				[LocalKeyPropForCapture](const void* ElementKey) { return LocalKeyPropForCapture->GetValueTypeHash(ElementKey); },
+				[LocalKeyPropForCapture](const void* A, const void* B) { return LocalKeyPropForCapture->Identical(A, B); }
 			);
+		});
 		uint8* Result = (Index >= 0) ? GetPairPtr(Index) : nullptr;
 		return Result;
 	}
@@ -3632,71 +3766,74 @@ public:
 	/** Finds the associated value from hash, rather than linearly searching */
 	uint8* FindValueFromHash(const void* KeyPtr)
 	{
-		FProperty* LocalKeyPropForCapture = KeyProp;
-		return Map->FindValue(
-			KeyPtr, 
-			MapLayout,
-			[LocalKeyPropForCapture](const void* ElementKey) { return LocalKeyPropForCapture->GetValueTypeHash(ElementKey); },
-			[LocalKeyPropForCapture](const void* A, const void* B) { return LocalKeyPropForCapture->Identical(A, B); }
+		return WithScriptMap([this, KeyPtr, LocalKeyPropForCapture = this->KeyProp](auto* Map)
+		{
+			return Map->FindValue(
+				KeyPtr,
+				MapLayout,
+				[LocalKeyPropForCapture](const void* ElementKey) { return LocalKeyPropForCapture->GetValueTypeHash(ElementKey); },
+				[LocalKeyPropForCapture](const void* A, const void* B) { return LocalKeyPropForCapture->Identical(A, B); }
 			);
+		});
 	}
 
 	/** Adds the (key, value) pair to the map, returning true if the element was added, or false if the element was already present and has been overwritten */
 	void AddPair(const void* KeyPtr, const void* ValuePtr)
 	{
-		FProperty* LocalKeyPropForCapture = KeyProp;
-		FProperty* LocalValuePropForCapture = ValueProp;
-		Map->Add(
-			KeyPtr,
-			ValuePtr,
-			MapLayout,
-			[LocalKeyPropForCapture](const void* ElementKey) { return LocalKeyPropForCapture->GetValueTypeHash(ElementKey); },
-			[LocalKeyPropForCapture](const void* A, const void* B) { return LocalKeyPropForCapture->Identical(A, B); },
-			[LocalKeyPropForCapture, KeyPtr](void* NewElementKey)
-			{
-				if (LocalKeyPropForCapture->PropertyFlags & CPF_ZeroConstructor)
-				{
-					FMemory::Memzero(NewElementKey, LocalKeyPropForCapture->GetSize());
-				}
-				else
+		return WithScriptMap([this, KeyPtr, ValuePtr, LocalKeyPropForCapture = this->KeyProp, LocalValuePropForCapture = this->ValueProp](auto* Map)
+		{
+			Map->Add(
+				KeyPtr,
+				ValuePtr,
+				MapLayout,
+				[LocalKeyPropForCapture](const void* ElementKey) { return LocalKeyPropForCapture->GetValueTypeHash(ElementKey); },
+				[LocalKeyPropForCapture](const void* A, const void* B) { return LocalKeyPropForCapture->Identical(A, B); },
+				[LocalKeyPropForCapture, KeyPtr](void* NewElementKey)
 				{
-					LocalKeyPropForCapture->InitializeValue(NewElementKey);
-				}
+					if (LocalKeyPropForCapture->PropertyFlags & CPF_ZeroConstructor)
+					{
+						FMemory::Memzero(NewElementKey, LocalKeyPropForCapture->GetSize());
+					}
+					else
+					{
+						LocalKeyPropForCapture->InitializeValue(NewElementKey);
+					}
 
-				LocalKeyPropForCapture->CopySingleValueToScriptVM(NewElementKey, KeyPtr);
-			},
-			[LocalValuePropForCapture, ValuePtr](void* NewElementValue)
-			{
-				if (LocalValuePropForCapture->PropertyFlags & CPF_ZeroConstructor)
+					LocalKeyPropForCapture->CopySingleValueToScriptVM(NewElementKey, KeyPtr);
+				},
+				[LocalValuePropForCapture, ValuePtr](void* NewElementValue)
 				{
-					FMemory::Memzero(NewElementValue, LocalValuePropForCapture->GetSize());
-				}
-				else
-				{
-					LocalValuePropForCapture->InitializeValue(NewElementValue);
-				}
+					if (LocalValuePropForCapture->PropertyFlags & CPF_ZeroConstructor)
+					{
+						FMemory::Memzero(NewElementValue, LocalValuePropForCapture->GetSize());
+					}
+					else
+					{
+						LocalValuePropForCapture->InitializeValue(NewElementValue);
+					}
 
-				LocalValuePropForCapture->CopySingleValueToScriptVM(NewElementValue, ValuePtr);
-			},
-			[LocalValuePropForCapture, ValuePtr](void* ExistingElementValue)
-			{
-				LocalValuePropForCapture->CopySingleValueToScriptVM(ExistingElementValue, ValuePtr);
-			},
-			[LocalKeyPropForCapture](void* ElementKey)
-			{
-				if (!(LocalKeyPropForCapture->PropertyFlags & (CPF_IsPlainOldData | CPF_NoDestructor)))
+					LocalValuePropForCapture->CopySingleValueToScriptVM(NewElementValue, ValuePtr);
+				},
+				[LocalValuePropForCapture, ValuePtr](void* ExistingElementValue)
 				{
-					LocalKeyPropForCapture->DestroyValue(ElementKey);
-				}
-			},
-			[LocalValuePropForCapture](void* ElementValue)
-			{
-				if (!(LocalValuePropForCapture->PropertyFlags & (CPF_IsPlainOldData | CPF_NoDestructor)))
+					LocalValuePropForCapture->CopySingleValueToScriptVM(ExistingElementValue, ValuePtr);
+				},
+				[LocalKeyPropForCapture](void* ElementKey)
 				{
-					LocalValuePropForCapture->DestroyValue(ElementValue);
+					if (!(LocalKeyPropForCapture->PropertyFlags & (CPF_IsPlainOldData | CPF_NoDestructor)))
+					{
+						LocalKeyPropForCapture->DestroyValue(ElementKey);
+					}
+				},
+				[LocalValuePropForCapture](void* ElementValue)
+				{
+					if (!(LocalValuePropForCapture->PropertyFlags & (CPF_IsPlainOldData | CPF_NoDestructor)))
+					{
+						LocalValuePropForCapture->DestroyValue(ElementValue);
+					}
 				}
-			}
-		);
+			);
+		});
 	}
 
 
@@ -3709,59 +3846,61 @@ public:
 	 **/
 	void* FindOrAdd(const void* KeyPtr)
 	{
-		FProperty* LocalKeyPropForCapture = KeyProp;
-		FProperty* LocalValuePropForCapture = ValueProp;
-		return Map->FindOrAdd(
-			KeyPtr,
-			MapLayout,
-			[LocalKeyPropForCapture](const void* ElementKey) { return LocalKeyPropForCapture->GetValueTypeHash(ElementKey); },
-			[LocalKeyPropForCapture](const void* A, const void* B) { return LocalKeyPropForCapture->Identical(A, B); },
-			[LocalKeyPropForCapture, LocalValuePropForCapture, KeyPtr](void* NewElementKey, void* NewElementValue)
-			{
-				if (LocalKeyPropForCapture->PropertyFlags & CPF_ZeroConstructor)
-				{
-					FMemory::Memzero(NewElementKey, LocalKeyPropForCapture->GetSize());
-				}
-				else
+		return WithScriptMap([this, KeyPtr, LocalKeyPropForCapture = this->KeyProp, LocalValuePropForCapture = this->ValueProp](auto* Map)
+		{
+			return Map->FindOrAdd(
+				KeyPtr,
+				MapLayout,
+				[LocalKeyPropForCapture](const void* ElementKey) { return LocalKeyPropForCapture->GetValueTypeHash(ElementKey); },
+				[LocalKeyPropForCapture](const void* A, const void* B) { return LocalKeyPropForCapture->Identical(A, B); },
+				[LocalKeyPropForCapture, LocalValuePropForCapture, KeyPtr](void* NewElementKey, void* NewElementValue)
 				{
-					LocalKeyPropForCapture->InitializeValue(NewElementKey);
-				}
+					if (LocalKeyPropForCapture->PropertyFlags & CPF_ZeroConstructor)
+					{
+						FMemory::Memzero(NewElementKey, LocalKeyPropForCapture->GetSize());
+					}
+					else
+					{
+						LocalKeyPropForCapture->InitializeValue(NewElementKey);
+					}
 
-				LocalKeyPropForCapture->CopySingleValue(NewElementKey, KeyPtr);
+					LocalKeyPropForCapture->CopySingleValue(NewElementKey, KeyPtr);
 
-				if (LocalValuePropForCapture->PropertyFlags & CPF_ZeroConstructor)
-				{
-					FMemory::Memzero(NewElementValue, LocalValuePropForCapture->GetSize());
-				}
-				else
-				{
-					LocalValuePropForCapture->InitializeValue(NewElementValue);
+					if (LocalValuePropForCapture->PropertyFlags & CPF_ZeroConstructor)
+					{
+						FMemory::Memzero(NewElementValue, LocalValuePropForCapture->GetSize());
+					}
+					else
+					{
+						LocalValuePropForCapture->InitializeValue(NewElementValue);
+					}
 				}
-			}
-		);
-
+			);
+		});
 	}
 
 
 	/** Removes the key and its associated value from the map */
 	bool RemovePair(const void* KeyPtr)
 	{
-		FProperty* LocalKeyPropForCapture = KeyProp;
-		if(uint8* Entry =  Map->FindValue(
-			KeyPtr, 
-			MapLayout,
-			[LocalKeyPropForCapture](const void* ElementKey) { return LocalKeyPropForCapture->GetValueTypeHash(ElementKey); },
-			[LocalKeyPropForCapture](const void* A, const void* B) { return LocalKeyPropForCapture->Identical(A, B); }
-			))
-		{
-			int32 Idx = (int32)((Entry - (uint8*)Map->GetData(0, MapLayout)) / MapLayout.SetLayout.Size);
-			RemoveAt(Idx);
-			return true;
-		}
-		else
+		return WithScriptMap([this, KeyPtr, LocalKeyPropForCapture = this->KeyProp](auto* Map)
 		{
-			return false;
-		}
+			if (uint8* Entry = Map->FindValue(
+				KeyPtr,
+				MapLayout,
+				[LocalKeyPropForCapture](const void* ElementKey) { return LocalKeyPropForCapture->GetValueTypeHash(ElementKey); },
+				[LocalKeyPropForCapture](const void* A, const void* B) { return LocalKeyPropForCapture->Identical(A, B); }
+			))
+			{
+				int32 Idx = (int32)((Entry - (uint8*)Map->GetData(0, MapLayout)) / MapLayout.SetLayout.Size);
+				RemoveAt(Idx);
+				return true;
+			}
+			else
+			{
+				return false;
+			}
+		});
 	}
 
 	/**
@@ -3797,17 +3936,16 @@ public:
 		return false;
 	}
 
-	static FScriptMapHelper CreateHelperFormInnerProperties(FProperty* InKeyProperty, FProperty* InValProperty, const void *InMap)
+	static FScriptMapHelper CreateHelperFormInnerProperties(FProperty* InKeyProperty, FProperty* InValProperty, const void *InMap, EMapPropertyFlags InMapFlags = EMapPropertyFlags::None)
 	{
-		check(InKeyProperty && InValProperty);
-
-		FScriptMapHelper ScriptMapHelper;
-		ScriptMapHelper.KeyProp = InKeyProperty;
-		ScriptMapHelper.ValueProp = InValProperty;
-		ScriptMapHelper.Map = (FScriptMap*)InMap;
-		ScriptMapHelper.MapLayout = FScriptMap::GetScriptLayout(InKeyProperty->GetSize(), InKeyProperty->GetMinAlignment(), InValProperty->GetSize(), InValProperty->GetMinAlignment());
-
-		return ScriptMapHelper;
+		return FScriptMapHelper(
+			Internal,
+			InKeyProperty,
+			InValProperty,
+			InMap,
+			FScriptMap::GetScriptLayout(InKeyProperty->GetSize(), InKeyProperty->GetMinAlignment(), InValProperty->GetSize(), InValProperty->GetMinAlignment()),
+			InMapFlags
+		);
 	}
 
 	class FIterator
@@ -3845,12 +3983,26 @@ public:
 	}
 
 private:
-	FScriptMapHelper()
-		: KeyProp(nullptr)
-		, ValueProp(nullptr)
-		, Map(nullptr)
-		, MapLayout(FScriptMap::GetScriptLayout(0, 1, 0, 1))
-	{}
+	FORCEINLINE FScriptMapHelper(EInternal, FProperty* InKeyProp, FProperty* InValueProp, const void* InMap, const FScriptMapLayout& InMapLayout, EMapPropertyFlags InMapFlags)
+		: KeyProp  (InKeyProp)
+		, ValueProp(InValueProp)
+		, MapLayout(InMapLayout)
+		, MapFlags (InMapFlags)
+	{
+		check(InKeyProp && InValueProp);
+
+		//@todo, we are casting away the const here
+		if (!!(InMapFlags & EMapPropertyFlags::UsesMemoryImageAllocator))
+		{
+			FreezableMap = (FFreezableScriptMap*)InMap;
+		}
+		else
+		{
+			HeapMap = (FScriptMap*)InMap;
+		}
+
+		check(KeyProp && ValueProp);
+	}
 
 	/**
 	 * Internal function to call into the property system to construct / initialize elements.
@@ -3865,7 +4017,7 @@ private:
 		bool bZeroKey   = !!(KeyProp  ->PropertyFlags & CPF_ZeroConstructor);
 		bool bZeroValue = !!(ValueProp->PropertyFlags & CPF_ZeroConstructor);
 
-		uint8* Dest = GetPairPtrWithoutCheck(Index);
+		void* Dest = WithScriptMap([this, Index](auto* Map) { return Map->GetData(Index, MapLayout); });
 
 		if (bZeroKey || bZeroValue)
 		{
@@ -3903,7 +4055,7 @@ private:
 		if (bDestroyKeys || bDestroyValues)
 		{
 			uint32 Stride  = MapLayout.SetLayout.Size;
-			uint8* PairPtr = GetPairPtrWithoutCheck(Index);
+			uint8* PairPtr = WithScriptMap([this, Index](auto* Map) { return (uint8*)Map->GetData(Index, MapLayout); });
 			if (bDestroyKeys)
 			{
 				if (bDestroyValues)
@@ -3956,7 +4108,7 @@ private:
 	 */
 	FORCEINLINE uint8* GetPairPtrWithoutCheck(int32 Index)
 	{
-		return (uint8*)Map->GetData(Index, MapLayout);
+		return WithScriptMap([this, Index](auto* Map) { return (uint8*)Map->GetData(Index, MapLayout); });
 	}
 
 	/**
@@ -3972,10 +4124,15 @@ private:
 	}
 
 public:
-	FProperty*       KeyProp;
-	FProperty*       ValueProp;
-	FScriptMap*      Map;
-	FScriptMapLayout MapLayout;
+	FProperty*        KeyProp;
+	FProperty*        ValueProp;
+	union
+	{
+		FScriptMap*          HeapMap;
+		FFreezableScriptMap* FreezableMap;
+	};
+	FScriptMapLayout  MapLayout;
+	EMapPropertyFlags MapFlags;
 };
 
 class FScriptMapHelper_InContainer : public FScriptMapHelper
diff --git a/Engine/Source/Runtime/D3D12RHI/Private/D3D12Adapter.cpp b/Engine/Source/Runtime/D3D12RHI/Private/D3D12Adapter.cpp
--- a/Engine/Source/Runtime/D3D12RHI/Private/D3D12Adapter.cpp
+++ b/Engine/Source/Runtime/D3D12RHI/Private/D3D12Adapter.cpp
@@ -666,17 +666,7 @@ void FD3D12Adapter::Cleanup()
 #endif
 
 	// Ask all initialized FRenderResources to release their RHI resources.
-	for (TLinkedList<FRenderResource*>::TIterator ResourceIt(FRenderResource::GetResourceList()); ResourceIt; ResourceIt.Next())
-	{
-		FRenderResource* Resource = *ResourceIt;
-		check(Resource->IsInitialized());
-		Resource->ReleaseRHI();
-	}
-
-	for (TLinkedList<FRenderResource*>::TIterator ResourceIt(FRenderResource::GetResourceList()); ResourceIt; ResourceIt.Next())
-	{
-		ResourceIt->ReleaseDynamicRHI();
-	}
+	FRenderResource::ReleaseRHIForAllResources();
 
 	FRHIResource::FlushPendingDeletes();
 
diff --git a/Engine/Source/Runtime/D3D12RHI/Private/D3D12Buffer.cpp b/Engine/Source/Runtime/D3D12RHI/Private/D3D12Buffer.cpp
--- a/Engine/Source/Runtime/D3D12RHI/Private/D3D12Buffer.cpp
+++ b/Engine/Source/Runtime/D3D12RHI/Private/D3D12Buffer.cpp
@@ -30,7 +30,11 @@ template FD3D12StructuredBuffer* FD3D12Adapter::CreateRHIBuffer<FD3D12Structured
 	uint32 Alignment, uint32 Stride, uint32 Size, uint32 InUsage,
 	FRHIResourceCreateInfo& CreateInfo);
 
-struct FRHICommandUpdateBuffer final : public FRHICommand<FRHICommandUpdateBuffer>
+struct FRHICommandUpdateBufferString
+{
+	static const TCHAR* TStr() { return TEXT("FRHICommandUpdateBuffer"); }
+};
+struct FRHICommandUpdateBuffer final : public FRHICommand<FRHICommandUpdateBuffer, FRHICommandUpdateBufferString>
 {
 	FD3D12ResourceLocation Source;
 	FD3D12ResourceLocation* Destination;
@@ -55,8 +59,12 @@ struct FRHICommandUpdateBuffer final : public FRHICommand<FRHICommandUpdateBuffe
 // This allows us to rename resources from the RenderThread i.e. all the 'hard' work of allocating a new resource
 // is done in parallel and this small function is called to switch the resource to point to the correct location
 // a the correct time.
+struct FRHICommandRenameUploadBufferString
+{
+	static const TCHAR* TStr() { return TEXT("FRHICommandRenameUploadBuffer"); }
+};
 template<typename ResourceType>
-struct FRHICommandRenameUploadBuffer final : public FRHICommand<FRHICommandRenameUploadBuffer<ResourceType>>
+struct FRHICommandRenameUploadBuffer final : public FRHICommand<FRHICommandRenameUploadBuffer<ResourceType>, FRHICommandRenameUploadBufferString>
 {
 	ResourceType* Resource;
 	FD3D12ResourceLocation NewLocation;
@@ -181,7 +189,11 @@ BufferType* FD3D12Adapter::CreateRHIBuffer(FRHICommandListImmediate* RHICmdList,
 			check(pData);
 			FMemory::Memcpy(pData, CreateInfo.ResourceArray->GetResourceData(), Size);
 
-			struct FD3D12RHICommandInitializeBuffer final : public FRHICommand<FD3D12RHICommandInitializeBuffer>
+			struct FD3D12RHICommandInitializeBufferString
+			{
+				static const TCHAR* TStr() { return TEXT("FD3D12RHICommandInitializeBuffer"); }
+			};
+			struct FD3D12RHICommandInitializeBuffer final : public FRHICommand<FD3D12RHICommandInitializeBuffer, FD3D12RHICommandInitializeBufferString>
 			{
 				BufferType* CurrentBuffer;
 				FD3D12ResourceLocation SrcResourceLoc;
diff --git a/Engine/Source/Runtime/D3D12RHI/Private/D3D12Commands.cpp b/Engine/Source/Runtime/D3D12RHI/Private/D3D12Commands.cpp
--- a/Engine/Source/Runtime/D3D12RHI/Private/D3D12Commands.cpp
+++ b/Engine/Source/Runtime/D3D12RHI/Private/D3D12Commands.cpp
@@ -57,12 +57,12 @@ void FD3D12DynamicRHI::SetupRecursiveResources()
 
 	{
 		TShaderMapRef<FLongGPUTaskPS> PixelShader(ShaderMap);
-		PixelShader->GetPixelShader();
+		PixelShader.GetPixelShader();
 	}
 
 	{
 		TShaderMapRef<FLongGPUTaskPS> PixelShader(ShaderMap);
-		PixelShader->GetPixelShader();
+		PixelShader.GetPixelShader();
 	}
 
 	// TODO: Waiting to integrate MSAA fix for ResolveShader.h
@@ -73,15 +73,15 @@ void FD3D12DynamicRHI::SetupRecursiveResources()
 	if (GMaxRHIShaderPlatform == SP_PCD3D_SM5 || GMaxRHIShaderPlatform == SP_XBOXONE_D3D12)
 	{
 		TShaderMapRef<FResolveDepthPS> ResolvePixelShader_Depth(ShaderMap);
-		ResolvePixelShader_Depth->GetPixelShader();
+		ResolvePixelShader_Depth.GetPixelShader();
 
 		TShaderMapRef<FResolveDepthPS> ResolvePixelShader_SingleSample(ShaderMap);
-		ResolvePixelShader_SingleSample->GetPixelShader();
+		ResolvePixelShader_SingleSample.GetPixelShader();
 	}
 	else
 	{
 		TShaderMapRef<FResolveDepthNonMSPS> ResolvePixelShader_DepthNonMS(ShaderMap);
-		ResolvePixelShader_DepthNonMS->GetPixelShader();
+		ResolvePixelShader_DepthNonMS.GetPixelShader();
 	}
 }
 
@@ -525,37 +525,37 @@ void FD3D12CommandContext::RHISetComputePipelineState(FRHIComputePipelineState*
 void FD3D12CommandContext::RHISetShaderTexture(FRHIGraphicsShader* ShaderRHI, uint32 TextureIndex, FRHITexture* NewTextureRHI)
 {
 	FD3D12TextureBase* const NewTexture = RetrieveTextureBase(NewTextureRHI);
-	switch (ShaderRHI->GetType())
+	switch (ShaderRHI->GetFrequency())
 	{
-	case FRHIShader::EType::Vertex:
+	case SF_Vertex:
 	{
 		FRHIVertexShader* VertexShaderRHI = static_cast<FRHIVertexShader*>(ShaderRHI);
 		VALIDATE_BOUND_SHADER(VertexShaderRHI);
 		StateCache.SetShaderResourceView<SF_Vertex>(NewTexture ? NewTexture->GetShaderResourceView() : nullptr, TextureIndex);
 	}
 	break;
-	case FRHIShader::EType::Hull:
+	case SF_Hull:
 	{
 		FRHIHullShader* HullShaderRHI = static_cast<FRHIHullShader*>(ShaderRHI);
 		VALIDATE_BOUND_SHADER(HullShaderRHI);
 		StateCache.SetShaderResourceView<SF_Hull>(NewTexture ? NewTexture->GetShaderResourceView() : nullptr, TextureIndex);
 	}
 	break;
-	case FRHIShader::EType::Domain:
+	case SF_Domain:
 	{
 		FRHIDomainShader* DomainShaderRHI = static_cast<FRHIDomainShader*>(ShaderRHI);
 		VALIDATE_BOUND_SHADER(DomainShaderRHI);
 		StateCache.SetShaderResourceView<SF_Domain>(NewTexture ? NewTexture->GetShaderResourceView() : nullptr, TextureIndex);
 	}
 	break;
-	case FRHIShader::EType::Geometry:
+	case SF_Geometry:
 	{
 		FRHIGeometryShader* GeometryShaderRHI = static_cast<FRHIGeometryShader*>(ShaderRHI);
 		VALIDATE_BOUND_SHADER(GeometryShaderRHI);
 		StateCache.SetShaderResourceView<SF_Geometry>(NewTexture ? NewTexture->GetShaderResourceView() : nullptr, TextureIndex);
 	}
 	break;
-	case FRHIShader::EType::Pixel:
+	case SF_Pixel:
 	{
 		FRHIPixelShader* PixelShaderRHI = static_cast<FRHIPixelShader*>(ShaderRHI);
 		VALIDATE_BOUND_SHADER(PixelShaderRHI);
@@ -563,7 +563,7 @@ void FD3D12CommandContext::RHISetShaderTexture(FRHIGraphicsShader* ShaderRHI, ui
 	}
 	break;
 	default:
-		checkf(0, TEXT("Undefined FRHIShader Type %d!"), (int32)ShaderRHI->GetType());
+		checkf(0, TEXT("Undefined FRHIShader Type %d!"), (int32)ShaderRHI->GetFrequency());
 	}
 }
 
@@ -624,37 +624,37 @@ void FD3D12CommandContext::RHISetUAVParameter(FRHIComputeShader* ComputeShaderRH
 void FD3D12CommandContext::RHISetShaderResourceViewParameter(FRHIGraphicsShader* ShaderRHI, uint32 TextureIndex, FRHIShaderResourceView* SRVRHI)
 {
 	FD3D12ShaderResourceView* const SRV = RetrieveObject<FD3D12ShaderResourceView>(SRVRHI);
-	switch (ShaderRHI->GetType())
+	switch (ShaderRHI->GetFrequency())
 	{
-	case FRHIShader::EType::Vertex:
+	case SF_Vertex:
 	{
 		FRHIVertexShader* VertexShaderRHI = static_cast<FRHIVertexShader*>(ShaderRHI);
 		VALIDATE_BOUND_SHADER(VertexShaderRHI);
 		StateCache.SetShaderResourceView<SF_Vertex>(SRV, TextureIndex);
 	}
 	break;
-	case FRHIShader::EType::Hull:
+	case SF_Hull:
 	{
 		FRHIHullShader* HullShaderRHI = static_cast<FRHIHullShader*>(ShaderRHI);
 		VALIDATE_BOUND_SHADER(HullShaderRHI);
 		StateCache.SetShaderResourceView<SF_Hull>(SRV, TextureIndex);
 	}
 	break;
-	case FRHIShader::EType::Domain:
+	case SF_Domain:
 	{
 		FRHIDomainShader* DomainShaderRHI = static_cast<FRHIDomainShader*>(ShaderRHI);
 		VALIDATE_BOUND_SHADER(DomainShaderRHI);
 		StateCache.SetShaderResourceView<SF_Domain>(SRV, TextureIndex);
 	}
 	break;
-	case FRHIShader::EType::Geometry:
+	case SF_Geometry:
 	{
 		FRHIGeometryShader* GeometryShaderRHI = static_cast<FRHIGeometryShader*>(ShaderRHI);
 		VALIDATE_BOUND_SHADER(GeometryShaderRHI);
 		StateCache.SetShaderResourceView<SF_Geometry>(SRV, TextureIndex);
 	}
 	break;
-	case FRHIShader::EType::Pixel:
+	case SF_Pixel:
 	{
 		FRHIPixelShader* PixelShaderRHI = static_cast<FRHIPixelShader*>(ShaderRHI);
 		VALIDATE_BOUND_SHADER(PixelShaderRHI);
@@ -662,7 +662,7 @@ void FD3D12CommandContext::RHISetShaderResourceViewParameter(FRHIGraphicsShader*
 	}
 	break;
 	default:
-		checkf(0, TEXT("Undefined FRHIShader Type %d!"), (int32)ShaderRHI->GetType());
+		checkf(0, TEXT("Undefined FRHIShader Type %d!"), (int32)ShaderRHI->GetFrequency());
 	}
 
 }
@@ -677,37 +677,37 @@ void FD3D12CommandContext::RHISetShaderResourceViewParameter(FRHIComputeShader*
 void FD3D12CommandContext::RHISetShaderSampler(FRHIGraphicsShader* ShaderRHI, uint32 SamplerIndex, FRHISamplerState* NewStateRHI)
 {
 	FD3D12SamplerState* NewState = RetrieveObject<FD3D12SamplerState>(NewStateRHI);
-	switch (ShaderRHI->GetType())
+	switch (ShaderRHI->GetFrequency())
 	{
-	case FRHIShader::EType::Vertex:
+	case SF_Vertex:
 	{
 		FRHIVertexShader* VertexShaderRHI = static_cast<FRHIVertexShader*>(ShaderRHI);
 		VALIDATE_BOUND_SHADER(VertexShaderRHI);
 		StateCache.SetSamplerState<SF_Vertex>(NewState, SamplerIndex);
 	}
 	break;
-	case FRHIShader::EType::Hull:
+	case SF_Hull:
 	{
 		FRHIHullShader* HullShaderRHI = static_cast<FRHIHullShader*>(ShaderRHI);
 		VALIDATE_BOUND_SHADER(HullShaderRHI);
 		StateCache.SetSamplerState<SF_Hull>(NewState, SamplerIndex);
 	}
 	break;
-	case FRHIShader::EType::Domain:
+	case SF_Domain:
 	{
 		FRHIDomainShader* DomainShaderRHI = static_cast<FRHIDomainShader*>(ShaderRHI);
 		VALIDATE_BOUND_SHADER(DomainShaderRHI);
 		StateCache.SetSamplerState<SF_Domain>(NewState, SamplerIndex);
 	}
 	break;
-	case FRHIShader::EType::Geometry:
+	case SF_Geometry:
 	{
 		FRHIGeometryShader* GeometryShaderRHI = static_cast<FRHIGeometryShader*>(ShaderRHI);
 		VALIDATE_BOUND_SHADER(GeometryShaderRHI);
 		StateCache.SetSamplerState<SF_Geometry>(NewState, SamplerIndex);
 	}
 	break;
-	case FRHIShader::EType::Pixel:
+	case SF_Pixel:
 	{
 		FRHIPixelShader* PixelShaderRHI = static_cast<FRHIPixelShader*>(ShaderRHI);
 		VALIDATE_BOUND_SHADER(PixelShaderRHI);
@@ -715,7 +715,7 @@ void FD3D12CommandContext::RHISetShaderSampler(FRHIGraphicsShader* ShaderRHI, ui
 	}
 	break;
 	default:
-		checkf(0, TEXT("Undefined FRHIShader Type %d!"), (int32)ShaderRHI->GetType());
+		checkf(0, TEXT("Undefined FRHIShader Type %d!"), (int32)ShaderRHI->GetFrequency());
 	}
 }
 
@@ -731,9 +731,9 @@ void FD3D12CommandContext::RHISetShaderUniformBuffer(FRHIGraphicsShader* ShaderR
 	//SCOPE_CYCLE_COUNTER(STAT_D3D12SetShaderUniformBuffer);
 	FD3D12UniformBuffer* Buffer = RetrieveObject<FD3D12UniformBuffer>(BufferRHI);
 	EShaderFrequency Stage = SF_NumFrequencies;
-	switch (ShaderRHI->GetType())
+	switch (ShaderRHI->GetFrequency())
 	{
-	case FRHIShader::EType::Vertex:
+	case SF_Vertex:
 	{
 		FRHIVertexShader* VertexShaderRHI = static_cast<FRHIVertexShader*>(ShaderRHI);
 		VALIDATE_BOUND_SHADER(VertexShaderRHI);
@@ -741,7 +741,7 @@ void FD3D12CommandContext::RHISetShaderUniformBuffer(FRHIGraphicsShader* ShaderR
 		Stage = SF_Vertex;
 	}
 	break;
-	case FRHIShader::EType::Hull:
+	case SF_Hull:
 	{
 		FRHIHullShader* HullShaderRHI = static_cast<FRHIHullShader*>(ShaderRHI);
 		VALIDATE_BOUND_SHADER(HullShaderRHI);
@@ -749,7 +749,7 @@ void FD3D12CommandContext::RHISetShaderUniformBuffer(FRHIGraphicsShader* ShaderR
 		Stage = SF_Hull;
 	}
 	break;
-	case FRHIShader::EType::Domain:
+	case SF_Domain:
 	{
 		FRHIDomainShader* DomainShaderRHI = static_cast<FRHIDomainShader*>(ShaderRHI);
 		VALIDATE_BOUND_SHADER(DomainShaderRHI);
@@ -757,7 +757,7 @@ void FD3D12CommandContext::RHISetShaderUniformBuffer(FRHIGraphicsShader* ShaderR
 		Stage = SF_Domain;
 	}
 	break;
-	case FRHIShader::EType::Geometry:
+	case SF_Geometry:
 	{
 		FRHIGeometryShader* GeometryShaderRHI = static_cast<FRHIGeometryShader*>(ShaderRHI);
 		VALIDATE_BOUND_SHADER(GeometryShaderRHI);
@@ -765,7 +765,7 @@ void FD3D12CommandContext::RHISetShaderUniformBuffer(FRHIGraphicsShader* ShaderR
 		Stage = SF_Geometry;
 	}
 	break;
-	case FRHIShader::EType::Pixel:
+	case SF_Pixel:
 	{
 		FRHIPixelShader* PixelShaderRHI = static_cast<FRHIPixelShader*>(ShaderRHI);
 		VALIDATE_BOUND_SHADER(PixelShaderRHI);
@@ -774,7 +774,7 @@ void FD3D12CommandContext::RHISetShaderUniformBuffer(FRHIGraphicsShader* ShaderR
 	}
 	break;
 	default:
-		checkf(0, TEXT("Undefined FRHIShader Type %d!"), (int32)ShaderRHI->GetType());
+		checkf(0, TEXT("Undefined FRHIShader Type %d!"), (int32)ShaderRHI->GetFrequency());
 		return;
 	}
 
@@ -808,37 +808,37 @@ void FD3D12CommandContext::RHISetShaderParameter(FRHIGraphicsShader* ShaderRHI,
 {
 	checkSlow(BufferIndex == 0);
 
-	switch (ShaderRHI->GetType())
+	switch (ShaderRHI->GetFrequency())
 	{
-	case FRHIShader::EType::Vertex:
+	case SF_Vertex:
 	{
 		FRHIVertexShader* VertexShaderRHI = static_cast<FRHIVertexShader*>(ShaderRHI);
 		VALIDATE_BOUND_SHADER(VertexShaderRHI);
 		VSConstantBuffer.UpdateConstant((const uint8*)NewValue, BaseIndex, NumBytes);
 	}
 	break;
-	case FRHIShader::EType::Hull:
+	case SF_Hull:
 	{
 		FRHIHullShader* HullShaderRHI = static_cast<FRHIHullShader*>(ShaderRHI);
 		VALIDATE_BOUND_SHADER(HullShaderRHI);
 		HSConstantBuffer.UpdateConstant((const uint8*)NewValue, BaseIndex, NumBytes);
 	}
 	break;
-	case FRHIShader::EType::Domain:
+	case SF_Domain:
 	{
 		FRHIDomainShader* DomainShaderRHI = static_cast<FRHIDomainShader*>(ShaderRHI);
 		VALIDATE_BOUND_SHADER(DomainShaderRHI);
 		DSConstantBuffer.UpdateConstant((const uint8*)NewValue, BaseIndex, NumBytes);
 	}
 	break;
-	case FRHIShader::EType::Geometry:
+	case SF_Geometry:
 	{
 		FRHIGeometryShader* GeometryShaderRHI = static_cast<FRHIGeometryShader*>(ShaderRHI);
 		VALIDATE_BOUND_SHADER(GeometryShaderRHI);
 		GSConstantBuffer.UpdateConstant((const uint8*)NewValue, BaseIndex, NumBytes);
 	}
 	break;
-	case FRHIShader::EType::Pixel:
+	case SF_Pixel:
 	{
 		FRHIPixelShader* PixelShaderRHI = static_cast<FRHIPixelShader*>(ShaderRHI);
 		VALIDATE_BOUND_SHADER(PixelShaderRHI);
@@ -846,7 +846,7 @@ void FD3D12CommandContext::RHISetShaderParameter(FRHIGraphicsShader* ShaderRHI,
 	}
 	break;
 	default:
-		checkf(0, TEXT("Undefined FRHIShader Type %d!"), (int32)ShaderRHI->GetType());
+		checkf(0, TEXT("Undefined FRHIShader Type %d!"), (int32)ShaderRHI->GetFrequency());
 	}
 }
 
diff --git a/Engine/Source/Runtime/D3D12RHI/Private/D3D12RHIPrivate.h b/Engine/Source/Runtime/D3D12RHI/Private/D3D12RHIPrivate.h
--- a/Engine/Source/Runtime/D3D12RHI/Private/D3D12RHIPrivate.h
+++ b/Engine/Source/Runtime/D3D12RHI/Private/D3D12RHIPrivate.h
@@ -263,12 +263,12 @@ public:
 	virtual FDepthStencilStateRHIRef RHICreateDepthStencilState(const FDepthStencilStateInitializerRHI& Initializer) final override;
 	virtual FBlendStateRHIRef RHICreateBlendState(const FBlendStateInitializerRHI& Initializer) final override;
 	virtual FVertexDeclarationRHIRef RHICreateVertexDeclaration(const FVertexDeclarationElementList& Elements) final override;
-	virtual FPixelShaderRHIRef RHICreatePixelShader(const TArray<uint8>& Code) final override;
-	virtual FVertexShaderRHIRef RHICreateVertexShader(const TArray<uint8>& Code) final override;
-	virtual FHullShaderRHIRef RHICreateHullShader(const TArray<uint8>& Code) final override;
-	virtual FDomainShaderRHIRef RHICreateDomainShader(const TArray<uint8>& Code) final override;
-	virtual FGeometryShaderRHIRef RHICreateGeometryShader(const TArray<uint8>& Code) final override;
-	virtual FComputeShaderRHIRef RHICreateComputeShader(const TArray<uint8>& Code) override;
+	virtual FPixelShaderRHIRef RHICreatePixelShader(TArrayView<const uint8> Code, const FSHAHash& Hash) final override;
+	virtual FVertexShaderRHIRef RHICreateVertexShader(TArrayView<const uint8> Code, const FSHAHash& Hash) final override;
+	virtual FHullShaderRHIRef RHICreateHullShader(TArrayView<const uint8> Code, const FSHAHash& Hash) final override;
+	virtual FDomainShaderRHIRef RHICreateDomainShader(TArrayView<const uint8> Code, const FSHAHash& Hash) final override;
+	virtual FGeometryShaderRHIRef RHICreateGeometryShader(TArrayView<const uint8> Code, const FSHAHash& Hash) final override;
+	virtual FComputeShaderRHIRef RHICreateComputeShader(TArrayView<const uint8> Code, const FSHAHash& Hash) override;
 	virtual FComputeFenceRHIRef RHICreateComputeFence(const FName& Name) final override;
 	virtual FGPUFenceRHIRef RHICreateGPUFence(const FName& Name) final override;
 	virtual FStagingBufferRHIRef RHICreateStagingBuffer() final override;
@@ -388,34 +388,34 @@ public:
 	// These will be un-commented as they are implemented.
 	//
 
-	virtual FVertexShaderRHIRef CreateVertexShader_RenderThread(class FRHICommandListImmediate& RHICmdList, const TArray<uint8>& Code) override final
+	virtual FVertexShaderRHIRef CreateVertexShader_RenderThread(class FRHICommandListImmediate& RHICmdList, TArrayView<const uint8> Code, const FSHAHash& Hash) override final
 	{
-		return RHICreateVertexShader(Code);
+		return RHICreateVertexShader(Code, Hash);
 	}
 
-	virtual FGeometryShaderRHIRef CreateGeometryShader_RenderThread(class FRHICommandListImmediate& RHICmdList, const TArray<uint8>& Code) override final
+	virtual FGeometryShaderRHIRef CreateGeometryShader_RenderThread(class FRHICommandListImmediate& RHICmdList, TArrayView<const uint8> Code, const FSHAHash& Hash) override final
 	{
-		return RHICreateGeometryShader(Code);
+		return RHICreateGeometryShader(Code, Hash);
 	}
 
-	virtual FHullShaderRHIRef CreateHullShader_RenderThread(class FRHICommandListImmediate& RHICmdList, const TArray<uint8>& Code) override final
+	virtual FHullShaderRHIRef CreateHullShader_RenderThread(class FRHICommandListImmediate& RHICmdList, TArrayView<const uint8> Code, const FSHAHash& Hash) override final
 	{
-		return RHICreateHullShader(Code);
+		return RHICreateHullShader(Code, Hash);
 	}
 
-	virtual FDomainShaderRHIRef CreateDomainShader_RenderThread(class FRHICommandListImmediate& RHICmdList, const TArray<uint8>& Code) override final
+	virtual FDomainShaderRHIRef CreateDomainShader_RenderThread(class FRHICommandListImmediate& RHICmdList, TArrayView<const uint8> Code, const FSHAHash& Hash) override final
 	{
-		return RHICreateDomainShader(Code);
+		return RHICreateDomainShader(Code, Hash);
 	}
 
-	virtual FPixelShaderRHIRef CreatePixelShader_RenderThread(class FRHICommandListImmediate& RHICmdList, const TArray<uint8>& Code) override final
+	virtual FPixelShaderRHIRef CreatePixelShader_RenderThread(class FRHICommandListImmediate& RHICmdList, TArrayView<const uint8> Code, const FSHAHash& Hash) override final
 	{
-		return RHICreatePixelShader(Code);
+		return RHICreatePixelShader(Code, Hash);
 	}
 
-	virtual FComputeShaderRHIRef CreateComputeShader_RenderThread(class FRHICommandListImmediate& RHICmdList, const TArray<uint8>& Code) override final
+	virtual FComputeShaderRHIRef CreateComputeShader_RenderThread(class FRHICommandListImmediate& RHICmdList, TArrayView<const uint8> Code, const FSHAHash& Hash) override final
 	{
-		return RHICreateComputeShader(Code);
+		return RHICreateComputeShader(Code, Hash);
 	}
 
 	virtual FVertexBufferRHIRef CreateVertexBuffer_RenderThread(class FRHICommandListImmediate& RHICmdList, uint32 Size, uint32 InUsage, FRHIResourceCreateInfo& CreateInfo);
@@ -460,7 +460,7 @@ public:
 
 	virtual FRayTracingGeometryRHIRef RHICreateRayTracingGeometry(const FRayTracingGeometryInitializer& Initializer) final override;
 	virtual FRayTracingSceneRHIRef RHICreateRayTracingScene(const FRayTracingSceneInitializer& Initializer) final override;
-	virtual FRayTracingShaderRHIRef RHICreateRayTracingShader(const TArray<uint8>& Code, EShaderFrequency ShaderFrequency) final override;
+	virtual FRayTracingShaderRHIRef RHICreateRayTracingShader(TArrayView<const uint8> Code, const FSHAHash& Hash, EShaderFrequency ShaderFrequency) final override;
 	virtual FRayTracingPipelineStateRHIRef RHICreateRayTracingPipelineState(const FRayTracingPipelineStateInitializer& Initializer) final override;
 
 #endif //D3D12_RHI_RAYTRACING
diff --git a/Engine/Source/Runtime/D3D12RHI/Private/D3D12RayTracing.cpp b/Engine/Source/Runtime/D3D12RHI/Private/D3D12RayTracing.cpp
--- a/Engine/Source/Runtime/D3D12RHI/Private/D3D12RayTracing.cpp
+++ b/Engine/Source/Runtime/D3D12RHI/Private/D3D12RayTracing.cpp
@@ -1428,9 +1428,9 @@ public:
 template<typename ShaderType>
 static FD3D12RayTracingShader* GetBuildInRayTracingShader()
 {
-	TShaderMap<FGlobalShaderType>* ShaderMap = GetGlobalShaderMap(ERHIFeatureLevel::SM5);
+	FGlobalShaderMap* ShaderMap = GetGlobalShaderMap(ERHIFeatureLevel::SM5);
 	auto Shader = ShaderMap->GetShader<ShaderType>();
-	FD3D12RayTracingShader* RayTracingShader = static_cast<FD3D12RayTracingShader*>(Shader->GetRayTracingShader());
+	FD3D12RayTracingShader* RayTracingShader = static_cast<FD3D12RayTracingShader*>(Shader.GetRayTracingShader());
 	return RayTracingShader;
 }
 
diff --git a/Engine/Source/Runtime/D3D12RHI/Private/D3D12RenderTarget.cpp b/Engine/Source/Runtime/D3D12RHI/Private/D3D12RenderTarget.cpp
--- a/Engine/Source/Runtime/D3D12RHI/Private/D3D12RenderTarget.cpp
+++ b/Engine/Source/Runtime/D3D12RHI/Private/D3D12RenderTarget.cpp
@@ -157,8 +157,8 @@ void FD3D12CommandContext::ResolveTextureUsingShader(
 	TShaderMapRef<FResolveVS> ResolveVertexShader(ShaderMap);
 	TShaderMapRef<TPixelShader> ResolvePixelShader(ShaderMap);
 
-	GraphicsPSOInit.BoundShaderState.VertexShaderRHI = GETSAFERHISHADER_VERTEX(*ResolveVertexShader);
-	GraphicsPSOInit.BoundShaderState.PixelShaderRHI = GETSAFERHISHADER_PIXEL(*ResolvePixelShader);
+	GraphicsPSOInit.BoundShaderState.VertexShaderRHI = ResolveVertexShader.GetVertexShader();
+	GraphicsPSOInit.BoundShaderState.PixelShaderRHI = ResolvePixelShader.GetPixelShader();
 	GraphicsPSOInit.PrimitiveType = PT_TriangleStrip;
 
 	SetGraphicsPipelineState(RHICmdList, GraphicsPSOInit, EApplyRendertargetOption::DoNothing);
diff --git a/Engine/Source/Runtime/D3D12RHI/Private/D3D12Shader.h b/Engine/Source/Runtime/D3D12RHI/Private/D3D12Shader.h
--- a/Engine/Source/Runtime/D3D12RHI/Private/D3D12Shader.h
+++ b/Engine/Source/Runtime/D3D12RHI/Private/D3D12Shader.h
@@ -190,6 +190,8 @@ public:
 class FD3D12RayTracingShader : public FRHIRayTracingShader, public FD3D12ShaderData
 {
 public:
+	explicit FD3D12RayTracingShader(EShaderFrequency InFrequency) : FRHIRayTracingShader(InFrequency) {}
+
 	/** The shader's bytecode. */
 	FD3D12ShaderBytecode ShaderBytecode;
 
diff --git a/Engine/Source/Runtime/D3D12RHI/Private/D3D12Shaders.cpp b/Engine/Source/Runtime/D3D12RHI/Private/D3D12Shaders.cpp
--- a/Engine/Source/Runtime/D3D12RHI/Private/D3D12Shaders.cpp
+++ b/Engine/Source/Runtime/D3D12RHI/Private/D3D12Shaders.cpp
@@ -65,12 +65,12 @@ static inline void InitUniformBufferStaticSlots(TShaderType* Shader)
 	}
 }
 
-FVertexShaderRHIRef FD3D12DynamicRHI::RHICreateVertexShader(const TArray<uint8>& Code)
+FVertexShaderRHIRef FD3D12DynamicRHI::RHICreateVertexShader(TArrayView<const uint8> Code, const FSHAHash& Hash)
 {
 	FShaderCodeReader ShaderCode(Code);
 	FD3D12VertexShader* Shader = new FD3D12VertexShader;
 
-	FMemoryReader Ar(Code, true);
+	FMemoryReaderView Ar(Code, true);
 	Ar << Shader->ShaderResourceTable;
 	int32 Offset = Ar.Tell();
 	const uint8* CodePtr = Code.GetData() + Offset;
@@ -93,13 +93,13 @@ FVertexShaderRHIRef FD3D12DynamicRHI::RHICreateVertexShader(const TArray<uint8>&
 	return Shader;
 }
 
-FPixelShaderRHIRef FD3D12DynamicRHI::RHICreatePixelShader(const TArray<uint8>& Code)
+FPixelShaderRHIRef FD3D12DynamicRHI::RHICreatePixelShader(TArrayView<const uint8> Code, const FSHAHash& Hash)
 {
 	FShaderCodeReader ShaderCode(Code);
 
 	FD3D12PixelShader* Shader = new FD3D12PixelShader;
 
-	FMemoryReader Ar(Code, true);
+	FMemoryReaderView Ar(Code, true);
 	Ar << Shader->ShaderResourceTable;
 	int32 Offset = Ar.Tell();
 	const uint8* CodePtr = Code.GetData() + Offset;
@@ -121,13 +121,13 @@ FPixelShaderRHIRef FD3D12DynamicRHI::RHICreatePixelShader(const TArray<uint8>& C
 	return Shader;
 }
 
-FHullShaderRHIRef FD3D12DynamicRHI::RHICreateHullShader(const TArray<uint8>& Code)
+FHullShaderRHIRef FD3D12DynamicRHI::RHICreateHullShader(TArrayView<const uint8> Code, const FSHAHash& Hash)
 {
 	FShaderCodeReader ShaderCode(Code);
 
 	FD3D12HullShader* Shader = new FD3D12HullShader;
 
-	FMemoryReader Ar(Code, true);
+	FMemoryReaderView Ar(Code, true);
 	Ar << Shader->ShaderResourceTable;
 	int32 Offset = Ar.Tell();
 	const uint8* CodePtr = Code.GetData() + Offset;
@@ -149,13 +149,13 @@ FHullShaderRHIRef FD3D12DynamicRHI::RHICreateHullShader(const TArray<uint8>& Cod
 	return Shader;
 }
 
-FDomainShaderRHIRef FD3D12DynamicRHI::RHICreateDomainShader(const TArray<uint8>& Code)
+FDomainShaderRHIRef FD3D12DynamicRHI::RHICreateDomainShader(TArrayView<const uint8> Code, const FSHAHash& Hash)
 {
 	FShaderCodeReader ShaderCode(Code);
 
 	FD3D12DomainShader* Shader = new FD3D12DomainShader;
 
-	FMemoryReader Ar(Code, true);
+	FMemoryReaderView Ar(Code, true);
 	Ar << Shader->ShaderResourceTable;
 	int32 Offset = Ar.Tell();
 	const uint8* CodePtr = Code.GetData() + Offset;
@@ -177,13 +177,13 @@ FDomainShaderRHIRef FD3D12DynamicRHI::RHICreateDomainShader(const TArray<uint8>&
 	return Shader;
 }
 
-FGeometryShaderRHIRef FD3D12DynamicRHI::RHICreateGeometryShader(const TArray<uint8>& Code)
+FGeometryShaderRHIRef FD3D12DynamicRHI::RHICreateGeometryShader(TArrayView<const uint8> Code, const FSHAHash& Hash)
 {
 	FShaderCodeReader ShaderCode(Code);
 
 	FD3D12GeometryShader* Shader = new FD3D12GeometryShader;
 
-	FMemoryReader Ar(Code, true);
+	FMemoryReaderView Ar(Code, true);
 	Ar << Shader->ShaderResourceTable;
 	int32 Offset = Ar.Tell();
 	const uint8* CodePtr = Code.GetData() + Offset;
@@ -205,13 +205,13 @@ FGeometryShaderRHIRef FD3D12DynamicRHI::RHICreateGeometryShader(const TArray<uin
 	return Shader;
 }
 
-FComputeShaderRHIRef FD3D12DynamicRHI::RHICreateComputeShader(const TArray<uint8>& Code)
+FComputeShaderRHIRef FD3D12DynamicRHI::RHICreateComputeShader(TArrayView<const uint8> Code, const FSHAHash& Hash)
 {
 	FShaderCodeReader ShaderCode(Code);
 
 	FD3D12ComputeShader* Shader = new FD3D12ComputeShader;
 
-	FMemoryReader Ar(Code, true);
+	FMemoryReaderView Ar(Code, true);
 	Ar << Shader->ShaderResourceTable;
 	int32 Offset = Ar.Tell();
 	const uint8* CodePtr = Code.GetData() + Offset;
@@ -246,12 +246,12 @@ FComputeShaderRHIRef FD3D12DynamicRHI::RHICreateComputeShader(const TArray<uint8
 
 #if D3D12_RHI_RAYTRACING
 
-FRayTracingShaderRHIRef FD3D12DynamicRHI::RHICreateRayTracingShader(const TArray<uint8>& Code, EShaderFrequency ShaderFrequency)
+FRayTracingShaderRHIRef FD3D12DynamicRHI::RHICreateRayTracingShader(TArrayView<const uint8> Code, const FSHAHash& Hash, EShaderFrequency ShaderFrequency)
 {
 	FShaderCodeReader ShaderCode(Code);
-	FD3D12RayTracingShader* Shader = new FD3D12RayTracingShader;
+	FD3D12RayTracingShader* Shader = new FD3D12RayTracingShader(ShaderFrequency);
 
-	FMemoryReader Ar(Code, true);
+	FMemoryReaderView Ar(Code, true);
 	Ar << Shader->ShaderResourceTable;
 	Ar << Shader->EntryPoint;
 	Ar << Shader->AnyHitEntryPoint;
diff --git a/Engine/Source/Runtime/D3D12RHI/Private/D3D12Texture.cpp b/Engine/Source/Runtime/D3D12RHI/Private/D3D12Texture.cpp
--- a/Engine/Source/Runtime/D3D12RHI/Private/D3D12Texture.cpp
+++ b/Engine/Source/Runtime/D3D12RHI/Private/D3D12Texture.cpp
@@ -74,8 +74,11 @@ template void TD3D12Texture2D<FD3D12BaseTexture2DArray>::GetReadBackHeapDesc(D3D
 template void TD3D12Texture2D<FD3D12BaseTextureCube>::GetReadBackHeapDesc(D3D12_PLACED_SUBRESOURCE_FOOTPRINT& OutFootprint, uint32 Subresource) const;
 
 /// @endcond
-
-struct FRHICommandUpdateTexture final : public FRHICommand<FRHICommandUpdateTexture>
+struct FRHICommandUpdateTextureString
+{
+	static const TCHAR* TStr() { return TEXT("FRHICommandUpdateTexture"); }
+};
+struct FRHICommandUpdateTexture final : public FRHICommand<FRHICommandUpdateTexture, FRHICommandUpdateTextureString>
 {
 	FD3D12TextureBase* TextureBase;
 	D3D12_TEXTURE_COPY_LOCATION DestCopyLocation;
@@ -114,7 +117,11 @@ struct FRHICommandUpdateTexture final : public FRHICommand<FRHICommandUpdateText
 	}
 };
 
-struct FRHICommandCopySubTextureRegion final : public FRHICommand<FRHICommandCopySubTextureRegion>
+struct FRHICommandCopySubTextureRegionString
+{
+	static const TCHAR* TStr() { return TEXT("FRHICommandCopySubTextureRegion"); }
+};
+struct FRHICommandCopySubTextureRegion final : public FRHICommand<FRHICommandCopySubTextureRegion, FRHICommandCopySubTextureRegionString>
 {
 	FD3D12TextureBase* DestTexture;
 	uint32 DestX;
@@ -143,7 +150,11 @@ struct FRHICommandCopySubTextureRegion final : public FRHICommand<FRHICommandCop
 	}
 };
 
-struct FD3D12RHICommandInitializeTexture final : public FRHICommand<FD3D12RHICommandInitializeTexture>
+struct FD3D12RHICommandInitializeTextureString
+{
+	static const TCHAR* TStr() { return TEXT("FD3D12RHICommandInitializeTexture"); }
+};
+struct FD3D12RHICommandInitializeTexture final : public FRHICommand<FD3D12RHICommandInitializeTexture, FD3D12RHICommandInitializeTextureString>
 {
 	FD3D12TextureBase* TextureBase;
 	FD3D12ResourceLocation SrcResourceLoc;
@@ -2067,7 +2078,11 @@ void FD3D12DynamicRHI::EndUpdateTexture3D_RenderThread(class FRHICommandListImme
 	EndUpdateTexture3D_Internal(UpdateData);
 }
 
-class FD3D12RHICmdEndMultiUpdateTexture3D : public FRHICommand<FD3D12RHICmdEndMultiUpdateTexture3D>
+struct FD3D12RHICmdEndMultiUpdateTexture3DString
+{
+	static const TCHAR* TStr() { return TEXT("FD3D12RHICmdEndMultiUpdateTexture3D"); }
+};
+class FD3D12RHICmdEndMultiUpdateTexture3D : public FRHICommand<FD3D12RHICmdEndMultiUpdateTexture3D, FD3D12RHICmdEndMultiUpdateTexture3DString>
 {
 public:
 	FD3D12RHICmdEndMultiUpdateTexture3D(TArray<FUpdateTexture3DData>& UpdateDataArray) :
@@ -2308,7 +2323,11 @@ FUpdateTexture3DData FD3D12DynamicRHI::BeginUpdateTexture3D_Internal(FRHITexture
 	return UpdateData;
 }
 
-class FD3D12RHICmdEndUpdateTexture3D : public FRHICommand<FD3D12RHICmdEndUpdateTexture3D>
+struct FD3D12RHICmdEndUpdateTexture3DString
+{
+	static const TCHAR* TStr() { return TEXT("FD3D12RHICmdEndUpdateTexture3D"); }
+};
+class FD3D12RHICmdEndUpdateTexture3D : public FRHICommand<FD3D12RHICmdEndUpdateTexture3D, FD3D12RHICmdEndUpdateTexture3DString>
 {
 public:
 	FD3D12RHICmdEndUpdateTexture3D(FUpdateTexture3DData& UpdateData) :
diff --git a/Engine/Source/Runtime/D3D12RHI/Private/D3D12UniformBuffer.cpp b/Engine/Source/Runtime/D3D12RHI/Private/D3D12UniformBuffer.cpp
--- a/Engine/Source/Runtime/D3D12RHI/Private/D3D12UniformBuffer.cpp
+++ b/Engine/Source/Runtime/D3D12RHI/Private/D3D12UniformBuffer.cpp
@@ -121,7 +121,11 @@ FUniformBufferRHIRef FD3D12DynamicRHI::RHICreateUniformBuffer(const void* Conten
 	return UniformBufferOut;
 }
 
-struct FRHICommandD3D12UpdateUniformBuffer final : public FRHICommand<FRHICommandD3D12UpdateUniformBuffer>
+struct FRHICommandD3D12UpdateUniformBufferString
+{
+	static const TCHAR* TStr() { return TEXT("FRHICommandD3D12UpdateUniformBuffer"); }
+};
+struct FRHICommandD3D12UpdateUniformBuffer final : public FRHICommand<FRHICommandD3D12UpdateUniformBuffer, FRHICommandD3D12UpdateUniformBufferString>
 {
 	FD3D12UniformBuffer* UniformBuffer;
 	FD3D12ResourceLocation UpdatedLocation;
@@ -195,7 +199,7 @@ void FD3D12DynamicRHI::RHIUpdateUniformBuffer(FRHIUniformBuffer* UniformBufferRH
 			}
 
 			checkf(Resource, TEXT("Invalid resource entry creating uniform buffer, %s.Resources[%u], ResourceType 0x%x."),
-				*Layout.GetDebugName().ToString(),
+				*Layout.GetDebugName(),
 				ResourceIndex,
 				Layout.Resources[ResourceIndex].MemberType);
 
diff --git a/Engine/Source/Runtime/D3D12RHI/Private/D3D12Viewport.cpp b/Engine/Source/Runtime/D3D12RHI/Private/D3D12Viewport.cpp
--- a/Engine/Source/Runtime/D3D12RHI/Private/D3D12Viewport.cpp
+++ b/Engine/Source/Runtime/D3D12RHI/Private/D3D12Viewport.cpp
@@ -945,7 +945,11 @@ void FD3D12CommandContextBase::RHIEndDrawingViewport(FRHIViewport* ViewportRHI,
 	}
 }
 
-struct FRHICommandSignalFrameFence final : public FRHICommand<FRHICommandSignalFrameFence>
+struct FRHICommandSignalFrameFenceString
+{
+	static const TCHAR* TStr() { return TEXT("FRHICommandSignalFrameFence"); }
+};
+struct FRHICommandSignalFrameFence final : public FRHICommand<FRHICommandSignalFrameFence, FRHICommandSignalFrameFenceString>
 {
 	ED3D12CommandQueueType QueueType;
 	FD3D12ManualFence* const Fence;
diff --git a/Engine/Source/Runtime/EmptyRHI/Private/EmptyShaders.cpp b/Engine/Source/Runtime/EmptyRHI/Private/EmptyShaders.cpp
--- a/Engine/Source/Runtime/EmptyRHI/Private/EmptyShaders.cpp
+++ b/Engine/Source/Runtime/EmptyRHI/Private/EmptyShaders.cpp
@@ -9,7 +9,7 @@
 
 /** Initialization constructor. */
 template<typename BaseResourceType>
-TEmptyBaseShader<BaseResourceType>::TEmptyBaseShader(const TArray<uint8>& InCode)
+TEmptyBaseShader<BaseResourceType>::TEmptyBaseShader(TArrayView<const uint8> Code)
 {
 
 }
@@ -23,37 +23,37 @@ TEmptyBaseShader<BaseResourceType>::~TEmptyBaseShader()
 }
 
 
-FVertexShaderRHIRef FEmptyDynamicRHI::RHICreateVertexShader(const TArray<uint8>& Code)
+FVertexShaderRHIRef FEmptyDynamicRHI::RHICreateVertexShader(TArrayView<const uint8> Code, const FSHAHash& Hash)
 {
 	FEmptyVertexShader* Shader = new FEmptyVertexShader(Code);
 	return Shader;
 }
 
-FPixelShaderRHIRef FEmptyDynamicRHI::RHICreatePixelShader(const TArray<uint8>& Code)
+FPixelShaderRHIRef FEmptyDynamicRHI::RHICreatePixelShader(TArrayView<const uint8> Code, const FSHAHash& Hash)
 {
 	FEmptyPixelShader* Shader = new FEmptyPixelShader(Code);
 	return Shader;
 }
 
-FHullShaderRHIRef FEmptyDynamicRHI::RHICreateHullShader(const TArray<uint8>& Code) 
+FHullShaderRHIRef FEmptyDynamicRHI::RHICreateHullShader(TArrayView<const uint8> Code, const FSHAHash& Hash)
 { 
 	FEmptyHullShader* Shader = new FEmptyHullShader(Code);
 	return Shader;
 }
 
-FDomainShaderRHIRef FEmptyDynamicRHI::RHICreateDomainShader(const TArray<uint8>& Code) 
+FDomainShaderRHIRef FEmptyDynamicRHI::RHICreateDomainShader(TArrayView<const uint8> Code, const FSHAHash& Hash)
 { 
 	FEmptyDomainShader* Shader = new FEmptyDomainShader(Code);
 	return Shader;
 }
 
-FGeometryShaderRHIRef FEmptyDynamicRHI::RHICreateGeometryShader(const TArray<uint8>& Code) 
+FGeometryShaderRHIRef FEmptyDynamicRHI::RHICreateGeometryShader(TArrayView<const uint8> Code, const FSHAHash& Hash)
 { 
 	FEmptyGeometryShader* Shader = new FEmptyGeometryShader(Code);
 	return Shader;
 }
 
-FComputeShaderRHIRef FEmptyDynamicRHI::RHICreateComputeShader(const TArray<uint8>& Code) 
+FComputeShaderRHIRef FEmptyDynamicRHI::RHICreateComputeShader(TArrayView<const uint8> Code, const FSHAHash& Hash)
 { 
 	FEmptyComputeShader* Shader = new FEmptyComputeShader(Code);
 	return Shader;
diff --git a/Engine/Source/Runtime/EmptyRHI/Public/EmptyRHI.h b/Engine/Source/Runtime/EmptyRHI/Public/EmptyRHI.h
--- a/Engine/Source/Runtime/EmptyRHI/Public/EmptyRHI.h
+++ b/Engine/Source/Runtime/EmptyRHI/Public/EmptyRHI.h
@@ -41,12 +41,12 @@ public:
 	virtual FDepthStencilStateRHIRef RHICreateDepthStencilState(const FDepthStencilStateInitializerRHI& Initializer) final override;
 	virtual FBlendStateRHIRef RHICreateBlendState(const FBlendStateInitializerRHI& Initializer) final override;
 	virtual FVertexDeclarationRHIRef RHICreateVertexDeclaration(const FVertexDeclarationElementList& Elements) final override;
-	virtual FPixelShaderRHIRef RHICreatePixelShader(const TArray<uint8>& Code) final override;
-	virtual FVertexShaderRHIRef RHICreateVertexShader(const TArray<uint8>& Code) final override;
-	virtual FHullShaderRHIRef RHICreateHullShader(const TArray<uint8>& Code) final override;
-	virtual FDomainShaderRHIRef RHICreateDomainShader(const TArray<uint8>& Code) final override;
-	virtual FGeometryShaderRHIRef RHICreateGeometryShader(const TArray<uint8>& Code) final override;
-	virtual FComputeShaderRHIRef RHICreateComputeShader(const TArray<uint8>& Code) final override;
+	virtual FPixelShaderRHIRef RHICreatePixelShader(TArrayView<const uint8> Code, const FSHAHash& Hash) final override;
+	virtual FVertexShaderRHIRef RHICreateVertexShader(TArrayView<const uint8> Code, const FSHAHash& Hash) final override;
+	virtual FHullShaderRHIRef RHICreateHullShader(TArrayView<const uint8> Code, const FSHAHash& Hash) final override;
+	virtual FDomainShaderRHIRef RHICreateDomainShader(TArrayView<const uint8> Code, const FSHAHash& Hash) final override;
+	virtual FGeometryShaderRHIRef RHICreateGeometryShader(TArrayView<const uint8> Code, const FSHAHash& Hash) final override;
+	virtual FComputeShaderRHIRef RHICreateComputeShader(TArrayView<const uint8> Code, const FSHAHash& Hash) final override;
 	virtual FBoundShaderStateRHIRef RHICreateBoundShaderState(FRHIVertexDeclaration* VertexDeclaration, FRHIVertexShader* VertexShader, FRHIHullShader* HullShader, FRHIDomainShaderFRHIDomainShader*, FRHIPixelShader* PixelShader, FRHIGeometryShader* GeometryShader) final override;
 	virtual FUniformBufferRHIRef RHICreateUniformBuffer(const void* Contents, const FRHIUniformBufferLayout& Layout, EUniformBufferUsage Usage, EUniformBufferValidation Validation) final override;
 	virtual FIndexBufferRHIRef RHICreateIndexBuffer(uint32 Stride, uint32 Size, uint32 InUsage, FRHIResourceCreateInfo& CreateInfo) final override;
diff --git a/Engine/Source/Runtime/EmptyRHI/Public/EmptyShaderResources.h b/Engine/Source/Runtime/EmptyRHI/Public/EmptyShaderResources.h
--- a/Engine/Source/Runtime/EmptyRHI/Public/EmptyShaderResources.h
+++ b/Engine/Source/Runtime/EmptyRHI/Public/EmptyShaderResources.h
@@ -18,7 +18,7 @@ public:
 	TEmptyBaseShader()
 	{
 	}
-	TEmptyBaseShader(const TArray<uint8>& InCode);
+	TEmptyBaseShader(TArrayView<const uint8> Code);
 
 	/** Destructor */
 	virtual ~TEmptyBaseShader();
diff --git a/Engine/Source/Runtime/Engine/Classes/Components/ActorComponent.h b/Engine/Source/Runtime/Engine/Classes/Components/ActorComponent.h
--- a/Engine/Source/Runtime/Engine/Classes/Components/ActorComponent.h
+++ b/Engine/Source/Runtime/Engine/Classes/Components/ActorComponent.h
@@ -23,6 +23,11 @@ class UActorComponent;
 class UAssetUserData;
 class ULevel;
 
+struct FRegisterComponentContext
+{
+	TArray<class UPrimitiveComponent*> AddPrimitiveBatches;
+};
+
 #if WITH_EDITOR
 class SWidget;
 struct FMinimalViewInfo;
@@ -515,7 +520,7 @@ private:
 	void ExecuteUnregisterEvents();
 
 	/** Calls OnRegister, CreateRenderState_Concurrent and OnCreatePhysicsState. */
-	void ExecuteRegisterEvents();
+	void ExecuteRegisterEvents(FRegisterComponentContext* Context = nullptr);
 
 	/** Utility function for each of the PostEditChange variations to call for the same behavior */
 	void ConsolidatedPostEditChange(const FPropertyChangedEvent& PropertyChangedEvent);
@@ -541,7 +546,7 @@ protected:
 	 * Used to create any rendering thread information for this component
 	 * @warning This is called concurrently on multiple threads (but never the same component concurrently)
 	 */
-	virtual void CreateRenderState_Concurrent();
+	virtual void CreateRenderState_Concurrent(FRegisterComponentContext* Context);
 
 	/** 
 	 * Called to send a transform update for this component to the rendering thread
@@ -677,7 +682,7 @@ public:
 	 * Registers a component with a specific world, which creates any visual/physical state
 	 * @param InWorld - The world to register the component with.
 	 */
-	void RegisterComponentWithWorld(UWorld* InWorld);
+	void RegisterComponentWithWorld(UWorld* InWorld, FRegisterComponentContext* Context = nullptr);
 
 	/** Overridable check for a component to indicate to its Owner that it should prevent the Actor from auto destroying when finished */
 	virtual bool IsReadyForOwnerToAutoDestroy() const { return true; }
diff --git a/Engine/Source/Runtime/Engine/Classes/Components/SkinnedMeshComponent.h b/Engine/Source/Runtime/Engine/Classes/Components/SkinnedMeshComponent.h
--- a/Engine/Source/Runtime/Engine/Classes/Components/SkinnedMeshComponent.h
+++ b/Engine/Source/Runtime/Engine/Classes/Components/SkinnedMeshComponent.h
@@ -815,7 +815,7 @@ protected:
 	//~ Begin UActorComponent Interface
 	virtual void OnRegister() override;
 	virtual void OnUnregister() override;
-	virtual void CreateRenderState_Concurrent() override;
+	virtual void CreateRenderState_Concurrent(FRegisterComponentContext* Context) override;
 	virtual void SendRenderDynamicData_Concurrent() override;
 	virtual void DestroyRenderState_Concurrent() override;
 	virtual bool RequiresGameThreadEndOfFrameRecreate() const override
@@ -1558,7 +1558,7 @@ public:
 
 		if (bWasRenderStateCreated && bIsRegistered)
 		{
-			Component->CreateRenderState_Concurrent();
+			Component->CreateRenderState_Concurrent(nullptr);
 		}
 	}
 };
diff --git a/Engine/Source/Runtime/Engine/Classes/DeviceProfiles/DeviceProfileManager.h b/Engine/Source/Runtime/Engine/Classes/DeviceProfiles/DeviceProfileManager.h
--- a/Engine/Source/Runtime/Engine/Classes/DeviceProfiles/DeviceProfileManager.h
+++ b/Engine/Source/Runtime/Engine/Classes/DeviceProfiles/DeviceProfileManager.h
@@ -82,13 +82,6 @@ public:
 	void RestoreDefaultDeviceProfile();
 
 
-	/**
-	 * Get the device profile .ini name.
-	 *
-	 * @return the device profile .ini name.
-	 */
-	const FString GetDeviceProfileIniName() const;
-
 	/**
 	 * Load the device profiles from the config file.
 	 */
diff --git a/Engine/Source/Runtime/Engine/Classes/Engine/Level.h b/Engine/Source/Runtime/Engine/Classes/Engine/Level.h
--- a/Engine/Source/Runtime/Engine/Classes/Engine/Level.h
+++ b/Engine/Source/Runtime/Engine/Classes/Engine/Level.h
@@ -764,7 +764,7 @@ public:
 	 * creates the BSP model components.
 	 * @param bRerunConstructionScripts	If we want to rerun construction scripts on actors in level
 	 */
-	ENGINE_API void UpdateLevelComponents(bool bRerunConstructionScripts);
+	ENGINE_API void UpdateLevelComponents(bool bRerunConstructionScripts, struct FRegisterComponentContext* Context = nullptr);
 
 	/**
 	 * Incrementally updates all components of actors associated with this level.
@@ -772,7 +772,7 @@ public:
 	 * @param NumComponentsToUpdate		Number of components to update in this run, 0 for all
 	 * @param bRerunConstructionScripts	If we want to rerun construction scripts on actors in level
 	 */
-	void IncrementalUpdateComponents( int32 NumComponentsToUpdate, bool bRerunConstructionScripts );
+	void IncrementalUpdateComponents( int32 NumComponentsToUpdate, bool bRerunConstructionScripts, struct FRegisterComponentContext* Context = nullptr);
 
 	/**
 	* Incrementally unregisters all components of actors associated with this level.
diff --git a/Engine/Source/Runtime/Engine/Classes/Engine/VolumeTexture.h b/Engine/Source/Runtime/Engine/Classes/Engine/VolumeTexture.h
--- a/Engine/Source/Runtime/Engine/Classes/Engine/VolumeTexture.h
+++ b/Engine/Source/Runtime/Engine/Classes/Engine/VolumeTexture.h
@@ -116,7 +116,7 @@ public:
 
 #endif
 
-	ENGINE_API static bool ShaderPlatformSupportsCompression(EShaderPlatform ShaderPlatform);
+	ENGINE_API static bool ShaderPlatformSupportsCompression(FStaticShaderPlatform ShaderPlatform);
 
 protected:
 
diff --git a/Engine/Source/Runtime/Engine/Classes/Engine/World.h b/Engine/Source/Runtime/Engine/Classes/Engine/World.h
--- a/Engine/Source/Runtime/Engine/Classes/Engine/World.h
+++ b/Engine/Source/Runtime/Engine/Classes/Engine/World.h
@@ -2425,7 +2425,7 @@ public:
 	 * @param	bRerunConstructionScripts	If we should rerun construction scripts on actors
 	 * @param	bCurrentLevelOnly			If true, affect only the current level.
 	 */
-	void UpdateWorldComponents(bool bRerunConstructionScripts, bool bCurrentLevelOnly);
+	void UpdateWorldComponents(bool bRerunConstructionScripts, bool bCurrentLevelOnly, FRegisterComponentContext* Context = nullptr);
 
 	/**
 	 * Updates cull distance volumes for a specified component or a specified actor or all actors
@@ -3169,7 +3169,7 @@ public:
 	 * @param InURL commandline URL
 	 * @param bResetTime (optional) whether the WorldSettings's TimeSeconds should be reset to zero
 	 */
-	void InitializeActorsForPlay(const FURL& InURL, bool bResetTime = true);
+	void InitializeActorsForPlay(const FURL& InURL, bool bResetTime = true, FRegisterComponentContext* Context = nullptr);
 
 	/**
 	 * Start gameplay. This will cause the game mode to transition to the correct state and call BeginPlay on all actors
diff --git a/Engine/Source/Runtime/Engine/Classes/GameFramework/Actor.h b/Engine/Source/Runtime/Engine/Classes/GameFramework/Actor.h
--- a/Engine/Source/Runtime/Engine/Classes/GameFramework/Actor.h
+++ b/Engine/Source/Runtime/Engine/Classes/GameFramework/Actor.h
@@ -2352,7 +2352,7 @@ public:
 	 * @param NumComponentsToRegister  Number of components to register in this run, 0 for all
 	 * @return true when all components were registered for this actor
 	 */
-	bool IncrementalRegisterComponents(int32 NumComponentsToRegister);
+	bool IncrementalRegisterComponents(int32 NumComponentsToRegister, FRegisterComponentContext* Context = nullptr);
 
 	/** Flags all component's render state as dirty	 */
 	void MarkComponentsRenderStateDirty();
