Fixed FAsyncPurge::TickGameThreadObjects not respecting GC time limit which could lead to bigger hitches on the game thread (regression)

diff --git a/Engine/Source/Runtime/CoreUObject/Private/UObject/GarbageCollection.cpp b/Engine/Source/Runtime/CoreUObject/Private/UObject/GarbageCollection.cpp
--- a/Engine/Source/Runtime/CoreUObject/Private/UObject/GarbageCollection.cpp
+++ b/Engine/Source/Runtime/CoreUObject/Private/UObject/GarbageCollection.cpp
@@ -341,6 +341,7 @@ class FAsyncPurge : public FRunnable
 		{
 			Object->~UObject();
 			GUObjectAllocator.FreeUObject(Object);
+			ProcessedObjectsCount++;
 
 			if (bUseTimeLimit && (ProcessedObjectsCount == TimeLimitEnforcementGranularityForDeletion) && !GameThreadObjects.IsEmpty())
 			{
