Fix for warning after upgrading to PVS-Studio 7.0.4:

V501: There are identical sub-expressions 'TPointerIsConvertibleFromTo < PropertyType, OtherPropertyType >::Value' to the left and to the right of the '||' operator.

#rnx

diff --git a/Engine/Source/Runtime/CoreUObject/Public/UObject/FieldPath.h b/Engine/Source/Runtime/CoreUObject/Public/UObject/FieldPath.h
--- a/Engine/Source/Runtime/CoreUObject/Public/UObject/FieldPath.h
+++ b/Engine/Source/Runtime/CoreUObject/Public/UObject/FieldPath.h
@@ -353,7 +353,7 @@ public:
 	FORCEINLINE bool operator==(const TFieldPath<OtherPropertyType> &Other) const
 	{
 		static_assert(TPointerIsConvertibleFromTo<OtherPropertyType, FField>::Value, "TFieldPath can only be compared with FField types");
-		static_assert(TPointerIsConvertibleFromTo<PropertyType, OtherPropertyType>::Value || TPointerIsConvertibleFromTo<PropertyType, OtherPropertyType>::Value, "Unable to compare TFieldPath with raw pointer - types are incompatible");
+		static_assert(TPointerIsConvertibleFromTo<PropertyType, OtherPropertyType>::Value, "Unable to compare TFieldPath with raw pointer - types are incompatible");
 
 		return Path == Other.Path;
 	}
@@ -366,7 +366,7 @@ public:
 	FORCEINLINE bool operator!=(const TFieldPath<OtherPropertyType> &Other) const
 	{
 		static_assert(TPointerIsConvertibleFromTo<OtherPropertyType, FField>::Value, "TFieldPath can only be compared with FField types");
-		static_assert(TPointerIsConvertibleFromTo<PropertyType, OtherPropertyType>::Value || TPointerIsConvertibleFromTo<PropertyType, OtherPropertyType>::Value, "Unable to compare TFieldPath with raw pointer - types are incompatible");
+		static_assert(TPointerIsConvertibleFromTo<PropertyType, OtherPropertyType>::Value, "Unable to compare TFieldPath with raw pointer - types are incompatible");
 
 		return Path != Other.Path;
 	}
@@ -379,7 +379,7 @@ public:
 	FORCEINLINE bool operator==(const OtherPropertyType* Other) const
 	{
 		static_assert(TPointerIsConvertibleFromTo<OtherPropertyType, FField>::Value, "TFieldPath can only be compared with FField types");
-		static_assert(TPointerIsConvertibleFromTo<PropertyType, OtherPropertyType>::Value || TPointerIsConvertibleFromTo<PropertyType, OtherPropertyType>::Value, "Unable to compare TFieldPath with raw pointer - types are incompatible");
+		static_assert(TPointerIsConvertibleFromTo<PropertyType, OtherPropertyType>::Value, "Unable to compare TFieldPath with raw pointer - types are incompatible");
 
 		return Get() == Other;
 	}
@@ -392,7 +392,7 @@ public:
 	FORCEINLINE bool operator!=(const OtherPropertyType* Other) const
 	{
 		static_assert(TPointerIsConvertibleFromTo<OtherPropertyType, FField>::Value, "TFieldPath can only be compared with FField types");
-		static_assert(TPointerIsConvertibleFromTo<PropertyType, OtherPropertyType>::Value || TPointerIsConvertibleFromTo<PropertyType, OtherPropertyType>::Value, "Unable to compare TFieldPath with raw pointer - types are incompatible");
+		static_assert(TPointerIsConvertibleFromTo<PropertyType, OtherPropertyType>::Value, "Unable to compare TFieldPath with raw pointer - types are incompatible");
 
 		return Get() != Other;
 	}
