Fix for warning after upgrading to PVS-Studio 7.0.4:

V501: There are identical sub-expressions 'TPointerIsConvertibleFromTo < T, TOther >::Value' to the left and to the right of the '||' operator.

#rnx

diff --git a/Engine/Source/Runtime/CoreUObject/Public/UObject/WeakFieldPtr.h b/Engine/Source/Runtime/CoreUObject/Public/UObject/WeakFieldPtr.h
--- a/Engine/Source/Runtime/CoreUObject/Public/UObject/WeakFieldPtr.h
+++ b/Engine/Source/Runtime/CoreUObject/Public/UObject/WeakFieldPtr.h
@@ -227,7 +227,7 @@ public:
 	FORCEINLINE bool operator==(const TWeakFieldPtr<TOther> &Other) const
 	{
 		static_assert(TPointerIsConvertibleFromTo<TOther, FField>::Value, "TWeakFieldPtr can only be compared with FField types");
-		static_assert(TPointerIsConvertibleFromTo<T, TOther>::Value || TPointerIsConvertibleFromTo<T, TOther>::Value, "Unable to compare TWeakFieldPtr with raw pointer - types are incompatible");
+		static_assert(TPointerIsConvertibleFromTo<T, TOther>::Value, "Unable to compare TWeakFieldPtr with raw pointer - types are incompatible");
 
 		return Field == Other.Field;
 	}
@@ -240,7 +240,7 @@ public:
 	FORCEINLINE bool operator!=(const TWeakFieldPtr<TOther> &Other) const
 	{
 		static_assert(TPointerIsConvertibleFromTo<TOther, FField>::Value, "TWeakFieldPtr can only be compared with FField types");
-		static_assert(TPointerIsConvertibleFromTo<T, TOther>::Value || TPointerIsConvertibleFromTo<T, TOther>::Value, "Unable to compare TWeakFieldPtr with raw pointer - types are incompatible");
+		static_assert(TPointerIsConvertibleFromTo<T, TOther>::Value, "Unable to compare TWeakFieldPtr with raw pointer - types are incompatible");
 
 		return Field != Other.Field;
 	}
@@ -253,7 +253,7 @@ public:
 	FORCEINLINE bool operator==(const TOther* Other) const
 	{
 		static_assert(TPointerIsConvertibleFromTo<TOther, FField>::Value, "TWeakFieldPtr can only be compared with FField types");
-		static_assert(TPointerIsConvertibleFromTo<T, TOther>::Value || TPointerIsConvertibleFromTo<T, TOther>::Value, "Unable to compare TWeakFieldPtr with raw pointer - types are incompatible");
+		static_assert(TPointerIsConvertibleFromTo<T, TOther>::Value, "Unable to compare TWeakFieldPtr with raw pointer - types are incompatible");
 
 		return Field == Other;
 	}
@@ -266,7 +266,7 @@ public:
 	FORCEINLINE bool operator!=(const TOther* Other) const
 	{
 		static_assert(TPointerIsConvertibleFromTo<TOther, FField>::Value, "TWeakFieldPtr can only be compared with FField types");
-		static_assert(TPointerIsConvertibleFromTo<T, TOther>::Value || TPointerIsConvertibleFromTo<T, TOther>::Value, "Unable to compare TWeakFieldPtr with raw pointer - types are incompatible");
+		static_assert(TPointerIsConvertibleFromTo<T, TOther>::Value, "Unable to compare TWeakFieldPtr with raw pointer - types are incompatible");
 
 		return Field != Other;
 	}
