Don't add multiple entries for the same c-style array property in the Class Net Cache.

diff --git a/Engine/Source/Runtime/CoreUObject/Private/UObject/CoreNet.cpp b/Engine/Source/Runtime/CoreUObject/Private/UObject/CoreNet.cpp
--- a/Engine/Source/Runtime/CoreUObject/Private/UObject/CoreNet.cpp
+++ b/Engine/Source/Runtime/CoreUObject/Private/UObject/CoreNet.cpp
@@ -195,6 +195,9 @@ const FClassNetCache* FClassNetCacheMgr::GetClassNetCache( UClass* Class )
 
 			// Add to cached fields on this class
 			Result->Fields.Add(FFieldNetCache(Property, ThisIndex, Checksum));
+			
+			// Skip over static array properties.
+			i += (Property->ArrayDim - 1);
 		}
 
 		for( int32 i = 0; i < Class->NetFields.Num(); i++ )
