Merging CL#11315165

Making sure UAnimBlueprintGeneratedClass::EvaluateGraphExposedInputs' ValueHandlerNodeProperties are properly resolved against their class which can be in a renamed package when loading.

+ TFieldPaths will now always try and use the provided struct when resolving themselves and update its path with the resolved package if it's different to the existting (serialized) value so that when they are re-resolved next time they will still point to the class in renamed package

diff --git a/Engine/Source/Runtime/CoreUObject/Private/UObject/FieldPath.cpp b/Engine/Source/Runtime/CoreUObject/Private/UObject/FieldPath.cpp
--- a/Engine/Source/Runtime/CoreUObject/Private/UObject/FieldPath.cpp
+++ b/Engine/Source/Runtime/CoreUObject/Private/UObject/FieldPath.cpp
@@ -96,16 +96,20 @@ FField* FFieldPath::TryToResolvePath(UStruct* InCurrentStruct, int32* OutOwnerIn
 	UObject* LastOuter = nullptr;
 	int32 PathIndex = Path.Num() - 1;
 	for (; PathIndex > 0; --PathIndex)
-	{
-		UObject* Outer = StaticFindObjectFast(UObject::StaticClass(), LastOuter, Path[PathIndex]);		
-		if (!Outer && PathIndex == (Path.Num() - 1) && InCurrentStruct)
+	{				
+		UObject* Outer = StaticFindObjectFast(UObject::StaticClass(), LastOuter, Path[PathIndex]);;
+
+		if (InCurrentStruct && PathIndex == (Path.Num() - 1))
 		{
-			// Try to resolve the package with the provided struct
-			// Sometimes packages are renamed when loading
-			UPackage* StructPackage = InCurrentStruct->GetOutermost();
-			if (StructPackage->FileName == Path[PathIndex])
+			UObject* CurrentOutermost = InCurrentStruct->GetOutermost();
+			// If a struct has been passed to this function and it's different than ther outer we found
+			// use it instead - this happens when a package with the class has been renammed
+			if (CurrentOutermost != Outer)
 			{
-				Outer = StructPackage;
+				Outer = CurrentOutermost;
+				// If we don't update the path then after a GC when this needs resolving we would resolve back to the unrenamed class package
+				FFieldPath* MutableThis = const_cast<FFieldPath*>(this);
+				MutableThis->Path[PathIndex] = Outer->GetFName();
 			}
 		}
 		if (!Outer)
diff --git a/Engine/Source/Runtime/Engine/Private/Animation/AnimBlueprintGeneratedClass.cpp b/Engine/Source/Runtime/Engine/Private/Animation/AnimBlueprintGeneratedClass.cpp
--- a/Engine/Source/Runtime/Engine/Private/Animation/AnimBlueprintGeneratedClass.cpp
+++ b/Engine/Source/Runtime/Engine/Private/Animation/AnimBlueprintGeneratedClass.cpp
@@ -288,6 +288,15 @@ void UAnimBlueprintGeneratedClass::Link(FArchive& Ar, bool bRelinkExistingProper
 	StateMachineNodeProperties.Empty();
 	InitializationNodeProperties.Empty();
 
+#if WITH_EDITORONLY_DATA
+	for (FExposedValueHandler& Handler : EvaluateGraphExposedInputs)
+	{
+		// handle potential renames of the class package
+		Handler.ValueHandlerNodeProperty.ClearCachedField();
+		Handler.ValueHandlerNodeProperty.Get(this);
+	}
+#endif // WITH_EDITORONLY_DATA
+
 #if WITH_EDITOR
 	// This relies on the entire class being fully loaded, this is not the case with EDL async-loading, in which case the functions are generated in PostLoad
 	GenerateAnimationBlueprintFunctions();
