Making sure FFieldPath is reset before generating a path from a string

diff --git a/Engine/Source/Runtime/CoreUObject/Private/UObject/FieldPath.cpp b/Engine/Source/Runtime/CoreUObject/Private/UObject/FieldPath.cpp
--- a/Engine/Source/Runtime/CoreUObject/Private/UObject/FieldPath.cpp
+++ b/Engine/Source/Runtime/CoreUObject/Private/UObject/FieldPath.cpp
@@ -50,6 +50,7 @@ void FFieldPath::Generate(const TCHAR* InFieldPathString)
 	// Expected format is: FullPackageName.Subobject[:Subobject:...]:FieldName
 	check(InFieldPathString);
 	
+	Path.Empty();
 	{
 		TCHAR NameBuffer[NAME_SIZE];
 		int32 NameIndex = 0;
