Reimplementing CL #11341031

Fixing TFieldPaths resolving themselves against the provided struct too aggresively (the behavior changed in CL #11315165)

diff --git a/Engine/Source/Runtime/CoreUObject/Private/UObject/FieldPath.cpp b/Engine/Source/Runtime/CoreUObject/Private/UObject/FieldPath.cpp
--- a/Engine/Source/Runtime/CoreUObject/Private/UObject/FieldPath.cpp
+++ b/Engine/Source/Runtime/CoreUObject/Private/UObject/FieldPath.cpp
@@ -89,7 +89,7 @@ void FFieldPath::Generate(const TCHAR* InFieldPathString)
 	}
 }
 
-FField* FFieldPath::TryToResolvePath(UStruct* InCurrentStruct, int32* OutOwnerIndex) const
+FField* FFieldPath::TryToResolvePath(UStruct* InCurrentStruct, int32* OutOwnerIndex, FFieldPath::EPathResolveType InResolveType /*= FFieldPath::UseStructIfOuterNotFound*/) const
 {
 	FField* Result = nullptr;
 
@@ -98,14 +98,15 @@ FField* FFieldPath::TryToResolvePath(UStruct* InCurrentStruct, int32* OutOwnerIn
 	int32 PathIndex = Path.Num() - 1;
 	for (; PathIndex > 0; --PathIndex)
 	{				
-		UObject* Outer = StaticFindObjectFast(UObject::StaticClass(), LastOuter, Path[PathIndex]);;
+		UObject* Outer = StaticFindObjectFast(UObject::StaticClass(), LastOuter, Path[PathIndex]);
 
 		if (InCurrentStruct && PathIndex == (Path.Num() - 1))
 		{
 			UObject* CurrentOutermost = InCurrentStruct->GetOutermost();
-			// If a struct has been passed to this function and it's different than ther outer we found
-			// use it instead - this happens when a package with the class has been renammed
-			if (CurrentOutermost != Outer)
+
+			if ((InResolveType == FFieldPath::UseStructIfOuterNotFound && !Outer) || // Outer is not found so try to use the provided struct Outer
+			    (InResolveType == FFieldPath::UseStructAlways && CurrentOutermost != Outer) // Prioritize the provided struct Outer over the resolved one
+			   )
 			{
 				Outer = CurrentOutermost;
 				// If we don't update the path then after a GC when this needs resolving we would resolve back to the unrenamed class package
diff --git a/Engine/Source/Runtime/CoreUObject/Public/UObject/FieldPath.h b/Engine/Source/Runtime/CoreUObject/Public/UObject/FieldPath.h
--- a/Engine/Source/Runtime/CoreUObject/Public/UObject/FieldPath.h
+++ b/Engine/Source/Runtime/CoreUObject/Public/UObject/FieldPath.h
@@ -16,6 +16,14 @@ class UObject;
 struct COREUOBJECT_API FFieldPath
 {
 protected:
+
+	/* Determines the behavior when resolving stored path */
+	enum EPathResolveType
+	{
+		UseStructIfOuterNotFound = 0,
+		UseStructAlways = 1
+	};
+
 	/** Untracked pointer to the resolved property */
 	mutable FField* ResolvedField;
 #if WITH_EDITORONLY_DATA
@@ -102,17 +110,17 @@ public:
 	 * @param OutOwnerIndex ObjectIndex of the Owner UObject
 	 * @return Resolved field or null
 	 */
-	FField* TryToResolvePath(UStruct* InCurrentStruct = nullptr, int32* OutOwnerIndex = nullptr) const;
+	FField* TryToResolvePath(UStruct* InCurrentStruct = nullptr, int32* OutOwnerIndex = nullptr, EPathResolveType InResolveType = FFieldPath::UseStructIfOuterNotFound) const;
 
 	/**
 	 * Tries to resolve the path and caches the result
 	 * @param ExpectedClass Expected class of the resolved field
 	 * @param InCurrentStruct Struct that's trying to resolve this field path	 
 	 */
-	FORCEINLINE void ResolveField(FFieldClass* ExpectedClass = FField::StaticClass(), UStruct* InCurrentStruct = nullptr) const
+	FORCEINLINE void ResolveField(FFieldClass* ExpectedClass = FField::StaticClass(), UStruct* InCurrentStruct = nullptr, EPathResolveType InResolveType = FFieldPath::UseStructIfOuterNotFound) const
 	{
 		int32 FoundOwner = -1;
-		FField* FoundField = TryToResolvePath(InCurrentStruct, &FoundOwner);
+		FField* FoundField = TryToResolvePath(InCurrentStruct, &FoundOwner, InResolveType);
 		if (FoundField && FoundField->IsA(ExpectedClass) 
 #if WITH_EDITORONLY_DATA
 			&& (!InitialFieldClass || FoundField->IsA(InitialFieldClass))
@@ -309,11 +317,18 @@ public:
 	{
 		if (NeedsResolving() && Path.Num())
 		{
-			ResolveField(PropertyType::StaticClass(), InCurrentStruct);
+			ResolveField(PropertyType::StaticClass(), InCurrentStruct, FFieldPath::UseStructIfOuterNotFound);
 		}
 		return static_cast<PropertyType*>(ResolvedField);
 	}
 
+	FORCEINLINE PropertyType* ResolveWithRenamedStructPackage(UStruct* InCurrentStruct)
+	{
+		ClearCachedField();
+		ResolveField(PropertyType::StaticClass(), InCurrentStruct, FFieldPath::UseStructAlways);
+		return static_cast<PropertyType*>(ResolvedField);
+	}
+
 	/**
 	* Dereference the weak pointer
 	**/
diff --git a/Engine/Source/Runtime/Engine/Private/Animation/AnimBlueprintGeneratedClass.cpp b/Engine/Source/Runtime/Engine/Private/Animation/AnimBlueprintGeneratedClass.cpp
--- a/Engine/Source/Runtime/Engine/Private/Animation/AnimBlueprintGeneratedClass.cpp
+++ b/Engine/Source/Runtime/Engine/Private/Animation/AnimBlueprintGeneratedClass.cpp
@@ -292,8 +292,7 @@ void UAnimBlueprintGeneratedClass::Link(FArchive& Ar, bool bRelinkExistingProper
 	for (FExposedValueHandler& Handler : EvaluateGraphExposedInputs)
 	{
 		// handle potential renames of the class package
-		Handler.ValueHandlerNodeProperty.ClearCachedField();
-		Handler.ValueHandlerNodeProperty.Get(this);
+		Handler.ValueHandlerNodeProperty.ResolveWithRenamedStructPackage(this);
 	}
 #endif // WITH_EDITORONLY_DATA
 
