Fixed some SA annotations

diff --git a/Engine/Source/Runtime/CoreUObject/Public/UObject/Field.h b/Engine/Source/Runtime/CoreUObject/Public/UObject/Field.h
--- a/Engine/Source/Runtime/CoreUObject/Public/UObject/Field.h
+++ b/Engine/Source/Runtime/CoreUObject/Public/UObject/Field.h
@@ -837,9 +837,7 @@ FUNCTION_NON_NULL_RETURN_END
 }
 
 template<typename FieldType>
-FUNCTION_NON_NULL_RETURN_START
 FORCEINLINE FieldType* CastFieldCheckedNullAllowed(FField* Src)
-FUNCTION_NON_NULL_RETURN_END
 {
 #if !DO_CHECK
 	return (FieldType*)Src;
@@ -851,9 +849,7 @@ FUNCTION_NON_NULL_RETURN_END
 }
 
 template<typename FieldType>
-FUNCTION_NON_NULL_RETURN_START
 FORCEINLINE const FieldType* CastFieldCheckedNullAllowed(const FField* Src)
-FUNCTION_NON_NULL_RETURN_END
 {
 #if !DO_CHECK
 	return (const FieldType*)Src;
