Allowing nested struct properties to be resolved with FProperties in linker placeholder classes (required for cooked bluerprints in the editor)

#rnx

diff --git a/Engine/Source/Runtime/CoreUObject/Private/UObject/LinkerPlaceholderBase.cpp b/Engine/Source/Runtime/CoreUObject/Private/UObject/LinkerPlaceholderBase.cpp
--- a/Engine/Source/Runtime/CoreUObject/Private/UObject/LinkerPlaceholderBase.cpp
+++ b/Engine/Source/Runtime/CoreUObject/Private/UObject/LinkerPlaceholderBase.cpp
@@ -356,10 +356,8 @@ FLinkerPlaceholderBase::FPlaceholderValuePropertyPath::FPlaceholderValueProperty
 				check(InnerStruct);
 				if (DEFERRED_DEPENDENCY_ENSURE(InnerStruct->IsChildOf(StructOwner)))
 				{
-					// With FProperties this should never happen
-					check(false);
-					//PropertyOuter = SerializingStructProp;
-					//PropertyChain.Add(SerializingStructProp);
+					PropertyOuter = SerializingStructProp;
+					PropertyChain.Add(SerializingStructProp);
 				}
 				else
 				{
