Avoid calling FixPackageNameCase in GetPackageLinker if the package already exists.  This causes FixPackageNameCase to go from 5m8s total time in a large dataset cook to 22s total time for a total time savings of 4m46s.

diff --git a/Engine/Source/Runtime/CoreUObject/Private/UObject/Linker.cpp b/Engine/Source/Runtime/CoreUObject/Private/UObject/Linker.cpp
--- a/Engine/Source/Runtime/CoreUObject/Private/UObject/Linker.cpp
+++ b/Engine/Source/Runtime/CoreUObject/Private/UObject/Linker.cpp
@@ -690,14 +690,18 @@ FLinkerLoad* GetPackageLinker
 			}
 			return nullptr;
 		}
-
+		UPackage* FilenamePkg = ExistingPackage;
+		if (!FilenamePkg)
+		{
 #if WITH_EDITORONLY_DATA
 		// Make sure the package name matches the name on disk
 		FPackageName::FixPackageNameCase(PackageNameToCreate, FPathViews::GetExtension(NewFilename));
 #endif
+			// Create the package with the provided long package name.
+			CreatedPackage = CreatePackage(nullptr, *PackageNameToCreate);
+			FilenamePkg = CreatedPackage;
+		}
 
-		// Create the package with the provided long package name.
-		UPackage* FilenamePkg = (ExistingPackage ? ExistingPackage : (CreatedPackage = CreatePackage(nullptr, *PackageNameToCreate)));
 		if (FilenamePkg && FilenamePkg != ExistingPackage && (LoadFlags & LOAD_PackageForPIE))
 		{
 			check(FilenamePkg);
