Guard against recursive calls to FTaskGraphInterface::ProcessThreadUntilIdle during async loading

diff --git a/Engine/Source/Runtime/CoreUObject/Private/Serialization/AsyncLoading.cpp b/Engine/Source/Runtime/CoreUObject/Private/Serialization/AsyncLoading.cpp
--- a/Engine/Source/Runtime/CoreUObject/Private/Serialization/AsyncLoading.cpp
+++ b/Engine/Source/Runtime/CoreUObject/Private/Serialization/AsyncLoading.cpp
@@ -4302,7 +4302,8 @@ EAsyncPackageState::Type FAsyncLoadingThread::ProcessLoadedPackages(bool bUseTim
 	}
 #if USE_EVENT_DRIVEN_ASYNC_LOAD_AT_BOOT_TIME
 	if (IsMultithreaded() && GEventDrivenLoaderEnabled &&
-		ENamedThreads::GetRenderThread() == ENamedThreads::GameThread) // render thread tasks are actually being sent to the game thread.
+		ENamedThreads::GetRenderThread() == ENamedThreads::GameThread &&
+		!FTaskGraphInterface::Get().IsThreadProcessingTasks(ENamedThreads::GameThread)) // render thread tasks are actually being sent to the game thread.
 	{
 		// The async loading thread might have queued some render thread tasks (we don't have a render thread yet, so these are actually sent to the game thread)
 		// We need to process them now before we do any postloads.
diff --git a/Engine/Source/Runtime/CoreUObject/Private/Serialization/AsyncLoading2.cpp b/Engine/Source/Runtime/CoreUObject/Private/Serialization/AsyncLoading2.cpp
--- a/Engine/Source/Runtime/CoreUObject/Private/Serialization/AsyncLoading2.cpp
+++ b/Engine/Source/Runtime/CoreUObject/Private/Serialization/AsyncLoading2.cpp
@@ -3533,7 +3533,7 @@ EAsyncPackageState::Type FAsyncLoadingThread2::ProcessLoadedPackagesFromGameThre
 			LoadedPackages.Reset();
 		}
 	}
-	if (IsMultithreaded() && ENamedThreads::GetRenderThread() == ENamedThreads::GameThread) // render thread tasks are actually being sent to the game thread.
+	if (IsMultithreaded() && ENamedThreads::GetRenderThread() == ENamedThreads::GameThread && !FTaskGraphInterface::Get().IsThreadProcessingTasks(ENamedThreads::GameThread)) // render thread tasks are actually being sent to the game thread.
 	{
 		// The async loading thread might have queued some render thread tasks (we don't have a render thread yet, so these are actually sent to the game thread)
 		// We need to process them now before we do any postloads.
