Workaround for Static Analysis false warnings

diff --git a/Engine/Source/Runtime/CoreUObject/Private/UObject/SavePackage.cpp b/Engine/Source/Runtime/CoreUObject/Private/UObject/SavePackage.cpp
--- a/Engine/Source/Runtime/CoreUObject/Private/UObject/SavePackage.cpp
+++ b/Engine/Source/Runtime/CoreUObject/Private/UObject/SavePackage.cpp
@@ -100,7 +100,7 @@ namespace SavePackageStats
 	static double MBWritten = 0.0;
 	TMap<FName, FArchiveDiffStats> PackageDiffStats;
 	static int32 NumberOfDifferentPackages = 0;
-	static FCookStatsManager::FAutoRegisterCallback RegisterCookStats([](FCookStatsManager::AddStatFuncRef AddStat)
+	void AddSavePackageStats(FCookStatsManager::AddStatFuncRef AddStat)
 	{
 		// Don't use FCookStatsManager::CreateKeyValueArray because there's just too many arguments. Don't need to overburden the compiler here.
 		TArray<FCookStatsManager::StringKeyValue> StatsList;
@@ -184,7 +184,8 @@ namespace SavePackageStats
 		#undef ADD_COOK_STAT		
 		
 		const FString TotalString = TEXT("Total");
-	});
+	}
+	static FCookStatsManager::FAutoRegisterCallback RegisterCookStats(AddSavePackageStats);
 }
 #endif
 
