Fix platform configs being saved in extension path for non-extension platforms (UE-81531).

- Platform extension path will be used only if the platform is an extension and no pre-existing configs are in place.
- Config reading code in C++ and C# is supposed to support both locations.

diff --git a/Engine/Source/Runtime/CoreUObject/Private/UObject/Obj.cpp b/Engine/Source/Runtime/CoreUObject/Private/UObject/Obj.cpp
--- a/Engine/Source/Runtime/CoreUObject/Private/UObject/Obj.cpp
+++ b/Engine/Source/Runtime/CoreUObject/Private/UObject/Obj.cpp
@@ -2552,13 +2552,26 @@ FString UObject::GetDefaultConfigFilename() const
 	FString OverridePlatform = GetFinalOverridePlatform(this);
 	if (OverridePlatform.Len())
 	{
-		FString PlatformExtPath = FString::Printf(TEXT("%s%s/Config"), *FPaths::ProjectPlatformExtensionsDir(), *OverridePlatform);
+		bool bIsPlatformExtension = FPaths::DirectoryExists(FPaths::Combine(FPaths::EnginePlatformExtensionsDir(), OverridePlatform));
 		FString RegularPath = FString::Printf(TEXT("%s%s"), *FPaths::SourceConfigDir(), *OverridePlatform, *OverridePlatform, *GetClass()->ClassConfigName.ToString());
+		FString SelectedPath = RegularPath;
 
-		// prefer the platform extension path if it exists, or if the regular path does not exist.
-		FString SelectedPath = (FPaths::DirectoryExists(*PlatformExtPath) || !FPaths::DirectoryExists(*RegularPath))
-			? PlatformExtPath
-			: RegularPath;
+		bool bPlatformConfigExistsInRegular = FPaths::DirectoryExists(*RegularPath);
+
+		// if the platform is an extension, create the new config in the extension path (Platforms/PlatformName/Config),
+		// unless there exists a platform config in the regular path (Config/PlatformName)
+
+		// PlatformExtension | ConfigExistsInRegularPath  |   Use path
+		//   false                  false                      regular
+		//   true                   false                      extension
+		//   false                  true                       regular
+		//   true                   true                       regular
+
+		// if the project already uses platform configs in the regular directory, just use that, otherwise check if this is a platform extensions
+		if (bIsPlatformExtension && !bPlatformConfigExistsInRegular)
+		{
+			SelectedPath = FString::Printf(TEXT("%s%s/Config"), *FPaths::ProjectPlatformExtensionsDir(), *OverridePlatform);
+		}
 
 		return FString::Printf(TEXT("%s/%s%s.ini"), *SelectedPath, *OverridePlatform, *GetClass()->ClassConfigName.ToString());
 	}
