Correct bad TMap related logic in ResolvePlaceholderValues - this could cause a crash when resolving placeholders in a TMap

diff --git a/Engine/Source/Runtime/CoreUObject/Private/UObject/LinkerPlaceholderBase.cpp b/Engine/Source/Runtime/CoreUObject/Private/UObject/LinkerPlaceholderBase.cpp
--- a/Engine/Source/Runtime/CoreUObject/Private/UObject/LinkerPlaceholderBase.cpp
+++ b/Engine/Source/Runtime/CoreUObject/Private/UObject/LinkerPlaceholderBase.cpp
@@ -153,11 +153,10 @@ int32 FLinkerPlaceholderObjectImpl::ResolvePlaceholderValues(const TArray<FField
 		}
 		else if (const FMapProperty* MapProperty = Property.Get<FMapProperty>())
 		{
+			const FProperty* NextProperty = PropertyChain[PropertyIndex - 1].Get<FProperty>();
 #if USE_DEFERRED_DEPENDENCY_CHECK_VERIFICATION_TESTS
-			const FProperty* NextProperty = PropertyChain[PropertyIndex - 1];
-			check(NextProperty == MapProperty->KeyProp);
+			check(NextProperty == MapProperty->KeyProp || NextProperty == MapProperty->ValueProp);
 #endif // USE_DEFERRED_DEPENDENCY_CHECK_VERIFICATION_TESTS
-
 			// because we can't know which map entry was set with a reference 
 			// to this object, we have to comb through them all
 			FScriptMapHelper MapHelper(MapProperty, ValueAddress);
@@ -167,11 +166,17 @@ int32 FLinkerPlaceholderObjectImpl::ResolvePlaceholderValues(const TArray<FField
 				if (MapHelper.IsValidIndex(MapIndex))
 				{
 					--Num;
-					uint8* KeyAddress = MapHelper.GetKeyPtr(MapIndex);
-					ReplacementCount += ResolvePlaceholderValues(PropertyChain, PropertyIndex - 1, KeyAddress, OldValue, ReplacementValue);
 
-					uint8* MapValueAddress = MapHelper.GetValuePtr(MapIndex);
-					ReplacementCount += ResolvePlaceholderValues(PropertyChain, PropertyIndex - 1, MapValueAddress, OldValue, ReplacementValue);
+					if (NextProperty == MapProperty->KeyProp)
+					{
+						uint8* KeyAddress = MapHelper.GetKeyPtr(MapIndex);
+						ReplacementCount += ResolvePlaceholderValues(PropertyChain, PropertyIndex - 1, KeyAddress, OldValue, ReplacementValue);
+					}
+					else if (NextProperty == MapProperty->ValueProp)
+					{
+						uint8* MapValueAddress = MapHelper.GetValuePtr(MapIndex);
+						ReplacementCount += ResolvePlaceholderValues(PropertyChain, PropertyIndex - 1, MapValueAddress, OldValue, ReplacementValue);
+					}
 				}
 			}
 
