Significantly improving MarkObjectsAsUnreachable phase of GC in scenarios where there's many root objects (like in the editor).

Before:
LogGarbage: Verbose: 24.557600 ms for Mark Phase (22485 Objects To Serialize)

After:
LogGarbage: Verbose: 0.290700 ms for Mark Phase (22485 Objects To Serialize)

diff --git a/Engine/Source/Runtime/CoreUObject/Private/UObject/GarbageCollection.cpp b/Engine/Source/Runtime/CoreUObject/Private/UObject/GarbageCollection.cpp
--- a/Engine/Source/Runtime/CoreUObject/Private/UObject/GarbageCollection.cpp
+++ b/Engine/Source/Runtime/CoreUObject/Private/UObject/GarbageCollection.cpp
@@ -1074,23 +1074,27 @@ public:
 	void MarkObjectsAsUnreachable(TArray<UObject*>& ObjectsToSerialize, const EObjectFlags KeepFlags)
 	{
 		const EInternalObjectFlags FastKeepFlags = EInternalObjectFlags::GarbageCollectionKeepFlags;
+		const int32 MaxNumberOfObjects = GUObjectArray.GetObjectArrayNum() - GUObjectArray.GetFirstGCIndex();
+		const int32 NumThreads = FMath::Max(1, FTaskGraphInterface::Get().GetNumWorkerThreads());
+		const int32 NumberOfObjectsPerThread = (MaxNumberOfObjects / NumThreads) + 1;		
 
-		TLockFreePointerListFIFO<UObject, PLATFORM_CACHE_LINE_SIZE> ObjectsToSerializeList;
 		TLockFreePointerListFIFO<FUObjectItem, PLATFORM_CACHE_LINE_SIZE> ClustersToDissolveList;
 		TLockFreePointerListFIFO<FUObjectItem, PLATFORM_CACHE_LINE_SIZE> KeepClusterRefsList;
-
-		int32 MaxNumberOfObjects = GUObjectArray.GetObjectArrayNum() - GUObjectArray.GetFirstGCIndex();
-		int32 NumThreads = FMath::Max(1, FTaskGraphInterface::Get().GetNumWorkerThreads());
-		int32 NumberOfObjectsPerThread = (MaxNumberOfObjects / NumThreads) + 1;		
+		FGCArrayStruct** ObjectsToSerializeArrays = new FGCArrayStruct*[NumThreads];
+		for (int32 ThreadIndex = 0; ThreadIndex < NumThreads; ++ThreadIndex)
+		{
+			ObjectsToSerializeArrays[ThreadIndex] = FGCArrayPool::Get().GetArrayStructFromPool();
+		}
 
 		// Iterate over all objects. Note that we iterate over the UObjectArray and usually check only internal flags which
 		// are part of the array so we don't suffer from cache misses as much as we would if we were to check ObjectFlags.
-		ParallelFor(NumThreads, [&ObjectsToSerializeList, &ClustersToDissolveList, &KeepClusterRefsList, FastKeepFlags, KeepFlags, NumberOfObjectsPerThread, NumThreads, MaxNumberOfObjects](int32 ThreadIndex)
+		ParallelFor(NumThreads, [ObjectsToSerializeArrays, &ClustersToDissolveList, &KeepClusterRefsList, FastKeepFlags, KeepFlags, NumberOfObjectsPerThread, NumThreads, MaxNumberOfObjects](int32 ThreadIndex)
 		{
 			int32 FirstObjectIndex = ThreadIndex * NumberOfObjectsPerThread + GUObjectArray.GetFirstGCIndex();
 			int32 NumObjects = (ThreadIndex < (NumThreads - 1)) ? NumberOfObjectsPerThread : (MaxNumberOfObjects - (NumThreads - 1) * NumberOfObjectsPerThread);
 			int32 LastObjectIndex = FMath::Min(GUObjectArray.GetObjectArrayNum() - 1, FirstObjectIndex + NumObjects - 1);
 			int32 ObjectCountDuringMarkPhase = 0;
+			TArray<UObject*>& LocalObjectsToSerialize = ObjectsToSerializeArrays[ThreadIndex]->ObjectsToSerialize;
 
 			for (int32 ObjectIndex = FirstObjectIndex; ObjectIndex <= LastObjectIndex; ++ObjectIndex)
 			{
@@ -1126,7 +1130,7 @@ public:
 							}
 						}
 
-						ObjectsToSerializeList.Push(Object);
+						LocalObjectsToSerialize.Add(Object);
 					}
 					// Regular objects or cluster root objects
 					else if (!bWithClusters || ObjectItem->GetOwnerIndex() <= 0)
@@ -1155,7 +1159,7 @@ public:
 						{
 							// IsValidLowLevel is extremely slow in this loop so only do it in debug
 							checkSlow(Object->IsValidLowLevel());
-							ObjectsToSerializeList.Push(Object);
+							LocalObjectsToSerialize.Add(Object);
 
 							if (bWithClusters)
 							{
@@ -1176,7 +1180,21 @@ public:
 			GObjectCountDuringLastMarkPhase.Add(ObjectCountDuringMarkPhase);
 		}, !bParallel);
 		
-		ObjectsToSerializeList.PopAll(ObjectsToSerialize);
+		// Collect all objects to serialize from all threads and put them into a single array
+		{
+			int32 NumObjectsToSerialize = 0;
+			for (int32 ThreadIndex = 0; ThreadIndex < NumThreads; ++ThreadIndex)
+			{
+				NumObjectsToSerialize += ObjectsToSerializeArrays[ThreadIndex]->ObjectsToSerialize.Num();
+			}
+			ObjectsToSerialize.Reserve(NumObjectsToSerialize);
+			for (int32 ThreadIndex = 0; ThreadIndex < NumThreads; ++ThreadIndex)
+			{
+				ObjectsToSerialize.Append(ObjectsToSerializeArrays[ThreadIndex]->ObjectsToSerialize);
+				FGCArrayPool::Get().ReturnToPool(ObjectsToSerializeArrays[ThreadIndex]);
+			}
+			delete[] ObjectsToSerializeArrays;
+		}
 
 		if (bWithClusters)
 		{
@@ -1259,7 +1277,7 @@ public:
 		{
 			const double StartTime = FPlatformTime::Seconds();
 			(this->*MarkObjectsFunctions[GetGCFunctionIndex(!bForceSingleThreaded, bWithClusters)])(ObjectsToSerialize, KeepFlags);
-			UE_LOG(LogGarbage, Verbose, TEXT("%f ms for Mark Phase (%d Objects To Serialize"), (FPlatformTime::Seconds() - StartTime) * 1000, ObjectsToSerialize.Num());
+			UE_LOG(LogGarbage, Verbose, TEXT("%f ms for MarkObjectsAsUnreachable Phase (%d Objects To Serialize)"), (FPlatformTime::Seconds() - StartTime) * 1000, ObjectsToSerialize.Num());
 		}
 
 		{
