Optimize Pak file indices by computing the PathHashIndex in unrealpak rather than using strings at runtime.

diff --git a/Engine/Source/Runtime/CoreUObject/Private/UObject/CoreRedirects.cpp b/Engine/Source/Runtime/CoreUObject/Private/UObject/CoreRedirects.cpp
--- a/Engine/Source/Runtime/CoreUObject/Private/UObject/CoreRedirects.cpp
+++ b/Engine/Source/Runtime/CoreUObject/Private/UObject/CoreRedirects.cpp
@@ -1,13 +1,17 @@
 // Copyright Epic Games, Inc. All Rights Reserved.
 
 #include "UObject/CoreRedirects.h"
+
+#include "UObject/Linker.h"
 #include "UObject/Package.h"
 #include "UObject/PropertyHelper.h"
-#include "UObject/Linker.h"
-#include "Templates/Casts.h"
+
+#include "GenericPlatform/GenericPlatformFile.h"
+#include "Misc/AutomationTest.h"
 #include "Misc/ConfigCacheIni.h"
+#include "Misc/PackageName.h"
 #include "Serialization/DeferredMessageLog.h"
-#include "Misc/AutomationTest.h"
+#include "Templates/Casts.h"
 
 FCoreRedirectObjectName::FCoreRedirectObjectName(const FString& InString)
 {
@@ -350,10 +354,23 @@ const TCHAR* FCoreRedirect::ParseValueChanges(const TCHAR* Buffer)
 	}
 }
 
+static FORCEINLINE bool CheckRedirectFlagsMatch(ECoreRedirectFlags FlagsA, ECoreRedirectFlags FlagsB)
+{
+	// For type, check it includes the matching type
+	const bool bTypesOverlap = !!((FlagsA & FlagsB) & ECoreRedirectFlags::Type_AllMask);
+
+	// For category, the bits must be an exact match
+	const bool bCategoriesMatch = (FlagsA & ECoreRedirectFlags::Category_AllMask) == (FlagsB & ECoreRedirectFlags::Category_AllMask);
+
+	// Options are not considered in this function; Flags will match at this point regardless of their bits in Options_AllMask
+
+	return bTypesOverlap && bCategoriesMatch;
+}
+
 bool FCoreRedirect::Matches(ECoreRedirectFlags InFlags, const FCoreRedirectObjectName& InName) const
 {
-	// Check flags for exact match
-	if (InFlags != RedirectFlags)
+	// Check flags for Type/Category match
+	if (!CheckRedirectFlagsMatch(InFlags, RedirectFlags))
 	{
 		return false;
 	}
@@ -489,24 +506,13 @@ const TMap<FString, FString>* FCoreRedirects::GetValueRedirects(ECoreRedirectFla
 	return nullptr;
 }
 
-static FORCEINLINE bool CheckRedirectFlagsMatch(ECoreRedirectFlags FlagsA, ECoreRedirectFlags FlagsB)
-{
-	// For type, check it includes the matching type
-	const bool bTypeMatches = !!((FlagsA & FlagsB) & ECoreRedirectFlags::Type_AllMask);
-
-	// For options, instance/remove must match exactly but allow both substring and non-substring matches
-	const bool bOptionsMatch = (FlagsA & ECoreRedirectFlags::Option_ExactMatchMask) == (FlagsB & ECoreRedirectFlags::Option_ExactMatchMask);
-
-	return bTypeMatches && bOptionsMatch;
-}
-
 bool FCoreRedirects::GetMatchingRedirects(ECoreRedirectFlags SearchFlags, const FCoreRedirectObjectName& OldObjectName, TArray<const FCoreRedirect*>& FoundRedirects)
 {
 	// Look for all redirects that match the given names and flags
 	bool bFound = false;
 	
-	// If we're not explicitly searching for packages or looking for removed things, add the implicit package redirects
-	const bool bSearchPackageRedirects = !(SearchFlags & ECoreRedirectFlags::Type_Package) && !(SearchFlags & ECoreRedirectFlags::Option_Removed);
+	// If we're not explicitly searching for packages or looking for removed things, add the implicit (Type=Package,Category=None) redirects
+	const bool bSearchPackageRedirects = !(SearchFlags & ECoreRedirectFlags::Type_Package) && !(SearchFlags & ECoreRedirectFlags::Category_Removed);
 
 	// Determine list of maps to look over, need to handle being passed multiple types in a bit mask
 	for (const TPair<ECoreRedirectFlags, FRedirectNameMap>& Pair : RedirectTypeMap)
@@ -540,8 +546,8 @@ bool FCoreRedirects::FindPreviousNames(ECoreRedirectFlags SearchFlags, const FCo
 	// Look for reverse direction redirects
 	bool bFound = false;
 
-	// If we're not explicitly searching for packages or looking for removed things, add the implicit package redirects
-	const bool bSearchPackageRedirects = !(SearchFlags & ECoreRedirectFlags::Type_Package) && !(SearchFlags & ECoreRedirectFlags::Option_Removed);
+	// If we're not explicitly searching for packages or looking for removed things, add the implicit (Type=Package,Category=None) redirects
+	const bool bSearchPackageRedirects = !(SearchFlags & ECoreRedirectFlags::Type_Package) && !(SearchFlags & ECoreRedirectFlags::Category_Removed);
 
 	// Determine list of maps to look over, need to handle being passed multiple Flagss in a bit mask
 	for (const TPair<ECoreRedirectFlags, FRedirectNameMap>& Pair : RedirectTypeMap)
@@ -581,23 +587,32 @@ bool FCoreRedirects::FindPreviousNames(ECoreRedirectFlags SearchFlags, const FCo
 bool FCoreRedirects::IsKnownMissing(ECoreRedirectFlags Type, const FCoreRedirectObjectName& ObjectName)
 {
 	TArray<const FCoreRedirect*> FoundRedirects;
-	return FCoreRedirects::GetMatchingRedirects(Type | ECoreRedirectFlags::Option_Removed, ObjectName, FoundRedirects);
+	return FCoreRedirects::GetMatchingRedirects(Type | ECoreRedirectFlags::Category_Removed, ObjectName, FoundRedirects);
 }
 
-bool FCoreRedirects::AddKnownMissing(ECoreRedirectFlags Type, const FCoreRedirectObjectName& ObjectName)
+bool FCoreRedirects::AddKnownMissing(ECoreRedirectFlags Type, const FCoreRedirectObjectName& ObjectName, ECoreRedirectFlags Channel)
 {
+	check((Channel & ~ECoreRedirectFlags::Option_MissingLoad) == ECoreRedirectFlags::None);
 	TArray<FCoreRedirect> NewRedirects;
-	NewRedirects.Emplace(Type | ECoreRedirectFlags::Option_Removed, ObjectName, FCoreRedirectObjectName());
+	NewRedirects.Emplace(Type | ECoreRedirectFlags::Category_Removed | Channel, ObjectName, FCoreRedirectObjectName());
 	return AddRedirectList(NewRedirects, TEXT("AddKnownMissing"));
 }
 
-bool FCoreRedirects::RemoveKnownMissing(ECoreRedirectFlags Type, const FCoreRedirectObjectName& ObjectName)
+bool FCoreRedirects::RemoveKnownMissing(ECoreRedirectFlags Type, const FCoreRedirectObjectName& ObjectName, ECoreRedirectFlags Channel)
 {
+	check((Channel & ~ECoreRedirectFlags::Option_MissingLoad) == ECoreRedirectFlags::None);
 	TArray<FCoreRedirect> RedirectsToRemove;
-	RedirectsToRemove.Emplace(Type | ECoreRedirectFlags::Option_Removed, ObjectName, FCoreRedirectObjectName());
+	RedirectsToRemove.Emplace(Type | ECoreRedirectFlags::Category_Removed | Channel, ObjectName, FCoreRedirectObjectName());
 	return RemoveRedirectList(RedirectsToRemove, TEXT("RemoveKnownMissing"));
 }
 
+void FCoreRedirects::ClearKnownMissing(ECoreRedirectFlags Type, ECoreRedirectFlags Channel)
+{
+	check((Channel & ~ECoreRedirectFlags::Option_MissingLoad) == ECoreRedirectFlags::None);
+	ECoreRedirectFlags RedirectFlags = Type | ECoreRedirectFlags::Category_Removed | Channel;
+	RedirectTypeMap.Remove(RedirectFlags);
+}
+
 bool FCoreRedirects::RunTests()
 {
 	bool bSuccess = true;
@@ -606,16 +621,17 @@ bool FCoreRedirects::RunTests()
 
 	TArray<FCoreRedirect> NewRedirects;
 
-	new (NewRedirects) FCoreRedirect(ECoreRedirectFlags::Type_Property, TEXT("Property"), TEXT("Property2"));
-	new (NewRedirects) FCoreRedirect(ECoreRedirectFlags::Type_Property, TEXT("Class.Property"), TEXT("Property3"));
-	new (NewRedirects) FCoreRedirect(ECoreRedirectFlags::Type_Property, TEXT("/game/PackageSpecific.Class.Property"), TEXT("Property4"));
-	new (NewRedirects) FCoreRedirect(ECoreRedirectFlags::Type_Property, TEXT("/game/Package.Class.OtherProperty"), TEXT("/game/Package.Class.OtherProperty2"));
-	new (NewRedirects) FCoreRedirect(ECoreRedirectFlags::Type_Class, TEXT("Class"), TEXT("Class2"));
-	new (NewRedirects) FCoreRedirect(ECoreRedirectFlags::Type_Class, TEXT("/game/Package.Class"), TEXT("Class3"));
-	new (NewRedirects) FCoreRedirect(ECoreRedirectFlags::Type_Class | ECoreRedirectFlags::Option_InstanceOnly, TEXT("/game/Package.Class"), TEXT("/game/Package.ClassInstance"));
-	new (NewRedirects) FCoreRedirect(ECoreRedirectFlags::Type_Package, TEXT("/game/Package"), TEXT("/game/Package2"));
-	new (NewRedirects) FCoreRedirect(ECoreRedirectFlags::Type_Package | ECoreRedirectFlags::Option_MatchSubstring, TEXT("/oldgame"), TEXT("/newgame"));
-	new (NewRedirects) FCoreRedirect(ECoreRedirectFlags::Type_Package | ECoreRedirectFlags::Option_Removed, TEXT("/game/RemovedPackage"), TEXT("/game/RemovedPackage"));
+	NewRedirects.Emplace(ECoreRedirectFlags::Type_Property, TEXT("Property"), TEXT("Property2"));
+	NewRedirects.Emplace(ECoreRedirectFlags::Type_Property, TEXT("Class.Property"), TEXT("Property3"));
+	NewRedirects.Emplace(ECoreRedirectFlags::Type_Property, TEXT("/Game/PackageSpecific.Class.Property"), TEXT("Property4"));
+	NewRedirects.Emplace(ECoreRedirectFlags::Type_Property, TEXT("/Game/Package.Class.OtherProperty"), TEXT("/Game/Package.Class.OtherProperty2"));
+	NewRedirects.Emplace(ECoreRedirectFlags::Type_Class, TEXT("Class"), TEXT("Class2"));
+	NewRedirects.Emplace(ECoreRedirectFlags::Type_Class, TEXT("/Game/Package.Class"), TEXT("Class3"));
+	NewRedirects.Emplace(ECoreRedirectFlags::Type_Class | ECoreRedirectFlags::Category_InstanceOnly, TEXT("/Game/Package.Class"), TEXT("/Game/Package.ClassInstance"));
+	NewRedirects.Emplace(ECoreRedirectFlags::Type_Package, TEXT("/Game/Package"), TEXT("/Game/Package2"));
+	NewRedirects.Emplace(ECoreRedirectFlags::Type_Package | ECoreRedirectFlags::Option_MatchSubstring, TEXT("/oldgame"), TEXT("/newgame"));
+	NewRedirects.Emplace(ECoreRedirectFlags::Type_Package | ECoreRedirectFlags::Category_Removed, TEXT("/Game/RemovedPackage"), TEXT("/Game/RemovedPackage"));
+	NewRedirects.Emplace(ECoreRedirectFlags::Type_Package | ECoreRedirectFlags::Category_Removed | ECoreRedirectFlags::Option_MissingLoad, TEXT("/Game/MissingLoadPackage"), TEXT("/Game/MissingLoadPackage"));
 
 	AddRedirectList(NewRedirects, TEXT("RunTests"));
 
@@ -635,19 +651,19 @@ bool FCoreRedirects::RunTests()
 	UE_LOG(LogLinker, Log, TEXT("Running FCoreRedirect Tests"));
 
 	// Package-specific property rename and package rename apply
-	Tests.Emplace(TEXT("/game/PackageSpecific.Class:Property"), TEXT("/game/PackageSpecific.Class:Property4"), ECoreRedirectFlags::Type_Property);
+	Tests.Emplace(TEXT("/Game/PackageSpecific.Class:Property"), TEXT("/Game/PackageSpecific.Class:Property4"), ECoreRedirectFlags::Type_Property);
 	// Verify . works as well
-	Tests.Emplace(TEXT("/game/PackageSpecific.Class.Property"), TEXT("/game/PackageSpecific.Class:Property4"), ECoreRedirectFlags::Type_Property);
+	Tests.Emplace(TEXT("/Game/PackageSpecific.Class.Property"), TEXT("/Game/PackageSpecific.Class:Property4"), ECoreRedirectFlags::Type_Property);
 	// Wrong type, no replacement
-	Tests.Emplace(TEXT("/game/PackageSpecific.Class:Property"), TEXT("/game/PackageSpecific.Class:Property"), ECoreRedirectFlags::Type_Function);
+	Tests.Emplace(TEXT("/Game/PackageSpecific.Class:Property"), TEXT("/Game/PackageSpecific.Class:Property"), ECoreRedirectFlags::Type_Function);
 	// Class-specific property rename and package rename apply
-	Tests.Emplace(TEXT("/game/Package.Class:Property"), TEXT("/game/Package2.Class:Property3"), ECoreRedirectFlags::Type_Property);
+	Tests.Emplace(TEXT("/Game/Package.Class:Property"), TEXT("/Game/Package2.Class:Property3"), ECoreRedirectFlags::Type_Property);
 	// Package-Specific class rename applies
-	Tests.Emplace(TEXT("/game/Package.Class"), TEXT("/game/Package2.Class3"), ECoreRedirectFlags::Type_Class);
+	Tests.Emplace(TEXT("/Game/Package.Class"), TEXT("/Game/Package2.Class3"), ECoreRedirectFlags::Type_Class);
 	// Generic class rename applies
-	Tests.Emplace(TEXT("/game/PackageOther.Class"), TEXT("/game/PackageOther.Class2"), ECoreRedirectFlags::Type_Class);
+	Tests.Emplace(TEXT("/Game/PackageOther.Class"), TEXT("/Game/PackageOther.Class2"), ECoreRedirectFlags::Type_Class);
 	// Check instance option
-	Tests.Emplace(TEXT("/game/Package.Class"), TEXT("/game/Package2.ClassInstance"), ECoreRedirectFlags::Type_Class | ECoreRedirectFlags::Option_InstanceOnly);
+	Tests.Emplace(TEXT("/Game/Package.Class"), TEXT("/Game/Package2.ClassInstance"), ECoreRedirectFlags::Type_Class | ECoreRedirectFlags::Category_InstanceOnly);
 	// Substring test
 	Tests.Emplace(TEXT("/oldgame/Package.DefaultClass"), TEXT("/newgame/Package.DefaultClass"), ECoreRedirectFlags::Type_Class);
 
@@ -666,40 +682,79 @@ bool FCoreRedirects::RunTests()
 	// Check reverse lookup
 	TArray<FCoreRedirectObjectName> OldNames;
 
-	FindPreviousNames(ECoreRedirectFlags::Type_Class, FCoreRedirectObjectName(TEXT("/game/PackageOther.Class2")), OldNames);
+	FindPreviousNames(ECoreRedirectFlags::Type_Class, FCoreRedirectObjectName(TEXT("/Game/PackageOther.Class2")), OldNames);
 
-	if (OldNames.Num() != 1 || OldNames[0].ToString() != TEXT("/game/PackageOther.Class"))
+	if (OldNames.Num() != 1 || OldNames[0].ToString() != TEXT("/Game/PackageOther.Class"))
 	{
 		bSuccess = false;
 		UE_LOG(LogLinker, Error, TEXT("FCoreRedirect Test Failed: ReverseLookup!"));
 	}
 
 	// Check removed
-	if (!IsKnownMissing(ECoreRedirectFlags::Type_Package, FCoreRedirectObjectName(TEXT("/game/RemovedPackage"))))
+	if (!IsKnownMissing(ECoreRedirectFlags::Type_Package, FCoreRedirectObjectName(TEXT("/Game/RemovedPackage"))))
 	{
 		bSuccess = false;
-		UE_LOG(LogLinker, Error, TEXT("FCoreRedirect Test Failed: /game/RemovedPackage should be removed!"));
+		UE_LOG(LogLinker, Error, TEXT("FCoreRedirect Test Failed: /Game/RemovedPackage should be removed!"));
 	}
 
-	if (IsKnownMissing(ECoreRedirectFlags::Type_Package, FCoreRedirectObjectName(TEXT("/game/NotRemovedPackage"))))
+	if (!IsKnownMissing(ECoreRedirectFlags::Type_Package, FCoreRedirectObjectName(TEXT("/Game/MissingLoadPackage"))))
 	{
 		bSuccess = false;
-		UE_LOG(LogLinker, Error, TEXT("FCoreRedirect Test Failed: /game/NotRemovedPackage should be removed!"));
+		UE_LOG(LogLinker, Error, TEXT("FCoreRedirect Test Failed: /Game/MissingLoadPackage should be removed!"));
 	}
 
-	AddKnownMissing(ECoreRedirectFlags::Type_Package, FCoreRedirectObjectName(TEXT("/game/NotRemovedPackage")));
+	if (IsKnownMissing(ECoreRedirectFlags::Type_Package, FCoreRedirectObjectName(TEXT("/Game/NotRemovedPackage"))))
+	{
+		bSuccess = false;
+		UE_LOG(LogLinker, Error, TEXT("FCoreRedirect Test Failed: /Game/NotRemovedPackage should be removed!"));
+	}
+
+	AddKnownMissing(ECoreRedirectFlags::Type_Package, FCoreRedirectObjectName(TEXT("/Game/NotRemovedMissingLoad")), ECoreRedirectFlags::Option_MissingLoad);
 
-	if (!IsKnownMissing(ECoreRedirectFlags::Type_Package, FCoreRedirectObjectName(TEXT("/game/NotRemovedPackage"))))
+	if (!IsKnownMissing(ECoreRedirectFlags::Type_Package, FCoreRedirectObjectName(TEXT("/Game/NotRemovedMissingLoad"))))
 	{
 		bSuccess = false;
-		UE_LOG(LogLinker, Error, TEXT("FCoreRedirect Test Failed: /game/NotRemovedPackage should be removed now!"));
+		UE_LOG(LogLinker, Error, TEXT("FCoreRedirect Test Failed: /Game/NotRemovedMissingLoad should be removed now!"));
 	}
 
-	RemoveKnownMissing(ECoreRedirectFlags::Type_Package, FCoreRedirectObjectName(TEXT("/game/NotRemovedPackage")));
+	RemoveKnownMissing(ECoreRedirectFlags::Type_Package, FCoreRedirectObjectName(TEXT("/Game/NotRemovedMissingLoad")), ECoreRedirectFlags::None);
 
-	if (IsKnownMissing(ECoreRedirectFlags::Type_Package, FCoreRedirectObjectName(TEXT("/game/NotRemovedPackage"))))
+	if (!IsKnownMissing(ECoreRedirectFlags::Type_Package, FCoreRedirectObjectName(TEXT("/Game/NotRemovedMissingLoad"))))
 	{
-		UE_LOG(LogLinker, Error, TEXT("FCoreRedirect Test Failed: /game/NotRemovedPackage should be removed!"));
+		bSuccess = false;
+		UE_LOG(LogLinker, Error, TEXT("FCoreRedirect Test Failed: RemoveKnownMissing of /Game/NotRemovedMissingLoad but with bIsMissingLoad=false should not have removed the redirect!"));
+	}
+
+	RemoveKnownMissing(ECoreRedirectFlags::Type_Package, FCoreRedirectObjectName(TEXT("/Game/NotRemovedMissingLoad")), ECoreRedirectFlags::Option_MissingLoad);
+
+	if (IsKnownMissing(ECoreRedirectFlags::Type_Package, FCoreRedirectObjectName(TEXT("/Game/NotRemovedMissingLoad"))))
+	{
+		bSuccess = false;
+		UE_LOG(LogLinker, Error, TEXT("FCoreRedirect Test Failed: /Game/NotRemovedMissingLoad should no longer be removed!"));
+	}
+
+	AddKnownMissing(ECoreRedirectFlags::Type_Package, FCoreRedirectObjectName(TEXT("/Game/NotRemovedPackage")), ECoreRedirectFlags::None);
+
+	if (!IsKnownMissing(ECoreRedirectFlags::Type_Package, FCoreRedirectObjectName(TEXT("/Game/NotRemovedPackage"))))
+	{
+		bSuccess = false;
+		UE_LOG(LogLinker, Error, TEXT("FCoreRedirect Test Failed: /Game/NotRemovedPackage should be removed now!"));
+	}
+
+	RemoveKnownMissing(ECoreRedirectFlags::Type_Package, FCoreRedirectObjectName(TEXT("/Game/NotRemovedPackage")), ECoreRedirectFlags::Option_MissingLoad);
+
+	if (!IsKnownMissing(ECoreRedirectFlags::Type_Package, FCoreRedirectObjectName(TEXT("/Game/NotRemovedPackage"))))
+	{
+		bSuccess = false;
+		UE_LOG(LogLinker, Error, TEXT("FCoreRedirect Test Failed: RemoveKnownMissing of /Game/NotRemovedPackage but with bIsMissingLoad=true should not have removed the redirect!"));
+	}
+
+	RemoveKnownMissing(ECoreRedirectFlags::Type_Package, FCoreRedirectObjectName(TEXT("/Game/NotRemovedPackage")), ECoreRedirectFlags::None);
+
+	if (IsKnownMissing(ECoreRedirectFlags::Type_Package, FCoreRedirectObjectName(TEXT("/Game/NotRemovedPackage"))))
+	{
+		bSuccess = false;
+		UE_LOG(LogLinker, Error, TEXT("FCoreRedirect Test Failed: /Game/NotRemovedPackage should no longer be removed!"));
 	}
 
 	// Restore old state
@@ -771,12 +826,12 @@ bool FCoreRedirects::ReadRedirectsFromIni(const FString& IniName)
 
 					if (bInstanceOnly)
 					{
-						NewFlags |= ECoreRedirectFlags::Option_InstanceOnly;
+						NewFlags |= ECoreRedirectFlags::Category_InstanceOnly;
 					}
 
 					if (bRemoved)
 					{
-						NewFlags |= ECoreRedirectFlags::Option_Removed;
+						NewFlags |= ECoreRedirectFlags::Category_Removed;
 					}
 
 					if (bMatchSubstring)
@@ -784,7 +839,7 @@ bool FCoreRedirects::ReadRedirectsFromIni(const FString& IniName)
 						NewFlags |= ECoreRedirectFlags::Option_MatchSubstring;
 					}
 
-					FCoreRedirect* Redirect = new (NewRedirects) FCoreRedirect(NewFlags, FCoreRedirectObjectName(OldName), FCoreRedirectObjectName(NewName));
+					FCoreRedirect* Redirect = &NewRedirects.Emplace_GetRef(NewFlags, FCoreRedirectObjectName(OldName), FCoreRedirectObjectName(NewName));
 
 					if (!OverrideClassName.IsEmpty())
 					{
@@ -1038,13 +1093,13 @@ ECoreRedirectFlags FCoreRedirects::GetFlagsForTypeClass(UClass *TypeClass)
 PRAGMA_DISABLE_OPTIMIZATION
 
 // The compiler doesn't like having a massive string table in a single function so split it up
-#define CLASS_REDIRECT(OldName, NewName) new (Redirects) FCoreRedirect(ECoreRedirectFlags::Type_Class, TEXT(OldName), TEXT(NewName))
-#define CLASS_REDIRECT_INSTANCES(OldName, NewName) new (Redirects) FCoreRedirect(ECoreRedirectFlags::Type_Class | ECoreRedirectFlags::Option_InstanceOnly, TEXT(OldName), TEXT(NewName))
-#define STRUCT_REDIRECT(OldName, NewName) new (Redirects) FCoreRedirect(ECoreRedirectFlags::Type_Struct, TEXT(OldName), TEXT(NewName))
-#define ENUM_REDIRECT(OldName, NewName) new (Redirects) FCoreRedirect(ECoreRedirectFlags::Type_Enum, TEXT(OldName), TEXT(NewName))
-#define PROPERTY_REDIRECT(OldName, NewName) new (Redirects) FCoreRedirect(ECoreRedirectFlags::Type_Property, TEXT(OldName), TEXT(NewName))
-#define FUNCTION_REDIRECT(OldName, NewName) new (Redirects) FCoreRedirect(ECoreRedirectFlags::Type_Function, TEXT(OldName), TEXT(NewName))
-#define PACKAGE_REDIRECT(OldName, NewName) new (Redirects) FCoreRedirect(ECoreRedirectFlags::Type_Package, TEXT(OldName), TEXT(NewName))
+#define CLASS_REDIRECT(OldName, NewName) Redirects.Emplace_GetRef(ECoreRedirectFlags::Type_Class, TEXT(OldName), TEXT(NewName))
+#define CLASS_REDIRECT_INSTANCES(OldName, NewName) Redirects.Emplace_GetRef(ECoreRedirectFlags::Type_Class | ECoreRedirectFlags::Category_InstanceOnly, TEXT(OldName), TEXT(NewName))
+#define STRUCT_REDIRECT(OldName, NewName) Redirects.Emplace_GetRef(ECoreRedirectFlags::Type_Struct, TEXT(OldName), TEXT(NewName))
+#define ENUM_REDIRECT(OldName, NewName) Redirects.Emplace_GetRef(ECoreRedirectFlags::Type_Enum, TEXT(OldName), TEXT(NewName))
+#define PROPERTY_REDIRECT(OldName, NewName) Redirects.Emplace_GetRef(ECoreRedirectFlags::Type_Property, TEXT(OldName), TEXT(NewName))
+#define FUNCTION_REDIRECT(OldName, NewName) Redirects.Emplace_GetRef(ECoreRedirectFlags::Type_Function, TEXT(OldName), TEXT(NewName))
+#define PACKAGE_REDIRECT(OldName, NewName) Redirects.Emplace_GetRef(ECoreRedirectFlags::Type_Package, TEXT(OldName), TEXT(NewName))
 
 static void RegisterNativeRedirects40(TArray<FCoreRedirect>& Redirects)
 {
@@ -1316,77 +1371,77 @@ static void RegisterNativeRedirects40(TArray<FCoreRedirect>& Redirects)
 	PROPERTY_REDIRECT("SlateBrush.TextureObject", "SlateBrush.ResourceObject");
 	PROPERTY_REDIRECT("WorldSettings.DefaultGameType", "WorldSettings.DefaultGameMode");
 
-	FCoreRedirect* PointLightComponent = CLASS_REDIRECT("PointLightComponent", "PointLightComponent");
-	PointLightComponent->ValueChanges.Add(TEXT("PointLightComponent0"), TEXT("LightComponent0"));
-
-	FCoreRedirect* DirectionalLightComponent = CLASS_REDIRECT("DirectionalLightComponent", "DirectionalLightComponent");
-	DirectionalLightComponent->ValueChanges.Add(TEXT("DirectionalLightComponent0"), TEXT("LightComponent0"));
-
-	FCoreRedirect* SpotLightComponent = CLASS_REDIRECT("SpotLightComponent", "SpotLightComponent");
-	SpotLightComponent->ValueChanges.Add(TEXT("SpotLightComponent0"), TEXT("LightComponent0"));
-
-	FCoreRedirect* ETransitionGetterType = ENUM_REDIRECT("ETransitionGetterType", "ETransitionGetter");
-	ETransitionGetterType->ValueChanges.Add(TEXT("TGT_ArbitraryState_GetBlendWeight"), TEXT("ETransitionGetter::ArbitraryState_GetBlendWeight"));
-	ETransitionGetterType->ValueChanges.Add(TEXT("TGT_CurrentState_ElapsedTime"), TEXT("ETransitionGetter::CurrentState_ElapsedTime"));
-	ETransitionGetterType->ValueChanges.Add(TEXT("TGT_CurrentState_GetBlendWeight"), TEXT("ETransitionGetter::CurrentState_GetBlendWeight"));
-	ETransitionGetterType->ValueChanges.Add(TEXT("TGT_CurrentTransitionDuration"), TEXT("ETransitionGetter::CurrentTransitionDuration"));
-	ETransitionGetterType->ValueChanges.Add(TEXT("TGT_SequencePlayer_GetCurrentTime"), TEXT("ETransitionGetter::AnimationAsset_GetCurrentTime"));
-	ETransitionGetterType->ValueChanges.Add(TEXT("TGT_SequencePlayer_GetCurrentTimeFraction"), TEXT("ETransitionGetter::AnimationAsset_GetCurrentTimeFraction"));
-	ETransitionGetterType->ValueChanges.Add(TEXT("TGT_SequencePlayer_GetLength"), TEXT("ETransitionGetter::AnimationAsset_GetLength"));
-	ETransitionGetterType->ValueChanges.Add(TEXT("TGT_SequencePlayer_GetTimeFromEnd"), TEXT("ETransitionGetter::AnimationAsset_GetTimeFromEnd"));
-	ETransitionGetterType->ValueChanges.Add(TEXT("TGT_SequencePlayer_GetTimeFromEndFraction"), TEXT("ETransitionGetter::AnimationAsset_GetTimeFromEndFraction"));
-
-	FCoreRedirect* EModifyFrequency = ENUM_REDIRECT("EModifyFrequency", "EComponentMobility");
-	EModifyFrequency->ValueChanges.Add(TEXT("MF_Dynamic"), TEXT("EComponentMobility::Movable"));
-	EModifyFrequency->ValueChanges.Add(TEXT("MF_OccasionallyModified"), TEXT("EComponentMobility::Stationary"));
-	EModifyFrequency->ValueChanges.Add(TEXT("MF_Static"), TEXT("EComponentMobility::Static"));
-
-	FCoreRedirect* EAttachLocationType = ENUM_REDIRECT("EAttachLocationType", "EAttachLocation");
-	EAttachLocationType->ValueChanges.Add(TEXT("EAttachLocationType_AbsoluteWorld"), TEXT("EAttachLocation::KeepWorldPosition"));
-	EAttachLocationType->ValueChanges.Add(TEXT("EAttachLocationType_RelativeOffset"), TEXT("EAttachLocation::KeepRelativeOffset"));
-	EAttachLocationType->ValueChanges.Add(TEXT("EAttachLocationType_SnapTo"), TEXT("EAttachLocation::SnapToTarget"));
-
-	FCoreRedirect* EAxis = ENUM_REDIRECT("EAxis", "EAxis");
-	EAxis->ValueChanges.Add(TEXT("AXIS_BLANK"), TEXT("EAxis::None"));
-	EAxis->ValueChanges.Add(TEXT("AXIS_NONE"), TEXT("EAxis::None"));
-	EAxis->ValueChanges.Add(TEXT("AXIS_X"), TEXT("EAxis::X"));
-	EAxis->ValueChanges.Add(TEXT("AXIS_Y"), TEXT("EAxis::Y"));
-	EAxis->ValueChanges.Add(TEXT("AXIS_Z"), TEXT("EAxis::Z"));
-
-	FCoreRedirect* EKeys = ENUM_REDIRECT("EKeys", "EKeys");
-	EKeys->ValueChanges.Add(TEXT("EKeys::XboxTypeS_A"), TEXT("EKeys::Gamepad_FaceButton_Bottom"));
-	EKeys->ValueChanges.Add(TEXT("EKeys::XboxTypeS_B"), TEXT("EKeys::Gamepad_FaceButton_Right"));
-	EKeys->ValueChanges.Add(TEXT("EKeys::XboxTypeS_X"), TEXT("EKeys::Gamepad_FaceButton_Left"));
-	EKeys->ValueChanges.Add(TEXT("EKeys::XboxTypeS_Y"), TEXT("EKeys::Gamepad_FaceButton_Top"));
-	EKeys->ValueChanges.Add(TEXT("EKeys::XboxTypeS_Back"), TEXT("EKeys::Gamepad_Special_Left"));
-	EKeys->ValueChanges.Add(TEXT("EKeys::XboxTypeS_Start"), TEXT("EKeys::Gamepad_Special_Right"));
-	EKeys->ValueChanges.Add(TEXT("EKeys::XboxTypeS_DPad_Down"), TEXT("EKeys::Gamepad_DPad_Down"));
-	EKeys->ValueChanges.Add(TEXT("EKeys::XboxTypeS_DPad_Left"), TEXT("EKeys::Gamepad_DPad_Left"));
-	EKeys->ValueChanges.Add(TEXT("EKeys::XboxTypeS_DPad_Right"), TEXT("EKeys::Gamepad_DPad_Right"));
-	EKeys->ValueChanges.Add(TEXT("EKeys::XboxTypeS_DPad_Up"), TEXT("EKeys::Gamepad_DPad_Up"));
-	EKeys->ValueChanges.Add(TEXT("EKeys::XboxTypeS_LeftShoulder"), TEXT("EKeys::Gamepad_LeftShoulder"));
-	EKeys->ValueChanges.Add(TEXT("EKeys::XboxTypeS_LeftThumbstick"), TEXT("EKeys::Gamepad_LeftThumbstick"));
-	EKeys->ValueChanges.Add(TEXT("EKeys::XboxTypeS_LeftTrigger"), TEXT("EKeys::Gamepad_LeftTrigger"));
-	EKeys->ValueChanges.Add(TEXT("EKeys::XboxTypeS_LeftTriggerAxis"), TEXT("EKeys::Gamepad_LeftTriggerAxis"));
-	EKeys->ValueChanges.Add(TEXT("EKeys::XboxTypeS_LeftX"), TEXT("EKeys::Gamepad_LeftX"));
-	EKeys->ValueChanges.Add(TEXT("EKeys::XboxTypeS_LeftY"), TEXT("EKeys::Gamepad_LeftY"));
-	EKeys->ValueChanges.Add(TEXT("EKeys::XboxTypeS_RightShoulder"), TEXT("EKeys::Gamepad_RightShoulder"));
-	EKeys->ValueChanges.Add(TEXT("EKeys::XboxTypeS_RightThumbstick"), TEXT("EKeys::Gamepad_RightThumbstick"));
-	EKeys->ValueChanges.Add(TEXT("EKeys::XboxTypeS_RightTrigger"), TEXT("EKeys::Gamepad_RightTrigger"));
-	EKeys->ValueChanges.Add(TEXT("EKeys::XboxTypeS_RightTriggerAxis"), TEXT("EKeys::Gamepad_RightTriggerAxis"));
-	EKeys->ValueChanges.Add(TEXT("EKeys::XboxTypeS_RightX"), TEXT("EKeys::Gamepad_RightX"));
-	EKeys->ValueChanges.Add(TEXT("EKeys::XboxTypeS_RightY"), TEXT("EKeys::Gamepad_RightY"));
-
-	FCoreRedirect* EMaxConcurrentResolutionRule = ENUM_REDIRECT("EMaxConcurrentResolutionRule", "EMaxConcurrentResolutionRule");
-	EMaxConcurrentResolutionRule->ValueChanges.Add(TEXT("EMaxConcurrentResolutionRule::StopFarthest"), TEXT("EMaxConcurrentResolutionRule::StopFarthestThenPreventNew"));
-
-
-	FCoreRedirect* EParticleEventType = ENUM_REDIRECT("EParticleEventType", "EParticleEventType");
-	EParticleEventType->ValueChanges.Add(TEXT("EPET_Kismet"), TEXT("EPET_Blueprint"));
-
-	FCoreRedirect* ETranslucencyLightingMode = ENUM_REDIRECT("ETranslucencyLightingMode", "ETranslucencyLightingMode");
-	ETranslucencyLightingMode->ValueChanges.Add(TEXT("TLM_PerPixel"), TEXT("TLM_VolumetricDirectional"));
-	ETranslucencyLightingMode->ValueChanges.Add(TEXT("TLM_PerPixelNonDirectional"), TEXT("TLM_VolumetricNonDirectional"));
+	FCoreRedirect& PointLightComponent = CLASS_REDIRECT("PointLightComponent", "PointLightComponent");
+	PointLightComponent.ValueChanges.Add(TEXT("PointLightComponent0"), TEXT("LightComponent0"));
+
+	FCoreRedirect& DirectionalLightComponent = CLASS_REDIRECT("DirectionalLightComponent", "DirectionalLightComponent");
+	DirectionalLightComponent.ValueChanges.Add(TEXT("DirectionalLightComponent0"), TEXT("LightComponent0"));
+
+	FCoreRedirect& SpotLightComponent = CLASS_REDIRECT("SpotLightComponent", "SpotLightComponent");
+	SpotLightComponent.ValueChanges.Add(TEXT("SpotLightComponent0"), TEXT("LightComponent0"));
+
+	FCoreRedirect& ETransitionGetterType = ENUM_REDIRECT("ETransitionGetterType", "ETransitionGetter");
+	ETransitionGetterType.ValueChanges.Add(TEXT("TGT_ArbitraryState_GetBlendWeight"), TEXT("ETransitionGetter::ArbitraryState_GetBlendWeight"));
+	ETransitionGetterType.ValueChanges.Add(TEXT("TGT_CurrentState_ElapsedTime"), TEXT("ETransitionGetter::CurrentState_ElapsedTime"));
+	ETransitionGetterType.ValueChanges.Add(TEXT("TGT_CurrentState_GetBlendWeight"), TEXT("ETransitionGetter::CurrentState_GetBlendWeight"));
+	ETransitionGetterType.ValueChanges.Add(TEXT("TGT_CurrentTransitionDuration"), TEXT("ETransitionGetter::CurrentTransitionDuration"));
+	ETransitionGetterType.ValueChanges.Add(TEXT("TGT_SequencePlayer_GetCurrentTime"), TEXT("ETransitionGetter::AnimationAsset_GetCurrentTime"));
+	ETransitionGetterType.ValueChanges.Add(TEXT("TGT_SequencePlayer_GetCurrentTimeFraction"), TEXT("ETransitionGetter::AnimationAsset_GetCurrentTimeFraction"));
+	ETransitionGetterType.ValueChanges.Add(TEXT("TGT_SequencePlayer_GetLength"), TEXT("ETransitionGetter::AnimationAsset_GetLength"));
+	ETransitionGetterType.ValueChanges.Add(TEXT("TGT_SequencePlayer_GetTimeFromEnd"), TEXT("ETransitionGetter::AnimationAsset_GetTimeFromEnd"));
+	ETransitionGetterType.ValueChanges.Add(TEXT("TGT_SequencePlayer_GetTimeFromEndFraction"), TEXT("ETransitionGetter::AnimationAsset_GetTimeFromEndFraction"));
+
+	FCoreRedirect& EModifyFrequency = ENUM_REDIRECT("EModifyFrequency", "EComponentMobility");
+	EModifyFrequency.ValueChanges.Add(TEXT("MF_Dynamic"), TEXT("EComponentMobility::Movable"));
+	EModifyFrequency.ValueChanges.Add(TEXT("MF_OccasionallyModified"), TEXT("EComponentMobility::Stationary"));
+	EModifyFrequency.ValueChanges.Add(TEXT("MF_Static"), TEXT("EComponentMobility::Static"));
+
+	FCoreRedirect& EAttachLocationType = ENUM_REDIRECT("EAttachLocationType", "EAttachLocation");
+	EAttachLocationType.ValueChanges.Add(TEXT("EAttachLocationType_AbsoluteWorld"), TEXT("EAttachLocation::KeepWorldPosition"));
+	EAttachLocationType.ValueChanges.Add(TEXT("EAttachLocationType_RelativeOffset"), TEXT("EAttachLocation::KeepRelativeOffset"));
+	EAttachLocationType.ValueChanges.Add(TEXT("EAttachLocationType_SnapTo"), TEXT("EAttachLocation::SnapToTarget"));
+
+	FCoreRedirect& EAxis = ENUM_REDIRECT("EAxis", "EAxis");
+	EAxis.ValueChanges.Add(TEXT("AXIS_BLANK"), TEXT("EAxis::None"));
+	EAxis.ValueChanges.Add(TEXT("AXIS_NONE"), TEXT("EAxis::None"));
+	EAxis.ValueChanges.Add(TEXT("AXIS_X"), TEXT("EAxis::X"));
+	EAxis.ValueChanges.Add(TEXT("AXIS_Y"), TEXT("EAxis::Y"));
+	EAxis.ValueChanges.Add(TEXT("AXIS_Z"), TEXT("EAxis::Z"));
+
+	FCoreRedirect& EKeys = ENUM_REDIRECT("EKeys", "EKeys");
+	EKeys.ValueChanges.Add(TEXT("EKeys::XboxTypeS_A"), TEXT("EKeys::Gamepad_FaceButton_Bottom"));
+	EKeys.ValueChanges.Add(TEXT("EKeys::XboxTypeS_B"), TEXT("EKeys::Gamepad_FaceButton_Right"));
+	EKeys.ValueChanges.Add(TEXT("EKeys::XboxTypeS_X"), TEXT("EKeys::Gamepad_FaceButton_Left"));
+	EKeys.ValueChanges.Add(TEXT("EKeys::XboxTypeS_Y"), TEXT("EKeys::Gamepad_FaceButton_Top"));
+	EKeys.ValueChanges.Add(TEXT("EKeys::XboxTypeS_Back"), TEXT("EKeys::Gamepad_Special_Left"));
+	EKeys.ValueChanges.Add(TEXT("EKeys::XboxTypeS_Start"), TEXT("EKeys::Gamepad_Special_Right"));
+	EKeys.ValueChanges.Add(TEXT("EKeys::XboxTypeS_DPad_Down"), TEXT("EKeys::Gamepad_DPad_Down"));
+	EKeys.ValueChanges.Add(TEXT("EKeys::XboxTypeS_DPad_Left"), TEXT("EKeys::Gamepad_DPad_Left"));
+	EKeys.ValueChanges.Add(TEXT("EKeys::XboxTypeS_DPad_Right"), TEXT("EKeys::Gamepad_DPad_Right"));
+	EKeys.ValueChanges.Add(TEXT("EKeys::XboxTypeS_DPad_Up"), TEXT("EKeys::Gamepad_DPad_Up"));
+	EKeys.ValueChanges.Add(TEXT("EKeys::XboxTypeS_LeftShoulder"), TEXT("EKeys::Gamepad_LeftShoulder"));
+	EKeys.ValueChanges.Add(TEXT("EKeys::XboxTypeS_LeftThumbstick"), TEXT("EKeys::Gamepad_LeftThumbstick"));
+	EKeys.ValueChanges.Add(TEXT("EKeys::XboxTypeS_LeftTrigger"), TEXT("EKeys::Gamepad_LeftTrigger"));
+	EKeys.ValueChanges.Add(TEXT("EKeys::XboxTypeS_LeftTriggerAxis"), TEXT("EKeys::Gamepad_LeftTriggerAxis"));
+	EKeys.ValueChanges.Add(TEXT("EKeys::XboxTypeS_LeftX"), TEXT("EKeys::Gamepad_LeftX"));
+	EKeys.ValueChanges.Add(TEXT("EKeys::XboxTypeS_LeftY"), TEXT("EKeys::Gamepad_LeftY"));
+	EKeys.ValueChanges.Add(TEXT("EKeys::XboxTypeS_RightShoulder"), TEXT("EKeys::Gamepad_RightShoulder"));
+	EKeys.ValueChanges.Add(TEXT("EKeys::XboxTypeS_RightThumbstick"), TEXT("EKeys::Gamepad_RightThumbstick"));
+	EKeys.ValueChanges.Add(TEXT("EKeys::XboxTypeS_RightTrigger"), TEXT("EKeys::Gamepad_RightTrigger"));
+	EKeys.ValueChanges.Add(TEXT("EKeys::XboxTypeS_RightTriggerAxis"), TEXT("EKeys::Gamepad_RightTriggerAxis"));
+	EKeys.ValueChanges.Add(TEXT("EKeys::XboxTypeS_RightX"), TEXT("EKeys::Gamepad_RightX"));
+	EKeys.ValueChanges.Add(TEXT("EKeys::XboxTypeS_RightY"), TEXT("EKeys::Gamepad_RightY"));
+
+	FCoreRedirect& EMaxConcurrentResolutionRule = ENUM_REDIRECT("EMaxConcurrentResolutionRule", "EMaxConcurrentResolutionRule");
+	EMaxConcurrentResolutionRule.ValueChanges.Add(TEXT("EMaxConcurrentResolutionRule::StopFarthest"), TEXT("EMaxConcurrentResolutionRule::StopFarthestThenPreventNew"));
+
+
+	FCoreRedirect& EParticleEventType = ENUM_REDIRECT("EParticleEventType", "EParticleEventType");
+	EParticleEventType.ValueChanges.Add(TEXT("EPET_Kismet"), TEXT("EPET_Blueprint"));
+
+	FCoreRedirect& ETranslucencyLightingMode = ENUM_REDIRECT("ETranslucencyLightingMode", "ETranslucencyLightingMode");
+	ETranslucencyLightingMode.ValueChanges.Add(TEXT("TLM_PerPixel"), TEXT("TLM_VolumetricDirectional"));
+	ETranslucencyLightingMode.ValueChanges.Add(TEXT("TLM_PerPixelNonDirectional"), TEXT("TLM_VolumetricNonDirectional"));
 }
 
 static void RegisterNativeRedirects46(TArray<FCoreRedirect>& Redirects)
@@ -1436,25 +1491,25 @@ static void RegisterNativeRedirects46(TArray<FCoreRedirect>& Redirects)
 	PROPERTY_REDIRECT("SplineMeshComponent.SplineXDir", "SplineMeshComponent.SplineUpDir");
 	PROPERTY_REDIRECT("TextureFactory.LightingModel", "TextureFactory.ShadingModel");
 
-	FCoreRedirect* EKinematicBonesUpdateToPhysics = ENUM_REDIRECT("EKinematicBonesUpdateToPhysics", "EKinematicBonesUpdateToPhysics");
-	EKinematicBonesUpdateToPhysics->ValueChanges.Add(TEXT("EKinematicBonesUpdateToPhysics::SkipFixedAndSimulatingBones"), TEXT("EKinematicBonesUpdateToPhysics::SkipAllBones"));
+	FCoreRedirect& EKinematicBonesUpdateToPhysics = ENUM_REDIRECT("EKinematicBonesUpdateToPhysics", "EKinematicBonesUpdateToPhysics");
+	EKinematicBonesUpdateToPhysics.ValueChanges.Add(TEXT("EKinematicBonesUpdateToPhysics::SkipFixedAndSimulatingBones"), TEXT("EKinematicBonesUpdateToPhysics::SkipAllBones"));
 
-	FCoreRedirect* EMaterialLightingModel = ENUM_REDIRECT("EMaterialLightingModel", "EMaterialShadingModel");
-	EMaterialLightingModel->ValueChanges.Add(TEXT("MLM_DefaultLit"), TEXT("MSM_DefaultLit"));
-	EMaterialLightingModel->ValueChanges.Add(TEXT("MLM_PreintegratedSkin"), TEXT("MSM_PreintegratedSkin"));
-	EMaterialLightingModel->ValueChanges.Add(TEXT("MLM_Subsurface"), TEXT("MSM_Subsurface"));
-	EMaterialLightingModel->ValueChanges.Add(TEXT("MLM_Unlit"), TEXT("MSM_Unlit"));
+	FCoreRedirect& EMaterialLightingModel = ENUM_REDIRECT("EMaterialLightingModel", "EMaterialShadingModel");
+	EMaterialLightingModel.ValueChanges.Add(TEXT("MLM_DefaultLit"), TEXT("MSM_DefaultLit"));
+	EMaterialLightingModel.ValueChanges.Add(TEXT("MLM_PreintegratedSkin"), TEXT("MSM_PreintegratedSkin"));
+	EMaterialLightingModel.ValueChanges.Add(TEXT("MLM_Subsurface"), TEXT("MSM_Subsurface"));
+	EMaterialLightingModel.ValueChanges.Add(TEXT("MLM_Unlit"), TEXT("MSM_Unlit"));
 
-	FCoreRedirect* ESmartNavLinkDir = ENUM_REDIRECT("ESmartNavLinkDir", "ENavLinkDirection");
-	ESmartNavLinkDir->ValueChanges.Add(TEXT("ESmartNavLinkDir::BothWays"), TEXT("ENavLinkDirection::BothWays"));
-	ESmartNavLinkDir->ValueChanges.Add(TEXT("ESmartNavLinkDir::OneWay"), TEXT("ENavLinkDirection::LeftToRight"));
+	FCoreRedirect& ESmartNavLinkDir = ENUM_REDIRECT("ESmartNavLinkDir", "ENavLinkDirection");
+	ESmartNavLinkDir.ValueChanges.Add(TEXT("ESmartNavLinkDir::BothWays"), TEXT("ENavLinkDirection::BothWays"));
+	ESmartNavLinkDir.ValueChanges.Add(TEXT("ESmartNavLinkDir::OneWay"), TEXT("ENavLinkDirection::LeftToRight"));
 
-	FCoreRedirect* EPhysicsType = ENUM_REDIRECT("EPhysicsType", "EPhysicsType");
-	EPhysicsType->ValueChanges.Add(TEXT("PhysType_Fixed"), TEXT("PhysType_Kinematic"));
-	EPhysicsType->ValueChanges.Add(TEXT("PhysType_Unfixed"), TEXT("PhysType_Simulated"));
+	FCoreRedirect& EPhysicsType = ENUM_REDIRECT("EPhysicsType", "EPhysicsType");
+	EPhysicsType.ValueChanges.Add(TEXT("PhysType_Fixed"), TEXT("PhysType_Kinematic"));
+	EPhysicsType.ValueChanges.Add(TEXT("PhysType_Unfixed"), TEXT("PhysType_Simulated"));
 
-	FCoreRedirect* ESceneTextureId = ENUM_REDIRECT("ESceneTextureId", "ESceneTextureId");
-	ESceneTextureId->ValueChanges.Add(TEXT("PPI_LightingModel"), TEXT("PPI_ShadingModelColor"));
+	FCoreRedirect& ESceneTextureId = ENUM_REDIRECT("ESceneTextureId", "ESceneTextureId");
+	ESceneTextureId.ValueChanges.Add(TEXT("PPI_LightingModel"), TEXT("PPI_ShadingModelColor"));
 
 	// 4.5
 
@@ -1679,22 +1734,22 @@ static void RegisterNativeRedirects49(TArray<FCoreRedirect>& Redirects)
 	PROPERTY_REDIRECT("MeshComponent.Materials", "MeshComponent.OverrideMaterials");
 	PROPERTY_REDIRECT("Pawn.AutoPossess", "Pawn.AutoPossessPlayer");
 
-	FCoreRedirect* ECollisionChannel = ENUM_REDIRECT("ECollisionChannel", "ECollisionChannel");
-	ECollisionChannel->ValueChanges.Add(TEXT("ECC_Default"), TEXT("ECC_Visibility"));
-	ECollisionChannel->ValueChanges.Add(TEXT("ECC_Dynamic"), TEXT("ECC_WorldDynamic"));
-	ECollisionChannel->ValueChanges.Add(TEXT("ECC_OverlapAll"), TEXT("ECC_OverlapAll_Deprecated"));
-	ECollisionChannel->ValueChanges.Add(TEXT("ECC_OverlapAllDynamic"), TEXT("ECC_OverlapAll_Deprecated"));
-	ECollisionChannel->ValueChanges.Add(TEXT("ECC_OverlapAllDynamic_Deprecated"), TEXT("ECC_OverlapAll_Deprecated"));
-	ECollisionChannel->ValueChanges.Add(TEXT("ECC_OverlapAllStatic"), TEXT("ECC_OverlapAll_Deprecated"));
-	ECollisionChannel->ValueChanges.Add(TEXT("ECC_OverlapAllStatic_Deprecated"), TEXT("ECC_OverlapAll_Deprecated"));
-	ECollisionChannel->ValueChanges.Add(TEXT("ECC_PawnMovement"), TEXT("ECC_Pawn"));
-	ECollisionChannel->ValueChanges.Add(TEXT("ECC_RigidBody"), TEXT("ECC_PhysicsBody"));
-	ECollisionChannel->ValueChanges.Add(TEXT("ECC_RigidBodyInteractable"), TEXT("ECC_PhysicsBody"));
-	ECollisionChannel->ValueChanges.Add(TEXT("ECC_TouchAll"), TEXT("ECC_OverlapAll_Deprecated"));
-	ECollisionChannel->ValueChanges.Add(TEXT("ECC_TouchAllDynamic"), TEXT("ECC_OverlapAll_Deprecated"));
-	ECollisionChannel->ValueChanges.Add(TEXT("ECC_TouchAllStatic"), TEXT("ECC_OverlapAll_Deprecated"));
-	ECollisionChannel->ValueChanges.Add(TEXT("ECC_VehicleMovement"), TEXT("ECC_Vehicle"));
-	ECollisionChannel->ValueChanges.Add(TEXT("ECC_WorldTrace"), TEXT("ECC_WorldStatic"));
+	FCoreRedirect& ECollisionChannel = ENUM_REDIRECT("ECollisionChannel", "ECollisionChannel");
+	ECollisionChannel.ValueChanges.Add(TEXT("ECC_Default"), TEXT("ECC_Visibility"));
+	ECollisionChannel.ValueChanges.Add(TEXT("ECC_Dynamic"), TEXT("ECC_WorldDynamic"));
+	ECollisionChannel.ValueChanges.Add(TEXT("ECC_OverlapAll"), TEXT("ECC_OverlapAll_Deprecated"));
+	ECollisionChannel.ValueChanges.Add(TEXT("ECC_OverlapAllDynamic"), TEXT("ECC_OverlapAll_Deprecated"));
+	ECollisionChannel.ValueChanges.Add(TEXT("ECC_OverlapAllDynamic_Deprecated"), TEXT("ECC_OverlapAll_Deprecated"));
+	ECollisionChannel.ValueChanges.Add(TEXT("ECC_OverlapAllStatic"), TEXT("ECC_OverlapAll_Deprecated"));
+	ECollisionChannel.ValueChanges.Add(TEXT("ECC_OverlapAllStatic_Deprecated"), TEXT("ECC_OverlapAll_Deprecated"));
+	ECollisionChannel.ValueChanges.Add(TEXT("ECC_PawnMovement"), TEXT("ECC_Pawn"));
+	ECollisionChannel.ValueChanges.Add(TEXT("ECC_RigidBody"), TEXT("ECC_PhysicsBody"));
+	ECollisionChannel.ValueChanges.Add(TEXT("ECC_RigidBodyInteractable"), TEXT("ECC_PhysicsBody"));
+	ECollisionChannel.ValueChanges.Add(TEXT("ECC_TouchAll"), TEXT("ECC_OverlapAll_Deprecated"));
+	ECollisionChannel.ValueChanges.Add(TEXT("ECC_TouchAllDynamic"), TEXT("ECC_OverlapAll_Deprecated"));
+	ECollisionChannel.ValueChanges.Add(TEXT("ECC_TouchAllStatic"), TEXT("ECC_OverlapAll_Deprecated"));
+	ECollisionChannel.ValueChanges.Add(TEXT("ECC_VehicleMovement"), TEXT("ECC_Vehicle"));
+	ECollisionChannel.ValueChanges.Add(TEXT("ECC_WorldTrace"), TEXT("ECC_WorldStatic"));
 
 	// 4.8
 
@@ -1710,20 +1765,20 @@ static void RegisterNativeRedirects49(TArray<FCoreRedirect>& Redirects)
 	
 	STRUCT_REDIRECT("ProceduralFoliageTypeData", "/Script/Foliage.FoliageTypeObject");
 
-	FCoreRedirect* EComponentCreationMethod = ENUM_REDIRECT("EComponentCreationMethod", "EComponentCreationMethod");
-	EComponentCreationMethod->ValueChanges.Add(TEXT("EComponentCreationMethod::ConstructionScript"), TEXT("EComponentCreationMethod::SimpleConstructionScript"));
+	FCoreRedirect& EComponentCreationMethod = ENUM_REDIRECT("EComponentCreationMethod", "EComponentCreationMethod");
+	EComponentCreationMethod.ValueChanges.Add(TEXT("EComponentCreationMethod::ConstructionScript"), TEXT("EComponentCreationMethod::SimpleConstructionScript"));
 
-	FCoreRedirect* EConstraintTransform = ENUM_REDIRECT("EConstraintTransform", "EConstraintTransform");
-	EConstraintTransform->ValueChanges.Add(TEXT("EConstraintTransform::Absoluate"), TEXT("EConstraintTransform::Absolute"));
+	FCoreRedirect& EConstraintTransform = ENUM_REDIRECT("EConstraintTransform", "EConstraintTransform");
+	EConstraintTransform.ValueChanges.Add(TEXT("EConstraintTransform::Absoluate"), TEXT("EConstraintTransform::Absolute"));
 
-	FCoreRedirect* ELockedAxis = ENUM_REDIRECT("ELockedAxis", "EDOFMode");
-	ELockedAxis->ValueChanges.Add(TEXT("Custom"), TEXT("EDOFMode::CustomPlane"));
-	ELockedAxis->ValueChanges.Add(TEXT("X"), TEXT("EDOFMode::YZPlane"));
-	ELockedAxis->ValueChanges.Add(TEXT("Y"), TEXT("EDOFMode::XZPlane"));
-	ELockedAxis->ValueChanges.Add(TEXT("Z"), TEXT("EDOFMode::XYPlane"));
+	FCoreRedirect& ELockedAxis = ENUM_REDIRECT("ELockedAxis", "EDOFMode");
+	ELockedAxis.ValueChanges.Add(TEXT("Custom"), TEXT("EDOFMode::CustomPlane"));
+	ELockedAxis.ValueChanges.Add(TEXT("X"), TEXT("EDOFMode::YZPlane"));
+	ELockedAxis.ValueChanges.Add(TEXT("Y"), TEXT("EDOFMode::XZPlane"));
+	ELockedAxis.ValueChanges.Add(TEXT("Z"), TEXT("EDOFMode::XYPlane"));
 
-	FCoreRedirect* EEndPlayReason = ENUM_REDIRECT("EEndPlayReason", "EEndPlayReason");
-	EEndPlayReason->ValueChanges.Add(TEXT("EEndPlayReason::ActorDestroyed"), TEXT("EEndPlayReason::Destroyed"));
+	FCoreRedirect& EEndPlayReason = ENUM_REDIRECT("EEndPlayReason", "EEndPlayReason");
+	EEndPlayReason.ValueChanges.Add(TEXT("EEndPlayReason::ActorDestroyed"), TEXT("EEndPlayReason::Destroyed"));
 
 	FUNCTION_REDIRECT("ActorComponent.ReceiveInitializeComponent", "ActorComponent.ReceiveBeginPlay");
 	FUNCTION_REDIRECT("ActorComponent.ReceiveUninitializeComponent", "ActorComponent.ReceiveEndPlay");
diff --git a/Engine/Source/Runtime/CoreUObject/Private/UObject/LinkerLoad.cpp b/Engine/Source/Runtime/CoreUObject/Private/UObject/LinkerLoad.cpp
--- a/Engine/Source/Runtime/CoreUObject/Private/UObject/LinkerLoad.cpp
+++ b/Engine/Source/Runtime/CoreUObject/Private/UObject/LinkerLoad.cpp
@@ -165,8 +165,8 @@ void FLinkerLoad::CreateActiveRedirectsMap(const FString& GEngineIniName)
 					if (NewSubobjName != NAME_None || OldSubobjName != NAME_None)
 					{
 						check(OldSubobjName != NAME_None && OldClassName != NAME_None );
-						FCoreRedirect* Redirect = new (NewRedirects) FCoreRedirect(ECoreRedirectFlags::Type_Class, OldClassName.ToString(), OldClassName.ToString());
-						Redirect->ValueChanges.Add(OldSubobjName.ToString(), NewSubobjName.ToString());
+						FCoreRedirect& Redirect = NewRedirects.Emplace_GetRef(ECoreRedirectFlags::Type_Class, OldClassName.ToString(), OldClassName.ToString());
+						Redirect.ValueChanges.Add(OldSubobjName.ToString(), NewSubobjName.ToString());
 					}
 					//instances only
 					else if( bInstanceOnly )
@@ -174,11 +174,11 @@ void FLinkerLoad::CreateActiveRedirectsMap(const FString& GEngineIniName)
 						// If NewClassName is none, register as removed instead
 						if (NewClassName == NAME_None)
 						{
-							FCoreRedirect* Redirect = new (NewRedirects) FCoreRedirect(ECoreRedirectFlags::Type_Class | ECoreRedirectFlags::Option_InstanceOnly | ECoreRedirectFlags::Option_Removed, OldClassName.ToString(), NewClassName.ToString());
+							NewRedirects.Emplace(ECoreRedirectFlags::Type_Class | ECoreRedirectFlags::Category_InstanceOnly | ECoreRedirectFlags::Category_Removed, OldClassName.ToString(), NewClassName.ToString());
 						}
 						else
 						{
-							FCoreRedirect* Redirect = new (NewRedirects) FCoreRedirect(ECoreRedirectFlags::Type_Class | ECoreRedirectFlags::Option_InstanceOnly, OldClassName.ToString(), NewClassName.ToString());
+							NewRedirects.Emplace(ECoreRedirectFlags::Type_Class | ECoreRedirectFlags::Category_InstanceOnly, OldClassName.ToString(), NewClassName.ToString());
 						}
 					}
 					//objects only on a per-object basis
@@ -195,21 +195,21 @@ void FLinkerLoad::CreateActiveRedirectsMap(const FString& GEngineIniName)
 						}
 						else
 						{
-							FCoreRedirect* Redirect = new(NewRedirects) FCoreRedirect(ECoreRedirectFlags::Type_Class, OldClassName.ToString(), NewClassName.ToString());
+							FCoreRedirect& Redirect = NewRedirects.Emplace_GetRef(ECoreRedirectFlags::Type_Class, OldClassName.ToString(), NewClassName.ToString());
 
 							if (!NewClassClass.IsNone() || !NewClassPackage.IsNone())
 							{
-								Redirect->OverrideClassName = FCoreRedirectObjectName(NewClassClass, NAME_None, NewClassPackage);
+								Redirect.OverrideClassName = FCoreRedirectObjectName(NewClassClass, NAME_None, NewClassPackage);
 							}
-							else if (Redirect->NewName.ObjectName.ToString().StartsWith(TEXT("E"), ESearchCase::CaseSensitive))
+							else if (Redirect.NewName.ObjectName.ToString().StartsWith(TEXT("E"), ESearchCase::CaseSensitive))
 							{
 								// This might be an enum, so we have to register it
-								FCoreRedirect* EnumRedirect = new(NewRedirects) FCoreRedirect(ECoreRedirectFlags::Type_Enum, OldClassName.ToString(), NewClassName.ToString());
+								NewRedirects.Emplace(ECoreRedirectFlags::Type_Enum, OldClassName.ToString(), NewClassName.ToString());
 							}
 							else
 							{
 								// This might be a struct redirect because many of them were registered incorrectly
-								FCoreRedirect* StructRedirect = new(NewRedirects) FCoreRedirect(ECoreRedirectFlags::Type_Struct, OldClassName.ToString(), NewClassName.ToString());
+								NewRedirects.Emplace(ECoreRedirectFlags::Type_Struct, OldClassName.ToString(), NewClassName.ToString());
 							}
 						}
 					}
@@ -222,7 +222,7 @@ void FLinkerLoad::CreateActiveRedirectsMap(const FString& GEngineIniName)
 					FParse::Value( *It.Value().GetValue(), TEXT("OldGameName="), OldGameName );
 					FParse::Value( *It.Value().GetValue(), TEXT("NewGameName="), NewGameName );
 
-					FCoreRedirect* Redirect = new(NewRedirects) FCoreRedirect(ECoreRedirectFlags::Type_Package, OldGameName.ToString(), NewGameName.ToString());
+					NewRedirects.Emplace(ECoreRedirectFlags::Type_Package, OldGameName.ToString(), NewGameName.ToString());
 				}
 				else if ( It.Key() == TEXT("ActiveStructRedirects") )
 				{
@@ -232,7 +232,7 @@ void FLinkerLoad::CreateActiveRedirectsMap(const FString& GEngineIniName)
 					FParse::Value( *It.Value().GetValue(), TEXT("OldStructName="), OldStructName );
 					FParse::Value( *It.Value().GetValue(), TEXT("NewStructName="), NewStructName );
 
-					FCoreRedirect* Redirect = new(NewRedirects) FCoreRedirect(ECoreRedirectFlags::Type_Struct, OldStructName.ToString(), NewStructName.ToString());
+					NewRedirects.Emplace(ECoreRedirectFlags::Type_Struct, OldStructName.ToString(), NewStructName.ToString());
 				}
 				else if ( It.Key() == TEXT("ActivePluginRedirects") )
 				{
@@ -245,7 +245,7 @@ void FLinkerLoad::CreateActiveRedirectsMap(const FString& GEngineIniName)
 					OldPluginName = FString(TEXT("/")) + OldPluginName + FString(TEXT("/"));
 					NewPluginName = FString(TEXT("/")) + NewPluginName + FString(TEXT("/"));
 
-					FCoreRedirect* Redirect = new(NewRedirects) FCoreRedirect(ECoreRedirectFlags::Type_Package | ECoreRedirectFlags::Option_MatchSubstring, OldPluginName, NewPluginName);
+					NewRedirects.Emplace(ECoreRedirectFlags::Type_Package | ECoreRedirectFlags::Option_MatchSubstring, OldPluginName, NewPluginName);
 				}
 				else if ( It.Key() == TEXT("KnownMissingPackages") )
 				{
@@ -253,7 +253,7 @@ void FLinkerLoad::CreateActiveRedirectsMap(const FString& GEngineIniName)
 
 					FParse::Value( *It.Value().GetValue(), TEXT("PackageName="), KnownMissingPackage );
 
-					FCoreRedirect* Redirect = new(NewRedirects) FCoreRedirect(ECoreRedirectFlags::Type_Package | ECoreRedirectFlags::Option_Removed, KnownMissingPackage.ToString(), FString());
+					NewRedirects.Emplace(ECoreRedirectFlags::Type_Package | ECoreRedirectFlags::Category_Removed, KnownMissingPackage.ToString(), FString());
 				}
 				else if (It.Key() == TEXT("TaggedPropertyRedirects"))
 				{
@@ -267,7 +267,7 @@ void FLinkerLoad::CreateActiveRedirectsMap(const FString& GEngineIniName)
 
 					check(ClassName != NAME_None && OldPropertyName != NAME_None && NewPropertyName != NAME_None);
 
-					FCoreRedirect* Redirect = new (NewRedirects) FCoreRedirect(ECoreRedirectFlags::Type_Property, FCoreRedirectObjectName(OldPropertyName, ClassName, NAME_None), FCoreRedirectObjectName(NewPropertyName, ClassName, NAME_None));
+					NewRedirects.Emplace(ECoreRedirectFlags::Type_Property, FCoreRedirectObjectName(OldPropertyName, ClassName, NAME_None), FCoreRedirectObjectName(NewPropertyName, ClassName, NAME_None));
 				}
 				else if (It.Key() == TEXT("EnumRedirects"))
 				{
@@ -284,8 +284,8 @@ void FLinkerLoad::CreateActiveRedirectsMap(const FString& GEngineIniName)
 					{
 						FParse::Value(*ConfigValue, TEXT("NewEnumEntry="), NewEnumEntry);
 						check(EnumName != NAME_None && OldEnumEntry != NAME_None && NewEnumEntry != NAME_None);
-						FCoreRedirect* Redirect = new(NewRedirects) FCoreRedirect(ECoreRedirectFlags::Type_Enum, EnumName.ToString(), EnumName.ToString());
-						Redirect->ValueChanges.Add(OldEnumEntry.ToString(), NewEnumEntry.ToString());
+						FCoreRedirect& Redirect = NewRedirects.Emplace_GetRef(ECoreRedirectFlags::Type_Enum, EnumName.ToString(), EnumName.ToString());
+						Redirect.ValueChanges.Add(OldEnumEntry.ToString(), NewEnumEntry.ToString());
 					}
 					else if (FParse::Value(*ConfigValue, TEXT("OldEnumSubstring="), OldEnumSubstring))
 					{
@@ -5236,7 +5236,7 @@ TArray<FName> FLinkerLoad::FindPreviousNamesForClass(FString CurrentClassPath, b
 	if (bIsInstance)
 	{
 		OldObjectNames.Empty();
-		if (FCoreRedirects::FindPreviousNames(ECoreRedirectFlags::Type_Class | ECoreRedirectFlags::Option_InstanceOnly, FCoreRedirectObjectName(CurrentClassPath), OldObjectNames))
+		if (FCoreRedirects::FindPreviousNames(ECoreRedirectFlags::Type_Class | ECoreRedirectFlags::Category_InstanceOnly, FCoreRedirectObjectName(CurrentClassPath), OldObjectNames))
 		{
 			for (FCoreRedirectObjectName& OldObjectName : OldObjectNames)
 			{
@@ -5285,7 +5285,7 @@ FName FLinkerLoad::FindNewNameForClass(FName OldClassName, bool bIsInstance)
 	if (bIsInstance)
 	{
 		// Also check instance types
-		NewName = FCoreRedirects::GetRedirectedName(ECoreRedirectFlags::Type_Class | ECoreRedirectFlags::Option_InstanceOnly, OldName);
+		NewName = FCoreRedirects::GetRedirectedName(ECoreRedirectFlags::Type_Class | ECoreRedirectFlags::Category_InstanceOnly, OldName);
 
 		if (NewName != OldName)
 		{
@@ -5323,6 +5323,15 @@ void FLinkerLoad::OnNewFileAdded(const FString& Filename)
 	}
 }
 
+void FLinkerLoad::OnPakFileMounted(const IPakFile& NewlyLoadedContainer)
+{
+	// To be strictly correct we should check every known missing Package to see whether it exists in the PakFile and remove it only if so.
+	// But the cost of that is be relatively high during loading, and the known missing system is for performance only.  So we instead clear the known missing
+	// on every pak file.
+	FCoreRedirects::ClearKnownMissing(ECoreRedirectFlags::Type_Package);
+}
+
+
 
 void FLinkerLoad::AddGameNameRedirect(const FName OldName, const FName NewName)
 {
@@ -5516,9 +5525,9 @@ FLinkerLoad::ELinkerStatus FLinkerLoad::FixupExportMap()
 			// Never modify the default object instances
 			if (!StrObjectName.StartsWith(DEFAULT_OBJECT_PREFIX))
 			{
-				FCoreRedirectObjectName NewClassInstanceName = FCoreRedirects::GetRedirectedName(ECoreRedirectFlags::Type_Class | ECoreRedirectFlags::Option_InstanceOnly, OldClassName);
+				FCoreRedirectObjectName NewClassInstanceName = FCoreRedirects::GetRedirectedName(ECoreRedirectFlags::Type_Class | ECoreRedirectFlags::Category_InstanceOnly, OldClassName);
 
-				bool bClassInstanceDeleted = FCoreRedirects::IsKnownMissing(ECoreRedirectFlags::Type_Class | ECoreRedirectFlags::Option_InstanceOnly, OldClassName);
+				bool bClassInstanceDeleted = FCoreRedirects::IsKnownMissing(ECoreRedirectFlags::Type_Class | ECoreRedirectFlags::Category_InstanceOnly, OldClassName);
 				if (bClassInstanceDeleted)
 				{
 					UE_LOG(LogLinker, Log, TEXT("FLinkerLoad::FixupExportMap() - Pkg<%s> [Obj<%s> Cls<%s> ClsPkg<%s>] -> removed"), *LinkerRoot->GetName(),
diff --git a/Engine/Source/Runtime/CoreUObject/Private/UObject/Obj.cpp b/Engine/Source/Runtime/CoreUObject/Private/UObject/Obj.cpp
--- a/Engine/Source/Runtime/CoreUObject/Private/UObject/Obj.cpp
+++ b/Engine/Source/Runtime/CoreUObject/Private/UObject/Obj.cpp
@@ -4346,6 +4346,7 @@ void InitUObject()
 	FModuleManager::Get().IsPackageLoadedCallback().BindStatic(Local::IsPackageLoaded);
 	
 	FCoreDelegates::NewFileAddedDelegate.AddStatic(FLinkerLoad::OnNewFileAdded);
+	FCoreDelegates::OnPakFileMounted2.AddStatic(FLinkerLoad::OnPakFileMounted);
 
 #if WITH_EDITOR
 PRAGMA_DISABLE_DEPRECATION_WARNINGS
diff --git a/Engine/Source/Runtime/CoreUObject/Public/UObject/CoreRedirects.h b/Engine/Source/Runtime/CoreUObject/Public/UObject/CoreRedirects.h
--- a/Engine/Source/Runtime/CoreUObject/Public/UObject/CoreRedirects.h
+++ b/Engine/Source/Runtime/CoreUObject/Public/UObject/CoreRedirects.h
@@ -9,6 +9,8 @@
 #include "CoreMinimal.h"
 #include "Misc/EnumClassFlags.h"
 
+class IPakFile;
+
 /** 
  * Flags describing the type and properties of this redirect
  */
@@ -16,7 +18,7 @@ enum class ECoreRedirectFlags : uint32
 {
 	None = 0,
 
-	// Core type of the thing being redirected, multiple can be set
+	// Core type of the Thing being redirected, multiple can be set.  A Query will only find Redirects that have at least one of the same Type bits set.
 	Type_Object =			0x00000001, // UObject
 	Type_Class =			0x00000002, // UClass
 	Type_Struct =			0x00000004, // UStruct
@@ -24,14 +26,22 @@ enum class ECoreRedirectFlags : uint32
 	Type_Function =			0x00000010, // UFunction
 	Type_Property =			0x00000020, // FProperty
 	Type_Package =			0x00000040, // UPackage
-	Type_AllMask =			0x0000FFFF, // Bit mask of all possible types
-
-	// Option flags, specify rules for this redirect
-	Option_InstanceOnly =	0x00010000, // Only redirect instances of this type, not the type itself
-	Option_Removed =		0x00020000, // This type was explicitly removed, new name isn't valid
-	Option_MatchSubstring = 0x00040000, // Does a slow substring match
-	Option_AllMask =		0xFFFF0000, // Bit mask of all possible options
-	Option_ExactMatchMask = Option_InstanceOnly | Option_Removed // These options must match exactly, substring is checked both ways
+	Type_AllMask =			0x0000FFFF, // Bit mask of all possible Types
+
+	// Category flags.  A Query will only match Redirects that have the same value for every category bit.
+	Category_InstanceOnly = 0x00010000, // Only redirect instances of this type, not the type itself
+	Category_Removed =		0x00020000, // This type was explicitly removed, new name isn't valid
+	Category_AllMask =		0x00FF0000, // Bit mask of all possible Categories
+
+	// Option flags.  Does not behave as a bit-match between Queries and Redirects.  Each one specifies a custom rule for how FCoreRedirects handles the Redirect.
+	Option_MatchSubstring = 0x01000000, // Does a slow substring match
+	Option_MissingLoad =	0x02000000, // An automatically-created redirect that was created in response to a missing Thing during load. Redirect will be removed if and when the Thing is loaded.
+	Option_AllMask =		0xFF000000, // Bit mask of all possible Options
+
+	// Deprecated Names
+	Option_InstanceOnly UE_DEPRECATED(4.25, "Use Category_InstanceOnly instead") = Category_InstanceOnly,
+	Option_Removed UE_DEPRECATED(4.25, "Use Category_Removed instead") = Category_Removed,
+	Option_ExactMatchMask UE_DEPRECATED(4.25, "Use Category_AllMask instead") = Category_AllMask,
 };
 ENUM_CLASS_FLAGS(ECoreRedirectFlags);
 
@@ -176,7 +186,7 @@ struct COREUOBJECT_API FCoreRedirect
 	/** Parses a char buffer into the ValueChanges map */
 	const TCHAR* ParseValueChanges(const TCHAR* Buffer);
 
-	/** Returns true if the passed in name matches requirements */
+	/** Returns true if the passed in name and flags match requirements */
 	bool Matches(ECoreRedirectFlags InFlags, const FCoreRedirectObjectName& InName) const;
 
 	/** Returns true if this has value redirects */
@@ -218,11 +228,25 @@ struct COREUOBJECT_API FCoreRedirects
 	/** Returns true if this name has been registered as explicitly missing */
 	static bool IsKnownMissing(ECoreRedirectFlags Type, const FCoreRedirectObjectName& ObjectName);
 
-	/** Adds this as a missing name */
-	static bool AddKnownMissing(ECoreRedirectFlags Type, const FCoreRedirectObjectName& ObjectName);
-
-	/** Removes this as a missing name */
-	static bool RemoveKnownMissing(ECoreRedirectFlags Type, const FCoreRedirectObjectName& ObjectName);
+	/**
+	  * Adds the given combination of (Type, ObjectName, Channel) as a missing name; IsKnownMissing queries will now find it
+	  *
+	  * @param Type Combination of the ECoreRedirectFlags::Type_* flags specifying the type of the object now known to be missing
+	  * @param ObjectName The name of the object now known to be missing
+	  * @param Channel may be Option_MissingLoad or Option_None; used to distinguish between detected-at-runtime and specified-by-ini
+	  */
+	static bool AddKnownMissing(ECoreRedirectFlags Type, const FCoreRedirectObjectName& ObjectName, ECoreRedirectFlags Channel = ECoreRedirectFlags::Option_MissingLoad);
+
+	/**
+	  * Removes the given combination of (Type, ObjectName, Channel) as a missing name
+	  *
+	  * @param Type Combination of the ECoreRedirectFlags::Type_* flags specifying the type of the object that has just been loaded.
+	  * @param ObjectName The name of the object that has just been loaded.
+	  * @param Channel may be Option_MissingLoad or Option_None; used to distinguish between detected-at-runtime and specified-by-ini
+	  */
+	static bool RemoveKnownMissing(ECoreRedirectFlags Type, const FCoreRedirectObjectName& ObjectName, ECoreRedirectFlags Channel = ECoreRedirectFlags::Option_MissingLoad);
+
+	static void ClearKnownMissing(ECoreRedirectFlags Type, ECoreRedirectFlags Channel = ECoreRedirectFlags::Option_MissingLoad);
 
 	/** Returns list of names it may have been before */
 	static bool FindPreviousNames(ECoreRedirectFlags Type, const FCoreRedirectObjectName& NewObjectName, TArray<FCoreRedirectObjectName>& PreviousNames);
diff --git a/Engine/Source/Runtime/CoreUObject/Public/UObject/LinkerLoad.h b/Engine/Source/Runtime/CoreUObject/Public/UObject/LinkerLoad.h
--- a/Engine/Source/Runtime/CoreUObject/Public/UObject/LinkerLoad.h
+++ b/Engine/Source/Runtime/CoreUObject/Public/UObject/LinkerLoad.h
@@ -10,6 +10,7 @@
 #include "UObject/Linker.h"
 
 class FLinkerPlaceholderBase;
+class IPakFile;
 class ULinkerPlaceholderExportObject;
 struct FScopedSlowTask;
 struct FUntypedBulkData;
@@ -252,6 +253,8 @@ public:
 	 */
 	COREUOBJECT_API static void OnNewFileAdded(const FString& Filename);
 
+	COREUOBJECT_API static void OnPakFileMounted(const IPakFile& PakFile);
+
 	/** 
 	 * Checks if the linker has any objects in the export table that require loading.
 	 */
diff --git a/Engine/Source/Runtime/Engine/Private/Streaming/StreamingManagerTexture.cpp b/Engine/Source/Runtime/Engine/Private/Streaming/StreamingManagerTexture.cpp
--- a/Engine/Source/Runtime/Engine/Private/Streaming/StreamingManagerTexture.cpp
+++ b/Engine/Source/Runtime/Engine/Private/Streaming/StreamingManagerTexture.cpp
@@ -134,7 +134,7 @@ FRenderAssetStreamingManager::FRenderAssetStreamingManager()
 
 	FCoreUObjectDelegates::GetPreGarbageCollectDelegate().AddRaw(this, &FRenderAssetStreamingManager::OnPreGarbageCollect);
 
-	FCoreDelegates::OnPakFileMounted.AddRaw(this, &FRenderAssetStreamingManager::OnPakFileMounted);
+	FCoreDelegates::OnPakFileMounted2.AddRaw(this, &FRenderAssetStreamingManager::OnPakFileMounted);
 }
 
 FRenderAssetStreamingManager::~FRenderAssetStreamingManager()
@@ -189,7 +189,7 @@ void FRenderAssetStreamingManager::OnPreGarbageCollect()
 
 
 
-void FRenderAssetStreamingManager::OnPakFileMounted(const TCHAR* PakFilename, const int32 ChunkId)
+void FRenderAssetStreamingManager::OnPakFileMounted(const IPakFile& PakFile)
 {
 	// clear the cached file exists checks which failed as they may now be loaded
 	bNewFilesLoaded = true;
diff --git a/Engine/Source/Runtime/Engine/Private/Streaming/StreamingManagerTexture.h b/Engine/Source/Runtime/Engine/Private/Streaming/StreamingManagerTexture.h
--- a/Engine/Source/Runtime/Engine/Private/Streaming/StreamingManagerTexture.h
+++ b/Engine/Source/Runtime/Engine/Private/Streaming/StreamingManagerTexture.h
@@ -15,6 +15,7 @@ TextureStreamingManager.h: Definitions of classes used for texture streaming.
 
 class AActor;
 class FRenderAssetStreamingMipCalcTask;
+class IPakFile;
 class UPrimitiveComponent;
 
 template<typename TTask> class FAsyncTask;
@@ -408,7 +409,7 @@ private:
 	TArray<int32> InflightRenderAssets;
 
 	TMap<FString, bool> CachedFileExistsChecks;
-	void OnPakFileMounted(const TCHAR* PakFilename, const int32 ChunkId);
+	void OnPakFileMounted(const IPakFile& PakFile);
 
 	// A critical section use around code that could be called in parallel with NotifyPrimitiveUpdated() or NotifyPrimitiveUpdated_Concurrent().
 	FCriticalSection CriticalSection;
