Mirror 11452607
Fix crash when compiling blueprints when other blueprints reference the properties of that blueprint
#rnx

diff --git a/Engine/Source/Runtime/CoreUObject/Private/UObject/Class.cpp b/Engine/Source/Runtime/CoreUObject/Private/UObject/Class.cpp
--- a/Engine/Source/Runtime/CoreUObject/Private/UObject/Class.cpp
+++ b/Engine/Source/Runtime/CoreUObject/Private/UObject/Class.cpp
@@ -4637,29 +4637,29 @@ void UClass::PurgeClass(bool bRecompilingOnLoad)
 		}
 		PropertyWrappers.Empty();
 	}
-	if (bRecompilingOnLoad)
+
+	// When compiling properties can't be immediately destroyed because we need 
+	// to fix up references to these properties. The caller of PurgeClass is 
+	// expected to call DestroyPropertiesPendingDestruction
+	FField* LastField = ChildProperties;
+	if (LastField)
 	{
-		// When compiling on load properties can't be immediately destroyed because new properties can be allocated
-		// in their place and this breaks script byte code reference replacement. See FBlueprintCompileReinstancer::AddReferencedObjects.
-		FField* LastField = ChildProperties;
-		if (LastField)
+		while (LastField->Next)
 		{
-			while (LastField->Next)
-			{
-				LastField = LastField->Next;
-			}
-			check(LastField->Next == nullptr);
-			LastField->Next = PropertiesPendingDestruction;
-			PropertiesPendingDestruction = ChildProperties;
-			ChildProperties = nullptr;
+			LastField = LastField->Next;
 		}
+		check(LastField->Next == nullptr);
+		LastField->Next = PropertiesPendingDestruction;
+		PropertiesPendingDestruction = ChildProperties;
+		ChildProperties = nullptr;
 	}
-	else
-#endif // WITH_EDITORONLY_DATA
+#else
 	{
 		// Destroy all properties owned by this struct
 		DestroyPropertyLinkedList(ChildProperties);
 	}
+#endif // WITH_EDITORONLY_DATA
+
 	FFieldPath::OnFieldDeleted();
 
 	DestroyUnversionedSchema(this);
diff --git a/Engine/Source/Runtime/CoreUObject/Public/UObject/Class.h b/Engine/Source/Runtime/CoreUObject/Public/UObject/Class.h
--- a/Engine/Source/Runtime/CoreUObject/Public/UObject/Class.h
+++ b/Engine/Source/Runtime/CoreUObject/Public/UObject/Class.h
@@ -3052,6 +3052,8 @@ public:
 	/** 
 	 * Purges out the properties of this class in preparation for it to be regenerated
 	 * @param bRecompilingOnLoad - true if we are recompiling on load
+	 *
+	 * In editor, properties are not freed until DestroyPropertiesPendingDestruction is called.
 	 */
 	virtual void PurgeClass(bool bRecompilingOnLoad);
 
