Add missing property wrapper type for blueprint event dispatcher's details customizations

diff --git a/Engine/Source/Runtime/CoreUObject/Private/UObject/PropertyPrivate.cpp b/Engine/Source/Runtime/CoreUObject/Private/UObject/PropertyPrivate.cpp
--- a/Engine/Source/Runtime/CoreUObject/Private/UObject/PropertyPrivate.cpp
+++ b/Engine/Source/Runtime/CoreUObject/Private/UObject/PropertyPrivate.cpp
@@ -919,4 +919,9 @@ IMPLEMENT_CORE_INTRINSIC_CLASS(UMulticastDelegatePropertyWrapper, UPropertyWrapp
 	}
 );
 
+IMPLEMENT_CORE_INTRINSIC_CLASS(UMulticastInlineDelegatePropertyWrapper, UMulticastDelegatePropertyWrapper,
+	{
+	}
+);
+
 #include "UObject/DefineUPropertyMacros.h"
\ No newline at end of file
diff --git a/Engine/Source/Runtime/CoreUObject/Public/UObject/UnrealType.h b/Engine/Source/Runtime/CoreUObject/Public/UObject/UnrealType.h
--- a/Engine/Source/Runtime/CoreUObject/Public/UObject/UnrealType.h
+++ b/Engine/Source/Runtime/CoreUObject/Public/UObject/UnrealType.h
@@ -5879,5 +5879,9 @@ class COREUOBJECT_API UMulticastDelegatePropertyWrapper : public UPropertyWrappe
 	DECLARE_CLASS_INTRINSIC(UMulticastDelegatePropertyWrapper, UPropertyWrapper, CLASS_Transient, TEXT("/Script/CoreUObject"));
 };
 
+class COREUOBJECT_API UMulticastInlineDelegatePropertyWrapper : public UMulticastDelegatePropertyWrapper
+{
+	DECLARE_CLASS_INTRINSIC(UMulticastInlineDelegatePropertyWrapper, UMulticastDelegatePropertyWrapper, CLASS_Transient, TEXT("/Script/CoreUObject"));
+};
 
 #include "UObject/DefineUPropertyMacros.h"
\ No newline at end of file
