Weak Interface Pointer: Changed the broken constructor TWeakInterfacePtr(const UObject* Object) to TWeakInterfacePtr(UObject* Object) (the const was removed)
The previous version was broken because it would throw this waring when it was used:
warning C4996: 'TWeakObjectPtr<UObject,FWeakObjectPtr>::operator =': Implicit conversions from const pointers to non-const TWeakObjectPtrs has been deprecated. Please const-correct this usage. Please update your code to the new API before upgrading to the next release, otherwise your project will no longer compile.

diff --git a/Engine/Source/Runtime/CoreUObject/Public/UObject/WeakInterfacePtr.h b/Engine/Source/Runtime/CoreUObject/Public/UObject/WeakInterfacePtr.h
--- a/Engine/Source/Runtime/CoreUObject/Public/UObject/WeakInterfacePtr.h
+++ b/Engine/Source/Runtime/CoreUObject/Public/UObject/WeakInterfacePtr.h
@@ -13,7 +13,7 @@ struct TWeakInterfacePtr
 {
 	TWeakInterfacePtr() : InterfaceInstance(nullptr) {}
 
-	TWeakInterfacePtr(const UObject* Object)
+	TWeakInterfacePtr(UObject* Object)
 	{
 		InterfaceInstance = Cast<TInterface>(Object);
 		if (InterfaceInstance != nullptr)
