Fix regression in UEnum Net Serialization

diff --git a/Engine/Source/Runtime/CoreUObject/Private/UObject/EnumProperty.cpp b/Engine/Source/Runtime/CoreUObject/Private/UObject/EnumProperty.cpp
--- a/Engine/Source/Runtime/CoreUObject/Private/UObject/EnumProperty.cpp
+++ b/Engine/Source/Runtime/CoreUObject/Private/UObject/EnumProperty.cpp
@@ -187,7 +187,9 @@ bool FEnumProperty::NetSerializeItem(FArchive& Ar, UPackageMap* Map, void* Data,
 	}
 	else
 	{
-		Ar.SerializeBits(Data, FMath::CeilLogTwo64(Enum->GetMaxEnumValue() + 1));
+        const uint64 MaxBits = ElementSize * 8;
+        const uint64 DesiredBits = FMath::CeilLogTwo64(Enum->GetMaxEnumValue() + 1);
+		Ar.SerializeBits(Data, FMath::Min(DesiredBits, MaxBits));
 	}
 	return 1;
 }
