Make StaticFindObject handle nested subobjects when searching for any package without a package pointer

This has the exact same semantics as the original 4.22 implementation based on EndsWith(GetPathName()), except this implementation ignores the exact delimiter being used, i.e. dots  and colons can be swapped. Nobody should be searching for object paths with wrong delimiters though.

diff --git a/Engine/Source/Runtime/CoreUObject/Private/UObject/UObjectHash.cpp b/Engine/Source/Runtime/CoreUObject/Private/UObject/UObjectHash.cpp
--- a/Engine/Source/Runtime/CoreUObject/Private/UObject/UObjectHash.cpp
+++ b/Engine/Source/Runtime/CoreUObject/Private/UObject/UObjectHash.cpp
@@ -489,53 +489,46 @@ UObject* StaticFindObjectFastExplicit( const UClass* ObjectClass, FName ObjectNa
 	return Result;
 }
 
+// Splits an object path into FNames representing an outer chain.
+//
+// Input path examples: "Object", "Package.Object", "Object:Subobject", "Object:Subobject.Nested", "Package.Object:Subobject", "Package.Object:Subobject.NestedSubobject"
 struct FObjectSearchPath
 {
 	FName Inner;
-	FName Outer;
-	FName OuterOuter;
+	TArray<FName, TInlineAllocator<8>> Outers;
 
-	// @param FullPath is one of four types: "Object", "Package.Object", "Object:Subobject" or "Package.Object:Subobject"
-	explicit FObjectSearchPath(FName FullPath)
+	explicit FObjectSearchPath(FName InPath)
 	{
 		TCHAR Buffer[NAME_SIZE];
-		FullPath.GetPlainNameString(Buffer);
+		InPath.GetPlainNameString(Buffer);
 
-		// Find last dot and last dot or colon
-		const TCHAR* LastDot = nullptr;
-		const TCHAR* LastDelimiter = nullptr;
 		constexpr FAsciiSet DotColon(".:");
-		for (const TCHAR* It = FAsciiSet::FindFirstOrEnd(Buffer, DotColon); *It != '\0'; It = FAsciiSet::FindFirstOrEnd(It + 1, DotColon)) 
+		const TCHAR* Begin = Buffer;
+		const TCHAR* End = FAsciiSet::FindFirstOrEnd(Begin, DotColon);
+		while (*End != '\0')
 		{
-			LastDot = *It == '.' ? It : LastDot;
-			LastDelimiter = It;
-		}
-			
-		if (!LastDelimiter) // "Object"
-		{
-			Inner = FullPath;
-		}
-		else if (!!LastDot && LastDelimiter > LastDot) // "Package.Object:Subobject"
-		{
-			OuterOuter = FName(LastDot - Buffer, Buffer);
-			SetOuterAndInner(LastDot + 1, LastDelimiter, FullPath.GetNumber());
-		}
-		else // "Package.Object" or "Object:Subobject"
-		{
-			SetOuterAndInner(Buffer, LastDelimiter, FullPath.GetNumber());
+			Outers.Add(FName(End - Begin, Begin));
+
+			Begin = End + 1;
+			End = FAsciiSet::FindFirstOrEnd(Begin, DotColon);
 		}
-	}
-	
-	void SetOuterAndInner(const TCHAR* OuterAndInner, const TCHAR* Delimiter, int32 InnerNumber)
-	{
-		int32 OuterLen = static_cast<int32>(Delimiter - OuterAndInner);
-		Outer = FName(OuterLen, OuterAndInner);
-		Inner = FName(Delimiter + 1, InnerNumber);
+
+		Inner = Outers.Num() == 0 ? InPath : FName(FStringView(Begin, End - Begin), InPath.GetNumber());
 	}
 
-	bool MatchOuterNames(UObject* OuterObject) const
+	bool MatchOuterNames(UObject* Outer) const
 	{
-		return Outer.IsNone() || (OuterObject->GetFName() == Outer && (OuterOuter.IsNone() || OuterObject->GetOuter()->GetFName() == OuterOuter));
+		for (int32 Idx = Outers.Num() - 1; Idx >= 0; --Idx)
+		{
+			if (!Outer || Outers[Idx] != Outer->GetFName())
+			{
+				return false;
+			}
+
+			Outer = Outer->GetOuter();
+		}
+
+		return true;
 	}
 };
 
