Fixed incorrect behavior of FField::GetFullGroupName

This was inconsistent with the old behavior of the UField version.

#rnx

diff --git a/Engine/Source/Runtime/CoreUObject/Private/UObject/Field.cpp b/Engine/Source/Runtime/CoreUObject/Private/UObject/Field.cpp
--- a/Engine/Source/Runtime/CoreUObject/Private/UObject/Field.cpp
+++ b/Engine/Source/Runtime/CoreUObject/Private/UObject/Field.cpp
@@ -644,17 +644,17 @@ FString FField::GetFullGroupName(bool bStartWithOuter) const
 			}
 			else
 			{
-				return Owner.ToField()->GetPathName();
+				return Owner.ToField()->GetPathName(GetOutermost());
 			}
 		}
 		else
 		{
-			return TEXT("");
+			return FString();
 		}
 	}
 	else
 	{
-		return GetPathName();
+		return GetPathName(GetOutermost());
 	}
 }
 
