FReferencerFinder will now properly find GGCObjectReferencer as one of the referencing objects

diff --git a/Engine/Source/Runtime/CoreUObject/Private/UObject/ReferencerFinder.cpp b/Engine/Source/Runtime/CoreUObject/Private/UObject/ReferencerFinder.cpp
--- a/Engine/Source/Runtime/CoreUObject/Private/UObject/ReferencerFinder.cpp
+++ b/Engine/Source/Runtime/CoreUObject/Private/UObject/ReferencerFinder.cpp
@@ -11,23 +11,33 @@ class FAllReferencesProcessor : public FSimpleReferenceProcessorBase
 {
 	const TSet<UObject*>& PotentiallyReferencedObjects;
 	TSet<UObject*>& ReferencingObjects;
+	UObject* CurrentObject;
 
 public:
 	FAllReferencesProcessor(const TSet<UObject*>& InPotentiallyReferencedObjects, TSet<UObject*>& OutReferencingObjects)
 		: PotentiallyReferencedObjects(InPotentiallyReferencedObjects)
 		, ReferencingObjects(OutReferencingObjects)
+		, CurrentObject(nullptr)
 	{
 	}
 	FORCEINLINE_DEBUGGABLE void HandleTokenStreamObjectReference(TArray<UObject*>& ObjectsToSerialize, UObject* ReferencingObject, UObject*& Object, const int32 TokenIndex, bool bAllowReferenceElimination)
 	{
-		if (Object != nullptr && Object != ReferencingObject)
-		{			
-			if (ReferencingObject && PotentiallyReferencedObjects.Contains(Object))
+		if (!ReferencingObject)
+		{
+			ReferencingObject = CurrentObject;
+		}
+		if (Object && ReferencingObject && Object != ReferencingObject)
+		{
+			if (PotentiallyReferencedObjects.Contains(Object))
 			{
 				ReferencingObjects.Add(ReferencingObject);
 			}
 		}
 	}
+	void SetCurrentObject(UObject* Obj)
+	{
+		CurrentObject = Obj;
+	}
 };
 typedef TDefaultReferenceCollector<FAllReferencesProcessor> FAllReferencesCollector;
 
