Fixes a bug where FieldPathProperties did not clear the old field path info before adding the new info. This resulted in bad data in the property and could result in lost work.

diff --git a/Engine/Source/Runtime/CoreUObject/Private/UObject/FieldPath.cpp b/Engine/Source/Runtime/CoreUObject/Private/UObject/FieldPath.cpp
--- a/Engine/Source/Runtime/CoreUObject/Private/UObject/FieldPath.cpp
+++ b/Engine/Source/Runtime/CoreUObject/Private/UObject/FieldPath.cpp
@@ -55,6 +55,8 @@ void FFieldPath::Generate(const TCHAR* InFieldPathString)
 		TCHAR NameBuffer[NAME_SIZE];
 		int32 NameIndex = 0;
 
+		Path.Reset();
+
 		// Construct names
 		while (true)
 		{
