Back out changelist 11604440. This change is redundant due to CL 11317919.

diff --git a/Engine/Source/Runtime/CoreUObject/Private/UObject/FieldPath.cpp b/Engine/Source/Runtime/CoreUObject/Private/UObject/FieldPath.cpp
--- a/Engine/Source/Runtime/CoreUObject/Private/UObject/FieldPath.cpp
+++ b/Engine/Source/Runtime/CoreUObject/Private/UObject/FieldPath.cpp
@@ -55,8 +55,6 @@ void FFieldPath::Generate(const TCHAR* InFieldPathString)
 		TCHAR NameBuffer[NAME_SIZE];
 		int32 NameIndex = 0;
 
-		Path.Reset();
-
 		// Construct names
 		while (true)
 		{
