FieldPath property's ImportText will now properly support full path names (including the field type) as well as delimited strings

diff --git a/Engine/Source/Runtime/CoreUObject/Private/UObject/FieldPathProperty.cpp b/Engine/Source/Runtime/CoreUObject/Private/UObject/FieldPathProperty.cpp
--- a/Engine/Source/Runtime/CoreUObject/Private/UObject/FieldPathProperty.cpp
+++ b/Engine/Source/Runtime/CoreUObject/Private/UObject/FieldPathProperty.cpp
@@ -5,6 +5,8 @@
 #include "UObject/Package.h"
 #include "UObject/UnrealTypePrivate.h"
 #include "UObject/LinkerLoad.h"
+#include "UObject/ConstructorHelpers.h"
+#include "Misc/Parse.h"
 
 // WARNING: This should always be the last include in any file that needs it (except .generated.h)
 #include "UObject/UndefineUPropertyMacros.h"
@@ -102,25 +104,46 @@ void FFieldPathProperty::ExportTextItem( FString& ValueStr, const void* Property
 
 const TCHAR* FFieldPathProperty::ImportText_Internal( const TCHAR* Buffer, void* Data, int32 PortFlags, UObject* Parent, FOutputDevice* ErrorText ) const
 {
+	check(Buffer);
 	FFieldPath* PathPtr = GetPropertyValuePtr(Data);
 	check(PathPtr);
+	FString PathName;
 
 	if (!(PortFlags & PPF_Delimited))
 	{
-		PathPtr->Generate(Buffer);
+		PathName = Buffer;
 		// in order to indicate that the value was successfully imported, advance the buffer past the last character that was imported
-		Buffer += FCString::Strlen(Buffer);
+		Buffer += PathName.Len();
 	}
 	else
+	{		
+		// Advance to the next delimiter (comma) or the end of the buffer
+		int32 SeparatorIndex = 0;
+		while (Buffer[SeparatorIndex] != '\0' && Buffer[SeparatorIndex] != ',')
 	{
-		FString Temp;
-		Buffer = FPropertyHelpers::ReadToken(Buffer, Temp, true);
-		if (!Buffer)
+			++SeparatorIndex;
+		}
+		// Copy the value string
+		PathName = FString(SeparatorIndex, Buffer);
+		// Advance the buffer to let the calling function know we succeeded
+		Buffer += SeparatorIndex + 1;
+	}
+
+	if (PathName.Len())
+	{
+		// Strip the class name if present, we don't need it
+		ConstructorHelpers::StripObjectClass(PathName);
+		if (PathName[0] == '\"')
 		{
+			FString UnquotedPathName;
+			if (!FParse::QuotedString(*PathName, UnquotedPathName))
+			{
+				UE_LOG(LogProperty, Warning, TEXT("FieldPathProperty: Bad quoted string: %s"), *PathName);
 			return nullptr;
 		}
-
-		PathPtr->Generate(*Temp);
+			PathName = MoveTemp(UnquotedPathName);
+		}
+		PathPtr->Generate(*PathName);
 	}
 
 	return Buffer;
@@ -134,7 +157,6 @@ void FFieldPathProperty::Serialize(FArchive& Ar)
 
 FString FFieldPathProperty::GetCPPMacroType(FString& ExtendedTypeText) const
 {
-	//return FProperty::GetCPPMacroType(ExtendedTypeText);
 	check(PropertyClass);
 	ExtendedTypeText = FString::Printf(TEXT("TFieldPath<F%s>"), *PropertyClass->GetName());
 	return TEXT("STRUCT");
