Timed Input Monitor - Timecode Provider is not updated when changed from the Project Settings (needs restart)
Skip filterering of affected/changed archetype instances if no set is provided.
#rnx

diff --git a/Engine/Source/Runtime/CoreUObject/Private/UObject/EditPropertyChain.cpp b/Engine/Source/Runtime/CoreUObject/Private/UObject/EditPropertyChain.cpp
--- a/Engine/Source/Runtime/CoreUObject/Private/UObject/EditPropertyChain.cpp
+++ b/Engine/Source/Runtime/CoreUObject/Private/UObject/EditPropertyChain.cpp
@@ -55,7 +55,7 @@ bool FEditPropertyChain::SetActiveMemberPropertyNode( FProperty* NewActiveMember
 
 bool FEditPropertyChain::IsArchetypeInstanceAffected( UObject* InInstance ) const
 {
-	return AffectedInstances.Contains(InInstance);
+	return !bFilterAffectedInstances || AffectedInstances.Contains(InInstance);
 }
 
 /**
diff --git a/Engine/Source/Runtime/CoreUObject/Public/UObject/UnrealType.h b/Engine/Source/Runtime/CoreUObject/Public/UObject/UnrealType.h
--- a/Engine/Source/Runtime/CoreUObject/Public/UObject/UnrealType.h
+++ b/Engine/Source/Runtime/CoreUObject/Public/UObject/UnrealType.h
@@ -5073,7 +5073,7 @@ class COREUOBJECT_API FEditPropertyChain : public TDoubleLinkedList<FProperty*>
 
 public:
 	/** Constructors */
-	FEditPropertyChain() : ActivePropertyNode(NULL), ActiveMemberPropertyNode(NULL) {}
+	FEditPropertyChain() : ActivePropertyNode(NULL), ActiveMemberPropertyNode(NULL), bFilterAffectedInstances(false) {}
 
 	/**
 	 * Sets the ActivePropertyNode to the node associated with the property specified.
@@ -5102,6 +5102,7 @@ public:
 	template<typename T>
 	void SetAffectedArchetypeInstances( T&& InAffectedInstances )
 	{
+		bFilterAffectedInstances = true;
 		AffectedInstances = Forward<T>(InAffectedInstances);
 	}
 	
@@ -5142,6 +5143,11 @@ protected:
 	 */
 	TSet<UObject*> AffectedInstances;
 
+	/**
+	 * Assume all archetype instances are affected unless a set of affected instances is provided.
+	 */
+	bool bFilterAffectedInstances;
+
 	/** TDoubleLinkedList interface */
 	/**
 	 * Updates the size reported by Num().  Child classes can use this function to conveniently
@@ -5190,6 +5196,7 @@ struct FPropertyChangedEvent
 		, MemberProperty(InProperty)
 		, ChangeType(InChangeType)
 		, ObjectIteratorIndex(INDEX_NONE)
+		, bFilterChangedInstances(false)
 		, TopLevelObjects(InTopLevelObjects)
 	{
 	}
@@ -5200,6 +5207,7 @@ struct FPropertyChangedEvent
 		, MemberProperty(InProperty)
 		, ChangeType(InChangeType)
 		, ObjectIteratorIndex(INDEX_NONE)
+		, bFilterChangedInstances(false)
 	{
 		if (InTopLevelObjects)
 		{
@@ -5226,6 +5234,7 @@ struct FPropertyChangedEvent
 	template<typename T>
 	void SetInstancesChanged(T&& InInstancesChanged)
 	{
+		bFilterChangedInstances = true;
 		InstancesChanged = Forward<T>(InInstancesChanged);
 	}
 
@@ -5254,7 +5263,7 @@ struct FPropertyChangedEvent
 	 */
 	bool HasArchetypeInstanceChanged(UObject* InInstance) const
 	{
-		return InstancesChanged.Contains(InInstance);
+		return !bFilterChangedInstances || InstancesChanged.Contains(InInstance);
 	}
 
 	/**
@@ -5301,6 +5310,9 @@ private:
 	//In the property window, multiple objects can be selected at once. In this case we want to know if an instance was updated for this operation (used in array/set/map context)
 	TSet<UObject*> InstancesChanged;
 
+	//Assume all archetype instances were changed unless a set of changed instances is provided.
+	bool bFilterChangedInstances;
+
 	/** List of top level objects being changed */
 	TArrayView<const UObject* const> TopLevelObjects;
 };
