Update ObjectMacros to include missing meta data specifiers to improve VAX/Visual studio search

#rnx

diff --git a/Engine/Source/Runtime/CoreUObject/Public/UObject/ObjectMacros.h b/Engine/Source/Runtime/CoreUObject/Public/UObject/ObjectMacros.h
--- a/Engine/Source/Runtime/CoreUObject/Public/UObject/ObjectMacros.h
+++ b/Engine/Source/Runtime/CoreUObject/Public/UObject/ObjectMacros.h
@@ -853,7 +853,17 @@ namespace UF
 		ServiceRequest,
 
 		/// This function is RPC service response
-		ServiceResponse
+		ServiceResponse,
+		
+		/// [FunctionMetadata]	Marks a UFUNCTION as accepting variadic arguments. Variadic functions may have extra terms they need to emit after the main set of function arguments
+		///						These are all considered wildcards so no type checking will be performed on them
+		Variadic,
+
+		/// [FunctionMetadata] Indicates the display name of the return value pin
+		ReturnDisplayName, 
+
+		/// [FunctionMetadata] Indicates that a particular function parameter is for internal use only, which means it will be both hidden and not connectible.
+		InternalUseParam, 
 	};
 }
 
@@ -978,6 +988,9 @@ namespace UP
 
 		/// Property shouldn't be serialized, can still be exported to text
 		SkipSerialization,
+
+		/// If true, the self pin should not be shown or connectable regardless of purity, const, etc. similar to InternalUseParam
+		HideSelfPin, 
 	};
 }
 
@@ -1263,6 +1276,15 @@ namespace UM
 		/// [PropertyMetadata] Used for SoftObjectPtr/SoftObjectPath properties to specify a reference should not be tracked. This reference will not be automatically cooked or saved into the asset registry for redirector/delete fixup.
 		Untracked,
 
+		/// [PropertyMetadata] For functions that should be compiled in development mode only.
+		DevelopmentOnly, 
+
+		/// [PropertyMetadata] (Internal use only) Used for the latent action manager to fix up a latent action with the VM
+		NeedsLatentFixup,
+
+		/// [PropertyMetadata] (Internal use only) Used for the latent action manager to track where it's re-entry should be
+		LatentCallbackTarget,
+
 		/// [PropertyMetadata] Causes FString and FName properties to have a limited set of options generated dynamically, e.g. meta=(GetOptions="FuncName")
 		///
 		/// UFUNCTION()
@@ -1427,6 +1449,45 @@ namespace UM
 
 		// [FunctionMetadata] Only valid in Blueprint Function Libraries. Mark this function as an exception to the class's general BlueprintThreadSafe metadata.
 		NotBlueprintThreadSafe,
+
+		/// [FunctionMetadata] [InterfaceMetadata] Metadata that flags function params that govern what type of object the function returns
+		DeterminesOutputType,
+
+		/// [FunctionMetadata] [InterfaceMetadata] Metadata that flags the function output param that will be controlled by the "MD_DynamicOutputType" pin
+		DynamicOutputParam,
+
+		/// [FunctionMetadata][InterfaceMetadata] Metadata to identify an DataTable Pin. Depending on which DataTable is selected, we display different RowName options
+		DataTablePin,
+
+		/// [FunctionMetadata][InterfaceMetadata] Metadata that flags TSet parameters that will have their type determined at blueprint compile time
+		SetParam,
+
+		/// [FunctionMetadata] [InterfaceMetadata] Metadata that flags TMap function parameters that will have their type determined at blueprint compile time
+		MapParam,
+
+		/// [FunctionMetadata] [InterfaceMetadata]  Metadata that flags TMap function parameters that will have their key type determined at blueprint compile time
+		MapKeyParam,
+
+		/// [FunctionMetadata][InterfaceMetadata] Metadata that flags TMap function parameter that will have their value type determined at blueprint compile time
+		MapValueParam,
+
+		/// [FunctionMetadata] [InterfaceMetadata] Metadata that identifies an integral property as a bitmask.
+		Bitmask,
+
+		/// [FunctionMetadata] [InterfaceMetadata] Metadata that associates a bitmask property with a bitflag enum.
+		BitmaskEnum,
+
+		/// [InterfaceMetadata] Metadata that identifies an enum as a set of explicitly-named bitflags.
+		Bitflags,
+
+		/// [InterfaceMetadata] Metadata that signals to the editor that enum values correspond to mask values instead of bitshift (index) values.
+		UseEnumValuesAsMaskValuesInEditor,
+
+		/// [InterfaceMetadata] Stub function used internally by animation blueprints
+		AnimBlueprintFunction,
+
+		/// [FunctionMetadata] [InterfaceMetadata] Metadata that flags TArray function parameters that will have their type determined at blueprint compile time
+		ArrayParam,
 	};
 
 	// Metadata usable in UINTERFACE
@@ -1738,6 +1799,18 @@ typedef uint32 ERenameFlags;
 
 typedef void (*FAsyncCompletionCallback)( UObject* LinkerRoot, void* CallbackUserData );
 
+namespace GameplayTagsManager
+{
+	enum
+	{
+		/// Used for filtering by tag widget
+		Categories,
+		
+		/// Used for filtering by tag widget for any parameters of the function that end up as BP pins
+		GameplayTagFilter,
+	};
+}
+
 /*-----------------------------------------------------------------------------
 	UObject.
 -----------------------------------------------------------------------------*/
