Set FlagsPrivate after loading our associated UProperty

#localization none
#tests cooked, pie

diff --git a/Engine/Source/Runtime/CoreUObject/Private/UObject/Field.cpp b/Engine/Source/Runtime/CoreUObject/Private/UObject/Field.cpp
--- a/Engine/Source/Runtime/CoreUObject/Private/UObject/Field.cpp
+++ b/Engine/Source/Runtime/CoreUObject/Private/UObject/Field.cpp
@@ -335,7 +335,7 @@ FField::~FField()
 FField::FField(UField* InField)
 	: Next(nullptr)
 	, NamePrivate(InField->GetFName())
-	, FlagsPrivate(InField->GetFlags())
+	, FlagsPrivate(RF_NoFlags)
 	, MetaDataMap(nullptr)
 {
 	check(InField);
@@ -344,6 +344,7 @@ FField::FField(UField* InField)
 		// The source UField needs to be loaded, otherwise we'll be copying default property values
 		InField->GetLinker()->Preload(InField);
 	}
+	FlagsPrivate = InField->GetFlags();
 	// Associate this FField with the UField we're constructing from so that next time something tries to convert it,
 	// it can already grab the cached new FField
 	InField->SetAssociatedFField(this);
