FieldPath property's ImportText will now properly support full path names (including the field type) as well as delimited strings

diff --git a/Engine/Source/Runtime/CoreUObject/Private/UObject/FieldPathProperty.cpp b/Engine/Source/Runtime/CoreUObject/Private/UObject/FieldPathProperty.cpp
--- a/Engine/Source/Runtime/CoreUObject/Private/UObject/FieldPathProperty.cpp
+++ b/Engine/Source/Runtime/CoreUObject/Private/UObject/FieldPathProperty.cpp
@@ -121,6 +121,7 @@ const TCHAR* FFieldPathProperty::ImportText_Internal( const TCHAR* Buffer, void*
 		// Advance to the next delimiter (comma) or the end of the buffer
 		int32 SeparatorIndex = 0;
 		while (Buffer[SeparatorIndex] != '\0' && Buffer[SeparatorIndex] != ',')
+		{
 	{
 			++SeparatorIndex;
 		}
@@ -140,7 +141,8 @@ const TCHAR* FFieldPathProperty::ImportText_Internal( const TCHAR* Buffer, void*
 			if (!FParse::QuotedString(*PathName, UnquotedPathName))
 			{
 				UE_LOG(LogProperty, Warning, TEXT("FieldPathProperty: Bad quoted string: %s"), *PathName);
-			return nullptr;
+				return nullptr;
+			}
 		}
 			PathName = MoveTemp(UnquotedPathName);
 		}
