Remove incorrect brackets introduced in merge
#rnx

diff --git a/Engine/Source/Runtime/CoreUObject/Private/UObject/FieldPathProperty.cpp b/Engine/Source/Runtime/CoreUObject/Private/UObject/FieldPathProperty.cpp
--- a/Engine/Source/Runtime/CoreUObject/Private/UObject/FieldPathProperty.cpp
+++ b/Engine/Source/Runtime/CoreUObject/Private/UObject/FieldPathProperty.cpp
@@ -122,7 +122,6 @@ const TCHAR* FFieldPathProperty::ImportText_Internal( const TCHAR* Buffer, void*
 		int32 SeparatorIndex = 0;
 		while (Buffer[SeparatorIndex] != '\0' && Buffer[SeparatorIndex] != ',')
 		{
-	{
 			++SeparatorIndex;
 		}
 		// Copy the value string
@@ -143,7 +142,6 @@ const TCHAR* FFieldPathProperty::ImportText_Internal( const TCHAR* Buffer, void*
 				UE_LOG(LogProperty, Warning, TEXT("FieldPathProperty: Bad quoted string: %s"), *PathName);
 				return nullptr;
 			}
-		}
 			PathName = MoveTemp(UnquotedPathName);
 		}
 		PathPtr->Generate(*PathName);
