In cooked builds using async archives, if we fail to load the package file summary then advance the async archive's load state out of "waiting for header" so that the failure can be processed gracefully and not trigger a check()

diff --git a/Engine/Source/Runtime/CoreUObject/Private/UObject/LinkerLoad.cpp b/Engine/Source/Runtime/CoreUObject/Private/UObject/LinkerLoad.cpp
--- a/Engine/Source/Runtime/CoreUObject/Private/UObject/LinkerLoad.cpp
+++ b/Engine/Source/Runtime/CoreUObject/Private/UObject/LinkerLoad.cpp
@@ -1135,140 +1135,152 @@ FLinkerLoad::ELinkerStatus FLinkerLoad::CreateLoader(
 	}
 }
 
-/**
- * Serializes the package file summary.
- */
-FLinkerLoad::ELinkerStatus FLinkerLoad::SerializePackageFileSummary()
+FLinkerLoad::ELinkerStatus FLinkerLoad::SerializePackageFileSummaryInternal()
 {
-	DECLARE_SCOPE_CYCLE_COUNTER( TEXT( "FLinkerLoad::SerializePackageFileSummary" ), STAT_LinkerLoad_SerializePackageFileSummary, STATGROUP_LinkerLoad );
-
-	if (bHasSerializedPackageFileSummary == false)
-	{
-		if (Loader->IsError())
-		{
-			UE_LOG(LogLinker, Warning, TEXT("The file '%s' contains unrecognizable data, check that it is of the expected type."), *Filename);
-			return LINKER_Failed;
-		}
-		if (bIsAsyncLoader)
-		{
-			GetAsyncLoader()->StartReadingHeader();
-		}
-
 #if WITH_EDITOR
-		if (LoadProgressScope)
-		{
+	if (LoadProgressScope)
+	{
 		LoadProgressScope->EnterProgressFrame(1);
-		}
+	}
 #endif
-		// Read summary from file.
-		StructuredArchiveRootRecord.GetValue() << SA_VALUE(TEXT("Summary"), Summary);
+	// Read summary from file.
+	StructuredArchiveRootRecord.GetValue() << SA_VALUE(TEXT("Summary"), Summary);
 
-		// Check tag.
-		if( Summary.Tag != PACKAGE_FILE_TAG )
-		{
-			UE_LOG(LogLinker, Warning, TEXT("The file '%s' contains unrecognizable data, check that it is of the expected type."), *Filename );
-			return LINKER_Failed;
-		}
+	// Check tag.
+	if (Summary.Tag != PACKAGE_FILE_TAG)
+	{
+		UE_LOG(LogLinker, Warning, TEXT("The file '%s' contains unrecognizable data, check that it is of the expected type."), *Filename);
+		return LINKER_Failed;
+	}
 
-		// Validate the summary.
-		if( Summary.GetFileVersionUE4() < VER_UE4_OLDEST_LOADABLE_PACKAGE)
-		{
-			UE_LOG(LogLinker, Warning, TEXT("The file %s was saved by a previous version which is not backwards compatible with this one. Min Required Version: %i  Package Version: %i"), *Filename, (int32)VER_UE4_OLDEST_LOADABLE_PACKAGE, Summary.GetFileVersionUE4() );
-			return LINKER_Failed;
-		}
-		// Don't load packages that are only compatible with an engine version newer than the current one.
-		if (!FEngineVersion::Current().IsCompatibleWith(Summary.CompatibleWithEngineVersion))
-		{
-			UE_LOG(LogLinker, Warning, TEXT("Asset '%s' has been saved with engine version newer than current and therefore can't be loaded. CurrEngineVersion: %s AssetEngineVersion: %s"), *Filename, *FEngineVersion::Current().ToString(), *Summary.CompatibleWithEngineVersion.ToString());
-			return LINKER_Failed;
-		}
-		
-		// Set desired property tag format
-		bool bUseUnversionedProperties = Summary.bUnversioned && CanUseUnversionedPropertySerialization();
-		SetUseUnversionedPropertySerialization(bUseUnversionedProperties);
-		Loader->SetUseUnversionedPropertySerialization(bUseUnversionedProperties);
-		
-		if (!FPlatformProperties::RequiresCookedData() && !Summary.SavedByEngineVersion.HasChangelist() && FEngineVersion::Current().HasChangelist())
+	// Validate the summary.
+	if (Summary.GetFileVersionUE4() < VER_UE4_OLDEST_LOADABLE_PACKAGE)
+	{
+		UE_LOG(LogLinker, Warning, TEXT("The file %s was saved by a previous version which is not backwards compatible with this one. Min Required Version: %i  Package Version: %i"), *Filename, (int32)VER_UE4_OLDEST_LOADABLE_PACKAGE, Summary.GetFileVersionUE4());
+		return LINKER_Failed;
+	}
+	// Don't load packages that are only compatible with an engine version newer than the current one.
+	if (!FEngineVersion::Current().IsCompatibleWith(Summary.CompatibleWithEngineVersion))
+	{
+		UE_LOG(LogLinker, Warning, TEXT("Asset '%s' has been saved with engine version newer than current and therefore can't be loaded. CurrEngineVersion: %s AssetEngineVersion: %s"), *Filename, *FEngineVersion::Current().ToString(), *Summary.CompatibleWithEngineVersion.ToString());
+		return LINKER_Failed;
+	}
+
+	// Set desired property tag format
+	bool bUseUnversionedProperties = Summary.bUnversioned && CanUseUnversionedPropertySerialization();
+	SetUseUnversionedPropertySerialization(bUseUnversionedProperties);
+	Loader->SetUseUnversionedPropertySerialization(bUseUnversionedProperties);
+
+	if (!FPlatformProperties::RequiresCookedData() && !Summary.SavedByEngineVersion.HasChangelist() && FEngineVersion::Current().HasChangelist())
+	{
+		// This warning can be disabled in ini with [Core.System] ZeroEngineVersionWarning=False
+		static struct FInitZeroEngineVersionWarning
 		{
-			// This warning can be disabled in ini with [Core.System] ZeroEngineVersionWarning=False
-			static struct FInitZeroEngineVersionWarning
+			bool bDoWarn;
+			FInitZeroEngineVersionWarning()
 			{
-				bool bDoWarn;
-				FInitZeroEngineVersionWarning()
+				if (!GConfig->GetBool(TEXT("Core.System"), TEXT("ZeroEngineVersionWarning"), bDoWarn, GEngineIni))
 				{
-					if (!GConfig->GetBool(TEXT("Core.System"), TEXT("ZeroEngineVersionWarning"), bDoWarn, GEngineIni))
-					{
-						bDoWarn = true;
-					}
+					bDoWarn = true;
 				}
-				FORCEINLINE operator bool() const { return bDoWarn; }
-			} ZeroEngineVersionWarningEnabled;			
-			UE_CLOG(ZeroEngineVersionWarningEnabled, LogLinker, Warning, TEXT("Asset '%s' has been saved with empty engine version. The asset will be loaded but may be incompatible."), *Filename );
-		}
+			}
+			FORCEINLINE operator bool() const { return bDoWarn; }
+		} ZeroEngineVersionWarningEnabled;
+		UE_CLOG(ZeroEngineVersionWarningEnabled, LogLinker, Warning, TEXT("Asset '%s' has been saved with empty engine version. The asset will be loaded but may be incompatible."), *Filename);
+	}
+
+	// Don't load packages that were saved with package version newer than the current one.
+	if ((Summary.GetFileVersionUE4() > GPackageFileUE4Version) || (Summary.GetFileVersionLicenseeUE4() > GPackageFileLicenseeUE4Version))
+	{
+		UE_LOG(LogLinker, Warning, TEXT("Unable to load package (%s) PackageVersion %i, MaxExpected %i : LicenseePackageVersion %i, MaxExpected %i."), *Filename, Summary.GetFileVersionUE4(), GPackageFileUE4Version, Summary.GetFileVersionLicenseeUE4(), GPackageFileLicenseeUE4Version);
+		return LINKER_Failed;
+	}
 
-		// Don't load packages that were saved with package version newer than the current one.
-		if( (Summary.GetFileVersionUE4() > GPackageFileUE4Version) || (Summary.GetFileVersionLicenseeUE4() > GPackageFileLicenseeUE4Version) )
+	// don't load packages that contain editor only data in builds that don't support that and vise versa
+	if (!FPlatformProperties::HasEditorOnlyData() && !(Summary.PackageFlags & PKG_FilterEditorOnly))
+	{
+		UE_LOG(LogLinker, Warning, TEXT("Unable to load package (%s). Package contains EditorOnly data which is not supported by the current build."), *Filename);
+		return LINKER_Failed;
+	}
+
+	// don't load packages that contain editor only data in builds that don't support that and vise versa
+	if (FPlatformProperties::HasEditorOnlyData() && !!(Summary.PackageFlags & PKG_FilterEditorOnly))
+	{
+		// This warning can be disabled in ini or project settings
+		if (!GAllowCookedDataInEditorBuilds)
 		{
-			UE_LOG(LogLinker, Warning, TEXT("Unable to load package (%s) PackageVersion %i, MaxExpected %i : LicenseePackageVersion %i, MaxExpected %i."), *Filename, Summary.GetFileVersionUE4(), GPackageFileUE4Version, Summary.GetFileVersionLicenseeUE4(), GPackageFileLicenseeUE4Version );
+			UE_LOG(LogLinker, Warning,
+				TEXT("Unable to load package (%s). Package contains cooked data which is not supported by the current build. Enable 'Allow Cooked Content In The Editor' in Project Settings under 'Engine - Cooker' section to load it."),
+				*Filename);
 			return LINKER_Failed;
 		}
+	}
+
+	if (FPlatformProperties::RequiresCookedData() &&
+		Summary.PreloadDependencyCount > 0 && Summary.PreloadDependencyOffset > 0 &&
+		!IsEventDrivenLoaderEnabledInCookedBuilds())
+	{
+		UE_LOG(LogLinker, Fatal, TEXT("Package %s contains preload dependency data but the current build does not support it. Make sure Event Driven Loader is enabled and rebuild the game executable."),
+			*GetArchiveName())
+	}
+
+#if PLATFORM_WINDOWS
+	if (!FPlatformProperties::RequiresCookedData() &&
+		// We can't check the post tag if the file is an EDL cooked package
+		!((Summary.PackageFlags & PKG_FilterEditorOnly) && Summary.PreloadDependencyCount > 0 && Summary.PreloadDependencyOffset > 0)
+		&& !IsTextFormat())
+	{
+		// check if this package version stored the 4-byte magic post tag
+		// get the offset of the post tag
+		int64 MagicOffset = TotalSize() - sizeof(uint32);
+		// store the current file offset
+		int64 OriginalOffset = Tell();
 
-		// don't load packages that contain editor only data in builds that don't support that and vise versa
-		if (!FPlatformProperties::HasEditorOnlyData() && !(Summary.PackageFlags & PKG_FilterEditorOnly))
+		uint32 Tag = 0;
+
+		// seek to the post tag and serialize it
+		Seek(MagicOffset);
+		*this << Tag;
+
+		if (Tag != PACKAGE_FILE_TAG)
 		{
-			UE_LOG(LogLinker, Warning, TEXT("Unable to load package (%s). Package contains EditorOnly data which is not supported by the current build."), *Filename );
+			UE_LOG(LogLinker, Warning, TEXT("Unable to load package (%s). Post Tag is not valid. File might be corrupted."), *Filename);
 			return LINKER_Failed;
 		}
 
-		// don't load packages that contain editor only data in builds that don't support that and vise versa
-		if (FPlatformProperties::HasEditorOnlyData() && !!(Summary.PackageFlags & PKG_FilterEditorOnly))
-		{
-			// This warning can be disabled in ini or project settings
-			if (!GAllowCookedDataInEditorBuilds)
-			{
-				UE_LOG(LogLinker, Warning, 
-					TEXT("Unable to load package (%s). Package contains cooked data which is not supported by the current build. Enable 'Allow Cooked Content In The Editor' in Project Settings under 'Engine - Cooker' section to load it."), 
-					*Filename);
-				return LINKER_Failed;
-			}
-		}
+		// seek back to the position after the package summary
+		Seek(OriginalOffset);
+	}
+#endif
+
+	return LINKER_Loaded;
+}
 
-		if (FPlatformProperties::RequiresCookedData() && 
-			Summary.PreloadDependencyCount > 0 && Summary.PreloadDependencyOffset > 0 &&
-			!IsEventDrivenLoaderEnabledInCookedBuilds())
+/**
+ * Serializes the package file summary.
+ */
+FLinkerLoad::ELinkerStatus FLinkerLoad::SerializePackageFileSummary()
+{
+	DECLARE_SCOPE_CYCLE_COUNTER( TEXT( "FLinkerLoad::SerializePackageFileSummary" ), STAT_LinkerLoad_SerializePackageFileSummary, STATGROUP_LinkerLoad );
+
+	if (bHasSerializedPackageFileSummary == false)
+	{
+		if (Loader->IsError())
 		{
-			UE_LOG(LogLinker, Fatal, TEXT("Package %s contains preload dependency data but the current build does not support it. Make sure Event Driven Loader is enabled and rebuild the game executable."),
-				*GetArchiveName())
+			UE_LOG(LogLinker, Warning, TEXT("The file '%s' contains unrecognizable data, check that it is of the expected type."), *Filename);
+			return LINKER_Failed;
 		}
-
-#if PLATFORM_WINDOWS
-		if (!FPlatformProperties::RequiresCookedData() && 
-			// We can't check the post tag if the file is an EDL cooked package
-			!((Summary.PackageFlags & PKG_FilterEditorOnly) && Summary.PreloadDependencyCount > 0 && Summary.PreloadDependencyOffset > 0)
-			&& !IsTextFormat())
+		if (bIsAsyncLoader)
 		{
-			// check if this package version stored the 4-byte magic post tag
-			// get the offset of the post tag
-			int64 MagicOffset = TotalSize() - sizeof(uint32);
-			// store the current file offset
-			int64 OriginalOffset = Tell();
-
-			uint32 Tag = 0;
-
-			// seek to the post tag and serialize it
-			Seek(MagicOffset);
-			*this << Tag;
+			GetAsyncLoader()->StartReadingHeader();
+		}
 
-			if (Tag != PACKAGE_FILE_TAG)
-			{
-				UE_LOG(LogLinker, Warning, TEXT("Unable to load package (%s). Post Tag is not valid. File might be corrupted."), *Filename);
-				return LINKER_Failed;
-			}
+		ELinkerStatus Status = SerializePackageFileSummaryInternal();
 
-			// seek back to the position after the package summary
-			Seek(OriginalOffset);
+		if (Status == LINKER_Failed && bIsAsyncLoader)
+		{
+			GetAsyncLoader()->EndReadingHeader();
 		}
-#endif // PLATFORM_WINDOWS
 
 		ELinkerStatus UpdateStatus = UpdateFromPackageFileSummary();
 		if (UpdateStatus != LINKER_Loaded)
diff --git a/Engine/Source/Runtime/CoreUObject/Public/UObject/LinkerLoad.h b/Engine/Source/Runtime/CoreUObject/Public/UObject/LinkerLoad.h
--- a/Engine/Source/Runtime/CoreUObject/Public/UObject/LinkerLoad.h
+++ b/Engine/Source/Runtime/CoreUObject/Public/UObject/LinkerLoad.h
@@ -915,10 +915,15 @@ protected: // Daniel L: Made this protected so I can override the constructor an
 	);
 private:
 	/**
-	 * Serializes the package file summary.
+	 * Start the process of serializing the package file summary if needed
 	 */
 	ELinkerStatus SerializePackageFileSummary();
 
+	/**
+	 * Does the actual serialization of  the package file summary.
+	 */
+	ELinkerStatus SerializePackageFileSummaryInternal();
+
 	/**
 	 * Updates the linker, loader and root package with data from the package file summary.
 	 * */
