Adding a MC Delegate to allow plugins/modules to change the tags we serialize on an asset.

diff --git a/Engine/Source/Runtime/CoreUObject/Private/UObject/Obj.cpp b/Engine/Source/Runtime/CoreUObject/Private/UObject/Obj.cpp
--- a/Engine/Source/Runtime/CoreUObject/Private/UObject/Obj.cpp
+++ b/Engine/Source/Runtime/CoreUObject/Private/UObject/Obj.cpp
@@ -86,6 +86,10 @@ static UPackage*			GObjTransientPkg								= NULL;
 	static TArray<FString>			DebugSpikeMarkNames;
 #endif
 
+#if WITH_EDITOR
+	UObject::FAssetRegistryTag::FOnGetObjectAssetRegistryTags UObject::FAssetRegistryTag::OnGetExtraObjectTags;
+#endif // WITH_EDITOR
+
 UObject::UObject( EStaticConstructor, EObjectFlags InFlags )
 : UObjectBaseUtility(InFlags | (!(InFlags & RF_Dynamic) ? (RF_MarkAsNative | RF_MarkAsRootSet) : RF_NoFlags))
 {
@@ -1815,6 +1819,9 @@ void UObject::GetAssetRegistryTags(TArray<FAssetRegistryTag>& OutTags) const
 	FAssetRegistryTag::GetAssetRegistryTagsFromSearchableProperties(this, OutTags);
 
 #if WITH_EDITOR
+	// Notify external sources that we need tags.
+	FAssetRegistryTag::OnGetExtraObjectTags.Broadcast(this, OutTags);
+
 	// Check if there's a UMetaData for this object that has tags that are requested in the settings to be transferred to the Asset Registry
 	const TSet<FName>& MetaDataTagsForAR = GetMetaDataTagsForAssetRegistry();
 	if (MetaDataTagsForAR.Num() > 0)
diff --git a/Engine/Source/Runtime/CoreUObject/Public/UObject/Object.h b/Engine/Source/Runtime/CoreUObject/Public/UObject/Object.h
--- a/Engine/Source/Runtime/CoreUObject/Public/UObject/Object.h
+++ b/Engine/Source/Runtime/CoreUObject/Public/UObject/Object.h
@@ -699,10 +699,16 @@ public:
 			: Name(InName), Value(InValue), Type(InType), DisplayFlags(InDisplayFlags) {}
 
 		/** Gathers a list of asset registry searchable tags from given objects properties */
-		COREUOBJECT_API static void GetAssetRegistryTagsFromSearchableProperties(const UObject* Object, TArray<FAssetRegistryTag>& OutTags);
+		COREUOBJECT_API static void GetAssetRegistryTagsFromSearchableProperties(const UObject* Object, TArray<FAssetRegistryTag>& InOutTags);
 
 		/** Returns true if this FName is a special UStruct that should be exported even if not tagged, with the struct name as the tag name */
 		COREUOBJECT_API static bool IsUniqueAssetRegistryTagStruct(FName StructName, ETagType& TagType);
+
+#if WITH_EDITOR
+		/** Callback  */
+		DECLARE_MULTICAST_DELEGATE_TwoParams(FOnGetObjectAssetRegistryTags, const UObject* /*Object*/, TArray<FAssetRegistryTag>& /*InOutTags*/);
+		COREUOBJECT_API static FOnGetObjectAssetRegistryTags OnGetExtraObjectTags;
+#endif // WITH_EDITOR
 	};
 
 	/**
