Potential fix for destroying blueprint property values on exit.

- The final GC pass will no longer gather unreachable objects first because it can mess up incremental GC if it's already in progress. Instead we're forcing incremental GC to complete if it's running (this is restoring behavior prior to CL #6535063)
- Added additional safeguards in a few places.

diff --git a/Engine/Source/Runtime/CoreUObject/Private/UObject/Class.cpp b/Engine/Source/Runtime/CoreUObject/Private/UObject/Class.cpp
--- a/Engine/Source/Runtime/CoreUObject/Private/UObject/Class.cpp
+++ b/Engine/Source/Runtime/CoreUObject/Private/UObject/Class.cpp
@@ -1591,7 +1591,7 @@ UStruct::~UStruct()
 	// This needs to happen after FinishDestroy which calls DestroyNonNativeProperties
 	// Also, Blueprint generated classes can have DestroyNonNativeProperties called on them after their FinishDestroy has been called
 	// so properties can only be deleted in the destructor
-	DestroyPropertyLinkedList(ChildProperties);
+	DestroyChildPropertiesAndResetPropertyLinks();
 }
 
 IMPLEMENT_FSTRUCTUREDARCHIVE_SERIALIZER(UStruct);
diff --git a/Engine/Source/Runtime/CoreUObject/Private/UObject/Obj.cpp b/Engine/Source/Runtime/CoreUObject/Private/UObject/Obj.cpp
--- a/Engine/Source/Runtime/CoreUObject/Private/UObject/Obj.cpp
+++ b/Engine/Source/Runtime/CoreUObject/Private/UObject/Obj.cpp
@@ -4434,11 +4434,16 @@ void StaticExit()
 		GUObjectArray.CloseDisregardForGC();
 	}
 
+	// Complete any pending incremental GC
+	if (IsIncrementalPurgePending())
+	{
+		IncrementalPurgeGarbage(false);
+	}
+
 	// Make sure no other threads manipulate UObjects
 	AcquireGCLock();
 
-	GatherUnreachableObjects(false);
-	IncrementalPurgeGarbage(false);
+	// Dissolve all clusters before the final GC pass
 	GUObjectClusters.DissolveClusters(true);
 
 	// Keep track of how many objects there are for GC stats as we simulate a mark pass.
@@ -4460,10 +4465,18 @@ void StaticExit()
 		FUObjectItem* ObjItem = *It;
 		checkSlow(ObjItem);
 		UObject* Obj = static_cast<UObject*>(ObjItem->Object);
-		if (Obj && !Obj->IsA<UField>()) // Skip Structures, properties, etc.. They could be still necessary while GC.
+		if (Obj) 
 		{
-			// Mark as unreachable so purge phase will kill it.
-			ObjItem->SetUnreachable();
+			// Skip Structures, properties, etc.. They could be still necessary while GC.
+			if (!Obj->IsA<UField>())
+			{
+				// Mark as unreachable so purge phase will kill it.
+				ObjItem->SetUnreachable();
+			}
+			else
+			{
+				ObjItem->ClearUnreachable();
+			}
 		}
 	}
 
