Fix IoStore CIS tests.

There was a mismerge of the fix 11091798 when it was moved from dev-core to a private branch which reintroduced the exception for FBulkDataBase::CanLoadFromDisk, so that it would pass inlined data if it was running from IoStore. This would then quite rightly trigger issues in our automated tests.

Restored the original fix in 11091798.

#rnx

diff --git a/Engine/Source/Runtime/CoreUObject/Private/Serialization/BulkData2.cpp b/Engine/Source/Runtime/CoreUObject/Private/Serialization/BulkData2.cpp
--- a/Engine/Source/Runtime/CoreUObject/Private/Serialization/BulkData2.cpp
+++ b/Engine/Source/Runtime/CoreUObject/Private/Serialization/BulkData2.cpp
@@ -999,18 +999,6 @@ bool FBulkDataBase::CanLoadFromDisk() const
 		return true;
 	}
 
-	// If the IoDispatcher is enabled then technically inlined data cannot load
-	// from disk, but it should act as thought it can up until the point that it
-	// actually tries to load. At which point it will give a log error.
-	//
-	// This is a temp hack while we work out if we want to allow additional inline
-	// loading at runtime (and pay the memory increase in the ToC) or prevent it entirely.
-	// TODO: Resolve this!
-	if (IsInlined() && IsIoDispatcherEnabled())
-	{
-		return true;
-	}
-
 	// If this BulkData has a fallback token then it can find it's filepath and load from disk
 	if (Data.Fallback.Token != InvalidToken)
 	{
