Add GetTypeHash() support to TWeakObjectPtr properties.

diff --git a/Engine/Source/Runtime/CoreUObject/Private/UObject/PropertyWeakObjectPtr.cpp b/Engine/Source/Runtime/CoreUObject/Private/UObject/PropertyWeakObjectPtr.cpp
--- a/Engine/Source/Runtime/CoreUObject/Private/UObject/PropertyWeakObjectPtr.cpp
+++ b/Engine/Source/Runtime/CoreUObject/Private/UObject/PropertyWeakObjectPtr.cpp
@@ -55,3 +55,8 @@ void FWeakObjectProperty::SetObjectPropertyValue(void* PropertyValueAddress, UOb
 {
 	SetPropertyValue(PropertyValueAddress, TCppType(Value));
 }
+
+uint32 FWeakObjectProperty::GetValueTypeHashInternal(const void* Src) const
+{
+	return GetTypeHash(*(FWeakObjectPtr*)Src);
+}
diff --git a/Engine/Source/Runtime/CoreUObject/Public/UObject/UnrealType.h b/Engine/Source/Runtime/CoreUObject/Public/UObject/UnrealType.h
--- a/Engine/Source/Runtime/CoreUObject/Public/UObject/UnrealType.h
+++ b/Engine/Source/Runtime/CoreUObject/Public/UObject/UnrealType.h
@@ -2279,6 +2279,9 @@ class COREUOBJECT_API FWeakObjectProperty : public TFObjectPropertyBase<FWeakObj
 
 	// FProperty interface
 	virtual void SerializeItem(FStructuredArchive::FSlot Slot, void* Value, void const* Defaults) const override;
+private:
+	virtual uint32 GetValueTypeHashInternal(const void* Src) const override;
+public:
 	// End of FProperty interface
 
 	// FObjectProperty interface
