Track placeholder objects stored off in PropertyObjectReferences

diff --git a/Engine/Source/Runtime/CoreUObject/Private/UObject/Class.cpp b/Engine/Source/Runtime/CoreUObject/Private/UObject/Class.cpp
--- a/Engine/Source/Runtime/CoreUObject/Private/UObject/Class.cpp
+++ b/Engine/Source/Runtime/CoreUObject/Private/UObject/Class.cpp
@@ -945,6 +945,24 @@ void UStruct::Link(FArchive& Ar, bool bRelinkExistingProperties)
 			CurrentField->AddReferencedObjects(PropertyReferenceCollector);
 		}
 		PropertyObjectReferences = PropertyReferenceCollector.UniqueReferences.Array();
+
+#if USE_CIRCULAR_DEPENDENCY_LOAD_DEFERRING
+		// The old (non-EDL) FLinkerLoad code paths create placeholder objects
+		// for classes and functions. We have to babysit these, just as we do
+		// for bytecode references (reusing the AddReferencingScriptExpr fn).
+		// Long term we should not use placeholder objects like this:
+		for(int32 I = 0; I < PropertyObjectReferences.Num(); ++I)
+		{
+			if (ULinkerPlaceholderClass* PlaceholderObj = Cast<ULinkerPlaceholderClass>(PropertyObjectReferences[I]))
+			{
+				// let the placeholder track the reference to it:
+				PlaceholderObj->AddReferencingScriptExpr(reinterpret_cast<UClass**>(&PropertyObjectReferences[I]));
+			}
+			// I don't currently see how placeholder functions could be present in this list, but that's
+			// a dangerous assumption.
+			ensure(!(PropertyObjectReferences[I]->IsA<ULinkerPlaceholderFunction>()));
+		}
+#endif // USE_CIRCULAR_DEPENDENCY_LOAD_DEFERRING
 	}
 
 #if WITH_EDITORONLY_DATA
