Clear NeedPostLoadSubobjects in all the places that blueprint compilation clears NeedPostLoad
This prevents non-component subobjects that were trashed during compilation from sticking around forever with a bad AsyncLoading flag, causing corruption or crashes
Some of these fixes aren't currently needed because CDOs are hacked to never have NeedPostLoadSubobjects set on them, but this is safer in case we change that later

diff --git a/Engine/Source/Editor/UnrealEd/Private/Kismet2/BlueprintEditorUtils.cpp b/Engine/Source/Editor/UnrealEd/Private/Kismet2/BlueprintEditorUtils.cpp
--- a/Engine/Source/Editor/UnrealEd/Private/Kismet2/BlueprintEditorUtils.cpp
+++ b/Engine/Source/Editor/UnrealEd/Private/Kismet2/BlueprintEditorUtils.cpp
@@ -736,7 +736,7 @@ void FBlueprintEditorUtils::PatchNewCDOIntoLinker(UObject* CDO, FLinkerLoad* Lin
 		if( OldCDO != nullptr )
 		{
 			EObjectFlags OldObjectFlags = OldCDO->GetFlags();
-			OldCDO->ClearFlags(RF_NeedLoad|RF_NeedPostLoad);
+			OldCDO->ClearFlags(RF_NeedLoad|RF_NeedPostLoad|RF_NeedPostLoadSubobjects);
 			OldCDO->SetLinker(nullptr, INDEX_NONE);
 			
 			// Copy flags from the old CDO.
diff --git a/Engine/Source/Runtime/CoreUObject/Private/Blueprint/BlueprintSupport.cpp b/Engine/Source/Runtime/CoreUObject/Private/Blueprint/BlueprintSupport.cpp
--- a/Engine/Source/Runtime/CoreUObject/Private/Blueprint/BlueprintSupport.cpp
+++ b/Engine/Source/Runtime/CoreUObject/Private/Blueprint/BlueprintSupport.cpp
@@ -693,7 +693,7 @@ bool FLinkerLoad::RegenerateBlueprintClass(UClass* LoadClass, UObject* ExportObj
 			{
 				BlueprintObject->ClearFlags(RF_BeingRegenerated);
 				// Fix up the linker so that the RegeneratedClass is used
-				LoadClass->ClearFlags(RF_NeedLoad | RF_NeedPostLoad);
+				LoadClass->ClearFlags(RF_NeedLoad | RF_NeedPostLoad | RF_NeedPostLoadSubobjects);
 			}
 		}
 	}
@@ -1974,7 +1974,7 @@ void FLinkerLoad::ResolveDeferredExports(UClass* LoadClass)
 			if (DeferredCDOIndex != INDEX_NONE)
 			{
 				const EObjectFlags OldFlags = BlueprintCDO->GetFlags();
-				BlueprintCDO->ClearFlags(RF_NeedLoad | RF_NeedPostLoad);
+				BlueprintCDO->ClearFlags(RF_NeedLoad | RF_NeedPostLoad | RF_NeedPostLoadSubobjects);
 				BlueprintCDO->SetLinker(this, DeferredCDOIndex, /*bShouldDetatchExisting =*/false);
 				BlueprintCDO->SetFlags(OldFlags);
 			}
diff --git a/Engine/Source/Runtime/CoreUObject/Private/UObject/LinkerLoad.cpp b/Engine/Source/Runtime/CoreUObject/Private/UObject/LinkerLoad.cpp
--- a/Engine/Source/Runtime/CoreUObject/Private/UObject/LinkerLoad.cpp
+++ b/Engine/Source/Runtime/CoreUObject/Private/UObject/LinkerLoad.cpp
@@ -550,7 +550,7 @@ void FLinkerLoad::PRIVATE_PatchNewObjectIntoExport(UObject* OldObject, UObject*
 
 		// Detach the old object to make room for the new
 		const EObjectFlags OldObjectFlags = OldObject->GetFlags();
-		OldObject->ClearFlags(RF_NeedLoad|RF_NeedPostLoad);
+		OldObject->ClearFlags(RF_NeedLoad|RF_NeedPostLoad|RF_NeedPostLoadSubobjects);
 		OldObject->SetLinker(nullptr, INDEX_NONE, true);
 
 		// Copy flags from the old CDO.
@@ -4514,7 +4514,7 @@ UObject* FLinkerLoad::CreateExport( int32 Index )
 				{
 					// For classes that are about to be regenerated, make sure we register them with the linker, so future references to this linker index will be valid
 					const EObjectFlags OldFlags = Export.Object->GetFlags();
-					Export.Object->ClearFlags(RF_NeedLoad|RF_NeedPostLoad);
+					Export.Object->ClearFlags(RF_NeedLoad|RF_NeedPostLoad|RF_NeedPostLoadSubobjects);
 					Export.Object->SetLinker( this, Index, false );
 					Export.Object->SetFlags(OldFlags);
 				}
