Fix StaticFindObject() search for any package with a package name that is only a substring of the full package name

diff --git a/Engine/Source/Runtime/CoreUObject/Private/UObject/UObjectHash.cpp b/Engine/Source/Runtime/CoreUObject/Private/UObject/UObjectHash.cpp
--- a/Engine/Source/Runtime/CoreUObject/Private/UObject/UObjectHash.cpp
+++ b/Engine/Source/Runtime/CoreUObject/Private/UObject/UObjectHash.cpp
@@ -489,6 +489,26 @@ UObject* StaticFindObjectFastExplicit( const UClass* ObjectClass, FName ObjectNa
 	return Result;
 }
 
+static bool NameEndsWith(FName Name, FName Suffix)
+{
+	if (Name == Suffix)
+	{
+		return true;
+	}
+
+	if (Name.GetNumber() != Suffix.GetNumber())
+	{
+		return false;
+	}
+		
+	TCHAR PlainName[NAME_SIZE];
+	TCHAR PlainSuffix[NAME_SIZE];
+	uint32 NameLen = Name.GetPlainNameString(PlainName);
+	uint32 SuffixLen = Suffix.GetPlainNameString(PlainSuffix);
+
+	return NameLen >= SuffixLen && FCString::Strnicmp(PlainName + NameLen - SuffixLen, PlainSuffix, SuffixLen) == 0;
+}
+
 // Splits an object path into FNames representing an outer chain.
 //
 // Input path examples: "Object", "Package.Object", "Object:Subobject", "Object:Subobject.Nested", "Package.Object:Subobject", "Package.Object:Subobject.NestedSubobject"
@@ -518,9 +538,14 @@ struct FObjectSearchPath
 
 	bool MatchOuterNames(UObject* Outer) const
 	{
-		for (int32 Idx = Outers.Num() - 1; Idx >= 0; --Idx)
+		if (Outers.Num() == 0)
 		{
-			if (!Outer || Outers[Idx] != Outer->GetFName())
+			return true;
+		}
+
+		for (int32 Idx = Outers.Num() - 1; Idx > 0; --Idx)
+		{
+			if (!Outer || Outer->GetFName() != Outers[Idx])
 			{
 				return false;
 			}
@@ -528,7 +553,7 @@ struct FObjectSearchPath
 			Outer = Outer->GetOuter();
 		}
 
-		return true;
+		return Outer && NameEndsWith(Outer->GetFName(), Outers[0]);
 	}
 };
 
