FReferenceFinder will not use multithreading until all native classes have been registered to avoid asserts when assembling token streams for classes that haven't had it assembled yet.

diff --git a/Engine/Source/Runtime/CoreUObject/Private/UObject/ReferencerFinder.cpp b/Engine/Source/Runtime/CoreUObject/Private/UObject/ReferencerFinder.cpp
--- a/Engine/Source/Runtime/CoreUObject/Private/UObject/ReferencerFinder.cpp
+++ b/Engine/Source/Runtime/CoreUObject/Private/UObject/ReferencerFinder.cpp
@@ -41,6 +41,14 @@ public:
 };
 typedef TDefaultReferenceCollector<FAllReferencesProcessor> FAllReferencesCollector;
 
+// Until all native UObject classes have been registered it's unsafe to run FReferencerFinder on multiple threads
+static bool GUObjectRegistrationComplete = false;
+
+void FReferencerFinder::NotifyRegistrationComplete()
+{
+	GUObjectRegistrationComplete = true;
+}
+
 TArray<UObject*> FReferencerFinder::GetAllReferencers(const TArray<UObject*>& Referencees, const TSet<UObject*>* ObjectsToIgnore )
 {
 	return GetAllReferencers(TSet<UObject*>(Referencees), ObjectsToIgnore);
@@ -104,7 +112,7 @@ TArray<UObject*> FReferencerFinder::GetAllReferencers(const TSet<UObject*>& Refe
 				FScopeLock ResultLock(&ResultCritical);
 				Ret.Append(ThreadResult.Array());
 			}
-		}, GIsInitialLoad ? EParallelForFlags::ForceSingleThread : EParallelForFlags::None);
+		}, GUObjectRegistrationComplete ? EParallelForFlags::None :  EParallelForFlags::ForceSingleThread );
 
 		UnlockUObjectHashTables();
 	}
diff --git a/Engine/Source/Runtime/CoreUObject/Public/UObject/ReferencerFinder.h b/Engine/Source/Runtime/CoreUObject/Public/UObject/ReferencerFinder.h
--- a/Engine/Source/Runtime/CoreUObject/Public/UObject/ReferencerFinder.h
+++ b/Engine/Source/Runtime/CoreUObject/Public/UObject/ReferencerFinder.h
@@ -13,4 +13,7 @@ class COREUOBJECT_API FReferencerFinder
 public:
 	static TArray<UObject*> GetAllReferencers(const TArray<UObject*>& Referencees, const TSet<UObject*>* ObjectsToIgnore);
 	static TArray<UObject*> GetAllReferencers(const TSet<UObject*>& Referencees, const TSet<UObject*>* ObjectsToIgnore);
+
+	/** Called when the initial load phase is complete and we're done registering native UObject classes */
+	static void NotifyRegistrationComplete();
 };
\ No newline at end of file
diff --git a/Engine/Source/Runtime/Launch/Private/LaunchEngineLoop.cpp b/Engine/Source/Runtime/Launch/Private/LaunchEngineLoop.cpp
--- a/Engine/Source/Runtime/Launch/Private/LaunchEngineLoop.cpp
+++ b/Engine/Source/Runtime/Launch/Private/LaunchEngineLoop.cpp
@@ -69,6 +69,7 @@
 	#include "UObject/Package.h"
 	#include "UObject/Linker.h"
 	#include "UObject/LinkerLoad.h"
+	#include "UObject/ReferencerFinder.h"
 #endif
 
 #if WITH_EDITOR
@@ -3185,6 +3186,7 @@ int32 FEngineLoop::PreInitPostStartupScreen(const TCHAR* CmdLine)
 		GUObjectArray.CloseDisregardForGC();
 	}
 	NotifyRegistrationComplete();
+	FReferencerFinder::NotifyRegistrationComplete();
 #endif // WITH_COREUOBJECT
 
 #if WITH_ENGINE
