Refresh the package localization cache when a plugin is manually mounted via MountPluginFromExternalSource

diff --git a/Engine/Source/Runtime/CoreUObject/Private/Internationalization/PackageLocalizationManager.cpp b/Engine/Source/Runtime/CoreUObject/Private/Internationalization/PackageLocalizationManager.cpp
--- a/Engine/Source/Runtime/CoreUObject/Private/Internationalization/PackageLocalizationManager.cpp
+++ b/Engine/Source/Runtime/CoreUObject/Private/Internationalization/PackageLocalizationManager.cpp
@@ -2,8 +2,9 @@
 
 #include "Internationalization/PackageLocalizationManager.h"
 #include "Misc/Paths.h"
-#include "Internationalization/Culture.h"
 #include "Misc/PackageName.h"
+#include "Interfaces/IPluginManager.h"
+#include "Internationalization/Culture.h"
 #include "Internationalization/IPackageLocalizationCache.h"
 #include "Internationalization/PackageLocalizationCache.h"
 
@@ -70,12 +71,15 @@ void FPackageLocalizationManager::InitializeFromCache(const TSharedRef<IPackageL
 {
 	ActiveCache = InCache;
 	ActiveCache->ConditionalUpdateCache();
+
+	// Allow the plugin manager to update the package localization cache by exposing access through a delegate.
+	// PluginManager is a Core class, but package localization functionality is added at the CoreUObject level.
+	IPluginManager::Get().SetUpdatePackageLocalizationCacheDelegate(IPluginManager::FUpdatePackageLocalizationCacheDelegate::CreateRaw(this, &FPackageLocalizationManager::ConditionalUpdateCache));
 }
 
 void FPackageLocalizationManager::InitializeFromDefaultCache()
 {
-	ActiveCache = MakeShareable(new FDefaultPackageLocalizationCache());
-	ActiveCache->ConditionalUpdateCache();
+	InitializeFromCache(MakeShared<FDefaultPackageLocalizationCache>());
 }
 
 FName FPackageLocalizationManager::FindLocalizedPackageName(const FName InSourcePackageName)
diff --git a/Engine/Source/Runtime/Projects/Private/PluginManager.cpp b/Engine/Source/Runtime/Projects/Private/PluginManager.cpp
--- a/Engine/Source/Runtime/Projects/Private/PluginManager.cpp
+++ b/Engine/Source/Runtime/Projects/Private/PluginManager.cpp
@@ -1303,6 +1303,11 @@ void FPluginManager::SetRegisterMountPointDelegate( const FRegisterMountPointDel
 	RegisterMountPointDelegate = Delegate;
 }
 
+void FPluginManager::SetUpdatePackageLocalizationCacheDelegate( const FUpdatePackageLocalizationCacheDelegate& Delegate )
+{
+	UpdatePackageLocalizationCacheDelegate = Delegate;
+}
+
 bool FPluginManager::AreRequiredPluginsAvailable()
 {
 	return ConfigureEnabledPlugins();
@@ -1484,6 +1489,9 @@ void FPluginManager::MountPluginFromExternalSource(const TSharedRef<FPlugin>& Pl
 				CoreSystemSection->AddUnique("Paths", MoveTemp(ContentDir));
 			}
 		}
+
+		// Update the localization cache for the newly added content directory
+		UpdatePackageLocalizationCacheDelegate.ExecuteIfBound();
 	}
 
 	// If it's a code module, also load the modules for it
diff --git a/Engine/Source/Runtime/Projects/Private/PluginManager.h b/Engine/Source/Runtime/Projects/Private/PluginManager.h
--- a/Engine/Source/Runtime/Projects/Private/PluginManager.h
+++ b/Engine/Source/Runtime/Projects/Private/PluginManager.h
@@ -100,6 +100,7 @@ public:
 	virtual bool LoadModulesForEnabledPlugins( const ELoadingPhase::Type LoadingPhase ) override;
 	virtual void GetLocalizationPathsForEnabledPlugins( TArray<FString>& OutLocResPaths ) override;
 	virtual void SetRegisterMountPointDelegate( const FRegisterMountPointDelegate& Delegate ) override;
+	virtual void SetUpdatePackageLocalizationCacheDelegate( const FUpdatePackageLocalizationCacheDelegate& Delegate ) override;
 	virtual bool AreRequiredPluginsAvailable() override;
 #if !IS_MONOLITHIC
 	virtual bool CheckModuleCompatibility(TArray<FString>& OutIncompatibleModules, TArray<FString>& OutIncompatibleEngineModules) override;
@@ -185,6 +186,10 @@ private:
 	    content path mounting functionality from Core. */
 	FRegisterMountPointDelegate RegisterMountPointDelegate;
 
+	/** Delegate for updating the package localization cache.  Bound by FPackageLocalizationManager code in 
+		CoreUObject, so that we can access localization cache functionality from Core. */
+	FUpdatePackageLocalizationCacheDelegate UpdatePackageLocalizationCacheDelegate;
+
 	/** Set when all the appropriate plugins have been marked as enabled */
 	bool bHaveConfiguredEnabledPlugins;
 
diff --git a/Engine/Source/Runtime/Projects/Public/Interfaces/IPluginManager.h b/Engine/Source/Runtime/Projects/Public/Interfaces/IPluginManager.h
--- a/Engine/Source/Runtime/Projects/Public/Interfaces/IPluginManager.h
+++ b/Engine/Source/Runtime/Projects/Public/Interfaces/IPluginManager.h
@@ -212,6 +212,17 @@ public:
 	 */
 	virtual void SetRegisterMountPointDelegate( const FRegisterMountPointDelegate& Delegate ) = 0;
 
+	/** Delegate type for updating the package localization cache.  Used internally by FPackageLocalizationManager code. */
+	DECLARE_DELEGATE( FUpdatePackageLocalizationCacheDelegate );
+
+	/**
+	 * Sets the delegate to call to update the package localization cache.  This is used internally by the plug-in manager system
+	 * and should not be called by you.  This is registered at application startup by FPackageLocalizationManager code in CoreUObject.
+	 *
+	 * @param	Delegate	The delegate to that will be called when plug-in manager needs to update the package localization cache
+	 */
+	virtual void SetUpdatePackageLocalizationCacheDelegate( const FUpdatePackageLocalizationCacheDelegate& Delegate ) = 0;
+
 	/**
 	 * Checks if all the required plug-ins are available. If not, will present an error dialog the first time a plug-in is loaded or this function is called.
 	 *
