Improved UnversionedPropertySerializationTest

Collect detailed information on roundtripping error to pinpoint problems much faster.

Implement Identical() for some movie scene evaluation types to fix UnversionedPropertySerializationTest

Previous Identical() implementation returned false since they were only compared to default value and was known to always differ from default. UPS test requires proper Identical() implementation.

#[fyi] andrew.rodham, austin.crismore

diff --git a/Engine/Source/Runtime/CoreUObject/Private/Serialization/UnversionedPropertySerializationTest.cpp b/Engine/Source/Runtime/CoreUObject/Private/Serialization/UnversionedPropertySerializationTest.cpp
--- a/Engine/Source/Runtime/CoreUObject/Private/Serialization/UnversionedPropertySerializationTest.cpp
+++ b/Engine/Source/Runtime/CoreUObject/Private/Serialization/UnversionedPropertySerializationTest.cpp
@@ -152,8 +152,26 @@ struct FUnversionedPropertyTest : public FUnversionedPropertyTestInput
 
 	static constexpr uint32 EqualsPortFlags = 0;
 
+	struct FPropertyDiff
+	{
+		const FProperty* Property;
+		const void* A;
+		const void* B;
+		const TCHAR* MismatchKind;
+
+		FString GetType() const
+		{
+			if (const FStructProperty* StructProperty = CastField<FStructProperty>(Property))
+			{
+				return StructProperty->Struct->GetName();
+			}
+
+			return Property->GetClass()->GetName();
+		}
+	};
+
 	// FProperty::Identical() flavor suited to comparing loaded instances
-	static bool Equals(const FProperty* Property, const void* A, const void* B)
+	static bool Equals(const FProperty* Property, const void* A, const void* B, FPropertyDiff& OutDiff)
 	{
 		if (Property->GetPropertyFlags() & (CPF_EditorOnly | CPF_Transient))
 		{
@@ -161,37 +179,43 @@ struct FUnversionedPropertyTest : public FUnversionedPropertyTestInput
 		}
 		else if (const FStructProperty* StructProperty = CastField<FStructProperty>(Property))
 		{
-			return Equals(StructProperty, A, B);
+			return Equals(StructProperty, A, B, OutDiff);
 		}
 		else if (const FArrayProperty* ArrayProperty = CastField<FArrayProperty>(Property))
 		{
-			return Equals(ArrayProperty, A, B);
+			return Equals(ArrayProperty, A, B, OutDiff);
 		}
 		else if (const FSetProperty* SetProperty = CastField<FSetProperty>(Property))
 		{
-			return Equals(SetProperty, A, B);
+			return Equals(SetProperty, A, B, OutDiff);
 		}
 		else if (const FMapProperty* MapProperty = CastField<FMapProperty>(Property))
 		{
-			return Equals(MapProperty, A, B);
+			return Equals(MapProperty, A, B, OutDiff);
 		}
-
-		return Property->Identical(A, B, EqualsPortFlags);
+		else if (!Property->Identical(A, B, EqualsPortFlags))
+		{
+			OutDiff = {Property, A, B, TEXT("Identical()")};
+			return false;
+		}
+		
+		return true;
 	}
 
-	static bool Equals(const FArrayProperty* Property, const void* A, const void* B)
+	static bool Equals(const FArrayProperty* Property, const void* A, const void* B, FPropertyDiff& OutDiff)
 	{
 		FScriptArrayHelper HelperA(Property, A);
 		FScriptArrayHelper HelperB(Property, B);
 
 		if (HelperA.Num() != HelperB.Num())
 		{
+			OutDiff = {Property, A, B, TEXT("Num()")};
 			return false;
 		}
 
 		for (int32 Idx = 0, Num = HelperA.Num(); Idx < Num; ++Idx)
 		{
-			if (!Equals(Property->Inner, HelperA.GetRawPtr(Idx), HelperB.GetRawPtr(Idx)))
+			if (!Equals(Property->Inner, HelperA.GetRawPtr(Idx), HelperB.GetRawPtr(Idx), OutDiff))
 			{
 				return false;
 			}
@@ -200,22 +224,23 @@ struct FUnversionedPropertyTest : public FUnversionedPropertyTestInput
 		return true;
 	}
 
-	static const uint8* FindElementPtr(const FScriptSetHelper& Helper, const uint8* Element)
+	static const uint8* FindElementPtr(const FScriptSetHelper& Helper, const uint8* Element, FPropertyDiff& OutDiff)
 	{
 		const FProperty* ElemProp = Helper.GetElementProperty();
 		int32 Index = Helper.Set->FindIndex(Element, Helper.SetLayout,
 											[ElemProp](const void* Element) { return ElemProp->GetValueTypeHash(Element); },
-											[ElemProp](const void* A, const void* B) { return Equals(ElemProp, A, B); });
+											[ElemProp, &OutDiff](const void* A, const void* B) { return Equals(ElemProp, A, B, OutDiff); });
 		return Index >= 0 ? Helper.GetElementPtr(Index) : nullptr;
 	}
 	
-	static bool Equals(const FSetProperty* Property, const void* A, const void* B)
+	static bool Equals(const FSetProperty* Property, const void* A, const void* B, FPropertyDiff& OutDiff)
 	{
 		FScriptSetHelper HelperA(Property, A);
 		FScriptSetHelper HelperB(Property, B);
 
 		if (HelperA.Num() != HelperB.Num())
 		{
+			OutDiff = {Property, A, B, TEXT("Num()")};
 			return false;
 		}
 
@@ -224,7 +249,7 @@ struct FUnversionedPropertyTest : public FUnversionedPropertyTestInput
 			if (HelperA.IsValidIndex(IndexA))
 			{
 				const uint8* ElemA = HelperA.GetElementPtr(IndexA);
-				const uint8* ElemB = FindElementPtr(HelperB, ElemA);
+				const uint8* ElemB = FindElementPtr(HelperB, ElemA, OutDiff);
 		
 				if (!ElemB)
 				{
@@ -236,16 +261,16 @@ struct FUnversionedPropertyTest : public FUnversionedPropertyTestInput
 		return true;
 	}
 
-	static const uint8* FindPairPtr(const FScriptMapHelper& Helper, const uint8* Key)
+	static const uint8* FindPairPtr(const FScriptMapHelper& Helper, const uint8* Key, FPropertyDiff& OutDiff)
 	{
 		const FProperty* KeyProp = Helper.GetKeyProperty();
 		int32 Index = Helper.HeapMap->FindPairIndex(Key, Helper.MapLayout,
 												[KeyProp](const void* Key) { return KeyProp->GetValueTypeHash(Key); },
-												[KeyProp](const void* A, const void* B) { return Equals(KeyProp, A, B); });
+												[KeyProp, &OutDiff](const void* A, const void* B) { return Equals(KeyProp, A, B, OutDiff); });
 		return Index >= 0 ? Helper.GetPairPtr(Index) : nullptr;
 	}
 	
-	static bool Equals(const FMapProperty* Property, const void* A, const void* B)
+	static bool Equals(const FMapProperty* Property, const void* A, const void* B, FPropertyDiff& OutDiff)
 	{
 		FScriptMapHelper HelperA(Property, A);
 		FScriptMapHelper HelperB(Property, B);
@@ -254,6 +279,7 @@ struct FUnversionedPropertyTest : public FUnversionedPropertyTestInput
 
 		if (HelperA.Num() != HelperB.Num())
 		{
+			OutDiff = {Property, A, B, TEXT("Num()")};
 			return false;
 		}
 		
@@ -262,14 +288,14 @@ struct FUnversionedPropertyTest : public FUnversionedPropertyTestInput
 			if (HelperA.IsValidIndex(IndexA))
 			{
 				const uint8* PairA = HelperA.GetPairPtr(IndexA);
-				const uint8* PairB = FindPairPtr(HelperB, PairA);
+				const uint8* PairB = FindPairPtr(HelperB, PairA, OutDiff);
 		
 				if (!PairB)
 				{
 					return false;
 				} 
 
-				if (!Equals(ValueProp, PairA + ValueOffset, PairB + ValueOffset))
+				if (!Equals(ValueProp, PairA + ValueOffset, PairB + ValueOffset, OutDiff))
 				{
 					return false;
 				}
@@ -279,7 +305,7 @@ struct FUnversionedPropertyTest : public FUnversionedPropertyTestInput
 		return true;
 	}
 
-	static bool Equals(const FStructProperty* Property, const void* A, const void* B)
+	static bool Equals(const FStructProperty* Property, const void* A, const void* B, FPropertyDiff& OutDiff)
 	{
 		UScriptStruct* Struct = Property->Struct;
 		if (Struct->StructFlags & STRUCT_IdenticalNative)
@@ -287,7 +313,12 @@ struct FUnversionedPropertyTest : public FUnversionedPropertyTestInput
 			bool bResult = false;
 			if (Struct->GetCppStructOps()->Identical(A, B, EqualsPortFlags, bResult))
 			{
-				return bResult;
+				if (!bResult)
+				{
+					OutDiff = {Property, A, B, TEXT("native operator==")};
+					return false;
+				}
+				return true;
 			}
 		}
 
@@ -296,7 +327,7 @@ struct FUnversionedPropertyTest : public FUnversionedPropertyTestInput
 		{
 			for (int32 Idx = 0, MaxIdx = It->ArrayDim; Idx < MaxIdx; ++Idx)
 			{
-				if (!Equals_InContainer(*It, A, B, Idx))
+				if (!Equals_InContainer(*It, A, B, Idx, OutDiff))
 				{
 					return false;
 				}
@@ -305,39 +336,56 @@ struct FUnversionedPropertyTest : public FUnversionedPropertyTestInput
 		return true;
 	}
 
-	static bool Equals_InContainer(const FProperty* Property, const void* A, const void* B, uint32 Idx)
+	static bool Equals_InContainer(const FProperty* Property, const void* A, const void* B, uint32 Idx, FPropertyDiff& OutDiff)
 	{
-		return Equals(Property, Property->ContainerPtrToValuePtr<void>(A, Idx), Property->ContainerPtrToValuePtr<void>(B, Idx));
+		return Equals(Property, Property->ContainerPtrToValuePtr<void>(A, Idx), Property->ContainerPtrToValuePtr<void>(B, Idx), OutDiff);
 	}
 
-	static FString GetValueAsText(FProperty* Property, uint32 ArrayIdx, void* StructInstance)
+	static FString GetValueAsText(const FProperty* Property, uint32 ArrayIdx, const void* Instance)
 	{
 		FString Value;
-		Property->ExportText_InContainer(ArrayIdx, Value, StructInstance, nullptr, nullptr, 0);
-		return Value;
+		Property->ExportText_InContainer(ArrayIdx, Value, Instance, nullptr, nullptr, 0);
+		return MoveTemp(Value);
 	}
 
+	static FString GetValueAsLimitedText(const FProperty* Property, uint32 ArrayIdx, const void* Instance)
+	{
+		FString Value = GetValueAsText(Property, ArrayIdx, Instance);
+		if (Value.Len() > 100)
+		{
+			Value.LeftInline(100);
+			Value += TEXT(" ... shortened");
+		}
+		return MoveTemp(Value);
+	}
 
 	void CheckEqual(FProperty* Property, void* VersionedInstance, void* UnversionedInstance)
 	{
 		for (int32 Idx = 0, Num = Property->ArrayDim; Idx < Num; ++Idx)
 		{
-			if (!Equals_InContainer(Property, VersionedInstance, UnversionedInstance, Idx))
+			FPropertyDiff VersionedUnversionedDiff = {};
+			if (!Equals_InContainer(Property, VersionedInstance, UnversionedInstance, Idx, VersionedUnversionedDiff))
 			{
-				bool VersionedOk = Equals_InContainer(Property, VersionedInstance, OriginalInstance, Idx);
-				bool UnversionedOk = Equals_InContainer(Property, UnversionedInstance, OriginalInstance, Idx);
+				FPropertyDiff VersionedOriginalDiff = {};
+				FPropertyDiff UnversionedOriginalDiff = {};
+				bool VersionedOk = Equals_InContainer(Property, VersionedInstance, OriginalInstance, Idx, VersionedOriginalDiff);
+				bool UnversionedOk = Equals_InContainer(Property, UnversionedInstance, OriginalInstance, Idx, UnversionedOriginalDiff);
 				const TCHAR* OkPaths = VersionedOk ? (UnversionedOk ? TEXT("Both paths") : TEXT("Versioned path"))
 													: (UnversionedOk ? TEXT("Unversioned path") : TEXT("Neither path"));
 
-				// When debugging test failures, put a breakpoint inside this if statement
+				FString	VersionedValue = GetValueAsLimitedText(VersionedUnversionedDiff.Property, Idx, VersionedUnversionedDiff.A);
+				FString	UnversionedValue = GetValueAsLimitedText(VersionedUnversionedDiff.Property, Idx, VersionedUnversionedDiff.B);
+				const FPropertyDiff& OriginalDiff = VersionedOk ? UnversionedOriginalDiff : VersionedOriginalDiff;
+				FString	OriginalValue = OriginalDiff.Property == VersionedUnversionedDiff.Property ? GetValueAsText(OriginalDiff.Property, Idx, OriginalDiff.B) : "(missing)";
+
 				if (FPlatformMisc::IsDebuggerPresent())
 				{
 					// These strings might be too long to fit in the assert message.
-					// We could write traversal code that identifies which nested property differs
-					// and only generates a text representation for that value. 
-					FString VersionedValue = GetValueAsText(Property, Idx, VersionedInstance);
-					FString UnversionedValue = GetValueAsText(Property, Idx, UnversionedInstance);
-					FString OriginalValue = GetValueAsText(Property, Idx, OriginalInstance);
+					FString EntireVersionedValue = GetValueAsText(Property, Idx, VersionedInstance);
+					FString EntireUnversionedValue = GetValueAsText(Property, Idx, UnversionedInstance);
+					FString EntireOriginalValue = GetValueAsText(Property, Idx, OriginalInstance);
+
+					// Put a breakpoint here if you need to debug UPS and/or TPS roundtripping
 
 					FSaveResult VersionedSaved2 = Save(EPath::Versioned);
 					FSaveResult UnversionedSaved2 = Save(EPath::Unversioned);
@@ -346,8 +394,11 @@ struct FUnversionedPropertyTest : public FUnversionedPropertyTestInput
 					FTestInstance UnversionedLoaded2 = Load(UnversionedSaved2);
 				}		
 
-				checkf(false, TEXT("The %s %s.%s roundtripped differently in versioned vs unversioned tagged property serialization. "
-					"%s loaded an instance equal to the original."), *Property->GetClass()->GetName(), *Struct->GetName(), *Property->GetName(), OkPaths);
+				checkf(false, TEXT("The %s %s.%s roundtripped differently in versioned / tagged vs unversioned property serialization. "
+					"%s loaded an instance equal to the original. "
+					"Inner mismatch in %s for the %s %s with UPS/TPS/Original values %s/%s/%s"), 
+					*Property->GetClass()->GetName(), *Struct->GetName(), *Property->GetName(), OkPaths,
+					VersionedUnversionedDiff.MismatchKind, *VersionedUnversionedDiff.GetType(), *VersionedUnversionedDiff.Property->GetName(), *VersionedValue, *UnversionedValue, *OriginalValue);
 			}
 		}		
 	}
diff --git a/Engine/Source/Runtime/MovieScene/Public/Evaluation/MovieSceneEvaluationTemplate.h b/Engine/Source/Runtime/MovieScene/Public/Evaluation/MovieSceneEvaluationTemplate.h
--- a/Engine/Source/Runtime/MovieScene/Public/Evaluation/MovieSceneEvaluationTemplate.h
+++ b/Engine/Source/Runtime/MovieScene/Public/Evaluation/MovieSceneEvaluationTemplate.h
@@ -78,7 +78,7 @@ struct FMovieSceneTrackFieldData
 	/** Only called for serialization. Returns false to always serialize. */
 	bool Identical(const FMovieSceneTrackFieldData* Other, uint32 PortFlags) const
 	{
-		return false;
+		return Field == Other->Field;
 	}
 
 	TMovieSceneEvaluationTree<FMovieSceneTrackIdentifier> Field;
@@ -133,7 +133,7 @@ struct FMovieSceneSubSectionFieldData
 	/** Only called for serialization. Returns false to always serialize. */
 	bool Identical(const FMovieSceneSubSectionFieldData* Other, uint32 PortFlags) const
 	{
-		return false;
+		return Field == Other->Field;
 	}
 
 	TMovieSceneEvaluationTree<FMovieSceneSubSectionData> Field;
diff --git a/Engine/Source/Runtime/MovieScene/Public/Evaluation/MovieSceneEvaluationTrack.h b/Engine/Source/Runtime/MovieScene/Public/Evaluation/MovieSceneEvaluationTrack.h
--- a/Engine/Source/Runtime/MovieScene/Public/Evaluation/MovieSceneEvaluationTrack.h
+++ b/Engine/Source/Runtime/MovieScene/Public/Evaluation/MovieSceneEvaluationTrack.h
@@ -46,7 +46,7 @@ struct FSectionEvaluationDataTree
 	/** Only called for serialization. Returns false to always serialize. */
 	bool Identical(const FSectionEvaluationDataTree* Other, uint32 PortFlags) const
 	{
-		return false;
+		return Tree == Other->Tree;
 	}
 
 	TMovieSceneEvaluationTree<FSectionEvaluationData> Tree;
