Fixed soft object paths with a class name (export path). Replaced call to TStringView::Right with a call to TStringView::Mid as the parameter is a position in the string.

diff --git a/Engine/Source/Runtime/CoreUObject/Private/Misc/PackageName.cpp b/Engine/Source/Runtime/CoreUObject/Private/Misc/PackageName.cpp
--- a/Engine/Source/Runtime/CoreUObject/Private/Misc/PackageName.cpp
+++ b/Engine/Source/Runtime/CoreUObject/Private/Misc/PackageName.cpp
@@ -1477,7 +1477,7 @@ bool FPackageName::ParseExportTextPath(FStringView InExportTextPath, FStringView
 
 		if (OutObjectPath)
 		{
-			*OutObjectPath = InExportTextPath.Right(Index + 1);
+			*OutObjectPath = InExportTextPath.Mid(Index + 1);
 			OutObjectPath->RemoveSuffix(InExportTextPath.EndsWith('\''));
 		}
 
