FKismetCompilerContext::ValidateVariableNames() will now look through both UFields and FProperties when looking for conflicting names

diff --git a/Engine/Source/Editor/KismetCompiler/Private/KismetCompiler.cpp b/Engine/Source/Editor/KismetCompiler/Private/KismetCompiler.cpp
--- a/Engine/Source/Editor/KismetCompiler/Private/KismetCompiler.cpp
+++ b/Engine/Source/Editor/KismetCompiler/Private/KismetCompiler.cpp
@@ -571,10 +571,11 @@ void FKismetCompilerContext::ValidateVariableNames()
 			}
 			else if (ParentClass->IsNative()) // the above case handles when the parent is a blueprint
 			{
-				if (UField* ExisingField = FindField<UField>(ParentClass, *VarNameStr))
+				FFieldVariant ExisingField = FindUFieldOrFProperty(ParentClass, *VarNameStr);
+				if (ExisingField.IsValid())
 				{
 					UE_LOG(LogK2Compiler, Warning, TEXT("ValidateVariableNames name %s (used in %s) is already taken by %s")
-						, *VarNameStr, *Blueprint->GetPathName(), *ExisingField->GetPathName());
+						, *VarNameStr, *Blueprint->GetPathName(), *ExisingField.GetPathName());
 					NewVarName = FBlueprintEditorUtils::FindUniqueKismetName(Blueprint, VarNameStr);
 				}
 			}
diff --git a/Engine/Source/Runtime/CoreUObject/Public/UObject/UnrealType.h b/Engine/Source/Runtime/CoreUObject/Public/UObject/UnrealType.h
--- a/Engine/Source/Runtime/CoreUObject/Public/UObject/UnrealType.h
+++ b/Engine/Source/Runtime/CoreUObject/Public/UObject/UnrealType.h
@@ -5566,6 +5566,26 @@ template <class T> T* FindField( const UStruct* Owner, const TCHAR* FieldName )
 	return FindField<T>(Owner, Name);
 }
 
+/** Finds FProperties or UFunctions and UEnums */
+inline FFieldVariant FindUFieldOrFProperty(const UStruct* Owner, FName FieldName)
+{
+	// Look for properties first as they're most often the runtime thing higher level code wants to find
+	FFieldVariant Result = FindField<FProperty>(Owner, FieldName);
+	if (!Result.IsValid())
+	{
+		Result = FindField<UField>(Owner, FieldName);
+	}
+	return Result;
+}
+
+/** Finds FProperties or UFunctions and UEnums */
+inline FFieldVariant FindUFieldOrFProperty(const UStruct* Owner, const TCHAR* FieldName)
+{
+	// lookup the string name in the Name hash
+	FName Name(FieldName, FNAME_Find);
+	return FindUFieldOrFProperty(Owner, Name);
+}
+
 /**
  * Search for the named field within the specified scope, including any Outer classes; assert on failure.
  *
