- Fixed a non-unity issue
#rnx

diff --git a/Engine/Source/Runtime/CoreUObject/Private/Serialization/AsyncLoading2.cpp b/Engine/Source/Runtime/CoreUObject/Private/Serialization/AsyncLoading2.cpp
--- a/Engine/Source/Runtime/CoreUObject/Private/Serialization/AsyncLoading2.cpp
+++ b/Engine/Source/Runtime/CoreUObject/Private/Serialization/AsyncLoading2.cpp
@@ -53,6 +53,7 @@
 #include "Serialization/LargeMemoryReader.h"
 #include "UObject/UObjectClusters.h"
 #include "ProfilingDebugging/CountersTrace.h"
+#include "HAL/LowLevelMemStats.h"
 
 #if UE_BUILD_DEVELOPMENT || UE_BUILD_DEBUG
 //PRAGMA_DISABLE_OPTIMIZATION
