Reducing the constant memory cost of keeping FProperty related references from 32 to 8 bytes per UStruct.

- Merged ScriptObjectReferences and PropertyObjectReferences into single ScriptAndPropertyObjectReferences array
- UnresolvedScriptProperties will now be a heap allocated array since it only needs to exist at load time

diff --git a/Engine/Source/Editor/KismetCompiler/Private/KismetCompiler.cpp b/Engine/Source/Editor/KismetCompiler/Private/KismetCompiler.cpp
--- a/Engine/Source/Editor/KismetCompiler/Private/KismetCompiler.cpp
+++ b/Engine/Source/Editor/KismetCompiler/Private/KismetCompiler.cpp
@@ -4304,7 +4304,7 @@ void FKismetCompilerContext::CompileFunctions(EInternalCompilerFlags InternalFla
 			}
 		}
 
-		// Fill ScriptObjectReferences arrays in functions
+		// Fill ScriptAndPropertyObjectReferences arrays in functions
 		if (bIsFullCompile && (0 == MessageLog.NumErrors)) // Backend_VM can generate errors, so bGenerateStubsOnly cannot be reused
 		{
 			for (FKismetFunctionContext& FunctionContext : FunctionList)
@@ -4312,8 +4312,7 @@ void FKismetCompilerContext::CompileFunctions(EInternalCompilerFlags InternalFla
 				if (FunctionContext.IsValid())
 				{
 					UFunction* Function = FunctionContext.Function; 
-					ensure(0 == Function->ScriptObjectReferences.Num());
-					FArchiveScriptReferenceCollector ObjRefCollector(Function->ScriptObjectReferences);
+					FArchiveScriptReferenceCollector ObjRefCollector(Function->ScriptAndPropertyObjectReferences);
 
 					for (int32 iCode = 0; iCode < Function->Script.Num();)
 					{
diff --git a/Engine/Source/Editor/KismetCompiler/Private/UserDefinedStructureCompilerUtils.cpp b/Engine/Source/Editor/KismetCompiler/Private/UserDefinedStructureCompilerUtils.cpp
--- a/Engine/Source/Editor/KismetCompiler/Private/UserDefinedStructureCompilerUtils.cpp
+++ b/Engine/Source/Editor/KismetCompiler/Private/UserDefinedStructureCompilerUtils.cpp
@@ -162,7 +162,7 @@ struct FUserDefinedStructureCompilerInner
 			StructToClean->DestroyChildPropertiesAndResetPropertyLinks();
 			StructToClean->Script.Empty();
 			StructToClean->MinAlignment = 0;
-			StructToClean->ScriptObjectReferences.Empty();
+			StructToClean->ScriptAndPropertyObjectReferences.Empty();
 			StructToClean->ErrorMessage.Empty();
 			StructToClean->SetStructTrashed(true);
 		}
diff --git a/Engine/Source/Runtime/CoreUObject/Private/UObject/Class.cpp b/Engine/Source/Runtime/CoreUObject/Private/UObject/Class.cpp
--- a/Engine/Source/Runtime/CoreUObject/Private/UObject/Class.cpp
+++ b/Engine/Source/Runtime/CoreUObject/Private/UObject/Class.cpp
@@ -632,10 +632,11 @@ UStruct::UStruct(EStaticConstructor, int32 InSize, int32 InMinAlignment, EObject
 	, ChildProperties(nullptr)
 	, PropertiesSize(InSize)
 	, MinAlignment(InMinAlignment)
-	, PropertyLink(NULL)
-	, RefLink(NULL)
-	, DestructorLink(NULL)
-	, PostConstructLink(NULL)
+	, PropertyLink(nullptr)
+	, RefLink(nullptr)
+	, DestructorLink(nullptr)
+	, PostConstructLink(nullptr)
+	, UnresolvedScriptProperties(nullptr)
 {
 }
 
@@ -650,6 +651,7 @@ UStruct::UStruct(UStruct* InSuperStruct, SIZE_T ParamsSize, SIZE_T Alignment)
 	, RefLink(nullptr)
 	, DestructorLink(nullptr)
 	, PostConstructLink(nullptr)
+	, UnresolvedScriptProperties(nullptr)
 {
 #if USTRUCT_FAST_ISCHILDOF_IMPL == USTRUCT_ISCHILDOF_STRUCTARRAY
 	this->ReinitializeBaseChainArray();
@@ -667,6 +669,7 @@ UStruct::UStruct(const FObjectInitializer& ObjectInitializer, UStruct* InSuperSt
 	, RefLink(nullptr)
 	, DestructorLink(nullptr)
 	, PostConstructLink(nullptr)
+	, UnresolvedScriptProperties(nullptr)
 {
 #if USTRUCT_FAST_ISCHILDOF_IMPL == USTRUCT_ISCHILDOF_STRUCTARRAY
 	this->ReinitializeBaseChainArray();
@@ -944,23 +947,23 @@ void UStruct::Link(FArchive& Ar, bool bRelinkExistingProperties)
 		{
 			CurrentField->AddReferencedObjects(PropertyReferenceCollector);
 		}
-		PropertyObjectReferences = PropertyReferenceCollector.UniqueReferences.Array();
+		ScriptAndPropertyObjectReferences.Append(PropertyReferenceCollector.UniqueReferences.Array());
 
 #if USE_CIRCULAR_DEPENDENCY_LOAD_DEFERRING
 		// The old (non-EDL) FLinkerLoad code paths create placeholder objects
 		// for classes and functions. We have to babysit these, just as we do
 		// for bytecode references (reusing the AddReferencingScriptExpr fn).
 		// Long term we should not use placeholder objects like this:
-		for(int32 I = 0; I < PropertyObjectReferences.Num(); ++I)
+		for(int32 I = 0; I < ScriptAndPropertyObjectReferences.Num(); ++I)
 		{
-			if (ULinkerPlaceholderClass* PlaceholderObj = Cast<ULinkerPlaceholderClass>(PropertyObjectReferences[I]))
+			if (ULinkerPlaceholderClass* PlaceholderObj = Cast<ULinkerPlaceholderClass>(ScriptAndPropertyObjectReferences[I]))
 			{
 				// let the placeholder track the reference to it:
-				PlaceholderObj->AddReferencingScriptExpr(reinterpret_cast<UClass**>(&PropertyObjectReferences[I]));
+				PlaceholderObj->AddReferencingScriptExpr(reinterpret_cast<UClass**>(&ScriptAndPropertyObjectReferences[I]));
 			}
 			// I don't currently see how placeholder functions could be present in this list, but that's
 			// a dangerous assumption.
-			ensure(!(PropertyObjectReferences[I]->IsA<ULinkerPlaceholderFunction>()));
+			ensure(!(ScriptAndPropertyObjectReferences[I]->IsA<ULinkerPlaceholderFunction>()));
 		}
 #endif // USE_CIRCULAR_DEPENDENCY_LOAD_DEFERRING
 	}
@@ -1610,6 +1613,7 @@ UStruct::~UStruct()
 	// Also, Blueprint generated classes can have DestroyNonNativeProperties called on them after their FinishDestroy has been called
 	// so properties can only be deleted in the destructor
 	DestroyChildPropertiesAndResetPropertyLinks();
+	DeleteUnresolvedScriptProperties();
 }
 
 IMPLEMENT_FSTRUCTUREDARCHIVE_SERIALIZER(UStruct);
@@ -1929,21 +1933,24 @@ void UStruct::PostLoad()
 	Super::PostLoad();
 
 	// Finally try to resolve all script properties that couldn't be resolved at load time
-	for (TPair<TFieldPath<FField>, int32>& MissingProperty : UnresolvedScriptProperties)
+	if (UnresolvedScriptProperties)
 	{
-		FField* ResolvedProperty = MissingProperty.Key.Get(this);
-		if (ResolvedProperty)
-		{			
-			check((int32)Script.Num() >= (int32)(MissingProperty.Value + sizeof(FField*)));
-			FField** TargetScriptPropertyPtr = (FField**)(Script.GetData() + MissingProperty.Value);
-			*TargetScriptPropertyPtr = ResolvedProperty;
-		}
-		else if (!MissingProperty.Key.IsEmpty())
+		for (TPair<TFieldPath<FField>, int32>& MissingProperty : *UnresolvedScriptProperties)
 		{
-			UE_LOG(LogClass, Warning, TEXT("Failed to resolve bytecode referenced field from path: %s when loading %s"), *MissingProperty.Key.ToString(), *GetFullName());
+			FField* ResolvedProperty = MissingProperty.Key.Get(this);
+			if (ResolvedProperty)
+			{
+				check((int32)Script.Num() >= (int32)(MissingProperty.Value + sizeof(FField*)));
+				FField** TargetScriptPropertyPtr = (FField**)(Script.GetData() + MissingProperty.Value);
+				*TargetScriptPropertyPtr = ResolvedProperty;
+			}
+			else if (!MissingProperty.Key.IsEmpty())
+			{
+				UE_LOG(LogClass, Warning, TEXT("Failed to resolve bytecode referenced field from path: %s when loading %s"), *MissingProperty.Key.ToString(), *GetFullName());
+			}
 		}
+		DeleteUnresolvedScriptProperties();
 	}
-	UnresolvedScriptProperties.Empty();
 }
 
 void UStruct::AddReferencedObjects(UObject* InThis, FReferenceCollector& Collector)
@@ -1955,11 +1962,7 @@ void UStruct::AddReferencedObjects(UObject* InThis, FReferenceCollector& Collect
 		// Required by the unified GC when running in the editor
 		Collector.AddReferencedObject( This->SuperStruct, This );
 		Collector.AddReferencedObject( This->Children, This );
-
-		for( int32 Index = 0; Index < This->ScriptObjectReferences.Num(); Index++ )
-		{
-			Collector.AddReferencedObject( This->ScriptObjectReferences[ Index ], This );
-		}
+		Collector.AddReferencedObjects(This->ScriptAndPropertyObjectReferences, This);
 	}
 #endif
 #if WITH_EDITORONLY_DATA
@@ -2125,8 +2128,7 @@ IMPLEMENT_CORE_INTRINSIC_CLASS(UStruct, UField,
 		// Note: None of the *Link members need to be emitted, as they only contain properties
 		// that are in the Children chain or SuperStruct->Children chains.
 
-		Class->EmitObjectArrayReference(STRUCT_OFFSET(UStruct, ScriptObjectReferences), TEXT("ScriptObjectReferences"));
-		Class->EmitObjectArrayReference(STRUCT_OFFSET(UStruct, PropertyObjectReferences), TEXT("PropertyObjectReferences"));		
+		Class->EmitObjectArrayReference(STRUCT_OFFSET(UStruct, ScriptAndPropertyObjectReferences), TEXT("ScriptAndPropertyObjectReferences"));
 	}
 );
 
@@ -4633,26 +4635,25 @@ void UClass::PurgeClass(bool bRecompilingOnLoad)
 	}
 #endif
 
-	ClassDefaultObject = NULL;
+	ClassDefaultObject = nullptr;
 
 	Interfaces.Empty();
 	NativeFunctionLookupTable.Empty();
-	SetSuperStruct(NULL);
-	Children = NULL;
+	SetSuperStruct(nullptr);
+	Children = nullptr;
 	Script.Empty();
 	MinAlignment = 0;
-	RefLink = NULL;
-	PropertyLink = NULL;
-	DestructorLink = NULL;
-	ClassAddReferencedObjects = NULL;
+	RefLink = nullptr;
+	PropertyLink = nullptr;
+	DestructorLink = nullptr;
+	ClassAddReferencedObjects = nullptr;
 
-	ScriptObjectReferences.Empty();
-	PropertyObjectReferences.Empty();
-	UnresolvedScriptProperties.Empty();
+	ScriptAndPropertyObjectReferences.Empty();
+	DeleteUnresolvedScriptProperties();
 
 	FuncMap.Empty();
 	ClearFunctionMapsCaches();
-	PropertyLink = NULL;
+	PropertyLink = nullptr;
 
 #if WITH_EDITORONLY_DATA
 	{
diff --git a/Engine/Source/Runtime/CoreUObject/Private/UObject/StructScriptLoader.cpp b/Engine/Source/Runtime/CoreUObject/Private/UObject/StructScriptLoader.cpp
--- a/Engine/Source/Runtime/CoreUObject/Private/UObject/StructScriptLoader.cpp
+++ b/Engine/Source/Runtime/CoreUObject/Private/UObject/StructScriptLoader.cpp
@@ -266,23 +266,29 @@ bool FStructScriptLoader::LoadStructWithScript(UStruct* DestScriptContainer, FAr
 
 	int32 BytecodeIndex = 0;
 	{
-		DestScriptContainer->UnresolvedScriptProperties.Empty();
-
 		FPropertyProxyArchive PropertyAr(Ar, BytecodeIndex, DestScriptContainer);
 		while (BytecodeIndex < BytecodeBufferSize)
 		{
 			DestScriptContainer->SerializeExpr(BytecodeIndex, PropertyAr);
 		}
 
-		DestScriptContainer->UnresolvedScriptProperties = MoveTemp(PropertyAr.UnresolvedProperties);
+		if (PropertyAr.UnresolvedProperties.Num())
+		{
+			DestScriptContainer->SetUnresolvedScriptProperties(PropertyAr.UnresolvedProperties);
+		}
+		else
+		{
+			// Make sure there's no stale properties in the array
+			DestScriptContainer->DeleteUnresolvedScriptProperties();
+		}
 	}
 	ensure(ScriptEndOffset == Ar.Tell());
 	checkf(BytecodeIndex == BytecodeBufferSize, TEXT("'%s' script expression-count mismatch; Expected: %i, Got: %i"), *DestScriptContainer->GetName(), BytecodeBufferSize, BytecodeIndex);
 
 	if (!GUObjectArray.IsDisregardForGC(DestScriptContainer))
 	{
-		DestScriptContainer->ScriptObjectReferences.Empty();
-		FArchiveScriptReferenceCollector ObjRefCollector(DestScriptContainer->ScriptObjectReferences);
+		DestScriptContainer->ScriptAndPropertyObjectReferences.Empty();
+		FArchiveScriptReferenceCollector ObjRefCollector(DestScriptContainer->ScriptAndPropertyObjectReferences);
 
 		BytecodeIndex = 0;
 		while (BytecodeIndex < BytecodeBufferSize)
@@ -305,5 +311,5 @@ int32 FStructScriptLoader::ResolveDeferredScriptLoads(FLinkerLoad* Linker)
 void FStructScriptLoader::ClearScriptCode(UStruct* ScriptContainer)
 {
 	ScriptContainer->Script.Empty(BytecodeBufferSize);
-	ScriptContainer->ScriptObjectReferences.Empty();
+	ScriptContainer->ScriptAndPropertyObjectReferences.Empty();
 }
diff --git a/Engine/Source/Runtime/CoreUObject/Private/UObject/StructScriptLoader.h b/Engine/Source/Runtime/CoreUObject/Private/UObject/StructScriptLoader.h
--- a/Engine/Source/Runtime/CoreUObject/Private/UObject/StructScriptLoader.h
+++ b/Engine/Source/Runtime/CoreUObject/Private/UObject/StructScriptLoader.h
@@ -76,7 +76,7 @@ public:
 private:
 	/**
 	 * Empties the specified struct's Script array, as well as its 
-	 * ScriptObjectReferences container.
+	 * ScriptAndPropertyObjectReferences container.
 	 * 
 	 * @param  ScriptContainer    The struct you want emptied.
 	 */
diff --git a/Engine/Source/Runtime/CoreUObject/Public/UObject/Class.h b/Engine/Source/Runtime/CoreUObject/Public/UObject/Class.h
--- a/Engine/Source/Runtime/CoreUObject/Public/UObject/Class.h
+++ b/Engine/Source/Runtime/CoreUObject/Public/UObject/Class.h
@@ -317,14 +317,12 @@ public:
 	/** In memory only: Linked list of properties requiring post constructor initialization */
 	FProperty* PostConstructLink;
 
-	/** Array of object references embedded in script code. Mirrored for easy access by realtime garbage collection code */
-	TArray<UObject*> ScriptObjectReferences;
-
-	/** Array of object references from struct properties. Mirrored here for fast access by the garbage collector */
-	TArray<UObject*> PropertyObjectReferences;
+	/** Array of object references embedded in script code and referenced by FProperties. Mirrored for easy access by realtime garbage collection code */
+	TArray<UObject*> ScriptAndPropertyObjectReferences;
 
+	typedef TArray<TPair<TFieldPath<FField>, int32>> FUnresolvedScriptPropertiesArray;
 	/** Contains a list of script properties that couldn't be resolved at load time */
-	TArray<TPair<TFieldPath<FField>, int32>> UnresolvedScriptProperties;
+	FUnresolvedScriptPropertiesArray* UnresolvedScriptProperties;
 
 #if WITH_EDITORONLY_DATA
 	/** List of wrapper UObjects for FProperties */
@@ -528,6 +526,26 @@ public:
 	friend struct FStructArchiveProxy;
 #endif // HACK_HEADER_GENERATOR
 
+	/** Sets the UnresolvedScriptProperties array */
+	void SetUnresolvedScriptProperties(FUnresolvedScriptPropertiesArray& InUnresolvedProperties)
+	{
+		if (!UnresolvedScriptProperties)
+		{
+			UnresolvedScriptProperties = new FUnresolvedScriptPropertiesArray();
+		}
+		*UnresolvedScriptProperties = MoveTemp(InUnresolvedProperties);
+	}
+
+	/** Deletes the UnresolvedScriptProperties array */
+	FORCEINLINE void DeleteUnresolvedScriptProperties()
+	{
+		if (UnresolvedScriptProperties)
+		{
+			delete UnresolvedScriptProperties;
+			UnresolvedScriptProperties = nullptr;
+		}
+	}
+
 protected:
 
 	/** Returns the property name from the guid */
diff --git a/Engine/Source/Runtime/CoreUObject/Public/UObject/PropertyProxyArchive.h b/Engine/Source/Runtime/CoreUObject/Public/UObject/PropertyProxyArchive.h
--- a/Engine/Source/Runtime/CoreUObject/Public/UObject/PropertyProxyArchive.h
+++ b/Engine/Source/Runtime/CoreUObject/Public/UObject/PropertyProxyArchive.h
@@ -21,7 +21,7 @@ class FPropertyProxyArchive : public FArchiveProxy
 public:
 
 	/** List of fields that could not be resolved at serialize time possibly due to their classes not being fully loaded yet */
-	TArray< TPair< TFieldPath<FField>, int32> > UnresolvedProperties;
+	UStruct::FUnresolvedScriptPropertiesArray UnresolvedProperties;
 	/** Current bytecode offset */
 	int32& BytecodeIndex;
 	/** Script container object */
