Reducing the constant memory cost of keeping FProperty related references from 32 to 8 bytes per UStruct.

- Merged ScriptObjectReferences and PropertyObjectReferences into single ScriptAndPropertyObjectReferences array
- UnresolvedScriptProperties will now be a heap allocated array since it only needs to exist at load time

#tests Cooked and launched PC client and server + 3 hours of bot soak tests in cooked buiild

diff --git a/Engine/Source/Runtime/CoreUObject/Private/UObject/Class.cpp b/Engine/Source/Runtime/CoreUObject/Private/UObject/Class.cpp
--- a/Engine/Source/Runtime/CoreUObject/Private/UObject/Class.cpp
+++ b/Engine/Source/Runtime/CoreUObject/Private/UObject/Class.cpp
@@ -1258,195 +1258,195 @@ void UStruct::SerializeVersionedTaggedProperties(FStructuredArchive::FSlot Slot,
 		else
 #endif // WITH_TEXT_ARCHIVE_SUPPORT
 		{
-			// Load tagged properties.
-			FStructuredArchive::FStream PropertiesStream = Slot.EnterStream();
+		// Load tagged properties.
+		FStructuredArchive::FStream PropertiesStream = Slot.EnterStream();
 
-			// This code assumes that properties are loaded in the same order they are saved in. This removes a n^2 search 
-			// and makes it an O(n) when properties are saved in the same order as they are loaded (default case). In the 
-			// case that a property was reordered the code falls back to a slower search.
+		// This code assumes that properties are loaded in the same order they are saved in. This removes a n^2 search 
+		// and makes it an O(n) when properties are saved in the same order as they are loaded (default case). In the 
+		// case that a property was reordered the code falls back to a slower search.
 			FProperty*	Property = PropertyLink;
-			bool		bAdvanceProperty	= false;
-			int32		RemainingArrayDim	= Property ? Property->ArrayDim : 0;
+		bool		bAdvanceProperty	= false;
+		int32		RemainingArrayDim	= Property ? Property->ArrayDim : 0;
 
-			// Load all stored properties, potentially skipping unknown ones.
-			while (true)
-			{
-				FStructuredArchive::FRecord PropertyRecord = PropertiesStream.EnterElement().EnterRecord();
+		// Load all stored properties, potentially skipping unknown ones.
+		while (true)
+		{
+			FStructuredArchive::FRecord PropertyRecord = PropertiesStream.EnterElement().EnterRecord();
 
-				FPropertyTag Tag;
-				PropertyRecord << SA_VALUE(TEXT("Tag"), Tag);
+			FPropertyTag Tag;
+			PropertyRecord << SA_VALUE(TEXT("Tag"), Tag);
 
-				if (Tag.Name.IsNone())
-				{
-					break;
-				}
+			if (Tag.Name.IsNone())
+			{
+				break;
+			}
 
-				// Move to the next property to be serialized
-				if( bAdvanceProperty && --RemainingArrayDim <= 0 )
+			// Move to the next property to be serialized
+				if (bAdvanceProperty && --RemainingArrayDim <= 0)
+			{
+				Property = Property->PropertyLinkNext;
+				// Skip over properties that don't need to be serialized.
+					while (Property && !Property->ShouldSerializeValue(UnderlyingArchive))
 				{
 					Property = Property->PropertyLinkNext;
-					// Skip over properties that don't need to be serialized.
-					while (Property && !Property->ShouldSerializeValue(UnderlyingArchive))
-					{
-						Property = Property->PropertyLinkNext;
-					}
-					RemainingArrayDim = Property ? Property->ArrayDim : 0;
 				}
-				bAdvanceProperty = false;
+				RemainingArrayDim = Property ? Property->ArrayDim : 0;
+			}
+			bAdvanceProperty = false;
 
-				// Optionally resolve properties using Guid Property tags in non cooked builds that support it.
-				if (bArePropertyGuidsAvailable && Tag.HasPropertyGuid)
+			// Optionally resolve properties using Guid Property tags in non cooked builds that support it.
+			if (bArePropertyGuidsAvailable && Tag.HasPropertyGuid)
+			{
+				// Use property guids from blueprint generated classes to redirect serialised data.
+				FName Result = FindPropertyNameFromGuid(Tag.PropertyGuid);
+				if (Result != NAME_None && Tag.Name != Result)
 				{
-					// Use property guids from blueprint generated classes to redirect serialised data.
-					FName Result = FindPropertyNameFromGuid(Tag.PropertyGuid);
-					if (Result != NAME_None && Tag.Name != Result)
-					{
-						Tag.Name = Result;
-					}
+					Tag.Name = Result;
 				}
-				// If this property is not the one we expect (e.g. skipped as it matches the default value), do the brute force search.
+			}
+			// If this property is not the one we expect (e.g. skipped as it matches the default value), do the brute force search.
 				if (Property == nullptr || Property->GetFName() != Tag.Name)
-				{
-					// No need to check redirects on platforms where everything is cooked. Always check for save games
+			{
+				// No need to check redirects on platforms where everything is cooked. Always check for save games
 					if (bUseRedirects && !UnderlyingArchive.HasAnyPortFlags(PPF_DuplicateForPIE | PPF_Duplicate))
+				{
+					for (UStruct* CheckStruct = GetOwnerStruct(); CheckStruct; CheckStruct = CheckStruct->GetSuperStruct())
 					{
-						for (UStruct* CheckStruct = GetOwnerStruct(); CheckStruct; CheckStruct = CheckStruct->GetSuperStruct())
-						{
 							FName NewTagName = FProperty::FindRedirectedPropertyName(CheckStruct, Tag.Name);
-							if (!NewTagName.IsNone())
-							{
-								Tag.Name = NewTagName;
-								break;
-							}
+						if (!NewTagName.IsNone())
+						{
+							Tag.Name = NewTagName;
+							break;
 						}
 					}
+				}
 
 					FProperty* CurrentProperty = Property;
-					// Search forward...
+				// Search forward...
 					for (; Property; Property = Property->PropertyLinkNext)
-					{
+				{
 						if (Property->GetFName() == Tag.Name)
-						{
-							break;
-						}
+					{
+						break;
 					}
-					// ... and then search from the beginning till we reach the current property if it's not found.
+				}
+				// ... and then search from the beginning till we reach the current property if it's not found.
 					if (Property == nullptr)
-					{
+				{
 						for (Property = PropertyLink; Property && Property != CurrentProperty; Property = Property->PropertyLinkNext)
-						{
+					{
 							if (Property->GetFName() == Tag.Name)
-							{
-								break;
-							}
-						}
-
-						if (Property == CurrentProperty)
 						{
-							// Property wasn't found.
-							Property = nullptr;
+							break;
 						}
 					}
 
-					RemainingArrayDim = Property ? Property->ArrayDim : 0;
+						if (Property == CurrentProperty)
+					{
+						// Property wasn't found.
+						Property = nullptr;
+					}
 				}
 
-				const int64 StartOfProperty = UnderlyingArchive.Tell();
+				RemainingArrayDim = Property ? Property->ArrayDim : 0;
+			}
 
-				if (!Property)
-				{
-					Property = CustomFindProperty(Tag.Name);
-				}
+			const int64 StartOfProperty = UnderlyingArchive.Tell();
 
-				if (Property)
-				{
-					FName PropID = Property->GetID();
+			if (!Property)
+			{
+				Property = CustomFindProperty(Tag.Name);
+			}
 
-					// Check if this is a struct property and we have a redirector
-					// No need to check redirects on platforms where everything is cooked. Always check for save games
-					if (bUseRedirects)
+			if (Property)
+			{
+				FName PropID = Property->GetID();
+
+				// Check if this is a struct property and we have a redirector
+				// No need to check redirects on platforms where everything is cooked. Always check for save games
+				if (bUseRedirects)
+				{
+					if (Tag.Type == NAME_StructProperty && PropID == NAME_StructProperty)
 					{
-						if (Tag.Type == NAME_StructProperty && PropID == NAME_StructProperty)
-						{
-							const FName NewName = FLinkerLoad::FindNewNameForStruct(Tag.StructName);
+						const FName NewName = FLinkerLoad::FindNewNameForStruct(Tag.StructName);
 							const FName StructName = CastFieldChecked<FStructProperty>(Property)->Struct->GetFName();
-							if (NewName == StructName)
-							{
-								Tag.StructName = NewName;
-							}
+						if (NewName == StructName)
+						{
+							Tag.StructName = NewName;
 						}
-						else if ((PropID == NAME_EnumProperty) && ((Tag.Type == NAME_EnumProperty) || (Tag.Type == NAME_ByteProperty)))
+					}
+					else if ((PropID == NAME_EnumProperty) && ((Tag.Type == NAME_EnumProperty) || (Tag.Type == NAME_ByteProperty)))
+					{
+						const FName NewName = FLinkerLoad::FindNewNameForEnum(Tag.EnumName);
+						if (!NewName.IsNone())
 						{
-							const FName NewName = FLinkerLoad::FindNewNameForEnum(Tag.EnumName);
-							if (!NewName.IsNone())
-							{
-								Tag.EnumName = NewName;
-							}
+							Tag.EnumName = NewName;
 						}
 					}
+				}
 
 #if WITH_EDITOR
-					if (BreakRecursionIfFullyLoad && BreakRecursionIfFullyLoad->HasAllFlags(RF_LoadCompleted))
-					{
-					}
-					else
+				if (BreakRecursionIfFullyLoad && BreakRecursionIfFullyLoad->HasAllFlags(RF_LoadCompleted))
+				{
+				}
+				else
 #endif // WITH_EDITOR
-					// editoronly properties should be skipped if we are NOT the editor, or we are 
-					// the editor but are cooking for console (editoronly implies notforconsole)
-					if ((Property->PropertyFlags & CPF_EditorOnly) && !FPlatformProperties::HasEditorOnlyData() && !GForceLoadEditorOnly)
-					{
-					}
-					// check for valid array index
-					else if (Tag.ArrayIndex >= Property->ArrayDim || Tag.ArrayIndex < 0)
-					{
-						UE_LOG(LogClass, Warning, TEXT("Array bound exceeded (var %s=%d, exceeds %s [0-%d] in package:  %s"),
-							*Tag.Name.ToString(), Tag.ArrayIndex, *GetName(), Property->ArrayDim - 1, *UnderlyingArchive.GetArchiveName());
-					}
-					else if (!Property->ShouldSerializeValue(UnderlyingArchive))
-					{
-						UE_CLOG((UnderlyingArchive.IsPersistent() && FPlatformProperties::RequiresCookedData()), LogClass, Warning, TEXT("Skipping saved property %s of %s since it is no longer serializable for asset:  %s. (Maybe resave asset?)"), *Tag.Name.ToString(), *GetName(), *UnderlyingArchive.GetArchiveName());
-					}
-					else
+				// editoronly properties should be skipped if we are NOT the editor, or we are 
+				// the editor but are cooking for console (editoronly implies notforconsole)
+				if ((Property->PropertyFlags & CPF_EditorOnly) && !FPlatformProperties::HasEditorOnlyData() && !GForceLoadEditorOnly)
+				{
+				}
+				// check for valid array index
+				else if (Tag.ArrayIndex >= Property->ArrayDim || Tag.ArrayIndex < 0)
+				{
+					UE_LOG(LogClass, Warning, TEXT("Array bound exceeded (var %s=%d, exceeds %s [0-%d] in package:  %s"),
+						*Tag.Name.ToString(), Tag.ArrayIndex, *GetName(), Property->ArrayDim - 1, *UnderlyingArchive.GetArchiveName());
+				}
+				else if (!Property->ShouldSerializeValue(UnderlyingArchive))
+				{
+					UE_CLOG((UnderlyingArchive.IsPersistent() && FPlatformProperties::RequiresCookedData()), LogClass, Warning, TEXT("Skipping saved property %s of %s since it is no longer serializable for asset:  %s. (Maybe resave asset?)"), *Tag.Name.ToString(), *GetName(), *UnderlyingArchive.GetArchiveName());
+				}
+				else
+				{
+					FStructuredArchive::FSlot ValueSlot = PropertyRecord.EnterField(SA_FIELD_NAME(TEXT("Value")));
+
+					switch (Property->ConvertFromType(Tag, ValueSlot, Data, DefaultsStruct))
 					{
-						FStructuredArchive::FSlot ValueSlot = PropertyRecord.EnterField(SA_FIELD_NAME(TEXT("Value")));
+						case EConvertFromTypeResult::Converted:
+							bAdvanceProperty = true;
+							break;
 
-						switch (Property->ConvertFromType(Tag, ValueSlot, Data, DefaultsStruct))
-						{
-							case EConvertFromTypeResult::Converted:
-								bAdvanceProperty = true;
-								break;
-
-							case EConvertFromTypeResult::UseSerializeItem:
-								if (Tag.Type != PropID)
-								{
-									UE_LOG(LogClass, Warning, TEXT("Type mismatch in %s of %s - Previous (%s) Current(%s) for package:  %s"), *Tag.Name.ToString(), *GetName(), *Tag.Type.ToString(), *PropID.ToString(), *UnderlyingArchive.GetArchiveName());
-								}
-								else
-								{
-									uint8* DestAddress = Property->ContainerPtrToValuePtr<uint8>(Data, Tag.ArrayIndex);
-									uint8* DefaultsFromParent = Property->ContainerPtrToValuePtrForDefaults<uint8>(DefaultsStruct, Defaults, Tag.ArrayIndex);
-
-									// This property is ok.
-									Tag.SerializeTaggedProperty(ValueSlot, Property, DestAddress, DefaultsFromParent);
-									bAdvanceProperty = !UnderlyingArchive.IsCriticalError();
-								}
-								break;
-
-							case EConvertFromTypeResult::CannotConvert:
-								break;
-
-							default:
-								check(false);
-						}
+						case EConvertFromTypeResult::UseSerializeItem:
+							if (Tag.Type != PropID)
+							{
+								UE_LOG(LogClass, Warning, TEXT("Type mismatch in %s of %s - Previous (%s) Current(%s) for package:  %s"), *Tag.Name.ToString(), *GetName(), *Tag.Type.ToString(), *PropID.ToString(), *UnderlyingArchive.GetArchiveName());
+							}
+							else
+							{
+								uint8* DestAddress = Property->ContainerPtrToValuePtr<uint8>(Data, Tag.ArrayIndex);
+								uint8* DefaultsFromParent = Property->ContainerPtrToValuePtrForDefaults<uint8>(DefaultsStruct, Defaults, Tag.ArrayIndex);
+
+								// This property is ok.
+								Tag.SerializeTaggedProperty(ValueSlot, Property, DestAddress, DefaultsFromParent);
+								bAdvanceProperty = !UnderlyingArchive.IsCriticalError();
+							}
+							break;
+
+						case EConvertFromTypeResult::CannotConvert:
+							break;
+
+						default:
+							check(false);
 					}
 				}
+			}
 
 				int64 Loaded = UnderlyingArchive.Tell() - StartOfProperty;
 
-				if (!bAdvanceProperty)
-				{
-					UnderlyingArchive.Seek(StartOfProperty + Tag.Size);
-				}
+			if (!bAdvanceProperty)
+			{
+				UnderlyingArchive.Seek(StartOfProperty + Tag.Size);
+			}
 				else
 				{
 					check(Tag.Size == Loaded);
@@ -4035,7 +4035,7 @@ void UClass::SetUpRuntimeReplicationData()
 					NetProperties.Add(Prop);
 				}
 			}
-		}
+			}
 
 		for(TFieldIterator<UField> It(this,EFieldIteratorFlags::ExcludeSuper); It; ++It)
 		{
@@ -4056,20 +4056,20 @@ void UClass::SetUpRuntimeReplicationData()
 		const bool bIsNativeClass = HasAnyClassFlags(CLASS_Native);
 		if (!bIsNativeClass)
 		{
-			// Sort NetProperties so that their ClassReps are sorted by memory offset
+		// Sort NetProperties so that their ClassReps are sorted by memory offset
 			struct FComparePropertyOffsets
-			{
+		{
 				FORCEINLINE bool operator()(FProperty& A, FProperty& B) const
-				{
-					// Ensure stable sort
+			{
+				// Ensure stable sort
 					if (A.GetOffset_ForGC() == B.GetOffset_ForGC())
-					{
-						return A.GetName() < B.GetName();
-					}
-
-					return A.GetOffset_ForGC() < B.GetOffset_ForGC();
+				{
+					return A.GetName() < B.GetName();
 				}
-			};
+
+				return A.GetOffset_ForGC() < B.GetOffset_ForGC();
+			}
+		};
 
 			Sort(NetProperties.GetData(), NetProperties.Num(), FComparePropertyOffsets());
 		}
@@ -4738,9 +4738,9 @@ bool UClass::HasProperty(FProperty* InProperty) const
 	{
 		UClass* PropertiesClass = InProperty->GetOwner<UClass>();
 		if (PropertiesClass)
-		{
-			return IsChildOf(PropertiesClass);
-		}
+	{
+		return IsChildOf(PropertiesClass);
+	}
 	}
 
 	return false;
