Fix for null access when verifying ScriptAndPropertyObjectReferences

diff --git a/Engine/Source/Runtime/CoreUObject/Private/UObject/Class.cpp b/Engine/Source/Runtime/CoreUObject/Private/UObject/Class.cpp
--- a/Engine/Source/Runtime/CoreUObject/Private/UObject/Class.cpp
+++ b/Engine/Source/Runtime/CoreUObject/Private/UObject/Class.cpp
@@ -954,16 +954,24 @@ void UStruct::Link(FArchive& Ar, bool bRelinkExistingProperties)
 		// for classes and functions. We have to babysit these, just as we do
 		// for bytecode references (reusing the AddReferencingScriptExpr fn).
 		// Long term we should not use placeholder objects like this:
-		for(int32 I = 0; I < ScriptAndPropertyObjectReferences.Num(); ++I)
+		for(int32 ReferenceIndex = ScriptAndPropertyObjectReferences.Num() - 1; ReferenceIndex >= 0; --ReferenceIndex)
 		{
-			if (ULinkerPlaceholderClass* PlaceholderObj = Cast<ULinkerPlaceholderClass>(ScriptAndPropertyObjectReferences[I]))
+			if (ScriptAndPropertyObjectReferences[ReferenceIndex])
 			{
-				// let the placeholder track the reference to it:
-				PlaceholderObj->AddReferencingScriptExpr(reinterpret_cast<UClass**>(&ScriptAndPropertyObjectReferences[I]));
+				if (ULinkerPlaceholderClass* PlaceholderObj = Cast<ULinkerPlaceholderClass>(ScriptAndPropertyObjectReferences[ReferenceIndex]))
+				{
+					// let the placeholder track the reference to it:
+					PlaceholderObj->AddReferencingScriptExpr(reinterpret_cast<UClass**>(&ScriptAndPropertyObjectReferences[ReferenceIndex]));
+				}
+				// I don't currently see how placeholder functions could be present in this list, but that's
+				// a dangerous assumption.
+				ensure(!(ScriptAndPropertyObjectReferences[ReferenceIndex]->IsA<ULinkerPlaceholderFunction>()));
+			}
+			else
+			{
+				// It's possible that in the process of recompilation one of the refernces got GC'd leaving a null ptr in the array
+				ScriptAndPropertyObjectReferences.RemoveAt(ReferenceIndex);
 			}
-			// I don't currently see how placeholder functions could be present in this list, but that's
-			// a dangerous assumption.
-			ensure(!(ScriptAndPropertyObjectReferences[I]->IsA<ULinkerPlaceholderFunction>()));
 		}
 #endif // USE_CIRCULAR_DEPENDENCY_LOAD_DEFERRING
 	}
