Fixed wrong return type of rarely used ExportTextPathToObjectPath() overload

diff --git a/Engine/Source/Runtime/CoreUObject/Private/Misc/PackageName.cpp b/Engine/Source/Runtime/CoreUObject/Private/Misc/PackageName.cpp
--- a/Engine/Source/Runtime/CoreUObject/Private/Misc/PackageName.cpp
+++ b/Engine/Source/Runtime/CoreUObject/Private/Misc/PackageName.cpp
@@ -1516,9 +1516,9 @@ FString FPackageName::ExportTextPathToObjectPath(const FString& InExportTextPath
 	return InExportTextPath;
 }
 
-const TCHAR* FPackageName::ExportTextPathToObjectPath(const TCHAR* InExportTextPath)
+FString FPackageName::ExportTextPathToObjectPath(const TCHAR* InExportTextPath)
 {
-	return ExportTextPathToObjectPath(FStringView(InExportTextPath)).GetData();
+	return ExportTextPathToObjectPath(FString(InExportTextPath));
 }
 
 FString FPackageName::ObjectPathToPackageName(const FString& InObjectPath)
diff --git a/Engine/Source/Runtime/CoreUObject/Public/Misc/PackageName.h b/Engine/Source/Runtime/CoreUObject/Public/Misc/PackageName.h
--- a/Engine/Source/Runtime/CoreUObject/Public/Misc/PackageName.h
+++ b/Engine/Source/Runtime/CoreUObject/Public/Misc/PackageName.h
@@ -446,7 +446,7 @@ public:
 	 */
 	static FStringView	ExportTextPathToObjectPath(FStringView InExportTextPath);
 	static FString		ExportTextPathToObjectPath(const FString& InExportTextPath);
-	static const TCHAR*	ExportTextPathToObjectPath(const TCHAR* InExportTextPath);
+	static FString		ExportTextPathToObjectPath(const TCHAR* InExportTextPath);
 
 	/** 
 	 * Returns the name of the package referred to by the specified object path
