Changed FName(F[Ansi]StringView) constructors to be explicit

Avoid unintentional construction of names as this can be expensive and bloats the name table.

#rnx

diff --git a/Engine/Source/Runtime/Core/Public/Misc/StringBuilder.h b/Engine/Source/Runtime/Core/Public/Misc/StringBuilder.h
--- a/Engine/Source/Runtime/Core/Public/Misc/StringBuilder.h
+++ b/Engine/Source/Runtime/Core/Public/Misc/StringBuilder.h
@@ -490,6 +490,7 @@ inline FStringBuilderBase&					operator<<(FStringBuilderBase& Builder, const ANS
 CORE_API FStringBuilderBase&				operator<<(FStringBuilderBase& Builder, const FString& Str);
 
 inline FStringBuilderBase&					operator<<(FStringBuilderBase& Builder, const FStringView& Str)					{ return Builder.Append(Str); }
+inline FStringBuilderBase&					operator<<(FStringBuilderBase& Builder, const FAnsiStringView& Str)				{ return Builder.AppendAnsi(Str); }
 
 inline FStringBuilderBase&					operator<<(FStringBuilderBase& Builder, int32 Value)							{ return Builder.Appendf(TEXT("%d"), Value); }
 inline FStringBuilderBase&					operator<<(FStringBuilderBase& Builder, uint32 Value)							{ return Builder.Appendf(TEXT("%u"), Value); }
diff --git a/Engine/Source/Runtime/Core/Public/UObject/NameTypes.h b/Engine/Source/Runtime/Core/Public/UObject/NameTypes.h
--- a/Engine/Source/Runtime/Core/Public/UObject/NameTypes.h
+++ b/Engine/Source/Runtime/Core/Public/UObject/NameTypes.h
@@ -770,8 +770,8 @@ public:
 	/** Create FName from non-null string with known length  */
 	FName(int32 Len, const WIDECHAR* Name, EFindName FindType=FNAME_Add);
 	FName(int32 Len, const ANSICHAR* Name, EFindName FindType=FNAME_Add);
-	FName(const FStringView& Name, EFindName FindType=FNAME_Add);
-	FName(const FAnsiStringView& Name, EFindName FindType=FNAME_Add);
+	explicit FName(const FStringView& Name, EFindName FindType=FNAME_Add);
+	explicit FName(const FAnsiStringView& Name, EFindName FindType=FNAME_Add);
 
 	/**
 	 * Create an FName. If FindType is FNAME_Find, and the string part of the name 
diff --git a/Engine/Source/Runtime/CoreUObject/Private/UObject/SoftObjectPath.cpp b/Engine/Source/Runtime/CoreUObject/Private/UObject/SoftObjectPath.cpp
--- a/Engine/Source/Runtime/CoreUObject/Private/UObject/SoftObjectPath.cpp
+++ b/Engine/Source/Runtime/CoreUObject/Private/UObject/SoftObjectPath.cpp
@@ -76,13 +76,13 @@ void FSoftObjectPath::SetPath(FWideStringView Path)
 		if (Path.FindChar(':', ColonIndex))
 		{
 			// Has a subobject, split on that then create a name from the temporary path
-			AssetPathName = Path.Left(ColonIndex);
+			AssetPathName = FName(Path.Left(ColonIndex));
 			SubPathString = Path.Mid(ColonIndex + 1);
 		}
 		else
 		{
 			// No Subobject
-			AssetPathName = Path;
+			AssetPathName = FName(Path);
 			SubPathString.Empty();
 		}
 	}
@@ -118,7 +118,7 @@ void FSoftObjectPath::SetPath(FName PathName)
 			if (Path.FindChar(':', ColonIndex))
 			{
 				// Has a subobject, split on that then create a name from the temporary path
-				AssetPathName = Path.Left(ColonIndex);
+				AssetPathName = FName(Path.Left(ColonIndex));
 				SubPathString = Path.Mid(ColonIndex + 1);
 			}
 			else
