Escape non-ASCII characters as UTF sequences when generating C++ literal strings

This improves portability between compilers and the default encoding of an OS install, and is inline with how we encode strings within hand written C++.

#rnx

diff --git a/Engine/Source/Runtime/CoreUObject/Private/UObject/PropertyStr.cpp b/Engine/Source/Runtime/CoreUObject/Private/UObject/PropertyStr.cpp
--- a/Engine/Source/Runtime/CoreUObject/Private/UObject/PropertyStr.cpp
+++ b/Engine/Source/Runtime/CoreUObject/Private/UObject/PropertyStr.cpp
@@ -4,6 +4,8 @@
 #include "UObject/ObjectMacros.h"
 #include "UObject/PropertyPortFlags.h"
 #include "UObject/UnrealType.h"
+#include "Misc/StringBuilder.h"
+#include "Misc/AsciiSet.h"
 
 /*-----------------------------------------------------------------------------
 	FStrProperty.
@@ -34,47 +36,105 @@ FString FStrProperty::GetCPPTypeForwardDeclaration() const
 // Necessary to fix Compiler Error C2026 and C1091
 FString FStrProperty::ExportCppHardcodedText(const FString& InSource, const FString& Indent)
 {
-	const FString Source = InSource.ReplaceCharWithEscapedChar();
-	FString Result{};
-	const int32 PreferredLineSize = 256;
-	int32 StartPos = 0;
+	constexpr FAsciiSet EscapableCharacters("\\\"\n\r\t");
+
+	auto EstimateExportedStringLength = [&EscapableCharacters](const UTF16CHAR* InStr)
+	{
+		int32 EstimatedLen = 0;
+		while (const UTF16CHAR Char = *InStr++)
+		{
+			if (EscapableCharacters.Contains(Char))
+			{
+				// Exported escaped
+				EstimatedLen += 2;
+			}
+			else if (Char > 0x7f)
+			{
+				// Exported as a UTF-16 sequence
+				EstimatedLen += 4;
+			}
+			else
+			{
+				// Exported verbatim
+				EstimatedLen += 1;
+			}
+		}
+		return EstimatedLen;
+	};
 
+	const int32 PreferredLineSize = 256;
 	const int32 LinesPerString = 16;
-	const bool bUseSubStrings = InSource.Len() > (LinesPerString * PreferredLineSize);
+
+	// Note: This is a no-op on platforms that are using a 16-bit TCHAR
+	FTCHARToUTF16 UTF16SourceString(*InSource, InSource.Len() + 1); // include the null terminator
+
+	const bool bUseSubStrings = EstimateExportedStringLength(UTF16SourceString.Get()) > (LinesPerString * PreferredLineSize);
 	int32 LineNum = 0;
+
+	TStringBuilder<1024> Result;
+
 	if (bUseSubStrings)
 	{
-		Result += TEXT("*(FString(");
+		Result << TEXT("*(FString(");
 	}
 
+	const UTF16CHAR* SourceBegin = UTF16SourceString.Get();
+	const UTF16CHAR* SourceIt = SourceBegin;
 	do
 	{
-		if (StartPos)
+		if (SourceIt > SourceBegin)
 		{
-			Result += TEXT("\n");
-			Result += Indent;
+			Result << TEXT("\n");
+			Result << Indent;
 		}
 
 		++LineNum;
-		if (bUseSubStrings && 0 == (LineNum % LinesPerString))
+		if (bUseSubStrings && (LineNum % LinesPerString) == 0)
 		{
-			Result += TEXT(") + FString(");
+			Result << TEXT(") + FString(");
 		}
 
-		int32 WantedSize = FMath::Min<int32>(Source.Len() - StartPos, PreferredLineSize);
-		while (((WantedSize + StartPos) < Source.Len()) && (Source[WantedSize + StartPos - 1] == TCHAR('\\')))
+		Result << TEXT("TEXT(\"");
 		{
-			WantedSize++;
+			int32 ResultStartLen = Result.Len();
+			while (*SourceIt && (Result.Len() - ResultStartLen) < PreferredLineSize)
+			{
+				if (EscapableCharacters.Contains(*SourceIt))
+				{
+					Result << TEXT('\\');
+					Result << (TCHAR)*SourceIt++;
+				}
+				else if (*SourceIt > 0x7f)
+				{
+					// If this character is part of a surrogate pair, then combine them and write them as a single UTF-32 sequence
+					// Otherwise just write out the character as a UTF-16 sequence
+					if (StringConv::IsHighSurrogate(*SourceIt) && StringConv::IsLowSurrogate(*(SourceIt + 1)))
+					{
+						const UTF32CHAR Codepoint = StringConv::EncodeSurrogate(*SourceIt, *(SourceIt + 1));
+						Result.Appendf(TEXT("\\U%08x"), Codepoint);
+						SourceIt += 2;
+					}
+					else
+					{
+						Result.Appendf(TEXT("\\u%04x"), *SourceIt++);
+					}
+				}
+				else
+				{
+					Result << (TCHAR)*SourceIt++;
+				}
+			}
 		}
-		Result += FString::Printf(TEXT("TEXT(\"%s\")"), *Source.Mid(StartPos, WantedSize));
-		StartPos += WantedSize;
-	} while (StartPos < Source.Len());
+		Result << TEXT("\")");
+	}
+	while (*SourceIt);
 
 	if (bUseSubStrings)
 	{
-		Result += TEXT("))");
+		Result << TEXT("))");
 	}
-	return Result;
+
+	return Result.ToString();
 }
 
 void FStrProperty::ExportTextItem( FString& ValueStr, const void* PropertyValue, const void* DefaultValue, UObject* Parent, int32 PortFlags, UObject* ExportRootScope ) const
