Suppress static analyzer warning 6286

diff --git a/Engine/Source/Runtime/CoreUObject/Private/Misc/PackageName.cpp b/Engine/Source/Runtime/CoreUObject/Private/Misc/PackageName.cpp
--- a/Engine/Source/Runtime/CoreUObject/Private/Misc/PackageName.cpp
+++ b/Engine/Source/Runtime/CoreUObject/Private/Misc/PackageName.cpp
@@ -986,7 +986,11 @@ bool FPackageName::DoesPackageExist(const FString& LongPackageName, const FGuid*
 
 	// Convert to filename (no extension yet).
 	FString Filename = LongPackageNameToFilename(PackageName, TEXT(""));
+
 	// Find the filename (with extension).
+
+	// suppress warning C6286: (<non-zero constant> || <expression>) is always a non-zero constant. <expression> is never evaluated and might have side effects.
+	CA_SUPPRESS(6286);
 	if (!WITH_EDITORONLY_DATA || !IsExtraPackage(PackageName))
 	{
 		bFoundFile = FindPackageFileWithoutExtension(Filename, Filename, InAllowTextFormats);
