FFieldPaths should no longer re-resolve themselves each time a blueprint is destroy and only re-resolve if the blueprint that owns the property they resolve to is destroyed / or recompiled.

#tests 1 hour bot soak tests in cooked build, quick some test, repro steps from FProperty related Jiras

diff --git a/Engine/Source/Runtime/Core/Public/UObject/WeakObjectPtrTemplates.h b/Engine/Source/Runtime/Core/Public/UObject/WeakObjectPtrTemplates.h
--- a/Engine/Source/Runtime/Core/Public/UObject/WeakObjectPtrTemplates.h
+++ b/Engine/Source/Runtime/Core/Public/UObject/WeakObjectPtrTemplates.h
@@ -29,6 +29,8 @@ struct TWeakObjectPtr;
 template<class T, class TWeakObjectPtrBase>
 struct TWeakObjectPtr : private TWeakObjectPtrBase
 {
+	friend struct FFieldPath;
+
 	// Although templated, these parameters are not intended to be anything other than the default,
 	// and are only templates for module organization reasons.
 	static_assert(TAreTypesEqual<TWeakObjectPtrBase, FWeakObjectPtr>::Value, "TWeakObjectPtrBase should not be overridden");
diff --git a/Engine/Source/Runtime/CoreUObject/Private/UObject/Class.cpp b/Engine/Source/Runtime/CoreUObject/Private/UObject/Class.cpp
--- a/Engine/Source/Runtime/CoreUObject/Private/UObject/Class.cpp
+++ b/Engine/Source/Runtime/CoreUObject/Private/UObject/Class.cpp
@@ -52,6 +52,7 @@
 #include "UObject/FastReferenceCollector.h"
 #include "UObject/PropertyProxyArchive.h"
 #include "UObject/FieldPath.h"
+#include "HAL/ThreadSafeCounter.h"
 
 
 // WARNING: This should always be the last include in any file that needs it (except .generated.h)
@@ -622,6 +623,14 @@ public:
 	}
 };
 
+#if WITH_EDITORONLY_DATA
+static int32 GetNextFieldPathSerialNumber()
+{
+	static FThreadSafeCounter GlobalSerialNumberCounter;
+	return GlobalSerialNumberCounter.Increment();
+}
+#endif // WITH_EDITORONLY_DATA
+
 //
 // Constructors.
 //
@@ -638,6 +647,9 @@ UStruct::UStruct(EStaticConstructor, int32 InSize, int32 InMinAlignment, EObject
 	, PostConstructLink(nullptr)
 	, UnresolvedScriptProperties(nullptr)
 {
+#if WITH_EDITORONLY_DATA
+	FieldPathSerialNumber = GetNextFieldPathSerialNumber();
+#endif // WITH_EDITORONLY_DATA
 }
 
 UStruct::UStruct(UStruct* InSuperStruct, SIZE_T ParamsSize, SIZE_T Alignment)
@@ -656,6 +668,9 @@ UStruct::UStruct(UStruct* InSuperStruct, SIZE_T ParamsSize, SIZE_T Alignment)
 #if USTRUCT_FAST_ISCHILDOF_IMPL == USTRUCT_ISCHILDOF_STRUCTARRAY
 	this->ReinitializeBaseChainArray();
 #endif
+#if WITH_EDITORONLY_DATA
+	FieldPathSerialNumber = GetNextFieldPathSerialNumber();
+#endif // WITH_EDITORONLY_DATA
 }
 
 UStruct::UStruct(const FObjectInitializer& ObjectInitializer, UStruct* InSuperStruct, SIZE_T ParamsSize, SIZE_T Alignment)
@@ -674,6 +689,9 @@ UStruct::UStruct(const FObjectInitializer& ObjectInitializer, UStruct* InSuperSt
 #if USTRUCT_FAST_ISCHILDOF_IMPL == USTRUCT_ISCHILDOF_STRUCTARRAY
 	this->ReinitializeBaseChainArray();
 #endif
+#if WITH_EDITORONLY_DATA
+	FieldPathSerialNumber = GetNextFieldPathSerialNumber();
+#endif // WITH_EDITORONLY_DATA
 }
 
 /**
@@ -1584,12 +1602,6 @@ void UStruct::FinishDestroy()
 {
 	DestroyUnversionedSchema(this);
 	Script.Empty();
-
-	if (ChildProperties)
-	{
-		FFieldPath::OnFieldDeleted();
-	}
-
 	Super::FinishDestroy();
 }
 
@@ -1612,6 +1624,9 @@ void UStruct::DestroyChildPropertiesAndResetPropertyLinks()
 	RefLink = nullptr;
 	DestructorLink = nullptr;
 	PostConstructLink = nullptr;
+#if WITH_EDITORONLY_DATA
+	FieldPathSerialNumber = GetNextFieldPathSerialNumber();
+#endif // WITH_EDITORONLY_DATA
 }
 
 UStruct::~UStruct()
@@ -1620,7 +1635,7 @@ UStruct::~UStruct()
 	// This needs to happen after FinishDestroy which calls DestroyNonNativeProperties
 	// Also, Blueprint generated classes can have DestroyNonNativeProperties called on them after their FinishDestroy has been called
 	// so properties can only be deleted in the destructor
-	DestroyChildPropertiesAndResetPropertyLinks();
+	DestroyPropertyLinkedList(ChildProperties);
 	DeleteUnresolvedScriptProperties();
 }
 
@@ -4687,6 +4702,8 @@ void UClass::PurgeClass(bool bRecompilingOnLoad)
 		PropertiesPendingDestruction = ChildProperties;
 		ChildProperties = nullptr;
 	}
+	// Update the serial number so that FFieldPaths that point to properties of this struct know they need to resolve themselves again
+	FieldPathSerialNumber = GetNextFieldPathSerialNumber();
 #else
 	{
 		// Destroy all properties owned by this struct
@@ -4694,8 +4711,6 @@ void UClass::PurgeClass(bool bRecompilingOnLoad)
 	}
 #endif // WITH_EDITORONLY_DATA
 
-	FFieldPath::OnFieldDeleted();
-
 	DestroyUnversionedSchema(this);
 }
 
@@ -4746,9 +4761,9 @@ bool UClass::HasProperty(FProperty* InProperty) const
 	{
 		UClass* PropertiesClass = InProperty->GetOwner<UClass>();
 		if (PropertiesClass)
-	{
-		return IsChildOf(PropertiesClass);
-	}
+		{
+			return IsChildOf(PropertiesClass);
+		}
 	}
 
 	return false;
diff --git a/Engine/Source/Runtime/CoreUObject/Private/UObject/FieldPath.cpp b/Engine/Source/Runtime/CoreUObject/Private/UObject/FieldPath.cpp
--- a/Engine/Source/Runtime/CoreUObject/Private/UObject/FieldPath.cpp
+++ b/Engine/Source/Runtime/CoreUObject/Private/UObject/FieldPath.cpp
@@ -33,9 +33,10 @@ void FFieldPath::Generate(FField* InField)
 		{
 			Path.Add(Iter.GetFName());
 		}
-		ResolvedFieldOwner = GUObjectArray.ObjectToIndex(InField->GetOwnerUObject());
-		SerialNumber = GlobalSerialNumber;
+		UStruct* OwnerStruct = CastChecked<UStruct>(InField->GetOwnerUObject());
+		ResolvedOwner = OwnerStruct;
 #if WITH_EDITORONLY_DATA
+		FieldPathSerialNumber = OwnerStruct->FieldPathSerialNumber;
 		InitialFieldClass = InField->GetClass();
 #endif // WITH_EDITORONLY_DATA
 	}
@@ -89,7 +90,7 @@ void FFieldPath::Generate(const TCHAR* InFieldPathString)
 	}
 }
 
-FField* FFieldPath::TryToResolvePath(UStruct* InCurrentStruct, int32* OutOwnerIndex, FFieldPath::EPathResolveType InResolveType /*= FFieldPath::UseStructIfOuterNotFound*/) const
+FField* FFieldPath::TryToResolvePath(UStruct* InCurrentStruct, UStruct** OutOwner, FFieldPath::EPathResolveType InResolveType /*= FFieldPath::UseStructIfOuterNotFound*/) const
 {
 	FField* Result = nullptr;
 
@@ -124,16 +125,18 @@ FField* FFieldPath::TryToResolvePath(UStruct* InCurrentStruct, int32* OutOwnerIn
 	{
 		check(PathIndex <= 1);
 		Result = FindFProperty<FField>(Owner, Path[PathIndex]);
-		if (Result && PathIndex > 0)
+		if (Result)
 		{
-			// Nested property
-			Result = Result->GetInnerFieldByName(Path[0]);
-		}	
-		
-	}
-	if (Result && OutOwnerIndex)
-	{
-		*OutOwnerIndex = GUObjectArray.ObjectToIndex(LastOuter);
+			if (PathIndex > 0)
+			{
+				// Nested property
+				Result = Result->GetInnerFieldByName(Path[0]);
+			}
+			if (OutOwner)
+			{
+				*OutOwner = Owner;
+			}
+		}
 	}
 
 	return Result;
@@ -179,12 +182,15 @@ void FFieldPath::GenerateFromUField(UField* InField)
 		Path.Add(Obj->GetFName());
 	}	
 }
-#endif // WITH_EDITORONLY_DATA
 
-int32 FFieldPath::GlobalSerialNumber = 0;
+bool FFieldPath::IsFieldPathSerialNumberIdentical(UStruct* InStruct) const
+{
+	return FieldPathSerialNumber == InStruct->FieldPathSerialNumber;
+}
 
-void FFieldPath::OnFieldDeleted()
+int32 FFieldPath::GetFieldPathSerialNumber(UStruct* InStruct) const
 {
-	checkSlow(IsInGameThread());
-	GlobalSerialNumber++;
-}
\ No newline at end of file
+	return InStruct->FieldPathSerialNumber;
+}
+
+#endif // WITH_EDITORONLY_DATA
diff --git a/Engine/Source/Runtime/CoreUObject/Public/UObject/Class.h b/Engine/Source/Runtime/CoreUObject/Public/UObject/Class.h
--- a/Engine/Source/Runtime/CoreUObject/Public/UObject/Class.h
+++ b/Engine/Source/Runtime/CoreUObject/Public/UObject/Class.h
@@ -327,6 +327,8 @@ public:
 #if WITH_EDITORONLY_DATA
 	/** List of wrapper UObjects for FProperties */
 	TArray<UPropertyWrapper*> PropertyWrappers;
+	/** Unique id incremented each time this class properties get destroyed */
+	int32 FieldPathSerialNumber;
 #endif
 
 	/** Cached schema for optimized unversioned property serialization, owned by this. */
diff --git a/Engine/Source/Runtime/CoreUObject/Public/UObject/FastReferenceCollector.h b/Engine/Source/Runtime/CoreUObject/Public/UObject/FastReferenceCollector.h
--- a/Engine/Source/Runtime/CoreUObject/Public/UObject/FastReferenceCollector.h
+++ b/Engine/Source/Runtime/CoreUObject/Public/UObject/FastReferenceCollector.h
@@ -819,19 +819,17 @@ private:
 					case GCRT_AddFieldPathReferencedObject:
 					{
 						FFieldPath*	FieldPathPtr = (FFieldPath*)(StackEntryData + ReferenceInfo.Offset);
-						FField*& Field = FieldPathPtr->GetResolvedFieldInternal();
-						int32& FieldOwnerIndex = FieldPathPtr->GetResolvedFieldOwnerInternal();
+						FUObjectItem* FieldOwnerItem = FieldPathPtr->GetResolvedOwnerItemInternal();
 						TokenReturnCount = ReferenceInfo.ReturnCount;
-						if (FieldOwnerIndex >= 0)
+						if (FieldOwnerItem)
 						{
-							UObject* OwnerObject = static_cast<UObject*>(GUObjectArray.IndexToObjectUnsafeForGC(FieldOwnerIndex)->Object);
+							UObject* OwnerObject = static_cast<UObject*>(FieldOwnerItem->Object);
 							UObject* PreviousOwner = OwnerObject;
 							ReferenceProcessor.HandleTokenStreamObjectReference(NewObjectsToSerialize, CurrentObject, OwnerObject, ReferenceTokenStreamIndex, true);
 							// Handle reference elimination (PendingKill owner)
 							if (PreviousOwner && !OwnerObject)
 							{
-								Field = nullptr;
-								FieldOwnerIndex = -1;
+								FieldPathPtr->ClearCachedFieldInternal();
 							}
 						}
 					}
@@ -843,18 +841,16 @@ private:
 						TokenReturnCount = ReferenceInfo.ReturnCount;
 						for (int32 FieldIndex = 0, FieldNum = FieldArray.Num(); FieldIndex < FieldNum; ++FieldIndex)
 						{
-							FField*& Field = FieldArray[FieldIndex].GetResolvedFieldInternal();
-							int32& FieldOwnerIndex = FieldArray[FieldIndex].GetResolvedFieldOwnerInternal();
-							if (FieldOwnerIndex >= 0)
+							FUObjectItem* FieldOwnerItem = FieldArray[FieldIndex].GetResolvedOwnerItemInternal();
+							if (FieldOwnerItem)
 							{
-								UObject* OwnerObject = static_cast<UObject*>(GUObjectArray.IndexToObjectUnsafeForGC(FieldOwnerIndex)->Object);
+								UObject* OwnerObject = static_cast<UObject*>(FieldOwnerItem->Object);
 								UObject* PreviousOwner = OwnerObject;
 								ReferenceProcessor.HandleTokenStreamObjectReference(NewObjectsToSerialize, CurrentObject, OwnerObject, ReferenceTokenStreamIndex, true);
 								// Handle reference elimination (PendingKill owner)
 								if (PreviousOwner && !OwnerObject)
 								{
-									Field = nullptr;
-									FieldOwnerIndex = -1;
+									FieldArray[FieldIndex].ClearCachedFieldInternal();
 								}
 							}
 						}
diff --git a/Engine/Source/Runtime/CoreUObject/Public/UObject/FieldPath.h b/Engine/Source/Runtime/CoreUObject/Public/UObject/FieldPath.h
--- a/Engine/Source/Runtime/CoreUObject/Public/UObject/FieldPath.h
+++ b/Engine/Source/Runtime/CoreUObject/Public/UObject/FieldPath.h
@@ -9,12 +9,23 @@
 #include "CoreMinimal.h"
 #include "UObject/Field.h"
 #include "HAL/ThreadSafeCounter.h"
+#include "UObject/WeakObjectPtr.h"
+#include "UObject/WeakObjectPtrTemplates.h"
+#include "UObject/UObjectArray.h"
 
+class UStruct;
 class UField;
-class UObject;
 
 struct COREUOBJECT_API FFieldPath
 {
+	// GC needs access to GetResolvedOwnerItemInternal and ClearCachedFieldInternal
+	template <bool bParallel, typename ReferenceProcessorType, typename CollectorType, typename ArrayPoolType, bool bAutoGenerateTokenStream>
+	friend class TFastReferenceCollector;
+
+	// TWeakFieldPtr needs access to ClearCachedField
+	template<class T>
+	friend struct TWeakFieldPtr;
+
 protected:
 
 	/* Determines the behavior when resolving stored path */
@@ -25,58 +36,77 @@ protected:
 	};
 
 	/** Untracked pointer to the resolved property */
-	mutable FField* ResolvedField;
+	mutable FField* ResolvedField = nullptr;
 #if WITH_EDITORONLY_DATA
 	/** In editor builds, store the original class of the resolved property in case it changes after recompiling BPs */
-	mutable FFieldClass* InitialFieldClass;
+	mutable FFieldClass* InitialFieldClass = nullptr;
+	/** In editor builds, fields may get deleted even though their owner struct remains */
+	mutable int32 FieldPathSerialNumber = 0;
 #endif
-	/** GC tracked index of property owner UObject */
-	mutable int32 ResolvedFieldOwner;
-	/** Serial number this FFieldPath was last resolved with */
-	mutable int32 SerialNumber;
+	/** The cached owner of this field. Even though implemented as a weak pointer, GC will keep a strong reference to it if exposed through UPROPERTY macro */
+	mutable TWeakObjectPtr<UStruct> ResolvedOwner;
 
 	/** Path to the FField object from the innermost FField to the outermost UObject (UPackage) */
 	TArray<FName> Path;
 
-	/** Global serial number that gets increased each time UStruct destyroys its properties */
-	static int32 GlobalSerialNumber;
-
 	FORCEINLINE bool NeedsResolving() const
 	{
-		return !ResolvedField || SerialNumber != GlobalSerialNumber;
+		if (ResolvedField)
+		{
+#if WITH_EDITORONLY_DATA
+			UStruct* Owner = ResolvedOwner.Get();
+			// In uncooked builds we also need to check if the serial number on the owner struct is identical
+			// It will change if the struct has been recompiled or its properties have been destroyed
+			if (Owner && IsFieldPathSerialNumberIdentical(Owner))
+			{
+				return false;
+			}
+#else
+			// The assumption is that if we already resolved a field and its owner is still valid, there's no need to resolve again
+			return !ResolvedOwner.IsValid();
+#endif // WITH_EDITORONLY_DATA
+		}
+		return true;
 	}
 
-public:
-
-	FFieldPath()
-		: ResolvedField(nullptr)
+	/** Clears the cached value so that the next time Get() is called, it will be resolved again */
+	FORCEINLINE void ClearCachedField() const
+	{
+		ResolvedField = nullptr;
 #if WITH_EDITORONLY_DATA
-		, InitialFieldClass(nullptr)
+		InitialFieldClass = nullptr;
+		FieldPathSerialNumber = 0;
 #endif // WITH_EDITORONLY_DATA
-		, ResolvedFieldOwner(-1)
-		, SerialNumber(-1)
-	{}
+		ResolvedOwner.Reset();
+	}
+
+private:
 
-	FFieldPath(FField* InField)
-		: ResolvedField(InField)
 #if WITH_EDITORONLY_DATA
-		, InitialFieldClass(nullptr)
-#endif // WITH_EDITORONLY_DATA
-		, ResolvedFieldOwner(-1)
-		, SerialNumber(-1)
+	/** Used to check if the serial number on the provided struct is identical to the one stored in this FFieldPath */
+	bool IsFieldPathSerialNumberIdentical(UStruct* InStruct) const;
+	/** Gets the serial number stored on the provided struct */
+	int32 GetFieldPathSerialNumber(UStruct* InStruct) const;
+#endif
+
+	/** FOR INTERNAL USE ONLY: gets the pointer to the resolved field without trying to resolve it */
+	FORCEINLINE FUObjectItem* GetResolvedOwnerItemInternal()
 	{
-		Generate(InField);
+		return ResolvedOwner.Internal_GetObjectItem();
+	}
+	FORCEINLINE void ClearCachedFieldInternal()
+	{
+		ResolvedField = nullptr;
+		ResolvedOwner.Reset();
 	}
 
-	FFieldPath(const FFieldPath& Other)
-		: ResolvedField(Other.ResolvedField)
-#if WITH_EDITORONLY_DATA
-		, InitialFieldClass(Other.InitialFieldClass)
-#endif // WITH_EDITORONLY_DATA
-		, ResolvedFieldOwner(Other.ResolvedFieldOwner)
-		, SerialNumber(Other.SerialNumber)
-		, Path(Other.Path)
+public:
+
+	FFieldPath() = default;
+
+	FFieldPath(FField* InField)
 	{
+		Generate(InField);
 	}
 
 #if WITH_EDITORONLY_DATA
@@ -93,24 +123,13 @@ public:
 	void GenerateFromUField(UField* InField);
 #endif
 
-	/** Clears the cached value so that the next time Get() is called, it will be resolved again */
-	inline void ClearCachedField()
-	{
-		ResolvedField = nullptr;
-#if WITH_EDITORONLY_DATA
-		InitialFieldClass = nullptr;
-#endif // WITH_EDITORONLY_DATA
-		ResolvedFieldOwner = -1;
-		SerialNumber = -1;
-	}
-
 	/**
 	 * Tries to resolve the path without caching the resolved pointer 
 	 * @param InCurrentStruct Struct that's trying to resolve this field path
 	 * @param OutOwnerIndex ObjectIndex of the Owner UObject
 	 * @return Resolved field or null
 	 */
-	FField* TryToResolvePath(UStruct* InCurrentStruct = nullptr, int32* OutOwnerIndex = nullptr, EPathResolveType InResolveType = FFieldPath::UseStructIfOuterNotFound) const;
+	FField* TryToResolvePath(UStruct* InCurrentStruct = nullptr, UStruct** OutOwner = nullptr, EPathResolveType InResolveType = FFieldPath::UseStructIfOuterNotFound) const;
 
 	/**
 	 * Tries to resolve the path and caches the result
@@ -119,7 +138,7 @@ public:
 	 */
 	FORCEINLINE void ResolveField(FFieldClass* ExpectedClass = FField::StaticClass(), UStruct* InCurrentStruct = nullptr, EPathResolveType InResolveType = FFieldPath::UseStructIfOuterNotFound) const
 	{
-		int32 FoundOwner = -1;
+		UStruct* FoundOwner = nullptr;
 		FField* FoundField = TryToResolvePath(InCurrentStruct, &FoundOwner, InResolveType);
 		if (FoundField && FoundField->IsA(ExpectedClass) 
 #if WITH_EDITORONLY_DATA
@@ -128,19 +147,19 @@ public:
 			)
 		{
 			ResolvedField = FoundField;
-			ResolvedFieldOwner = FoundOwner;
-			SerialNumber = GlobalSerialNumber;
+			ResolvedOwner = FoundOwner;
 #if WITH_EDITORONLY_DATA
 			if (!InitialFieldClass)
 			{
 				InitialFieldClass = FoundField->GetClass();
 			}
+			check(FoundOwner);
+			FieldPathSerialNumber = GetFieldPathSerialNumber(FoundOwner);
 #endif // WITH_EDITORONLY_DATA
 		}
 		else
 		{
-			ResolvedField = nullptr;
-			ResolvedFieldOwner = -1;
+			ClearCachedField();
 		}
 	}
 
@@ -184,19 +203,6 @@ public:
 		}
 		return Ar;
 	}
-
-	/** FOR INTERNAL USE ONLY: gets the pointer to the resolved field without trying to resolve it */
-	inline FField*& GetResolvedFieldInternal()
-	{
-		return ResolvedField;
-	}
-	inline int32& GetResolvedFieldOwnerInternal()
-	{
-		return ResolvedFieldOwner;
-	}
-
-	/** Called when fields have been deleted to bump the global serial number and invalidate cached pointers */
-	static void OnFieldDeleted();
 };
 
 template<class PropertyType>
diff --git a/Engine/Source/Runtime/CoreUObject/Public/UObject/WeakObjectPtr.h b/Engine/Source/Runtime/CoreUObject/Public/UObject/WeakObjectPtr.h
--- a/Engine/Source/Runtime/CoreUObject/Public/UObject/WeakObjectPtr.h
+++ b/Engine/Source/Runtime/CoreUObject/Public/UObject/WeakObjectPtr.h
@@ -21,6 +21,8 @@ struct FWeakObjectPtr
 {
 public:
 
+	friend struct FFieldPath;
+
 	/** Null constructor **/
 	FORCEINLINE FWeakObjectPtr()
 	{
