Fixed escaping of control characters when converting a string literal to C++

Fallout from CL# 12343442

#rnx

diff --git a/Engine/Source/Runtime/CoreUObject/Private/UObject/PropertyStr.cpp b/Engine/Source/Runtime/CoreUObject/Private/UObject/PropertyStr.cpp
--- a/Engine/Source/Runtime/CoreUObject/Private/UObject/PropertyStr.cpp
+++ b/Engine/Source/Runtime/CoreUObject/Private/UObject/PropertyStr.cpp
@@ -101,8 +101,24 @@ FString FStrProperty::ExportCppHardcodedText(const FString& InSource, const FStr
 			{
 				if (EscapableCharacters.Contains(*SourceIt))
 				{
+					const TCHAR CharToEscape = (TCHAR)*SourceIt++;
+
 					Result << TEXT('\\');
-					Result << (TCHAR)*SourceIt++;
+					switch (CharToEscape)
+					{
+					case TEXT('\n'):
+						Result << TEXT('n');
+						break;
+					case TEXT('\r'):
+						Result << TEXT('r');
+						break;
+					case TEXT('\t'):
+						Result << TEXT('t');
+						break;
+					default:
+						Result << CharToEscape;
+						break;
+					}
 				}
 				else if (*SourceIt > 0x7f)
 				{
