Fix compile from incorrect predeclaration

diff --git a/Engine/Source/Runtime/CoreUObject/Public/UObject/FieldPath.h b/Engine/Source/Runtime/CoreUObject/Public/UObject/FieldPath.h
--- a/Engine/Source/Runtime/CoreUObject/Public/UObject/FieldPath.h
+++ b/Engine/Source/Runtime/CoreUObject/Public/UObject/FieldPath.h
@@ -19,7 +19,7 @@ class UField;
 struct COREUOBJECT_API FFieldPath
 {
 	// GC needs access to GetResolvedOwnerItemInternal and ClearCachedFieldInternal
-	template <bool bParallel, typename ReferenceProcessorType, typename CollectorType, typename ArrayPoolType, bool bAutoGenerateTokenStream>
+	template <bool bParallel, typename ReferenceProcessorType, typename CollectorType, typename ArrayPoolType, bool bAutoGenerateTokenStream, bool bIgnoreNoopTokens>
 	friend class TFastReferenceCollector;
 
 	// TWeakFieldPtr needs access to ClearCachedField
