Remove includes of StringBuilder.h in headers

#rnx

diff --git a/Engine/Source/Runtime/Core/Private/IO/IoDispatcherFileBackend.cpp b/Engine/Source/Runtime/Core/Private/IO/IoDispatcherFileBackend.cpp
--- a/Engine/Source/Runtime/Core/Private/IO/IoDispatcherFileBackend.cpp
+++ b/Engine/Source/Runtime/Core/Private/IO/IoDispatcherFileBackend.cpp
@@ -2,6 +2,7 @@
 
 #include "IO/IoDispatcherFileBackend.h"
 #include "Misc/ScopeRWLock.h"
+#include "Misc/StringBuilder.h"
 #include "ProfilingDebugging/CountersTrace.h"
 #include "HAL/PlatformFilemanager.h"
 #include "GenericPlatform/GenericPlatformFile.h"
diff --git a/Engine/Source/Runtime/Core/Public/IO/IoDispatcher.h b/Engine/Source/Runtime/Core/Public/IO/IoDispatcher.h
--- a/Engine/Source/Runtime/Core/Public/IO/IoDispatcher.h
+++ b/Engine/Source/Runtime/Core/Public/IO/IoDispatcher.h
@@ -6,7 +6,6 @@
 #include "Containers/StringView.h"
 #include "Containers/UnrealString.h"
 #include "Logging/LogMacros.h"
-#include "Misc/StringBuilder.h"
 #include "Templates/RefCounting.h"
 #include "Templates/UnrealTemplate.h"
 #include "Templates/TypeCompatibleBytes.h"
diff --git a/Engine/Source/Runtime/CoreUObject/Private/UObject/LinkerLoad.cpp b/Engine/Source/Runtime/CoreUObject/Private/UObject/LinkerLoad.cpp
--- a/Engine/Source/Runtime/CoreUObject/Private/UObject/LinkerLoad.cpp
+++ b/Engine/Source/Runtime/CoreUObject/Private/UObject/LinkerLoad.cpp
@@ -17,6 +17,7 @@
 #include "Misc/PackageName.h"
 #include "Blueprint/BlueprintSupport.h"
 #include "Misc/SecureHash.h"
+#include "Misc/StringBuilder.h"
 #include "ProfilingDebugging/DebuggingDefines.h"
 #include "Logging/TokenizedMessage.h"
 #include "UObject/LinkerPlaceholderBase.h"
diff --git a/Engine/Source/Runtime/CoreUObject/Private/UObject/PrimaryAssetId.cpp b/Engine/Source/Runtime/CoreUObject/Private/UObject/PrimaryAssetId.cpp
--- a/Engine/Source/Runtime/CoreUObject/Private/UObject/PrimaryAssetId.cpp
+++ b/Engine/Source/Runtime/CoreUObject/Private/UObject/PrimaryAssetId.cpp
@@ -1,11 +1,10 @@
 // Copyright Epic Games, Inc. All Rights Reserved.
 
 #include "UObject/PrimaryAssetId.h"
+#include "Misc/StringBuilder.h"
+#include "UObject/ObjectRedirector.h"
 #include "UObject/PropertyPortFlags.h"
 #include "UObject/UnrealType.h"
-#include "UObject/ObjectRedirector.h"
-
-	
 
 bool FPrimaryAssetType::ExportTextItem(FString& ValueStr, FPrimaryAssetType const& DefaultValue, UObject* Parent, int32 PortFlags, UObject* ExportRootScope) const
 {
@@ -139,3 +138,7 @@ bool FPrimaryAssetId::SerializeFromMismatchedTag(struct FPropertyTag const& Tag,
 	return false;
 }
 
+FStringBuilderBase& operator<<(FStringBuilderBase& Builder, const FPrimaryAssetId& Id)
+{
+	return Builder << Id.PrimaryAssetType.GetName() << ":" << Id.PrimaryAssetName;
+}
diff --git a/Engine/Source/Runtime/CoreUObject/Private/UObject/Property.cpp b/Engine/Source/Runtime/CoreUObject/Private/UObject/Property.cpp
--- a/Engine/Source/Runtime/CoreUObject/Private/UObject/Property.cpp
+++ b/Engine/Source/Runtime/CoreUObject/Private/UObject/Property.cpp
@@ -7,6 +7,7 @@
 #include "CoreMinimal.h"
 #include "Misc/AsciiSet.h"
 #include "Misc/Guid.h"
+#include "Misc/StringBuilder.h"
 #include "Math/RandomStream.h"
 #include "Logging/LogScopedCategoryAndVerbosityOverride.h"
 #include "UObject/ObjectMacros.h"
diff --git a/Engine/Source/Runtime/CoreUObject/Private/UObject/PropertyBaseObject.cpp b/Engine/Source/Runtime/CoreUObject/Private/UObject/PropertyBaseObject.cpp
--- a/Engine/Source/Runtime/CoreUObject/Private/UObject/PropertyBaseObject.cpp
+++ b/Engine/Source/Runtime/CoreUObject/Private/UObject/PropertyBaseObject.cpp
@@ -9,6 +9,7 @@
 #include "UObject/PropertyHelper.h"
 #include "UObject/LinkerPlaceholderClass.h"
 #include "UObject/LinkerPlaceholderExportObject.h"
+#include "Misc/StringBuilder.h"
 
 // WARNING: This should always be the last include in any file that needs it (except .generated.h)
 #include "UObject/UndefineUPropertyMacros.h"
diff --git a/Engine/Source/Runtime/CoreUObject/Private/UObject/PropertyName.cpp b/Engine/Source/Runtime/CoreUObject/Private/UObject/PropertyName.cpp
--- a/Engine/Source/Runtime/CoreUObject/Private/UObject/PropertyName.cpp
+++ b/Engine/Source/Runtime/CoreUObject/Private/UObject/PropertyName.cpp
@@ -4,6 +4,7 @@
 #include "UObject/ObjectMacros.h"
 #include "UObject/PropertyPortFlags.h"
 #include "UObject/UnrealType.h"
+#include "Misc/StringBuilder.h"
 
 /*-----------------------------------------------------------------------------
 	FNameProperty.
diff --git a/Engine/Source/Runtime/CoreUObject/Public/UObject/PrimaryAssetId.h b/Engine/Source/Runtime/CoreUObject/Public/UObject/PrimaryAssetId.h
--- a/Engine/Source/Runtime/CoreUObject/Public/UObject/PrimaryAssetId.h
+++ b/Engine/Source/Runtime/CoreUObject/Public/UObject/PrimaryAssetId.h
@@ -3,7 +3,6 @@
 #pragma once
 
 #include "CoreMinimal.h"
-#include "Misc/StringBuilder.h"
 
 /**
  * A primary asset type, represented as an FName internally and implicitly convertible back and forth
@@ -165,7 +164,4 @@ struct FPrimaryAssetId
 	friend struct Z_Construct_UScriptStruct_FPrimaryAssetId_Statics;
 };
 
-inline FStringBuilderBase& operator<<(FStringBuilderBase& Builder, const FPrimaryAssetId& Id)
-{
-	return Builder << Id.PrimaryAssetType.GetName() << ":" << Id.PrimaryAssetName;
-}
\ No newline at end of file
+COREUOBJECT_API FStringBuilderBase& operator<<(FStringBuilderBase& Builder, const FPrimaryAssetId& Id);
\ No newline at end of file
diff --git a/Engine/Source/Runtime/CoreUObject/Public/UObject/SoftObjectPath.h b/Engine/Source/Runtime/CoreUObject/Public/UObject/SoftObjectPath.h
--- a/Engine/Source/Runtime/CoreUObject/Public/UObject/SoftObjectPath.h
+++ b/Engine/Source/Runtime/CoreUObject/Public/UObject/SoftObjectPath.h
@@ -3,9 +3,10 @@
 #pragma once
 
 #include "CoreMinimal.h"
+#include "Containers/StringView.h"
 #include "HAL/ThreadSafeCounter.h"
-#include "UObject/Class.h"
 #include "HAL/ThreadSingleton.h"
+#include "UObject/Class.h"
 
 /**
  * A struct that contains a string reference to an object, either a top level asset or a subobject.
