- Fixed a compile error
#rnx

diff --git a/Engine/Source/Runtime/CoreUObject/Private/UObject/LinkerLoad.cpp b/Engine/Source/Runtime/CoreUObject/Private/UObject/LinkerLoad.cpp
--- a/Engine/Source/Runtime/CoreUObject/Private/UObject/LinkerLoad.cpp
+++ b/Engine/Source/Runtime/CoreUObject/Private/UObject/LinkerLoad.cpp
@@ -40,6 +40,7 @@
 #include "Serialization/LoadTimeTracePrivate.h"
 #include "HAL/FileManager.h"
 #include "UObject/CoreRedirects.h"
+#include "Misc/StringBuilder.h"
 
 class FTexture2DResourceMem;
 
