#UE4 Added '+' to the list of RegularTokenChars in property parsing code. Object paths can contain a plus, so they must be able to be parsed

diff --git a/Engine/Source/Runtime/CoreUObject/Private/UObject/Property.cpp b/Engine/Source/Runtime/CoreUObject/Private/UObject/Property.cpp
--- a/Engine/Source/Runtime/CoreUObject/Private/UObject/Property.cpp
+++ b/Engine/Source/Runtime/CoreUObject/Private/UObject/Property.cpp
@@ -379,7 +379,7 @@ FORCEINLINE constexpr bool IsValidTokenStart(TCHAR FirstChar, bool bDottedNames)
 
 FORCEINLINE constexpr FStringView ParsePropertyToken(const TCHAR* Str, bool DottedNames)
 {
-	constexpr FAsciiSet RegularTokenChars = AlphaNumericChars  + '_' + '-' + '"';
+	constexpr FAsciiSet RegularTokenChars = AlphaNumericChars  + '_' + '-' + '+' + '"';
 	constexpr FAsciiSet DottedTokenChars = RegularTokenChars + '.' + '/' + SUBOBJECT_DELIMITER_CHAR;
 	FAsciiSet CurrentTokenChars = DottedNames ? DottedTokenChars : RegularTokenChars;
 
