Copying //UE4/Release-4.25Plus-Dev to Release-4.25Plus (//UE4/Release-4.25Plus)

diff --git a/Engine/Source/Runtime/CoreUObject/Private/Serialization/AsyncLoading2.cpp b/Engine/Source/Runtime/CoreUObject/Private/Serialization/AsyncLoading2.cpp
--- a/Engine/Source/Runtime/CoreUObject/Private/Serialization/AsyncLoading2.cpp
+++ b/Engine/Source/Runtime/CoreUObject/Private/Serialization/AsyncLoading2.cpp
@@ -56,13 +56,92 @@
 #include "HAL/LowLevelMemStats.h"
 
 #if UE_BUILD_DEVELOPMENT || UE_BUILD_DEBUG
-//PRAGMA_DISABLE_OPTIMIZATION
+PRAGMA_DISABLE_OPTIMIZATION
 #endif
 
+FArchive& operator<<(FArchive& Ar, FNameMapEntry& NameMapEntry)
+{
+	Ar << NameMapEntry.Index;
+	Ar << NameMapEntry.Number;
+
+	return Ar;
+}
+
+FArchive& operator<<(FArchive& Ar, FExportBundleEntry& ExportBundleEntry)
+{
+	Ar << ExportBundleEntry.LocalExportIndex;
+	Ar << ExportBundleEntry.CommandType;
+
+	return Ar;
+}
+
+FArchive& operator<<(FArchive& Ar, FExportBundleHeader& ExportBundleHeader)
+{
+	Ar << ExportBundleHeader.FirstEntryIndex;
+	Ar << ExportBundleHeader.EntryCount;
+
+	return Ar;
+}
+
+FArchive& operator<<(FArchive& Ar, FExportBundleMetaEntry& ExportBundleMetaEntry)
+{
+	Ar << ExportBundleMetaEntry.LoadOrder;
+	Ar << ExportBundleMetaEntry.PayloadSize;
+
+	return Ar;
+}
+
+FArchive& operator<<(FArchive& Ar, FScriptObjectEntry& ScriptObjectEntry)
+{
+	Ar << ScriptObjectEntry.ObjectName;
+	Ar << ScriptObjectEntry.OuterIndex;
+	Ar << ScriptObjectEntry.CDOClassIndex;
+	return Ar;
+}
+
+FArchive& operator<<(FArchive& Ar, FExportMapEntry& ExportMapEntry)
+{
+	Ar << ExportMapEntry.CookedSerialOffset;
+	Ar << ExportMapEntry.CookedSerialSize;
+	Ar << ExportMapEntry.ObjectName.NameEntryIndex << ExportMapEntry.ObjectName.Number;
+	Ar << ExportMapEntry.OuterIndex;
+	Ar << ExportMapEntry.ClassIndex;
+	Ar << ExportMapEntry.SuperIndex;
+	Ar << ExportMapEntry.TemplateIndex;
+	Ar << ExportMapEntry.GlobalImportIndex;
+
+	uint32 ObjectFlags = uint32(ExportMapEntry.ObjectFlags);
+	Ar << ObjectFlags;
+	
+	if (Ar.IsLoading())
+	{
+		ExportMapEntry.ObjectFlags = EObjectFlags(ObjectFlags);
+	}
+
+	uint8 FilterFlags = uint8(ExportMapEntry.FilterFlags);
+	Ar << FilterFlags;
+
+	if (Ar.IsLoading())
+	{
+		ExportMapEntry.FilterFlags = EExportFilterFlags(FilterFlags);
+	}
+
+	uint64 Pad = 0;
+	Ar.Serialize(&Pad, 7);
+
+	return Ar;
+}
+
+#if WITH_ASYNCLOADING2
+
 #ifndef ALT2_VERIFY_ASYNC_FLAGS
 #define ALT2_VERIFY_ASYNC_FLAGS DO_CHECK
 #endif
 
+#ifndef ALT2_VERIFY_RECURSIVE_LOADS
+#define ALT2_VERIFY_RECURSIVE_LOADS DO_CHECK
+#endif
+
 #ifndef ALT2_LOG_VERBOSE
 #define ALT2_LOG_VERBOSE DO_CHECK
 #endif
@@ -106,48 +185,35 @@ TRACE_DECLARE_INT_COUNTER(PendingBundleIoRequests, TEXT("AsyncLoading/PendingBun
 struct FAsyncPackage2;
 class FAsyncLoadingThread2;
 
-class FSimpleArchive
+class FSimpleArchive final
 	: public FArchive
 {
 public:
 	FSimpleArchive(const uint8* BufferPtr, uint64 BufferSize)
 	{
-#if DEVIRTUALIZE_FLinkerLoad_Serialize
 		ActiveFPLB->OriginalFastPathLoadBuffer = BufferPtr;
 		ActiveFPLB->StartFastPathLoadBuffer = BufferPtr;
 		ActiveFPLB->EndFastPathLoadBuffer = BufferPtr + BufferSize;
-#endif
 	}
 
 	int64 TotalSize() override
 	{
-#if DEVIRTUALIZE_FLinkerLoad_Serialize
 		return ActiveFPLB->EndFastPathLoadBuffer - ActiveFPLB->OriginalFastPathLoadBuffer;
-#else
-		return 0;
-#endif
 	}
 
 	int64 Tell() override
 	{
-#if DEVIRTUALIZE_FLinkerLoad_Serialize
 		return ActiveFPLB->StartFastPathLoadBuffer - ActiveFPLB->OriginalFastPathLoadBuffer;
-#else
-		return 0;
-#endif
 	}
 
 	void Seek(int64 Position) override
 	{
-#if DEVIRTUALIZE_FLinkerLoad_Serialize
 		ActiveFPLB->StartFastPathLoadBuffer = ActiveFPLB->OriginalFastPathLoadBuffer + Position;
 		check(ActiveFPLB->StartFastPathLoadBuffer <= ActiveFPLB->EndFastPathLoadBuffer);
-#endif
 	}
 
 	void Serialize(void* Data, int64 Length) override
 	{
-#if DEVIRTUALIZE_FLinkerLoad_Serialize
 		if (!Length || IsError())
 		{
 			return;
@@ -155,58 +221,9 @@ public:
 		check(ActiveFPLB->StartFastPathLoadBuffer + Length <= ActiveFPLB->EndFastPathLoadBuffer);
 		FMemory::Memcpy(Data, ActiveFPLB->StartFastPathLoadBuffer, Length);
 		ActiveFPLB->StartFastPathLoadBuffer += Length;
-#endif
 	}
 };
 
-FArchive& operator<<(FArchive& Ar, FExportBundleEntry& ExportBundleEntry)
-{
-	Ar << ExportBundleEntry.LocalExportIndex;
-	Ar << ExportBundleEntry.CommandType;
-
-	return Ar;
-}
-
-FArchive& operator<<(FArchive& Ar, FExportBundleHeader& ExportBundleHeader)
-{
-	Ar << ExportBundleHeader.FirstEntryIndex;
-	Ar << ExportBundleHeader.EntryCount;
-
-	return Ar;
-}
-
-FArchive& operator<<(FArchive& Ar, FExportMapEntry& ExportMapEntry)
-{
-	Ar << ExportMapEntry.SerialSize;
-	Ar << ExportMapEntry.ObjectName[0] << ExportMapEntry.ObjectName[1];
-	Ar << ExportMapEntry.OuterIndex;
-	Ar << ExportMapEntry.ClassIndex;
-	Ar << ExportMapEntry.SuperIndex;
-	Ar << ExportMapEntry.TemplateIndex;
-	Ar << ExportMapEntry.GlobalImportIndex;
-
-	uint32 ObjectFlags = uint32(ExportMapEntry.ObjectFlags);
-	Ar << ObjectFlags;
-	
-	if (Ar.IsLoading())
-	{
-		ExportMapEntry.ObjectFlags = EObjectFlags(ObjectFlags);
-	}
-
-	uint8 FilterFlags = uint8(ExportMapEntry.FilterFlags);
-	Ar << FilterFlags;
-
-	if (Ar.IsLoading())
-	{
-		ExportMapEntry.FilterFlags = EExportFilterFlags(FilterFlags);
-	}
-
-	uint64 Pad = 0;
-	Ar.Serialize(&Pad, 7);
-
-	return Ar;
-}
-
 struct FExportObject
 {
 	UObject* Object = nullptr;
@@ -364,25 +381,6 @@ private:
 	}
 };
 
-struct FPackageStoreEntry
-{
-	FName Name;
-	FPackageId SourcePackageId;
-	int32 ExportCount;
-	int32 ExportBundleCount;
-	int32 FirstExportBundleIndex;
-	int32 FirstGlobalImport;
-	int32 GlobalImportCount;
-	int32 ImportedPackagesCount;
-	int32 ImportedPackagesOffset;
-};
-
-struct FPackageStoreInitialLoadEntry
-{
-	int32 ScriptArcsOffset;
-	int32 ScriptArcsCount;
-};
-
 class FGlobalNameMap
 {
 public:
@@ -413,17 +411,16 @@ public:
 		IoDispatcher.FreeBatch(Batch);
 	}
 
-	FName GetName(const uint32 NameIndex, const uint32 NameNumber) const
+	FName GetName(const FNameMapIndex& NameMapIndex) const
 	{
-		FNameEntryId NameEntry = NameEntries[NameIndex];
-		return FName::CreateFromDisplayId(NameEntry, NameNumber);
+		FNameEntryId NameEntry = NameEntries[NameMapIndex.NameEntryIndex];
+		return FName::CreateFromDisplayId(NameEntry, NameMapIndex.Number);
 	}
 
-	FName FromSerializedName(const FName& SerializedName) const
+	FNameMapEntry FromSerializedNameMapEntry(const FNameMapEntry& SerializedNameMapEntry) const
 	{
-		const int32 EntryIndex = SerializedName.GetComparisonIndex().ToUnstableInt();
-		FNameEntryId NameEntry = NameEntries[EntryIndex];
-		return FName::CreateFromDisplayId(NameEntry, SerializedName.GetNumber());
+		FNameEntryId NameEntry = NameEntries[SerializedNameMapEntry.Index];
+		return FNameMapEntry { NameEntry.ToUnstableInt(), SerializedNameMapEntry.Number };
 	}
 
 	const TArray<FNameEntryId>& GetNameEntries() const
@@ -437,13 +434,85 @@ private:
 
 struct FGlobalImportStore
 {
-	/** Serialized data */
-	int32 Count = 0;
-	int32 ScriptImportCount = 0;
-	FName* Names = nullptr;
-	FPackageIndex* ScriptImportOuters = nullptr;
-	/** Runtime data */
-	FGlobalImport* Imports = nullptr;
+	TArray<FGlobalImport> Imports[FPackageObjectIndex::ImportTypeCount];
+	TArray<FScriptObjectEntry> ScriptObjectEntries;
+	TArray<FName> PackageImportNames;
+
+	inline FName GetName(FPackageObjectIndex GlobalIndex)
+	{
+		check(GlobalIndex.IsImport());
+		FName Name;
+		if (GlobalIndex.IsPackageImport())
+		{
+			Name = PackageImportNames[GlobalIndex.GetIndex()];
+		}
+		else if (GlobalIndex.IsScriptImport() && ScriptObjectEntries.Num() > 0)
+		{
+			Name = ScriptObjectEntries[GlobalIndex.GetIndex()].ObjectName.ToFName();
+		}
+		return Name;
+	}
+
+	inline FGlobalImport& GetImport(FPackageObjectIndex GlobalIndex)
+	{
+		check(GlobalIndex.IsImport());
+		FGlobalImport& Import = Imports[GlobalIndex.GetType()][GlobalIndex.GetIndex()];
+		return Import;
+	}
+
+	inline UObject* GetImportObject(FPackageObjectIndex GlobalIndex)
+	{
+		FGlobalImport& Import = GetImport(GlobalIndex);
+		return Import.GetObject();
+	}
+
+	UObject* FindScriptImportObjectFromIndex(int32 ScriptImportIndex);
+
+	FORCENOINLINE UObject* FindOrCreateScriptObject(int32 ScriptImportIndex)
+	{
+		UObject* Object = FindScriptImportObjectFromIndex(ScriptImportIndex);
+		if (!Object)
+		{
+			const FScriptObjectEntry& Entry = ScriptObjectEntries[ScriptImportIndex];
+			const FPackageObjectIndex& CDOClassIndex = Entry.CDOClassIndex;
+			if (CDOClassIndex.IsScriptImport())
+			{
+				UObject* CDOClassObject = FindScriptImportObjectFromIndex(CDOClassIndex.GetIndex());
+				if (CDOClassObject)
+				{
+					// UObjectLoadAllCompiledInDefaultProperties is creating CDOs from a flat list.
+					// One CDO may call LoadObject which may depend on a CDO later in the list, then just create it here.
+					// Recursive loads may be a problem, PostCDOConstruct exists as a workaround.
+					UClass* Class = CastChecked<UClass>(CDOClassObject);
+					UObject* CDO = Class->GetDefaultObject();
+					(void)CDO;
+					Object = FindScriptImportObjectFromIndex(ScriptImportIndex);
+				}
+			}
+		}
+		return Object;
+	}
+
+	inline UObject* FindOrGetImportObject(FPackageObjectIndex GlobalIndex)
+	{
+		check(GlobalIndex.IsImport());
+		UObject* Object = GetImportObject(GlobalIndex);
+		if (!Object && GlobalIndex.IsScriptImport() && GIsInitialLoad)
+		{
+			return FindOrCreateScriptObject(GlobalIndex.GetIndex());
+		}
+		return Object;
+	}
+
+	void StoreGlobalImportObject(FPackageObjectIndex GlobalIndex, UObject* InObject)
+	{
+		if (GlobalIndex.IsImport())
+		{
+			FGlobalImport& Import = Imports[GlobalIndex.GetType()][GlobalIndex.GetIndex()];
+			Import.SetObject(InObject);
+		}
+	}
+
 	/** Reference tracking for GC management */
 	TArray<UObject*> KeepAliveObjects;
 	bool bNeedToHandleGarbageCollect = false;
@@ -461,14 +530,8 @@ public:
 	TMap<FName, FPackageId> PackageNameToPackageId;
 
 	FGlobalImportStore ImportStore;
-	FIoBuffer InitialLoadIoBuffer;
 	FPackageStoreEntry* StoreEntries = nullptr;
-	FPackageStoreInitialLoadEntry* InitialLoadStoreEntries = nullptr;
-	int32* ImportedPackages = nullptr;
-	FExportBundleMetaEntry* ExportBundleMetaEntries = nullptr;
-	int32* ScriptArcs = nullptr;
 	int32 PackageCount = 0;
-	int32 ImportedPackagesCount = 0;
 	int32 ScriptArcsCount = 0;
 
 public:
@@ -477,6 +540,9 @@ public:
 		FIoBuffer IoBuffer;
 		FEvent* Event = FPlatformProcess::GetSynchEventFromPool();
 
+		FIoBuffer InitialLoadIoBuffer;
+		FEvent* InitialLoadEvent = FPlatformProcess::GetSynchEventFromPool();
+
 		{
 			TRACE_CPUPROFILER_EVENT_SCOPE(LoadPackageStoreTocIo);
 
@@ -489,6 +555,15 @@ public:
 					Event->Trigger();
 				});
 
+			IoDispatcher.ReadWithCallback(
+				CreateIoChunkId(0, 0, EIoChunkType::LoaderInitialLoadMeta),
+				FIoReadOptions(),
+				[InitialLoadEvent, &InitialLoadIoBuffer](TIoStatusOr<FIoBuffer> Result)
+				{
+					InitialLoadIoBuffer = Result.ConsumeValueOrDie();
+					InitialLoadEvent->Trigger();
+				});
+
 			Event->Wait();
 		}
 		
@@ -498,9 +573,9 @@ public:
 			TRACE_CPUPROFILER_EVENT_SCOPE(LoadPackageStoreTocFixup);
 
 			int32 PackageByteCount = 0;
+			GlobalMetaAr << PackageCount;
 			GlobalMetaAr << PackageByteCount;
 
-			PackageCount = PackageByteCount / sizeof(FPackageStoreEntry);
 			StoreEntries = reinterpret_cast<FPackageStoreEntry*>(FMemory::Malloc(PackageByteCount));
 
 			// In-place loading
@@ -511,7 +586,7 @@ public:
 				TRACE_CPUPROFILER_EVENT_SCOPE(LoadPackageStoreFNameFixup);
 				for (int I = 0; I < PackageCount; ++I)
 				{
-					StoreEntries[I].Name = GlobalNameMap.FromSerializedName(StoreEntries[I].Name);
+					StoreEntries[I].Name = GlobalNameMap.FromSerializedNameMapEntry(StoreEntries[I].Name);
 				}
 			}
 
@@ -522,60 +597,33 @@ public:
 				PackageNameToPackageId.Reserve(PackageCount + PackageCount / 10);
 				for (int I = 0; I < PackageCount; ++I)
 				{
-					PackageNameToPackageId.Add(StoreEntries[I].Name, FPackageId::FromIndex(I));
+					PackageNameToPackageId.Add(StoreEntries[I].Name.ToFName(), FPackageId::FromIndex(I));
 				}
 			}
 		}
 
 		{
-			TRACE_CPUPROFILER_EVENT_SCOPE(LoadPackageStoreImportedPackagesFixup);
+			TRACE_CPUPROFILER_EVENT_SCOPE(LoadPackageStoreGlobalImports);
 
-			int32 PackageImportByteCount = 0;
-			GlobalMetaAr << PackageImportByteCount;
+			int32 Count = 0;
+			GlobalMetaAr << Count; 
 
-			ImportedPackagesCount = PackageImportByteCount / sizeof(int32);
-			ImportedPackages = reinterpret_cast<int32*>(FMemory::Malloc(PackageImportByteCount));
-			GlobalMetaAr.Serialize(ImportedPackages, PackageImportByteCount);
-		}
-
-		{
-			TRACE_CPUPROFILER_EVENT_SCOPE(LoadPackageStoreGlobalImportNamesFixup);
+			TArray<FGlobalImport>& Imports = ImportStore.Imports[FPackageObjectIndex::PackageImport];
+			TArray<FName>& Names = ImportStore.PackageImportNames;
+			Imports.AddDefaulted(Count);
+			Names.Reserve(Count);
 
-			int32 ImportByteCount = 0;
-			GlobalMetaAr << ImportByteCount; 
-
-			ImportStore.Count = ImportByteCount / sizeof(FName);
-
-			FName* SerializedNames = reinterpret_cast<FName*>(IoBuffer.Data() + GlobalMetaAr.Tell());
-			GlobalMetaAr.Seek(GlobalMetaAr.Tell() + ImportByteCount);
-			ImportStore.Names = reinterpret_cast<FName*>(FMemory::Malloc(ImportByteCount));
-
-			for (int I = 0; I < ImportStore.Count; ++I)
+			FNameMapEntry* NameMapEntries = reinterpret_cast<FNameMapEntry*>(IoBuffer.Data() + GlobalMetaAr.Tell());
+			for (int I = 0; I < Count; ++I)
 			{
-				ImportStore.Names[I] = GlobalNameMap.FromSerializedName(SerializedNames[I]);
+				Names.Emplace(NameMapEntries[I].ToFName());
 			}
+			GlobalMetaAr.Seek(GlobalMetaAr.Tell() + Count * sizeof(FName));
 		}
 
-		{
-			TRACE_CPUPROFILER_EVENT_SCOPE(LoadPackageStoreInitGlobalImports);
-			ImportStore.Imports = new FGlobalImport[ImportStore.Count];
-		}
-
-		{
-			TRACE_CPUPROFILER_EVENT_SCOPE(LoadPackageStoreBundleMeta);
-			
-			int32 BundleMetaByteCount;
-			GlobalMetaAr << BundleMetaByteCount;
-			ExportBundleMetaEntries = reinterpret_cast<FExportBundleMetaEntry*>(FMemory::Malloc(BundleMetaByteCount));
-			GlobalMetaAr.Serialize(ExportBundleMetaEntries, BundleMetaByteCount);
-		}
-
-
 		{
 			TRACE_CPUPROFILER_EVENT_SCOPE(LoadPackageStoreLocalization);
 
-			using FLocalizedPackageIdMap = TMap<int32, int32>;
-			using FCulturePackageMap = TMap<FString, FLocalizedPackageIdMap>;
 			FCulturePackageMap CulturePackageMap;
 
 			GlobalMetaAr << CulturePackageMap;
@@ -583,83 +631,71 @@ public:
 			FString CurrentCultureName = FInternationalization::Get().GetCurrentCulture()->GetName();
 			FParse::Value(FCommandLine::Get(), TEXT("CULTURE="), CurrentCultureName);
 
-			FLocalizedPackageIdMap* LocalizedPackageIdMap = CulturePackageMap.Find(CurrentCultureName);
+			FSourceToLocalizedPackageIdMap* LocalizedPackageIdMap = CulturePackageMap.Find(CurrentCultureName);
 			if (LocalizedPackageIdMap)
 			{
 				for (auto It = LocalizedPackageIdMap->CreateIterator(); It; ++It)
 				{
-					const FName& SourceName = StoreEntries[It.Key()].Name;
+					const FName SourceName = StoreEntries[It.Key().ToIndex()].Name.ToFName();
 					FPackageId* FoundPackageId = PackageNameToPackageId.Find(SourceName);
 
 					UE_CLOG(!FoundPackageId, LogStreaming, Warning,
 						TEXT("Skip remapping for localized package %s (%d) since the source package %s (%d) is unknown."),
-						It.Value(),
-						*StoreEntries[It.Value()].Name.ToString(),
-						It.Key(),
+						It.Value().ToIndex(),
+						*StoreEntries[It.Value().ToIndex()].Name.ToFName().ToString(),
+						It.Key().ToIndex(),
 						*SourceName.ToString());
 
 					if (FoundPackageId)
 					{
-						const FPackageId LocalizedPackageId = FPackageId::FromIndex(It.Value());
+						const FPackageId LocalizedPackageId = It.Value();
 						*FoundPackageId = LocalizedPackageId;
 
 #if ALT2_LOG_VERBOSE
 						UE_LOG(LogStreaming, Verbose, TEXT("LoadPackageStore: RemapLocalizedPackage: %s (%d) %s (%d)."),
-							*StoreEntries[It.Key()].Name.ToString(),
-							It.Key(),
-							*StoreEntries[It.Value()].Name.ToString(),
-							It.Value());
+							*StoreEntries[It.Key().ToIndex()].Name.ToFName().ToString(),
+							It.Key().ToIndex(),
+							*StoreEntries[It.Value().ToIndex()].Name.ToFName().ToString(),
+							It.Value().ToIndex());
 #endif
 					}
 				}
 
-				for (int32 I = 0; I < ImportedPackagesCount; ++I)
+				for (int I = 0; I < PackageCount; ++I)
 				{
-					if (int32* Value = LocalizedPackageIdMap->Find(ImportedPackages[I]))
+					FPackageStoreEntry& StoreEntry = StoreEntries[I];
+					for (FPackageId& ImportedPackageIndex : StoreEntry.ImportedPackages)
 					{
-						ImportedPackages[I] = *Value;
+						if (FPackageId* Value = LocalizedPackageIdMap->Find(ImportedPackageIndex))
+						{
+							ImportedPackageIndex = *Value;
+						}
 					}
 				}
-
 			}
 		}
 
 		// Load initial loading meta data
 		{
-			FIoBuffer TempIoBuffer;
 			TRACE_CPUPROFILER_EVENT_SCOPE(LoadPackageStoreInitialLoadIo);
-			IoDispatcher.ReadWithCallback(
-				CreateIoChunkId(0, 0, EIoChunkType::LoaderInitialLoadMeta),
-				FIoReadOptions(),
-				[Event, &TempIoBuffer](TIoStatusOr<FIoBuffer> Result)
-				{
-					TempIoBuffer = Result.ConsumeValueOrDie();
-					Event->Trigger();
-				});
-
-			Event->Wait();
-			FPlatformProcess::ReturnSynchEventToPool(Event);
-			InitialLoadIoBuffer = TempIoBuffer;
+			InitialLoadEvent->Wait();
+			FPlatformProcess::ReturnSynchEventToPool(InitialLoadEvent);
 		}
 
 		{
 			TRACE_CPUPROFILER_EVENT_SCOPE(LoadPackageStoreInitialLoadFixup);
-			uint8* Data = InitialLoadIoBuffer.Data();
-
-			int32 TempPackageCount = *reinterpret_cast<int32*>(Data);
-			Data += 4;
-			ImportStore.ScriptImportCount = *reinterpret_cast<int32*>(Data);
-			Data += 4;
+			FLargeMemoryReader InitialLoadArchive(InitialLoadIoBuffer.Data(), InitialLoadIoBuffer.DataSize());
+			int32 NumScriptObjects = 0;
+			InitialLoadArchive << NumScriptObjects;
+			ImportStore.ScriptObjectEntries = MakeArrayView(reinterpret_cast<const FScriptObjectEntry*>(InitialLoadIoBuffer.Data() + InitialLoadArchive.Tell()), NumScriptObjects);
 
-			check(TempPackageCount == PackageCount);
-
-			InitialLoadStoreEntries = reinterpret_cast<FPackageStoreInitialLoadEntry*>(Data);
-			Data += TempPackageCount * sizeof(FPackageStoreInitialLoadEntry);
-
-			ImportStore.ScriptImportOuters = reinterpret_cast<FPackageIndex*>(Data);
-			Data += ImportStore.ScriptImportCount * sizeof(FPackageIndex);
+			for (FScriptObjectEntry& ScriptObjectEntry : ImportStore.ScriptObjectEntries)
+			{
+				ScriptObjectEntry.ObjectName = GlobalNameMap.FromSerializedNameMapEntry(ScriptObjectEntry.ObjectName);
+			}
 
-			ScriptArcs = reinterpret_cast<int32*>(Data);
+			TArray<FGlobalImport>& Imports = ImportStore.Imports[FPackageObjectIndex::ScriptImport];
+			Imports.AddDefaulted(ImportStore.ScriptObjectEntries.Num());
 		}
 
 		FPackageName::DoesPackageExistOverride().BindLambda([this](FName InPackageName)
@@ -671,12 +707,9 @@ public:
 	void FinalizeInitialLoad()
 	{
 		ImportStore.FindAllScriptImports();
-		InitialLoadStoreEntries = nullptr;
-		ScriptArcs = nullptr;
-		InitialLoadIoBuffer = FIoBuffer();
 
 		UE_LOG(LogStreaming, Display, TEXT("AsyncLoading2 - InitialLoad Finalized: Script Imports: %d"),
-			ImportStore.ScriptImportCount);
+			ImportStore.Imports[FPackageObjectIndex::ScriptImport].Num());
 	}
 
 	inline FGlobalImportStore& GetGlobalImportStore()
@@ -704,92 +737,31 @@ public:
 		return NewId;
 	}
 
-	inline FPackageIndex* GetGlobalImportScriptOuters(int32& OutCount) const
-	{
-		if (GIsInitialLoad)
-		{
-			OutCount = ImportStore.ScriptImportCount;
-			return ImportStore.ScriptImportOuters;
-		}
-		OutCount = 0;
-		return nullptr;
-	}
-
-	inline FName* GetGlobalImportNames(int32& OutCount) const
-	{
-		OutCount = ImportStore.Count;
-		return ImportStore.Names;
-	}
-
-	inline FGlobalImport* GetGlobalImports()
-	{
-		return ImportStore.Imports;
-	}
-
 	inline FPackageStoreEntry* GetGlobalStoreEntries()
 	{
 		return StoreEntries;
 	}
 
-	inline int32* GetPackageImportedPackages(FPackageId PackageId, int32& OutCount) const 
+	inline const FPackageStoreEntry& GetStoreEntry(FPackageId PackageId) const
 	{
-		FPackageStoreEntry& Entry = StoreEntries[PackageId.ToIndex()];
-		OutCount = Entry.ImportedPackagesCount;
-		if (Entry.ImportedPackagesCount >= 0)
-		{
-			return ImportedPackages + Entry.ImportedPackagesOffset / sizeof(int32);
-		}
-		return nullptr;
-	}
-
-	inline int32* GetPackageScriptArcs(FPackageId PackageId, int32& OutCount) const 
-	{
-		if (ensure(GIsInitialLoad))
-		{
-			FPackageStoreInitialLoadEntry& Entry = InitialLoadStoreEntries[PackageId.ToIndex()];
-			OutCount = Entry.ScriptArcsCount;
-			return ScriptArcs + Entry.ScriptArcsOffset / sizeof(int32);
-		}
-		OutCount = 0;
-		return nullptr;
-	}
-
-	inline int32 GetPackageExportCount(FPackageId PackageId) const 
-	{
-		return StoreEntries[PackageId.ToIndex()].ExportCount;
-	}
-
-	inline TTuple<FExportBundleMetaEntry*, uint32> GetPackageExportBundleMetaEntries(FPackageId PackageId) const
-	{
-		const FPackageStoreEntry* StoreEntry = StoreEntries + PackageId.ToIndex();
-		return MakeTuple(ExportBundleMetaEntries + StoreEntry->FirstExportBundleIndex, uint32(StoreEntry->ExportBundleCount));
+		return StoreEntries[PackageId.ToIndex()];
 	}
 };
 
 struct FPackageImportStore
 {
-	FPackageStoreEntry* GlobalStoreEntries = nullptr;
-	FGlobalImport* GlobalImports = nullptr;
-	const FName* GlobalImportNames = nullptr;
-	const FPackageIndex* GlobalImportScriptOuters = nullptr;
-	const int32* ImportedPackages = nullptr;
-	const int32* ImportMap = nullptr;
-	FPackageId PackageId;
-	int32 GlobalScriptImportCount = 0;
-	int32 ImportedPackagesCount = 0;
+	FPackageStore& GlobalPackageStore;
+	FGlobalImportStore& GlobalImportStore;
+	const FPackageStoreEntry& StoreEntry;
+	const FPackageObjectIndex* ImportMap = nullptr;
 	int32 ImportMapCount = 0;
 
-	FPackageImportStore(FPackageStore& GlobalPackageStore, FPackageId InPackageId) : PackageId(InPackageId)
+	FPackageImportStore(FPackageStore& InGlobalPackageStore, const FPackageStoreEntry& InStoreEntry)
+		: GlobalPackageStore(InGlobalPackageStore)
+		, GlobalImportStore(GlobalPackageStore.ImportStore)
+		, StoreEntry(InStoreEntry)
 	{
 		TRACE_CPUPROFILER_EVENT_SCOPE(NewPackageImportStore);
-		{
-			int32 GlobalImportCount = 0;
-			GlobalStoreEntries = GlobalPackageStore.GetGlobalStoreEntries();
-			GlobalImports = GlobalPackageStore.GetGlobalImports();
-			GlobalImportNames = GlobalPackageStore.GetGlobalImportNames(GlobalImportCount);
-		}
-		GlobalImportScriptOuters = GlobalPackageStore.GetGlobalImportScriptOuters(GlobalScriptImportCount);
-		ImportedPackages = GlobalPackageStore.GetPackageImportedPackages(PackageId, ImportedPackagesCount);
 		AddGlobalImportObjectReferences();
 	}
 
@@ -798,18 +770,9 @@ struct FPackageImportStore
 		check(!ImportMap);
 	}
 
-	UObject* FindImportFromGlobalIndex(int32 GlobalImportIndex);
-
-	inline UObject* GetImportFromGlobalIndex(FPackageIndex GlobalIndex)
+	inline UObject* GetImportObject(FPackageObjectIndex GlobalIndex)
 	{
-		check(GlobalIndex.IsImport());
-		const int32 GlobalImportIndex = GlobalIndex.ToImport();
-		return GlobalImports[GlobalImportIndex].GetObject();
-	}
-
-	inline UObject* GetImportFromGlobalIndex(int32 GlobalImportIndex)
-	{
-		return GlobalImports[GlobalImportIndex].GetObject();
+		return GlobalImportStore.GetImportObject(GlobalIndex);
 	}
 
 	inline bool IsValidLocalImportIndex(FPackageIndex LocalIndex)
@@ -818,27 +781,21 @@ struct FPackageImportStore
 		return LocalIndex.IsImport() && LocalIndex.ToImport() < ImportMapCount;
 	}
 
-	inline UObject* FindOrGetImportFromLocalIndex(FPackageIndex LocalIndex)
+	inline UObject* FindOrGetImportObjectFromLocalIndex(FPackageIndex LocalIndex)
 	{
 		check(LocalIndex.IsImport());
 		check(ImportMap);
 		const int32 LocalImportIndex = LocalIndex.ToImport();
 		check(LocalImportIndex < ImportMapCount);
-		const int32 GlobalImportIndex = ImportMap[LocalIndex.ToImport()];
-		return
-			GIsInitialLoad && GlobalImportIndex < GlobalScriptImportCount ?
-			FindImportFromGlobalIndex(GlobalImportIndex) :
-			GetImportFromGlobalIndex(GlobalImportIndex);
+		const FPackageObjectIndex GlobalIndex = ImportMap[LocalIndex.ToImport()];
+		check(GlobalIndex.IsImport());
+		return GlobalImportStore.FindOrGetImportObject(GlobalIndex);
 	}
 
-	inline UObject* FindOrGetImportFromGlobalIndex(FPackageIndex GlobalIndex)
+	inline UObject* FindOrGetImportObject(FPackageObjectIndex GlobalIndex)
 	{
 		check(GlobalIndex.IsImport());
-		const int32 GlobalImportIndex = GlobalIndex.ToImport();
-		return
-			GIsInitialLoad && GlobalImportIndex < GlobalScriptImportCount ?
-			FindImportFromGlobalIndex(GlobalImportIndex) :
-			GetImportFromGlobalIndex(GlobalImportIndex);
+		return GlobalImportStore.FindOrGetImportObject(GlobalIndex);
 	}
 
 	inline FName GetNameFromLocalIndex(FPackageIndex LocalIndex)
@@ -848,18 +805,15 @@ struct FPackageImportStore
 		const int32 LocalImportIndex = LocalIndex.ToImport();
 		if (LocalImportIndex < ImportMapCount)
 		{
-			const int32 GlobalImportIndex = ImportMap[LocalImportIndex];
-			return GlobalImportNames[GlobalImportIndex];
+			const FPackageObjectIndex GlobalIndex = ImportMap[LocalImportIndex];
+			return GlobalImportStore.GetName(GlobalIndex);
 		}
 		return FName();
 	}
 
-	void StoreGlobalImportObject(int32 GlobalImportIndex, UObject* InObject)
+	void StoreGlobalImportObject(FPackageObjectIndex GlobalIndex, UObject* InObject)
 	{
-		if (GlobalImportIndex >= 0)
-		{
-			GlobalImports[GlobalImportIndex].SetObject(InObject);
-		}
+		GlobalImportStore.StoreGlobalImportObject(GlobalIndex, InObject);
 	}
 
 	void ClearReferences()
@@ -870,52 +824,118 @@ struct FPackageImportStore
 private:
 	void AddGlobalImportObjectReferences()
 	{
-		for (int32 LocalPackageIndex = 0; LocalPackageIndex < ImportedPackagesCount; ++LocalPackageIndex)
+		for (const FPackageId& GlobalPackageId : StoreEntry.ImportedPackages)
 		{
-			int32 GlobalPackageIndex = ImportedPackages[LocalPackageIndex];
-			FPackageStoreEntry& Entry = GlobalStoreEntries[GlobalPackageIndex];
-			for (int32 ObjectImportIndex = 0; ObjectImportIndex < Entry.GlobalImportCount; ObjectImportIndex++)
+			const FPackageStoreEntry& ImportEntry = GlobalPackageStore.GetStoreEntry(GlobalPackageId);
+			for (const FPackageObjectIndex& GlobalImportIndex : ImportEntry.PublicExports)
 			{
-				int32 GlobalImportIndex = Entry.FirstGlobalImport + ObjectImportIndex;
-				GlobalImports[GlobalImportIndex].AddRef();
+				GlobalImportStore.GetImport(GlobalImportIndex).AddRef();
 			}
 		}
 		// Add ref counts to own imports to speed up OnPreGarbageCollect
-		FPackageStoreEntry& Entry = GlobalStoreEntries[PackageId.ToIndex()];
-		for (int32 ObjectImportIndex = 0; ObjectImportIndex < Entry.GlobalImportCount; ObjectImportIndex++)
+		for (const FPackageObjectIndex& GlobalImportIndex : StoreEntry.PublicExports)
 		{
-			int32 GlobalImportIndex = Entry.FirstGlobalImport + ObjectImportIndex;
-			GlobalImports[GlobalImportIndex].AddRef();
+			GlobalImportStore.GetImport(GlobalImportIndex).AddRef();
 		}
 	}
 
 	void ReleaseGlobalImportObjectReferences()
 	{
-		for (int32 LocalPackageIndex = 0; LocalPackageIndex < ImportedPackagesCount; ++LocalPackageIndex)
+		for (const FPackageId& GlobalPackageId : StoreEntry.ImportedPackages)
 		{
-			int32 GlobalPackageIndex = ImportedPackages[LocalPackageIndex];
-			FPackageStoreEntry& Entry = GlobalStoreEntries[GlobalPackageIndex];
-			for (int32 ObjectImportIndex = 0; ObjectImportIndex < Entry.GlobalImportCount; ObjectImportIndex++)
+			const FPackageStoreEntry& ImportEntry = GlobalPackageStore.GetStoreEntry(GlobalPackageId);
+			for (const FPackageObjectIndex& GlobalImportIndex : ImportEntry.PublicExports)
 			{
-				int32 GlobalImportIndex = Entry.FirstGlobalImport + ObjectImportIndex;
-				GlobalImports[GlobalImportIndex].ReleaseRef();
+				GlobalImportStore.GetImport(GlobalImportIndex).ReleaseRef();
 			}
 		}
 		// Release ref counts to own imports to speed up OnPreGarbageCollect
-		FPackageStoreEntry& Entry = GlobalStoreEntries[PackageId.ToIndex()];
-		for (int32 ObjectImportIndex = 0; ObjectImportIndex < Entry.GlobalImportCount; ObjectImportIndex++)
+		for (const FPackageObjectIndex& GlobalImportIndex : StoreEntry.PublicExports)
 		{
-			int32 GlobalImportIndex = Entry.FirstGlobalImport + ObjectImportIndex;
-			GlobalImports[GlobalImportIndex].ReleaseRef();
+			GlobalImportStore.GetImport(GlobalImportIndex).ReleaseRef();
 		}
 	}
 };
 	
-class FSimpleExportArchive
-	: public FSimpleArchive
+class FExportArchive final : public FArchive
 {
 public:
-	FSimpleExportArchive(const uint8* BufferPtr, uint64 BufferSize) : FSimpleArchive(BufferPtr, BufferSize) {}
+	FExportArchive(const uint8* AllExportDataPtr, const uint8* CurrentExportPtr, uint64 AllExportDataSize)
+	{
+		ActiveFPLB->OriginalFastPathLoadBuffer = AllExportDataPtr;
+		ActiveFPLB->StartFastPathLoadBuffer = CurrentExportPtr;
+		ActiveFPLB->EndFastPathLoadBuffer = AllExportDataPtr + AllExportDataSize;
+	}
+
+	void ExportBufferBegin(uint64 InExportCookedFileSerialOffset, uint64 InExportSerialSize)
+	{
+		CookedSerialOffset = InExportCookedFileSerialOffset;
+		BufferSerialOffset = (ActiveFPLB->StartFastPathLoadBuffer - ActiveFPLB->OriginalFastPathLoadBuffer);
+		CookedSerialSize = InExportSerialSize;
+	}
+
+	void ExportBufferEnd()
+	{
+		CookedSerialOffset = 0;
+		BufferSerialOffset = 0;
+		CookedSerialSize = 0;
+	}
+
+	void CheckBufferPosition(const TCHAR* Text, uint64 Offset = 0)
+	{
+#if DO_CHECK
+		const uint64 BufferPosition = (ActiveFPLB->StartFastPathLoadBuffer - ActiveFPLB->OriginalFastPathLoadBuffer) + Offset;
+		const bool bIsInsideExportBuffer =
+			(BufferSerialOffset <= BufferPosition) && (BufferPosition <= BufferSerialOffset + CookedSerialSize);
+
+		UE_ASYNC_PACKAGE_CLOG(
+			!bIsInsideExportBuffer,
+			Error, *PackageDesc, TEXT("FExportArchive::InvalidPosition"),
+			TEXT("%s: Position %llu is outside of the current export buffer (%lld,%lld)."),
+			Text,
+			BufferPosition,
+			BufferSerialOffset, BufferSerialOffset + CookedSerialSize);
+#endif
+	}
+
+	void Skip(int64 InBytes)
+	{
+		CheckBufferPosition(TEXT("InvalidSkip"), InBytes);
+		ActiveFPLB->StartFastPathLoadBuffer += InBytes;
+	}
+
+	virtual int64 TotalSize() override
+	{
+		return CookedHeaderSize + (ActiveFPLB->EndFastPathLoadBuffer - ActiveFPLB->OriginalFastPathLoadBuffer);
+	}
+
+	virtual int64 Tell() override
+	{
+		int64 CookedFilePosition = (ActiveFPLB->StartFastPathLoadBuffer - ActiveFPLB->OriginalFastPathLoadBuffer);
+		CookedFilePosition -= BufferSerialOffset;
+		CookedFilePosition += CookedSerialOffset;
+		return CookedFilePosition;
+	}
+
+	virtual void Seek(int64 Position) override
+	{
+		uint64 BufferPosition = (uint64)Position;
+		BufferPosition -= CookedSerialOffset;
+		BufferPosition += BufferSerialOffset;
+		ActiveFPLB->StartFastPathLoadBuffer = ActiveFPLB->OriginalFastPathLoadBuffer + BufferPosition;
+		CheckBufferPosition(TEXT("InvalidSeek"));
+	}
+
+	virtual void Serialize(void* Data, int64 Length) override
+	{
+		if (!Length || ArIsError)
+		{
+			return;
+		}
+		CheckBufferPosition(TEXT("InvalidSerialize"), (uint64)Length);
+		FMemory::Memcpy(Data, ActiveFPLB->StartFastPathLoadBuffer, Length);
+		ActiveFPLB->StartFastPathLoadBuffer += Length;
+	}
 
 	void UsingCustomVersion(const FGuid& Key) override {};
 	using FArchive::operator<<; // For visibility of the overloads we don't override
@@ -969,10 +989,10 @@ public:
 
 #if ALT2_LOG_VERBOSE
 				const FExportMapEntry& Export = ExportMap[ExportIndex];
-				FNameEntryId NameEntry = (*GlobalNameMap)[Export.ObjectName[0]];
-				FName ObjectName = FName::CreateFromDisplayId(NameEntry, Export.ObjectName[1]);
+				FNameEntryId NameEntry = (*GlobalNameMap)[Export.ObjectName.NameEntryIndex];
+				FName ObjectName = FName::CreateFromDisplayId(NameEntry, Export.ObjectName.Number);
 				UE_ASYNC_PACKAGE_CLOG_VERBOSE(!Object, VeryVerbose, *PackageDesc,
-					TEXT("FSimpleExportArchive: Object"), TEXT("Export %s at index %d is null."),
+					TEXT("FExportArchive: Object"), TEXT("Export %s at index %d is null."),
 					*ObjectName.ToString(), 
 					ExportIndex);
 #endif
@@ -986,10 +1006,10 @@ public:
 		{
 			if (ImportStore->IsValidLocalImportIndex(Index))
 			{
-				Object = ImportStore->FindOrGetImportFromLocalIndex(Index);
+				Object = ImportStore->FindOrGetImportObjectFromLocalIndex(Index);
 
 				UE_ASYNC_PACKAGE_CLOG_VERBOSE(!Object, Log, *PackageDesc,
-					TEXT("FSimpleExportArchive: Object"), TEXT("Import %s at index %d is null"),
+					TEXT("FExportArchive: Object"), TEXT("Import %s at index %d is null"),
 					*ImportStore->GetNameFromLocalIndex(Index).ToString(),
 					Index.ToImport());
 			}
@@ -1061,12 +1081,16 @@ private:
 
 	FAsyncPackageDesc2* PackageDesc = nullptr;
 	FPackageImportStore* ImportStore = nullptr;
+	TArray<FExternalReadCallback>* ExternalReadDependencies;
 	const int32* PackageNameMap = nullptr;
 	const TArray<FNameEntryId>* GlobalNameMap = nullptr;
 	const FExportObjects* Exports = nullptr;
 	const FExportMapEntry* ExportMap = nullptr;
 	int32 ExportCount = 0;
-	TArray<FExternalReadCallback>* ExternalReadDependencies;
+	uint32 CookedHeaderSize = 0;
+	uint64 CookedSerialOffset = 0;
+	uint64 CookedSerialSize = 0;
+	uint64 BufferSerialOffset = 0;
 };
 
 enum class EAsyncPackageLoadingState2 : uint8
@@ -1242,40 +1266,54 @@ struct FAsyncLoadingThreadState2
 		}
 	}
 
-	void SetTimeLimit(bool bUseTimeLimit, float TimeLimit)
+	void SetTimeLimit(bool bInUseTimeLimit, double InTimeLimit)
 	{
-
+		bUseTimeLimit = bInUseTimeLimit;
+		TimeLimit = InTimeLimit;
+		StartTime = FPlatformTime::Seconds();
 	}
 
-	bool IsTimeLimitExceeded()
+	bool IsTimeLimitExceeded(const TCHAR* InLastTypeOfWorkPerformed = nullptr, UObject* InLastObjectWorkWasPerformedOn = nullptr)
 	{
-		/*static double LastTestTime = -1.0;
 		bool bTimeLimitExceeded = false;
+
 		if (bUseTimeLimit)
 		{
 			double CurrentTime = FPlatformTime::Seconds();
-			bTimeLimitExceeded = CurrentTime - InTickStartTime > InTimeLimit;
+			bTimeLimitExceeded = CurrentTime - StartTime > TimeLimit;
 
 			if (bTimeLimitExceeded && GWarnIfTimeLimitExceeded)
 			{
-				IsTimeLimitExceededPrint(InTickStartTime, CurrentTime, LastTestTime, InTimeLimit, InLastTypeOfWorkPerformed, InLastObjectWorkWasPerformedOn);
+				IsTimeLimitExceededPrint(StartTime, CurrentTime, LastTestTime, TimeLimit, InLastTypeOfWorkPerformed, InLastObjectWorkWasPerformedOn);
 			}
+
 			LastTestTime = CurrentTime;
 		}
+
 		if (!bTimeLimitExceeded)
 		{
 			bTimeLimitExceeded = IsGarbageCollectionWaiting();
 			UE_CLOG(bTimeLimitExceeded, LogStreaming, Verbose, TEXT("Timing out async loading due to Garbage Collection request"));
 		}
-		return bTimeLimitExceeded;*/
-		return false;
+
+		return bTimeLimitExceeded;
+	}
+
+	bool UseTimeLimit()
+	{
+		return bUseTimeLimit;
 	}
 
 	FAsyncLoadEventGraphAllocator& GraphAllocator;
 	TArray<TTuple<FEventLoadNode2*, uint32>> DeferredFreeNodes;
 	TArray<TTuple<FEventLoadNode2**, uint32>> DeferredFreeArcs;
 	TArray<FEventLoadNode2*> NodesToFire;
+	FEventLoadNode2* CurrentEventNode = nullptr;
 	bool bShouldFireNodes = true;
+	bool bUseTimeLimit = false;
+	double TimeLimit = 0.0;
+	double StartTime = 0.0;
+	double LastTestTime = -1.0;
 	static uint32 TlsSlot;
 };
 
@@ -1292,7 +1330,6 @@ struct FAsyncPackage2
 
 	FAsyncPackage2(const FAsyncPackageDesc2& InDesc,
 		FAsyncLoadingThread2& InAsyncLoadingThread,
-		IEDLBootNotificationManager& InEDLBootNotificationManager,
 		FAsyncLoadEventGraphAllocator& InGraphAllocator,
 		const FAsyncLoadEventSpec* EventSpecs);
 	virtual ~FAsyncPackage2();
@@ -1414,11 +1451,6 @@ struct FAsyncPackage2
 	/** Returns the UPackage wrapped by this, if it is valid */
 	UPackage* GetLoadedPackage();
 
-#if WITH_EDITOR
-	/** Gets all assets loaded by this async package, used in the editor */
-	void GetLoadedAssets(TArray<FWeakObjectPtr>& AssetList);
-#endif
-
 	/** Checks if all dependencies (imported packages) of this package have been fully loaded */
 	bool AreAllDependenciesFullyLoaded(TSet<FPackageId>& VisitedPackages);
 
@@ -1443,11 +1475,15 @@ private:
 
 	/** Basic information associated with this package */
 	FAsyncPackageDesc2 Desc;
+	/** Package store information associated with this package */
+	const FPackageStoreEntry& StoreEntry;
 	/** Package which is going to have its exports and imports loaded									*/
 	UPackage*				LinkerRoot;
 	/** Call backs called when we finished loading this package											*/
 	using FCompletionCallback = TUniquePtr<FLoadPackageAsyncDelegate>;
 	TArray<FCompletionCallback, TInlineAllocator<2>> CompletionCallbacks;
+	/** Current bundle entry index in the current export bundle */
+	int32						ExportBundleEntryIndex = 0;
 	/** Current index into ExternalReadDependencies array used to spread wating for external reads over several frames			*/
 	int32						ExternalReadIndex = 0;
 	/** Current index into ObjLoaded array used to spread routing PostLoad over several frames			*/
@@ -1473,10 +1509,6 @@ private:
 
 	/** List of all request handles */
 	TArray<int32, TInlineAllocator<2>> RequestIDs;
-#if WITH_EDITORONLY_DATA
-	/** Index of the meta-data object within the linkers export table (unset if not yet processed, although may still be INDEX_NONE if there is no meta-data) */
-	TOptional<int32> MetaDataIndex;
-#endif // WITH_EDITORONLY_DATA
 	/** Number of times we recursed to load this package. */
 	int32 ReentryCount;
 	TArray<FAsyncPackage2*> ImportedAsyncPackages;
@@ -1484,7 +1516,6 @@ private:
 	TArray<UObject*> OwnedObjects;
 	/** Cached async loading thread object this package was created by */
 	FAsyncLoadingThread2& AsyncLoadingThread;
-	IEDLBootNotificationManager& EDLBootNotificationManager;
 	FAsyncLoadEventGraphAllocator& GraphAllocator;
 
 	FEventLoadNode2* PackageNodes = nullptr;
@@ -1492,7 +1523,9 @@ private:
 	uint32 ExportBundleNodeCount = 0;
 
 	FIoBuffer IoBuffer;
-	const uint8* SerialDataPtr = nullptr;
+	const uint8* CurrentExportDataPtr = nullptr;
+	const uint8* AllExportDataPtr = nullptr;
+	uint32 CookedHeaderSize = 0;
 
 	// FZenLinkerLoad
 	TArray<FExternalReadCallback> ExternalReadDependencies;
@@ -1502,7 +1535,6 @@ private:
 	FExportObjects Exports;
 	FPackageImportStore ImportStore;
 
-	FExportBundleMetaEntry* ExportBundleMetaEntries = nullptr;
 	int32 ExportBundleCount = 0;
 	const FExportBundleHeader* ExportBundles = nullptr;
 	const FExportBundleEntry* ExportBundleEntries = nullptr;
@@ -1532,11 +1564,11 @@ public:
 	static EAsyncPackageState::Type Event_Delete(FAsyncPackage2* Package, int32);
 
 	void EventDrivenCreateExport(int32 LocalExportIndex);
-	void EventDrivenSerializeExport(int32 LocalExportIndex, FSimpleExportArchive& Ar);
+	void EventDrivenSerializeExport(int32 LocalExportIndex, FExportArchive& Ar);
 
-	UObject* EventDrivenIndexToObject(FPackageIndex Index, bool bCheckSerialized);
+	UObject* EventDrivenIndexToObject(FPackageObjectIndex Index, bool bCheckSerialized);
 	template<class T>
-	T* CastEventDrivenIndexToObject(FPackageIndex Index, bool bCheckSerialized)
+	T* CastEventDrivenIndexToObject(FPackageObjectIndex Index, bool bCheckSerialized)
 	{
 		UObject* Result = EventDrivenIndexToObject(Index, bCheckSerialized);
 		if (!Result)
@@ -1563,7 +1595,6 @@ public:
 
 private:
 	void CreateNodes(const FAsyncLoadEventSpec* EventSpecs);
-	void SetupScriptArcs();
 	void SetupSerializedArcs(const uint8* GraphData, uint64 GraphDataSize);
 
 	/**
@@ -1686,7 +1717,7 @@ class FAsyncLoadingThread2 final
 {
 	friend struct FAsyncPackage2;
 public:
-	FAsyncLoadingThread2(FIoDispatcher& IoDispatcher, IEDLBootNotificationManager& InEDLBootNotificationManager);
+	FAsyncLoadingThread2(FIoDispatcher& IoDispatcher);
 	virtual ~FAsyncLoadingThread2();
 
 private:
@@ -1701,6 +1732,10 @@ private:
 	/** [ASYNC/GAME THREAD] true if the async thread is actually started. We don't start it until after we boot because the boot process on the game thread can create objects that are also being created by the loader */
 	bool bThreadStarted = false;
 
+#if ALT2_VERIFY_RECURSIVE_LOADS
+	int32 LoadRecursionLevel = 0;
+#endif
+
 	/** [ASYNC/GAME THREAD] Event used to signal loading should be cancelled */
 	FEvent* CancelLoadingEvent;
 	/** [ASYNC/GAME THREAD] Event used to signal that the async loading thread should be suspended */
@@ -1717,15 +1752,17 @@ private:
 	FCriticalSection LoadedPackagesCritical;
 	TArray<FAsyncPackage2*> LoadedPackagesToProcess;
 	TArray<FAsyncPackage2*> PackagesToDelete;
-#if WITH_EDITOR
-	TArray<FWeakObjectPtr> LoadedAssets;
-#endif
+	
+	struct FQueuedFailedPackageCallback
+	{
+		FName PackageName;
+		TUniquePtr<FLoadPackageAsyncDelegate> Callback;
+	};
+	TArray<FQueuedFailedPackageCallback> QueuedFailedPackageCallbacks;
 
 	FCriticalSection AsyncPackagesCritical;
 	TMap<FPackageId, FAsyncPackage2*> AsyncPackageLookup;
 
-	IEDLBootNotificationManager& EDLBootNotificationManager;
-
 	TQueue<FAsyncPackage2*, EQueueMode::Mpsc> ExternalReadQueue;
 	FThreadSafeCounter WaitingForIoBundleCounter;
 	FThreadSafeCounter WaitingForPostLoadCounter;
@@ -2048,7 +2085,7 @@ private:
 	{
 		if (Desc.PackageIdToLoad.ToIndex() < (uint32)GlobalPackageStore.PackageCount)
 		{
-			return new FAsyncPackage2(Desc, *this, EDLBootNotificationManager, GraphAllocator, EventSpecs.GetData());
+			return new FAsyncPackage2(Desc, *this, GraphAllocator, EventSpecs.GetData());
 		}
 		else
 		{
@@ -2136,9 +2173,10 @@ void FAsyncLoadingThread2::InitializeLoading()
 
 	AsyncThreadReady.Increment();
 
-	UE_LOG(LogStreaming, Display, TEXT("AsyncLoading2 - Initialized: Packages: %d, GlobalImports: %d, FNames: %d"),
+	UE_LOG(LogStreaming, Display, TEXT("AsyncLoading2 - Initialized: Packages: %d, PackageImports: %d, ScriptImports %d, FNames: %d"),
 		GlobalPackageStore.PackageCount,
-		GlobalPackageStore.ImportStore.Count,
+		GlobalPackageStore.ImportStore.Imports[FPackageObjectIndex::PackageImport].Num(),
+		GlobalPackageStore.ImportStore.Imports[FPackageObjectIndex::ScriptImport].Num(),
 		GlobalNameMap.GetNameEntries().Num());
 }
 
@@ -2188,53 +2226,73 @@ FAsyncPackage2* FAsyncLoadingThread2::FindOrInsertPackage(FAsyncPackageDesc2* De
 bool FAsyncLoadingThread2::CreateAsyncPackagesFromQueue()
 {
 	TRACE_CPUPROFILER_EVENT_SCOPE(CreateAsyncPackagesFromQueue);
+	
+	FAsyncLoadingThreadState2& ThreadState = *FAsyncLoadingThreadState2::Get();
+	bool bPackagesCreated = false;
+	const int32 TimeSliceGranularity = ThreadState.UseTimeLimit() ? 4 : MAX_int32;
 	TArray<FAsyncPackageDesc2*> QueueCopy;
-	{
-		FScopeLock QueueLock(&QueueCritical);
-		QueueCopy.Append(QueuedPackages);
-		QueuedPackages.Reset();
-	}
 
-	for (FAsyncPackageDesc2* PackageDesc : QueueCopy)
+	do
 	{
-		bool bInserted;
-		FAsyncPackage2* Package = FindOrInsertPackage(PackageDesc, bInserted);
-
-		if (bInserted)
 		{
-			UE_ASYNC_PACKAGE_LOG(Verbose, *PackageDesc, TEXT("CreateAsyncPackages: AddPackage"),
-				TEXT("Start loading package."));
-		}
-		else if (!Package)
-		{
-			UE_ASYNC_PACKAGE_LOG(Warning, *PackageDesc, TEXT("CreateAsyncPackages: SkipPackage"),
-				TEXT("Skipping unknown package, probably a temp package that has already been completely loaded"));
-		}
-		else
-		{
-			UE_ASYNC_PACKAGE_LOG_VERBOSE(Verbose, *PackageDesc, TEXT("CreateAsyncPackages: UpdatePackage"),
-				TEXT("Package is alreay being loaded."));
-		}
+			QueueCopy.Reset();
+			FScopeLock QueueLock(&QueueCritical);
 
-		--QueuedPackagesCounter;
-		if (Package)
-		{
-			TRACE_CPUPROFILER_EVENT_SCOPE(ImportPackages);
-			Package->ImportPackagesRecursive();
-		}
-		if (bInserted)
-		{
-			Package->StartLoading();
+			const int32 NumPackagesToCopy = FMath::Min(TimeSliceGranularity, QueuedPackages.Num());
+			if (NumPackagesToCopy > 0)
+			{
+				QueueCopy.Append(QueuedPackages.GetData(), NumPackagesToCopy);
+				QueuedPackages.RemoveAt(0, NumPackagesToCopy, false);
+			}
+			else
+			{
+				break;
+			}
 		}
-		delete PackageDesc;
 
-		if (Package)
+		for (FAsyncPackageDesc2* PackageDesc : QueueCopy)
 		{
-			StartBundleIoRequests();
+			bool bInserted;
+			FAsyncPackage2* Package = FindOrInsertPackage(PackageDesc, bInserted);
+
+			if (bInserted)
+			{
+				UE_ASYNC_PACKAGE_LOG(Verbose, *PackageDesc, TEXT("CreateAsyncPackages: AddPackage"),
+					TEXT("Start loading package."));
+			}
+			else if (!Package)
+			{
+				UE_ASYNC_PACKAGE_LOG(Warning, *PackageDesc, TEXT("CreateAsyncPackages: SkipPackage"),
+					TEXT("Skipping unknown package, probably a temp package that has already been completely loaded"));
+			}
+			else
+			{
+				UE_ASYNC_PACKAGE_LOG_VERBOSE(Verbose, *PackageDesc, TEXT("CreateAsyncPackages: UpdatePackage"),
+					TEXT("Package is alreay being loaded."));
+			}
+
+			--QueuedPackagesCounter;
+			if (Package)
+			{
+				{
+					TRACE_CPUPROFILER_EVENT_SCOPE(ImportPackages);
+					Package->ImportPackagesRecursive();
+				}
+
+				if (bInserted)
+				{
+					Package->StartLoading();
+				}
+
+				StartBundleIoRequests();
+			}
+			delete PackageDesc;
 		}
-	}
 
-	return QueueCopy.Num() > 0;
+		bPackagesCreated |= QueueCopy.Num() > 0;
+	} while (!ThreadState.IsTimeLimitExceeded(TEXT("CreateAsyncPackagesFromQueue")));
+
+	return bPackagesCreated;
 }
 
 void FAsyncLoadingThread2::AddBundleIoRequest(FAsyncPackage2* Package, const FExportBundleMetaEntry& BundleMetaEntry)
@@ -2385,7 +2443,7 @@ void FEventLoadNode2::Fire()
 #endif
 
 	FAsyncLoadingThreadState2* ThreadState = FAsyncLoadingThreadState2::Get();
-	if (Spec->bExecuteImmediately && ThreadState)
+	if (Spec->bExecuteImmediately && ThreadState && !ThreadState->CurrentEventNode)
 	{
 		Execute(*ThreadState);
 	}
@@ -2399,10 +2457,16 @@ void FEventLoadNode2::Execute(FAsyncLoadingThreadState2& ThreadState)
 {
 	//TRACE_CPUPROFILER_EVENT_SCOPE(ExecuteEvent);
 	check(BarrierCount.Load() == 0);
+	check(!ThreadState.CurrentEventNode || ThreadState.CurrentEventNode == this);
+
+	ThreadState.CurrentEventNode = this;
 	EAsyncPackageState::Type State = Spec->Func(Package, ImportOrExportIndex);
-	check(State == EAsyncPackageState::Complete);
-	bDone.Store(1);
-	ProcessDependencies(ThreadState);
+	if (State == EAsyncPackageState::Complete)
+	{
+		ThreadState.CurrentEventNode = nullptr;
+		bDone.Store(1);
+		ProcessDependencies(ThreadState);
+	}
 }
 
 void FEventLoadNode2::ProcessDependencies(FAsyncLoadingThreadState2& ThreadState)
@@ -2476,6 +2540,13 @@ void FAsyncLoadEventQueue2::Push(FEventLoadNode2* Node)
 
 bool FAsyncLoadEventQueue2::PopAndExecute(FAsyncLoadingThreadState2& ThreadState)
 {
+	if (ThreadState.CurrentEventNode)
+	{
+		check(!ThreadState.CurrentEventNode->IsDone());
+		ThreadState.CurrentEventNode->Execute(ThreadState);
+		return true;
+	}
+
 	FEventLoadNode2* Node = nullptr;
 	{
 		uint64 LocalHead = Head.Load();
@@ -2642,70 +2713,68 @@ void FAsyncPackage2::SetupSerializedArcs(const uint8* GraphData, uint64 GraphDat
 }
 
 static UObject* GFindExistingScriptImport(int32 GlobalImportIndex,
-	FGlobalImport* GlobalImports,
-	const FPackageIndex* GlobalImportScriptOuters,
-	const FName* GlobalImportNames)
+	TArray<FGlobalImport>& ScriptImports,
+	const TArray<FScriptObjectEntry>& ScriptObjectEntries)
 {
-	FGlobalImport& GlobalImport = GlobalImports[GlobalImportIndex];
-	UObject* Object = GlobalImport.GetObject();
+	FGlobalImport& ScriptImport = ScriptImports[GlobalImportIndex];
+	UObject* Object = ScriptImport.GetObject();
 	if (!Object)
 	{
-		const FPackageIndex& OuterIndex = GlobalImportScriptOuters[GlobalImportIndex];
-		const FName& ObjectName = GlobalImportNames[GlobalImportIndex];
-		if (OuterIndex.IsNull())
+		const FScriptObjectEntry& Entry = ScriptObjectEntries[GlobalImportIndex];
+		if (Entry.OuterIndex.IsNull())
 		{
-			Object = StaticFindObjectFast(UPackage::StaticClass(), nullptr, ObjectName, true);
+			Object = StaticFindObjectFast(UPackage::StaticClass(), nullptr, Entry.ObjectName.ToFName(), true);
 		}
 		else
 		{
-			UObject* Outer = GFindExistingScriptImport(OuterIndex.ToImport(), GlobalImports, GlobalImportScriptOuters, GlobalImportNames);
+			UObject* Outer = GFindExistingScriptImport(Entry.OuterIndex.GetIndex(), ScriptImports, ScriptObjectEntries);
 			if (Outer)
 			{
-				Object = StaticFindObjectFast(UObject::StaticClass(), Outer, ObjectName, false, true);
+				Object = StaticFindObjectFast(UObject::StaticClass(), Outer, Entry.ObjectName.ToFName(), false, true);
 			}
 		}
-		GlobalImport.SetObject(Object);
+		ScriptImport.SetObject(Object);
 	}
 	return Object;
 }
 
-UObject* FPackageImportStore::FindImportFromGlobalIndex(int32 GlobalImportIndex)
+UObject* FGlobalImportStore::FindScriptImportObjectFromIndex(int32 GlobalImportIndex)
 {
-	check(GlobalImportScriptOuters);
-	check(GlobalImportIndex < GlobalScriptImportCount);
+	check(ScriptObjectEntries.Num() > 0);
 	return GFindExistingScriptImport(
-		GlobalImportIndex,
-		GlobalImports,
-		GlobalImportScriptOuters,
-		GlobalImportNames);
+		GlobalImportIndex, Imports[FPackageObjectIndex::ScriptImport], ScriptObjectEntries);
 }
 
 void FGlobalImportStore::FindAllScriptImports()
 {
 	TRACE_CPUPROFILER_EVENT_SCOPE(FindAllScriptImports);
-	check(ScriptImportOuters);
+	const int32 ScriptImportCount = ScriptObjectEntries.Num();
+	check(ScriptImportCount > 0);
 	for (int32 GlobalImportIndex = 0; GlobalImportIndex < ScriptImportCount; ++GlobalImportIndex)
 	{
-		UObject* Object = GFindExistingScriptImport(GlobalImportIndex, Imports, ScriptImportOuters, Names);
+		UObject* Object = GFindExistingScriptImport(
+			GlobalImportIndex, Imports[FPackageObjectIndex::ScriptImport], ScriptObjectEntries);
 #if DO_CHECK
 		if (!Object)
 		{
-			FPackageIndex Outer = ScriptImportOuters[GlobalImportIndex];
-			if (Outer.IsNull())
+			FScriptObjectEntry& Entry = ScriptObjectEntries[GlobalImportIndex];
+			if (Entry.OuterIndex.IsNull())
 			{
-				UE_LOG(LogStreaming, Warning, TEXT("AsyncLoading2 - Failed to find import script package after initial load: %s"),
-					*Names[GlobalImportIndex].ToString());
+				UE_LOG(LogStreaming, Warning,
+					TEXT("AsyncLoading2 - Failed to find import script package after initial load: %s"),
+					*ScriptObjectEntries[GlobalImportIndex].ObjectName.ToFName().ToString());
 			}
 			else
 			{
-				UE_LOG(LogStreaming, Warning, TEXT("AsyncLoading2 - Failed to find import script object after initial load: %s - %s"),
-					*Names[Outer.ToImport()].ToString(),
-					*Names[GlobalImportIndex].ToString());
+				UE_LOG(LogStreaming, Warning,
+					TEXT("AsyncLoading2 - Failed to find import script object after initial load: %s - %s"),
+					*ScriptObjectEntries[Entry.OuterIndex.GetIndex()].ObjectName.ToFName().ToString(),
+					*ScriptObjectEntries[GlobalImportIndex].ObjectName.ToFName().ToString());
 			}
 		}
 #endif
 	}
-	ScriptImportOuters = nullptr;
+	ScriptObjectEntries.Empty();
 }
 
 void FAsyncPackage2::ImportPackagesRecursive()
@@ -2717,24 +2786,20 @@ void FAsyncPackage2::ImportPackagesRecursive()
 	bHasImportedPackagesRecursive = true;
 
 	FPackageStore& GlobalPackageStore = AsyncLoadingThread.GlobalPackageStore;
-	int32 ImportedPackageCount = 0;
-	int32* Imports = GlobalPackageStore.GetPackageImportedPackages(Desc.PackageIdToLoad, ImportedPackageCount);
-	if (!Imports)
+	int32 ImportedPackageCount = StoreEntry.ImportedPackages.Num();
+	if (!ImportedPackageCount)
 	{
 		return;
 	}
 
-	for (int32 LocalImportIndex = 0; LocalImportIndex < ImportedPackageCount; ++LocalImportIndex)
+	for (const FPackageId& ImportedPackageId : StoreEntry.ImportedPackages)
 	{
 		// AreAllImportsInImportPackageAlreadyLoaded?
-		int32 GlobalPackageIndex = Imports[LocalImportIndex];
-		FPackageStoreEntry& Entry = GlobalPackageStore.StoreEntries[GlobalPackageIndex];
+		const FPackageStoreEntry& ImportEntry = GlobalPackageStore.GetStoreEntry(ImportedPackageId);
 		bool bNeedToLoadPackage = false;
-		// start from 1 to skip checking the UPackage
-		for (int32 ObjectImportIndex = 1; ObjectImportIndex < Entry.GlobalImportCount; ObjectImportIndex++)
+		for (const FPackageObjectIndex& GlobalImportIndex : ImportEntry.PublicExports)
 		{
-			int32 GlobalImportIndex = Entry.FirstGlobalImport + ObjectImportIndex;
-			UObject* ImportedObject = ImportStore.GetImportFromGlobalIndex(GlobalImportIndex);
+			UObject* ImportedObject = ImportStore.GetImportObject(GlobalImportIndex);
 			if (!ImportedObject || !ImportedObject->HasAllFlags(RF_WasLoaded | RF_LoadCompleted))
 			{
 				bNeedToLoadPackage = true;
@@ -2744,11 +2809,12 @@ void FAsyncPackage2::ImportPackagesRecursive()
 
 		if (bNeedToLoadPackage)
 		{
-			const FPackageId PackageIdToLoad = FPackageId::FromIndex(GlobalPackageIndex);
-			const FName NameToLoad = Entry.Name;
+			const FGlobalNameMap& GlobalNameMap = AsyncLoadingThread.GlobalNameMap;
+			const FPackageId PackageIdToLoad = ImportedPackageId;
+			const FName NameToLoad = ImportEntry.Name.ToFName();
 			const FName SourceName =
-				Entry.SourcePackageId.IsValid() ?
-				GlobalPackageStore.StoreEntries[Entry.SourcePackageId.ToIndex()].Name :
+				ImportEntry.SourcePackageId.IsValid() ?
+				GlobalPackageStore.StoreEntries[ImportEntry.SourcePackageId.ToIndex()].Name.ToFName() :
 				NameToLoad;
 
 			FAsyncPackageDesc2 PackageDesc(INDEX_NONE, PackageIdToLoad, PackageIdToLoad, SourceName, NameToLoad);
@@ -2797,50 +2863,10 @@ void FAsyncPackage2::StartLoading()
 	LoadStartTime = FPlatformTime::Seconds();
 
 	check(ExportBundleCount > 0);
-	FExportBundleMetaEntry& BundleMetaEntry = ExportBundleMetaEntries[0];
+	const FExportBundleMetaEntry& BundleMetaEntry = StoreEntry.ExportBundles[0];
 	AsyncLoadingThread.AddBundleIoRequest(this, BundleMetaEntry);
 }
 
-void FAsyncPackage2::SetupScriptArcs()
-{
-	TRACE_CPUPROFILER_EVENT_SCOPE(SetupScriptArcs);
-
-	FPackageStore& GlobalPackageStore = AsyncLoadingThread.GlobalPackageStore;
-	int32 ScriptArcsCount = 0;
-	int32* ScriptArcs = GlobalPackageStore.GetPackageScriptArcs(Desc.PackageIdToLoad, ScriptArcsCount);
-	for (int32 ScriptArcIndex = 0; ScriptArcIndex < ScriptArcsCount * 2;)
-	{
-		int32 GlobalImportIndex = ScriptArcs[ScriptArcIndex++];
-		int32 ToNodeIndex = ScriptArcs[ScriptArcIndex++];
-		FEventLoadNode2* ToNode = PackageNodes + ToNodeIndex;
-
-		FPackageIndex OuterMostIndex = FPackageIndex::FromImport(GlobalImportIndex);
-		FPackageIndex OuterMostNonPackageIndex = OuterMostIndex;
-		while (true)
-		{
-			check(!OuterMostIndex.IsNull() && OuterMostIndex.IsImport());
-			FPackageIndex NextOuterMostIndex = ImportStore.GlobalImportScriptOuters[OuterMostIndex.ToImport()];
-
-			if (NextOuterMostIndex.IsNull())
-			{
-				break;
-			}
-			OuterMostNonPackageIndex = OuterMostIndex;
-			OuterMostIndex = NextOuterMostIndex;
-		}
-		FName PackageName = ImportStore.GlobalImportNames[OuterMostIndex.ToImport()];
-		FName OuterMostNonPackageObjectName = ImportStore.GlobalImportNames[OuterMostNonPackageIndex.ToImport()];
-		// OuterMostNonPackageIndex is used here because if it is a CDO or subobject, etc,
-		// we wait for the outermost thing that is not a package
-		const bool bWaitingForCompiledInImport = EDLBootNotificationManager.AddWaitingPackage(
-			this, PackageName, OuterMostNonPackageObjectName, FPackageIndex::FromImport(ToNodeIndex));
-		if (bWaitingForCompiledInImport)
-		{
-			PackageNodes[ToNodeIndex].AddBarrier();
-		}
-	}
-}
-
 EAsyncPackageState::Type FAsyncPackage2::Event_ProcessExportBundle(FAsyncPackage2* Package, int32 ExportBundleIndex)
 {
 	TRACE_CPUPROFILER_EVENT_SCOPE(Event_ProcessExportBundle);
@@ -2873,33 +2899,31 @@ EAsyncPackageState::Type FAsyncPackage2::Event_ProcessExportBundle(FAsyncPackage
 
 	check(ExportBundleIndex < Package->ExportBundleCount);
 
-	if (ExportBundleIndex == 0)
+	if (Package->AsyncPackageLoadingState == EAsyncPackageLoadingState2::WaitingForSummary)
 	{
 		TRACE_CPUPROFILER_EVENT_SCOPE(ProcessPackageSummary);
 
-		check(Package->AsyncPackageLoadingState == EAsyncPackageLoadingState2::WaitingForSummary);
+		check(ExportBundleIndex == 0);
+		check(Package->ExportBundleEntryIndex == 0);
 
-		Package->SerialDataPtr = Package->IoBuffer.Data();
-		const uint8* PackageSummaryData = Package->SerialDataPtr;
+		const uint8* PackageSummaryData = Package->IoBuffer.Data();
 		const FPackageSummary* PackageSummary = reinterpret_cast<const FPackageSummary*>(PackageSummaryData);
+		const uint8* GraphData = PackageSummaryData + PackageSummary->GraphDataOffset;
+		const uint64 PackageSummarySize = GraphData + PackageSummary->GraphDataSize - PackageSummaryData;
 
+		Package->CookedHeaderSize = PackageSummary->CookedHeaderSize;
 		Package->PackageNameMap = reinterpret_cast<const int32*>(PackageSummaryData + PackageSummary->NameMapOffset);
-		Package->ImportStore.ImportMap = reinterpret_cast<const int32*>(PackageSummaryData + PackageSummary->ImportMapOffset);
+		Package->ImportStore.ImportMap = reinterpret_cast<const FPackageObjectIndex*>(PackageSummaryData + PackageSummary->ImportMapOffset);
 		Package->ImportStore.ImportMapCount = (PackageSummary->ExportMapOffset - PackageSummary->ImportMapOffset) / sizeof(int32);
 		Package->ExportMap = reinterpret_cast<const FExportMapEntry*>(PackageSummaryData + PackageSummary->ExportMapOffset);
 		Package->ExportBundles = reinterpret_cast<const FExportBundleHeader*>(PackageSummaryData + PackageSummary->ExportBundlesOffset);
 		Package->ExportBundleEntries = reinterpret_cast<const FExportBundleEntry*>(Package->ExportBundles + Package->ExportBundleCount);
 
-		const uint8* GraphData = PackageSummaryData + PackageSummary->GraphDataOffset;
-
-		uint64 PackageSummarySize = GraphData + PackageSummary->GraphDataSize - PackageSummaryData;
-		Package->SerialDataPtr += PackageSummarySize;
-
 		Package->CreateUPackage(PackageSummary);
-
 		Package->SetupSerializedArcs(GraphData, PackageSummary->GraphDataSize);
 
-		check(Package->AsyncPackageLoadingState == EAsyncPackageLoadingState2::WaitingForSummary);
+		Package->AllExportDataPtr = PackageSummaryData + PackageSummarySize;
+		Package->CurrentExportDataPtr = Package->AllExportDataPtr;
 		Package->AsyncPackageLoadingState = EAsyncPackageLoadingState2::ProcessNewImportsAndExports;
 
 		TRACE_LOADTIME_PACKAGE_SUMMARY(Package, PackageSummarySize, Package->ImportStore.ImportMapCount, Package->ExportCount);
@@ -2907,8 +2931,8 @@ EAsyncPackageState::Type FAsyncPackage2::Event_ProcessExportBundle(FAsyncPackage
 
 	check(Package->AsyncPackageLoadingState == EAsyncPackageLoadingState2::ProcessNewImportsAndExports);
 
-	uint64 ExportsBufferSize = Package->IoBuffer.DataSize() - (Package->SerialDataPtr - Package->IoBuffer.Data());
-	FSimpleExportArchive Ar(Package->SerialDataPtr, ExportsBufferSize);
+	const uint64 AllExportDataSize = Package->IoBuffer.DataSize() - (Package->AllExportDataPtr - Package->IoBuffer.Data());
+	FExportArchive Ar(Package->AllExportDataPtr, Package->CurrentExportDataPtr, AllExportDataSize);
 	{
 		Ar.SetUE4Ver(Package->LinkerRoot->LinkerPackageVersion);
 		Ar.SetLicenseeUE4Ver(Package->LinkerRoot->LinkerLicenseeVersion);
@@ -2923,7 +2947,8 @@ EAsyncPackageState::Type FAsyncPackage2::Event_ProcessExportBundle(FAsyncPackage
 		}
 		Ar.ArAllowLazyLoading = true;
 
-		// FSimpleExportArchive special fields
+		// FExportArchive special fields
+		Ar.CookedHeaderSize = Package->CookedHeaderSize;
 		Ar.PackageDesc = &Package->Desc;
 		Ar.PackageNameMap = Package->PackageNameMap;
 		Ar.GlobalNameMap = &Package->AsyncLoadingThread.GlobalNameMap.GetNameEntries();
@@ -2934,61 +2959,63 @@ EAsyncPackageState::Type FAsyncPackage2::Event_ProcessExportBundle(FAsyncPackage
 		Ar.ExternalReadDependencies = &Package->ExternalReadDependencies;
 	}
 	const FExportBundleHeader* ExportBundle = Package->ExportBundles + ExportBundleIndex;
-	
-	const FExportBundleEntry* BundleEntry = Package->ExportBundleEntries + ExportBundle->FirstEntryIndex;
-	const FExportBundleEntry* BundleEntryEnd = BundleEntry + ExportBundle->EntryCount;
+	const FExportBundleEntry* BundleEntries = Package->ExportBundleEntries + ExportBundle->FirstEntryIndex;
+	const FExportBundleEntry* BundleEntry = BundleEntries + Package->ExportBundleEntryIndex;
+	const FExportBundleEntry* BundleEntryEnd = BundleEntries + ExportBundle->EntryCount;
 	check(BundleEntry <= BundleEntryEnd);
 	while (BundleEntry < BundleEntryEnd)
 	{
-		const FExportMapEntry& Export = Package->ExportMap[BundleEntry->LocalExportIndex];
-
-		if (FilterExport(Export.FilterFlags))
+		if (FAsyncLoadingThreadState2::Get()->IsTimeLimitExceeded(TEXT("Event_ProcessExportBundle")))
 		{
-			Package->Exports[BundleEntry->LocalExportIndex].bFiltered = true;
-			if (BundleEntry->CommandType == FExportBundleEntry::ExportCommandType_Serialize)
-			{
-				const uint64 SerialSize = Package->ExportMap[BundleEntry->LocalExportIndex].SerialSize;
-				Package->SerialDataPtr += SerialSize;
-				Ar.Seek(Ar.Tell() + SerialSize);
-			}
-
-			++BundleEntry;
-			continue;
+			return EAsyncPackageState::TimeOut;
 		}
+		const FExportMapEntry& ExportMapEntry = Package->ExportMap[BundleEntry->LocalExportIndex];
+		FExportObject& Export = Package->Exports[BundleEntry->LocalExportIndex];
+		Export.bFiltered = FilterExport(ExportMapEntry.FilterFlags);
 
 		if (BundleEntry->CommandType == FExportBundleEntry::ExportCommandType_Create)
 		{
-			Package->EventDrivenCreateExport(BundleEntry->LocalExportIndex);
+			if (!Export.bFiltered)
+			{
+				Package->EventDrivenCreateExport(BundleEntry->LocalExportIndex);
+			}
 		}
 		else
 		{
 			check(BundleEntry->CommandType == FExportBundleEntry::ExportCommandType_Serialize);
-			const uint64 ExportSerialSize = Export.SerialSize;
-			check(Package->SerialDataPtr + ExportSerialSize <= Package->IoBuffer.Data() + Package->IoBuffer.DataSize());
-			UObject* Object = Package->Exports[BundleEntry->LocalExportIndex].Object;
-			check(Object);
-			if (Object->HasAnyFlags(RF_NeedLoad))
+
+			const uint64 CookedSerialSize = ExportMapEntry.CookedSerialSize;
+			UObject* Object = Export.Object;
+
+			check(Package->CurrentExportDataPtr + CookedSerialSize <= Package->IoBuffer.Data() + Package->IoBuffer.DataSize());
+			check(Object || Export.bFiltered);
+
+			Ar.ExportBufferBegin(ExportMapEntry.CookedSerialOffset, ExportMapEntry.CookedSerialSize);
+			if (!Export.bFiltered && Object->HasAnyFlags(RF_NeedLoad))
 			{
-				TRACE_LOADTIME_SERIALIZE_EXPORT_SCOPE(Object, ExportSerialSize);
+				TRACE_LOADTIME_SERIALIZE_EXPORT_SCOPE(Object, CookedSerialSize);
 				const int64 Pos = Ar.Tell();
-				check(ExportSerialSize <= uint64(Ar.TotalSize() - Pos));
+				check(CookedSerialSize <= uint64(Ar.TotalSize() - Pos));
 
 				Package->EventDrivenSerializeExport(BundleEntry->LocalExportIndex, Ar);
-				checkf(ExportSerialSize == uint64(Ar.Tell() - Pos), TEXT("Expect read size: %llu - Actual read size: %llu"), ExportSerialSize, uint64(Ar.Tell() - Pos));
+				checkf(CookedSerialSize == uint64(Ar.Tell() - Pos), TEXT("Expect read size: %llu - Actual read size: %llu"), CookedSerialSize, uint64(Ar.Tell() - Pos));
 			}
 			else
 			{
-				// If we are not going to serialize the data we need to skip it!
-				Ar.Seek(Ar.Tell() + ExportSerialSize);
+				Ar.Skip(CookedSerialSize);
 			}
+			Ar.ExportBufferEnd();
 
-			check(!Object->HasAnyFlags(RF_NeedLoad));
+			check((Object && !Object->HasAnyFlags(RF_NeedLoad)) || Export.bFiltered);
 
-			Package->SerialDataPtr += ExportSerialSize;
+			Package->CurrentExportDataPtr += CookedSerialSize;
 		}
 		++BundleEntry;
+		++Package->ExportBundleEntryIndex;
 	}
 
+	Package->ExportBundleEntryIndex = 0;
+
 	if (ExportBundleIndex + 1 < Package->ExportBundleCount)
 	{
 		Package->GetExportBundleNode(ExportBundle_Process, ExportBundleIndex + 1)->ReleaseBarrier();
@@ -3014,13 +3041,15 @@ EAsyncPackageState::Type FAsyncPackage2::Event_ProcessExportBundle(FAsyncPackage
 
 	if (ExportBundleIndex == 0)
 	{
-		Package->AsyncLoadingThread.BundleIoRequestCompleted(Package->ExportBundleMetaEntries[0]);
+		check(Package->ExportBundleCount > 0);
+		const FExportBundleMetaEntry& BundleMetaEntry = Package->StoreEntry.ExportBundles[0];
+		Package->AsyncLoadingThread.BundleIoRequestCompleted(BundleMetaEntry);
 	}
 
 	return EAsyncPackageState::Complete;
 }
 
-UObject* FAsyncPackage2::EventDrivenIndexToObject(FPackageIndex Index, bool bCheckSerialized)
+UObject* FAsyncPackage2::EventDrivenIndexToObject(FPackageObjectIndex Index, bool bCheckSerialized)
 {
 	UObject* Result = nullptr;
 	if (Index.IsNull())
@@ -3033,7 +3062,7 @@ UObject* FAsyncPackage2::EventDrivenIndexToObject(FPackageIndex Index, bool bChe
 	}
 	else if (Index.IsImport())
 	{
-		Result = ImportStore.FindOrGetImportFromGlobalIndex(Index);
+		Result = ImportStore.FindOrGetImportObject(Index);
 		check(Result);
 	}
 #if DO_CHECK
@@ -3075,7 +3104,7 @@ void FAsyncPackage2::EventDrivenCreateExport(int32 LocalExportIndex)
 	{
 		TRACE_CPUPROFILER_EVENT_SCOPE(ObjectNameFixup);
 		const FGlobalNameMap& GlobalNameMap = AsyncLoadingThread.GlobalNameMap;
-		ObjectName = GlobalNameMap.GetName(Export.ObjectName[0], Export.ObjectName[1]);
+		ObjectName = GlobalNameMap.GetName(Export.ObjectName);
 	}
 
 	TRACE_LOADTIME_CREATE_EXPORT_SCOPE(this, &Object);
@@ -3219,7 +3248,7 @@ void FAsyncPackage2::EventDrivenCreateExport(int32 LocalExportIndex)
 	ImportStore.StoreGlobalImportObject(Export.GlobalImportIndex, Object);
 }
 
-void FAsyncPackage2::EventDrivenSerializeExport(int32 LocalExportIndex, FSimpleExportArchive& Ar)
+void FAsyncPackage2::EventDrivenSerializeExport(int32 LocalExportIndex, FExportArchive& Ar)
 {
 	TRACE_CPUPROFILER_EVENT_SCOPE(SerializeExport);
 
@@ -3272,7 +3301,6 @@ void FAsyncPackage2::EventDrivenSerializeExport(int32 LocalExportIndex, FSimpleE
 		TRACE_CPUPROFILER_EVENT_SCOPE(SerializeObject);
 		Object->Serialize(Ar);
 	}
-
 	Ar.TemplateForGetArchetypeFromLoader = nullptr;
 
 	Object->SetFlags(RF_LoadCompleted);
@@ -3388,6 +3416,25 @@ void FAsyncLoadingThread2::AddToLoadedPackages(FAsyncPackage2* Package)
 	LoadedPackages.Add(Package);
 }
 
+#if ALT2_VERIFY_RECURSIVE_LOADS 
+struct FScopedLoadRecursionVerifier
+{
+	int32& Level;
+	FScopedLoadRecursionVerifier(int32& InLevel) : Level(InLevel)
+	{
+		UE_CLOG(Level > 0, LogStreaming, Error, TEXT("Entering recursive load level: %d"), Level);
+		++Level;
+		check(Level == 1);
+	}
+	~FScopedLoadRecursionVerifier()
+	{
+		--Level;
+		UE_CLOG(Level > 0, LogStreaming, Error, TEXT("Leaving recursive load level: %d"), Level);
+		check(Level == 0);
+	}
+};
+#endif
+
 EAsyncPackageState::Type FAsyncLoadingThread2::ProcessAsyncLoadingFromGameThread(int32& OutPackagesProcessed)
 {
 	SCOPED_LOADTIMER(AsyncLoadingTime);
@@ -3398,6 +3445,9 @@ EAsyncPackageState::Type FAsyncLoadingThread2::ProcessAsyncLoadingFromGameThread
 	const bool bNeedsHeartbeatTick = !FAsyncLoadingThread2::IsMultithreaded();
 	OutPackagesProcessed = 0;
 
+#if ALT2_VERIFY_RECURSIVE_LOADS 
+	FScopedLoadRecursionVerifier LoadRecursionVerifier(this->LoadRecursionLevel);
+#endif
 	FAsyncLoadingTickScope2 InAsyncLoadingTick(*this);
 	uint32 LoopIterations = 0;
 
@@ -3417,7 +3467,7 @@ EAsyncPackageState::Type FAsyncLoadingThread2::ProcessAsyncLoadingFromGameThread
 				FThreadHeartBeat::Get().HeartBeat();
 			}
 
-			if (ThreadState.IsTimeLimitExceeded())
+			if (ThreadState.IsTimeLimitExceeded(TEXT("ProcessAsyncLoadingFromGameThread")))
 			{
 				return EAsyncPackageState::TimeOut;
 			}
@@ -3472,24 +3522,20 @@ EAsyncPackageState::Type FAsyncLoadingThread2::ProcessAsyncLoadingFromGameThread
 
 bool FAsyncPackage2::AreAllDependenciesFullyLoadedInternal(FAsyncPackage2* Package, TSet<FPackageId>& VisitedPackages, FPackageId& OutPackageId)
 {
-	TArrayView<const int32> ImportedPackageIndices(Package->ImportStore.ImportedPackages, Package->ImportStore.ImportedPackagesCount);
-
-	for (const int32 ImportedPackageIndex : ImportedPackageIndices)
+	for (const FPackageId& ImportedPackageId : Package->StoreEntry.ImportedPackages)
 	{
-		const FPackageId PackageId = FPackageId::FromIndex(ImportedPackageIndex);
-
-		if (VisitedPackages.Contains(PackageId))
+		if (VisitedPackages.Contains(ImportedPackageId))
 		{
 			continue;
 		}
-		VisitedPackages.Add(PackageId);
+		VisitedPackages.Add(ImportedPackageId);
 
-		FAsyncPackage2* AsyncRoot = AsyncLoadingThread.GetAsyncPackage(PackageId);
+		FAsyncPackage2* AsyncRoot = AsyncLoadingThread.GetAsyncPackage(ImportedPackageId);
 		if (AsyncRoot)
 		{
 			if (AsyncRoot->DeferredPostLoadIndex < AsyncRoot->ExportCount)
 			{
-				OutPackageId = PackageId;
+				OutPackageId = ImportedPackageId;
 				return false;
 			}
 
@@ -3533,19 +3579,26 @@ EAsyncPackageState::Type FAsyncLoadingThread2::ProcessLoadedPackagesFromGameThre
 			LoadedPackages.Reset();
 		}
 	}
-	if (IsMultithreaded() && ENamedThreads::GetRenderThread() == ENamedThreads::GameThread && !FTaskGraphInterface::Get().IsThreadProcessingTasks(ENamedThreads::GameThread)) // render thread tasks are actually being sent to the game thread.
+	if (IsMultithreaded() &&
+		ENamedThreads::GetRenderThread() == ENamedThreads::GameThread &&
+		!FTaskGraphInterface::Get().IsThreadProcessingTasks(ENamedThreads::GameThread)) // render thread tasks are actually being sent to the game thread.
 	{
 		// The async loading thread might have queued some render thread tasks (we don't have a render thread yet, so these are actually sent to the game thread)
 		// We need to process them now before we do any postloads.
 		FTaskGraphInterface::Get().ProcessThreadUntilIdle(ENamedThreads::GameThread);
-		if (FAsyncLoadingThreadState2::Get()->IsTimeLimitExceeded())
+		if (FAsyncLoadingThreadState2::Get()->IsTimeLimitExceeded(TEXT("ProcessLoadedPackagesFromGameThread")))
 		{
 			return EAsyncPackageState::TimeOut;
 		}
 	}
 
-		
-	bDidSomething = LoadedPackagesToProcess.Num() > 0;
+	bDidSomething = LoadedPackagesToProcess.Num() > 0 || QueuedFailedPackageCallbacks.Num() > 0;
+	for (FQueuedFailedPackageCallback& QueuedFailedPackageCallback : QueuedFailedPackageCallbacks)
+	{
+		QueuedFailedPackageCallback.Callback->ExecuteIfBound(QueuedFailedPackageCallback.PackageName, nullptr, EAsyncLoadingResult::Failed);
+	}
+	QueuedFailedPackageCallbacks.Empty();
+	FAsyncLoadingThreadState2& ThreadState = *FAsyncLoadingThreadState2::Get();
 	for (int32 PackageIndex = 0; PackageIndex < LoadedPackagesToProcess.Num() && !IsAsyncLoadingSuspended(); ++PackageIndex)
 	{
 		FAsyncPackage2* Package = LoadedPackagesToProcess[PackageIndex];
@@ -3577,10 +3630,7 @@ EAsyncPackageState::Type FAsyncLoadingThread2::ProcessLoadedPackagesFromGameThre
 				TRACE_CPUPROFILER_EVENT_SCOPE(PackageCompletionCallbacks);
 				Package->CallCompletionCallbacks(LoadingResult);
 			}
-#if WITH_EDITOR
-			// In the editor we need to find any assets and add them to list for later callback
-			Package->GetLoadedAssets(LoadedAssets);
-#endif
+
 			// We don't need the package anymore
 			check(!Package->bAddedForDelete);
 			check(!PackagesToDelete.Contains(Package));
@@ -3591,7 +3641,7 @@ EAsyncPackageState::Type FAsyncLoadingThread2::ProcessLoadedPackagesFromGameThre
 			UE_ASYNC_PACKAGE_LOG(Verbose, Package->Desc, TEXT("GameThread: LoadCompleted"),
 				TEXT("All loading of package is done, and the async package and load request will be deleted."));
 
-			if (FlushRequestID != INDEX_NONE && !ContainsRequestID(FlushRequestID))
+			if (ThreadState.IsTimeLimitExceeded(TEXT("ProcessAsyncLoadingFromGameThread")) || (FlushRequestID != INDEX_NONE && !ContainsRequestID(FlushRequestID)))
 			{
 				// The only package we care about has finished loading, so we're good to exit
 				break;
@@ -3655,27 +3705,6 @@ EAsyncPackageState::Type FAsyncLoadingThread2::ProcessLoadedPackagesFromGameThre
 
 	if (Result == EAsyncPackageState::Complete)
 	{
-#if WITH_EDITORONLY_DATA
-		// This needs to happen after loading new blueprints in the editor, and this is handled in EndLoad for synchronous loads
-		FBlueprintSupport::FlushReinstancingQueue();
-#endif
-
-#if WITH_EDITOR
-		// In editor builds, call the asset load callback. This happens in both editor and standalone to match EndLoad
-		TArray<FWeakObjectPtr> TempLoadedAssets = LoadedAssets;
-		LoadedAssets.Reset();
-
-		// Make a copy because LoadedAssets could be modified by one of the OnAssetLoaded callbacks
-		for (const FWeakObjectPtr& WeakAsset : TempLoadedAssets)
-		{
-			// It may have been unloaded/marked pending kill since being added, ignore those cases
-			if (UObject* LoadedAsset = WeakAsset.Get())
-			{
-				FCoreUObjectDelegates::OnAssetLoaded.Broadcast(LoadedAsset);
-			}
-		}
-#endif
-
 		// We're not done until all packages have been deleted
 		Result = PackagesToDelete.Num() ? EAsyncPackageState::PendingImports : EAsyncPackageState::Complete;
 	}
@@ -3738,15 +3767,6 @@ EAsyncPackageState::Type FAsyncLoadingThread2::TickAsyncLoadingFromGameThread(bo
 					FDeferredMessageLog::Flush();
 				}
 			}
-
-			if (!bDidSomething && GIsInitialLoad)
-			{
-				 bDidSomething = EDLBootNotificationManager.FireCompletedCompiledInImports();
-				 if (!bDidSomething)
-				 {
-				 	bDidSomething = EDLBootNotificationManager.ConstructWaitingBootObjects();
-				 }
-			}
 		}
 
 		// Call update callback once per tick on the game thread
@@ -3756,9 +3776,8 @@ EAsyncPackageState::Type FAsyncLoadingThread2::TickAsyncLoadingFromGameThread(bo
 	return Result;
 }
 
-FAsyncLoadingThread2::FAsyncLoadingThread2(FIoDispatcher& InIoDispatcher, IEDLBootNotificationManager& InEDLBootNotificationManager)
+FAsyncLoadingThread2::FAsyncLoadingThread2(FIoDispatcher& InIoDispatcher)
 	: Thread(nullptr)
-	, EDLBootNotificationManager(InEDLBootNotificationManager)
 	, IoDispatcher(InIoDispatcher)
 {
 	GEventDrivenLoaderEnabled = true;
@@ -4067,7 +4086,7 @@ EAsyncPackageState::Type FAsyncLoadingThread2::TickAsyncThreadFromGameThread(boo
 	int32 ProcessedRequests = 0;
 	if (AsyncThreadReady.GetValue())
 	{
-		if (IsGarbageCollectionWaiting() || FAsyncLoadingThreadState2::Get()->IsTimeLimitExceeded())
+		if (IsGarbageCollectionWaiting() || FAsyncLoadingThreadState2::Get()->IsTimeLimitExceeded(TEXT("TickAsyncThreadFromGameThread")))
 		{
 			Result = EAsyncPackageState::TimeOut;
 		}
@@ -4199,9 +4218,8 @@ void FGlobalImportStore::OnPreGarbageCollect(bool bInIsLoadingPackages)
 	bNeedToHandleGarbageCollect = bInIsLoadingPackages;
 
 	int32 NumWeak = 0;
-	for (int32 GlobalImportIndex = ScriptImportCount; GlobalImportIndex < Count; ++GlobalImportIndex)
+	for (FGlobalImport& GlobalImport : Imports[FPackageObjectIndex::PackageImport])
 	{
-		FGlobalImport& GlobalImport = Imports[GlobalImportIndex];
 		UObject* Object = GlobalImport.GetObjectIfRawPointer();
 		if (!Object)
 		{
@@ -4236,10 +4254,10 @@ void FGlobalImportStore::OnPreGarbageCollect(bool bInIsLoadingPackages)
 #ifdef ALT2_VERIFY_ASYNC_FLAGS
 	if (!bInIsLoadingPackages)
 	{
-		for (int32 GlobalImportIndex = 0; GlobalImportIndex < Count; ++GlobalImportIndex)
+		for (FGlobalImport& GlobalImport : Imports[FPackageObjectIndex::PackageImport])
 		{
-			check(GlobalImportIndex < ScriptImportCount || Imports[GlobalImportIndex].GetObjectIfRawPointer() == nullptr);
-			check(Imports[GlobalImportIndex].GetRefCount() == 0);
+			check(GlobalImport.GetObjectIfRawPointer() == nullptr);
+			check(GlobalImport.GetRefCount() == 0);
 		}
 		VerifyLoadFlagsWhenFinishedLoading();
 	}
@@ -4288,6 +4306,15 @@ void FAsyncLoadingThread2::OnPostGarbageCollect()
  */
 void FAsyncLoadingThread2::NotifyConstructedDuringAsyncLoading(UObject* Object, bool bSubObject)
 {
+	FUObjectThreadContext& ThreadContext = FUObjectThreadContext::Get();
+	if (!ThreadContext.AsyncPackage)
+	{
+		// Something is creating objects on the async loading thread outside of the actual async loading code
+		// e.g. ShaderCodeLibrary::OnExternalReadCallback doing FTaskGraphInterface::Get().WaitUntilTaskCompletes(Event);
+		// TODO: Change the code in this callback to only ever set the AsyncLoading flag on serialized objects in the current package
+		return;
+	}
+
 	// Mark objects created during async loading process (e.g. from within PostLoad or CreateExport) as async loaded so they 
 	// cannot be found. This requires also keeping track of them so we can remove the async loading flag later one when we 
 	// finished routing PostLoad to all objects.
@@ -4295,8 +4322,6 @@ void FAsyncLoadingThread2::NotifyConstructedDuringAsyncLoading(UObject* Object,
 	{
 		Object->SetInternalFlags(EInternalObjectFlags::AsyncLoading);
 	}
-	FUObjectThreadContext& ThreadContext = FUObjectThreadContext::Get();
-	check(ThreadContext.AsyncPackage); // Otherwise something is wrong and we're creating objects outside of async loading code
 	FAsyncPackage2* AsyncPackage2 = (FAsyncPackage2*)ThreadContext.AsyncPackage;
 	AsyncPackage2->AddOwnedObjectFromCallback(Object, /*bForceAdd*/ bSubObject);
 }
@@ -4317,10 +4342,10 @@ void FAsyncLoadingThread2::FireCompletedCompiledInImport(void* AsyncPackage, FPa
 FAsyncPackage2::FAsyncPackage2(
 	const FAsyncPackageDesc2& InDesc,
 	FAsyncLoadingThread2& InAsyncLoadingThread,
-	IEDLBootNotificationManager& InEDLBootNotificationManager,
 	FAsyncLoadEventGraphAllocator& InGraphAllocator,
 	const FAsyncLoadEventSpec* EventSpecs)
 : Desc(InDesc)
+, StoreEntry(InAsyncLoadingThread.GlobalPackageStore.GetStoreEntry(Desc.PackageIdToLoad))
 , LinkerRoot(nullptr)
 , PostLoadIndex(0)
 , DeferredPostLoadIndex(0)
@@ -4333,10 +4358,8 @@ FAsyncPackage2::FAsyncPackage2(
 , LoadStartTime(0.0)
 , LoadPercentage(0)
 , AsyncLoadingThread(InAsyncLoadingThread)
-, EDLBootNotificationManager(InEDLBootNotificationManager)
 , GraphAllocator(InGraphAllocator)
-// Begin EDL specific properties
-, ImportStore(AsyncLoadingThread.GlobalPackageStore, Desc.PackageIdToLoad)
+, ImportStore(AsyncLoadingThread.GlobalPackageStore, StoreEntry)
 , AsyncPackageLoadingState(EAsyncPackageLoadingState2::NewPackage)
 , bAllExportsSerialized(false)
 {
@@ -4344,18 +4367,12 @@ FAsyncPackage2::FAsyncPackage2(
 	TRACE_LOADTIME_NEW_ASYNC_PACKAGE(this, InDesc.NameToLoad);
 	AddRequestID(InDesc.RequestID);
 
-	TTuple<FExportBundleMetaEntry*, uint32> PackageExportBundles = AsyncLoadingThread.GlobalPackageStore.GetPackageExportBundleMetaEntries(Desc.PackageIdToLoad);
-	ExportBundleMetaEntries = PackageExportBundles.Get<0>();
-	ExportBundleCount = PackageExportBundles.Get<1>();
-	ExportCount = AsyncLoadingThread.GlobalPackageStore.GetPackageExportCount(Desc.PackageIdToLoad);
+	ExportBundleCount = StoreEntry.ExportBundles.Num();
+	ExportCount = StoreEntry.ExportCount;
 	Exports.AddDefaulted(ExportCount);
 	OwnedObjects.Reserve(ExportCount + 1); // +1 for UPackage
 
 	CreateNodes(EventSpecs);
-	if (GIsInitialLoad)
-	{
-		SetupScriptArcs();
-	}
 }
 
 void FAsyncPackage2::CreateNodes(const FAsyncLoadEventSpec* EventSpecs)
@@ -4561,7 +4578,6 @@ void FAsyncPackage2::CreateUPackage(const FPackageSummary* PackageSummary)
 		check(LinkerRoot->HasAnyFlags(RF_WasLoaded));
 	}
 
-	ImportStore.StoreGlobalImportObject(PackageSummary->GlobalImportIndex, LinkerRoot);
 	AddOwnedObjectWithAsyncFlag(LinkerRoot, /*bForceAdd*/ !bCreatedLinkerRoot);
 	check(LinkerRoot->HasAnyInternalFlags(EInternalObjectFlags::Async));
 
@@ -4624,7 +4640,7 @@ EAsyncPackageState::Type FAsyncPackage2::PostLoadObjects()
 	const bool bIsMultithreaded = AsyncLoadingThread.IsMultithreaded();
 
 	// PostLoad objects.
-	while (PostLoadIndex < ExportCount && !FAsyncLoadingThreadState2::Get()->IsTimeLimitExceeded())
+	while (PostLoadIndex < ExportCount && !FAsyncLoadingThreadState2::Get()->IsTimeLimitExceeded(TEXT("PostLoadObject")))
 	{
 		const FExportObject& Export = Exports[PostLoadIndex++];
 		if (Export.bFiltered)
@@ -4670,7 +4686,7 @@ EAsyncPackageState::Type FAsyncPackage2::PostLoadDeferredObjects()
 
 	while (DeferredPostLoadIndex < ExportCount && 
 		!AsyncLoadingThread.IsAsyncLoadingSuspended() &&
-		!FAsyncLoadingThreadState2::Get()->IsTimeLimitExceeded())
+		!FAsyncLoadingThreadState2::Get()->IsTimeLimitExceeded(TEXT("PostLoadDeferredObjects")))
 	{
 		const FExportObject& Export = Exports[DeferredPostLoadIndex++];
 		if (Export.bFiltered)
@@ -4701,7 +4717,7 @@ EAsyncPackageState::Type FAsyncPackage2::PostLoadDeferredObjects()
 		// Clear async loading flags (we still want RF_Async, but EInternalObjectFlags::AsyncLoading can be cleared)
 		while (DeferredFinalizeIndex < ExportCount &&
 			!AsyncLoadingThread.IsAsyncLoadingSuspended() &&
-			!FAsyncLoadingThreadState2::Get()->IsTimeLimitExceeded())
+			!FAsyncLoadingThreadState2::Get()->IsTimeLimitExceeded(TEXT("PostLoadDeferredObjects")))
 		{
 			const FExportObject& Export = Exports[DeferredFinalizeIndex++];
 			if (Export.bFiltered)
@@ -4787,7 +4803,7 @@ EAsyncPackageState::Type FAsyncPackage2::CreateClusters()
 {
 	while (DeferredClusterIndex < ExportCount &&
 			!AsyncLoadingThread.IsAsyncLoadingSuspended() &&
-			!FAsyncLoadingThreadState2::Get()->IsTimeLimitExceeded())
+			!FAsyncLoadingThreadState2::Get()->IsTimeLimitExceeded(TEXT("CreateClusters")))
 	{
 		const FExportObject& Export = Exports[DeferredClusterIndex++];
 
@@ -4981,7 +4997,13 @@ int32 FAsyncLoadingThread2::LoadPackage(const FString& InName, const FGuid* InGu
 	{
 		FAsyncPackageDesc2 PackageDesc(RequestID, PackageId, PackageIdToLoad, PackageName, PackageNameToLoad);
 		UE_ASYNC_PACKAGE_LOG(Warning, PackageDesc, TEXT("LoadPackage: SkipPackage"), TEXT("Skipping unknown package, the provided package does not exist"));
-		InCompletionDelegate.ExecuteIfBound(PackageName, nullptr, EAsyncLoadingResult::Failed);
+		if (InCompletionDelegate.IsBound())
+		{
+			// Queue completion callback and execute at next process loaded packages call to maintain behavior compatibility with old loader
+			FQueuedFailedPackageCallback& QueuedFailedPackageCallback = QueuedFailedPackageCallbacks.AddDefaulted_GetRef();
+			QueuedFailedPackageCallback.PackageName = PackageName;
+			QueuedFailedPackageCallback.Callback.Reset(new FLoadPackageAsyncDelegate(InCompletionDelegate));
+		}
 	}
 
 	return RequestID;
@@ -5092,7 +5114,13 @@ EAsyncPackageState::Type FAsyncLoadingThread2::ProcessLoadingUntilCompleteFromGa
 	return TimeLimit <= 0 ? EAsyncPackageState::TimeOut : EAsyncPackageState::Complete;
 }
 
-IAsyncPackageLoader* MakeAsyncPackageLoader2(FIoDispatcher& InIoDispatcher, IEDLBootNotificationManager& InEDLBootNotificationManager)
+IAsyncPackageLoader* MakeAsyncPackageLoader2(FIoDispatcher& InIoDispatcher)
 {
-	return new FAsyncLoadingThread2(InIoDispatcher, InEDLBootNotificationManager);
+	return new FAsyncLoadingThread2(InIoDispatcher);
 }
+
+#endif //WITH_ASYNCLOADING2
+
+#if UE_BUILD_DEVELOPMENT || UE_BUILD_DEBUG
+PRAGMA_ENABLE_OPTIMIZATION
+#endif
diff --git a/Engine/Source/Runtime/CoreUObject/Private/Serialization/AsyncPackageLoader.cpp b/Engine/Source/Runtime/CoreUObject/Private/Serialization/AsyncPackageLoader.cpp
--- a/Engine/Source/Runtime/CoreUObject/Private/Serialization/AsyncPackageLoader.cpp
+++ b/Engine/Source/Runtime/CoreUObject/Private/Serialization/AsyncPackageLoader.cpp
@@ -7,6 +7,7 @@
 #include "UObject/LinkerLoad.h"
 #include "Misc/CoreDelegates.h"
 #include "IO/IoDispatcher.h"
+#include "HAL/IConsoleManager.h"
 
 volatile int32 GIsLoaderCreated;
 TUniquePtr<IAsyncPackageLoader> GPackageLoader;
@@ -36,9 +37,17 @@ struct FEDLBootNotificationManager
 	TArray<UClass*> CDORecursiveStack;
 	TArray<UClass*> CDORecursives;
 	FCriticalSection EDLBootNotificationManagerLock;
+	bool bEnabled = true;
+
+	void Disable()
+	{
+		PathToState.Empty();
+		PathsToFire.Empty();
+		bEnabled = false;
+	}
 
 	// return true if you are waiting for this compiled in object
-	bool AddWaitingPackage(void* Pkg, FName PackageName, FName ObjectName, FPackageIndex Import) override
+	bool AddWaitingPackage(void* Pkg, FName PackageName, FName ObjectName, FPackageIndex Import, bool bIgnoreMissingPackage) override
 	{
 		if (PackageName == GLongCoreUObjectPackageName)
 		{
@@ -62,7 +71,7 @@ struct FEDLBootNotificationManager
 			}
 			if (!ExistingState)
 			{
-				UE_LOG(LogStreaming, Fatal, TEXT("Compiled in export %s not found; it was never registered."), *LongFName.ToString());
+				UE_CLOG(!bIgnoreMissingPackage, LogStreaming, Fatal, TEXT("Compiled in export %s not found; it was never registered."), *LongFName.ToString());
 				return false;
 			}
 		}
@@ -81,7 +90,7 @@ struct FEDLBootNotificationManager
 
 	void NotifyRegistrationEvent(const TCHAR* PackageName, const TCHAR* Name, ENotifyRegistrationType NotifyRegistrationType, ENotifyRegistrationPhase NotifyRegistrationPhase, UObject *(*InRegister)(), bool InbDynamic)
 	{
-		if (!GIsInitialLoad)
+		if (!bEnabled || !GIsInitialLoad)
 		{
 			return;
 		}
@@ -152,10 +161,13 @@ struct FEDLBootNotificationManager
 
 	void NotifyRegistrationComplete()
 	{
+		if (!bEnabled)
+		{
+			return;
+		}
 #if USE_EVENT_DRIVEN_ASYNC_LOAD_AT_BOOT_TIME
 		FireCompletedCompiledInImports(true);
 		FlushAsyncLoading();
-		GPackageLoader->StartThread();
 #endif
 #if !HACK_HEADER_GENERATOR
 		check(!GIsInitialLoad && IsInGameThread());
@@ -184,8 +196,7 @@ struct FEDLBootNotificationManager
 			UE_LOG(LogStreaming, Fatal, TEXT("Initial load is complete, but we still have %d imports to fire (listed above)."), PathsToFire.Num());
 		}
 #endif
-		PathToState.Empty();
-		PathsToFire.Empty();
+		Disable();
 	}
 
 	bool ConstructWaitingBootObjects() override
@@ -503,11 +514,14 @@ IAsyncPackageLoader& GetAsyncPackageLoader()
 
 void InitAsyncThread()
 {
+#if WITH_ASYNCLOADING2
 	if (FIoDispatcher::IsInitialized())
 	{
-		GPackageLoader.Reset(MakeAsyncPackageLoader2(FIoDispatcher::Get(), GetGEDLBootNotificationManager()));
+		GetGEDLBootNotificationManager().Disable();
+		GPackageLoader.Reset(MakeAsyncPackageLoader2(FIoDispatcher::Get()));
 	}
 	else
+#endif
 	{
 		GPackageLoader = MakeUnique<FAsyncLoadingThread>(/** ThreadIndex = */ 0, GetGEDLBootNotificationManager());
 	}
@@ -644,6 +658,8 @@ void NotifyRegistrationEvent(const TCHAR* PackageName, const TCHAR* Name, ENotif
 void NotifyRegistrationComplete()
 {
 	GetGEDLBootNotificationManager().NotifyRegistrationComplete();
+	GPackageLoader->FlushLoading(INDEX_NONE);
+	GPackageLoader->StartThread();
 }
 
 void NotifyConstructedDuringAsyncLoading(UObject* Object, bool bSubObject)
@@ -662,3 +678,55 @@ void ResetAsyncLoadingStats()
 	GFlushAsyncLoadingCount = 0;
 	GSyncLoadCount = 0;
 }
+
+int32 GWarnIfTimeLimitExceeded = 0;
+static FAutoConsoleVariableRef CVarWarnIfTimeLimitExceeded(
+	TEXT("s.WarnIfTimeLimitExceeded"),
+	GWarnIfTimeLimitExceeded,
+	TEXT("Enables log warning if time limit for time-sliced package streaming has been exceeded."),
+	ECVF_Default
+);
+
+float GTimeLimitExceededMultiplier = 1.5f;
+static FAutoConsoleVariableRef CVarTimeLimitExceededMultiplier(
+	TEXT("s.TimeLimitExceededMultiplier"),
+	GTimeLimitExceededMultiplier,
+	TEXT("Multiplier for time limit exceeded warning time threshold."),
+	ECVF_Default
+);
+
+float GTimeLimitExceededMinTime = 0.005f;
+static FAutoConsoleVariableRef CVarTimeLimitExceededMinTime(
+	TEXT("s.TimeLimitExceededMinTime"),
+	GTimeLimitExceededMinTime,
+	TEXT("Minimum time the time limit exceeded warning will be triggered by."),
+	ECVF_Default
+);
+
+void IsTimeLimitExceededPrint(
+	double InTickStartTime,
+	double CurrentTime,
+	double LastTestTime,
+	float InTimeLimit, 
+	const TCHAR* InLastTypeOfWorkPerformed,
+	UObject* InLastObjectWorkWasPerformedOn)
+{
+	static double LastPrintStartTime = -1.0;
+	// Log single operations that take longer than time limit (but only in cooked builds)
+	if (LastPrintStartTime != InTickStartTime &&
+		(CurrentTime - InTickStartTime) > GTimeLimitExceededMinTime &&
+		(CurrentTime - InTickStartTime) > (GTimeLimitExceededMultiplier * InTimeLimit))
+	{
+		float EstimatedTimeForThisStep = (CurrentTime - InTickStartTime) * 1000;
+		if (LastTestTime > InTickStartTime)
+		{
+			EstimatedTimeForThisStep = (CurrentTime - LastTestTime) * 1000;
+		}
+		LastPrintStartTime = InTickStartTime;
+		UE_LOG(LogStreaming, Warning, TEXT("IsTimeLimitExceeded: %s %s Load Time %5.2fms   Last Step Time %5.2fms"),
+			InLastTypeOfWorkPerformed ? InLastTypeOfWorkPerformed : TEXT("unknown"),
+			InLastObjectWorkWasPerformedOn ? *InLastObjectWorkWasPerformedOn->GetFullName() : TEXT("nullptr"),
+			(CurrentTime - InTickStartTime) * 1000,
+			EstimatedTimeForThisStep);
+	}
+}
