Copying //UE4/Release-4.25Plus-Dev to Release-4.25Plus (//UE4/Release-4.25Plus)

diff --git a/Engine/Source/Runtime/CoreUObject/Private/Serialization/BulkData.cpp b/Engine/Source/Runtime/CoreUObject/Private/Serialization/BulkData.cpp
--- a/Engine/Source/Runtime/CoreUObject/Private/Serialization/BulkData.cpp
+++ b/Engine/Source/Runtime/CoreUObject/Private/Serialization/BulkData.cpp
@@ -5,6 +5,7 @@
 #include "HAL/FileManager.h"
 #include "Misc/ScopeLock.h"
 #include "Misc/Paths.h"
+#include "Misc/PathViews.h"
 #include "Serialization/LargeMemoryWriter.h"
 #include "Serialization/LargeMemoryReader.h"
 #include "HAL/IConsoleManager.h"
@@ -730,6 +731,36 @@ void FUntypedBulkData::ForceBulkDataResident()
 #endif // WITH_EDITOR
 }
 
+bool FUntypedBulkData::StartAsyncLoading()
+{
+	DECLARE_SCOPE_CYCLE_COUNTER(TEXT("FUntypedBulkData::StartAsyncLoading"), STAT_UBD_StartSerializingBulkData, STATGROUP_Memory);
+
+	if (!IsAsyncLoadingComplete())
+	{
+		return true; // Early out if an asynchronous load is already in progress.
+	}
+
+	if (IsBulkDataLoaded())
+	{
+		return false; // Early out if we do not need to actually load any data
+	}
+
+	if (!CanLoadFromDisk())
+	{
+		return false; // Early out if we cannot load from disk
+	}
+
+	check(SerializeFuture.IsValid() == false);
+
+	SerializeFuture = Async(EAsyncExecution::ThreadPool, [=]()
+	{
+		AsyncLoadBulkData();
+		return true;
+	});
+
+	return true;
+}
+
 /**
  * Sets the passed in bulk data flags.
  *
@@ -737,12 +768,12 @@ void FUntypedBulkData::ForceBulkDataResident()
  */
 void FUntypedBulkData::SetBulkDataFlags( uint32 BulkDataFlagsToSet )
 {
-	BulkDataFlags |= BulkDataFlagsToSet;
+	BulkDataFlags = EBulkDataFlags(BulkDataFlags | BulkDataFlagsToSet);
 }
 
 void FUntypedBulkData::ResetBulkDataFlags(uint32 BulkDataFlagsToSet)
 {
-	BulkDataFlags = BulkDataFlagsToSet;
+	BulkDataFlags = (EBulkDataFlags)BulkDataFlagsToSet;
 }
 
 /**
@@ -782,20 +813,13 @@ uint32 FUntypedBulkData::GetBulkDataAlignment() const
  */
 void FUntypedBulkData::ClearBulkDataFlags( uint32 BulkDataFlagsToClear )
 {
-	BulkDataFlags &= ~BulkDataFlagsToClear;
+	BulkDataFlags = EBulkDataFlags(BulkDataFlags & ~BulkDataFlagsToClear);
 }
 
-/**
- * Load the resource data in the BulkDataAsync
- *
- * @param BulkDataFlagsToClear	Bulk data flags to clear
- */
 void FUntypedBulkData::AsyncLoadBulkData()
 {
 	BulkDataAsync.Reallocate(GetBulkDataSize(), BulkDataAlignment);
 
-	UE_CLOG(GEventDrivenLoaderEnabled, LogSerialization, Error, TEXT("Attempt to stream bulk data with EDL enabled. This is not desireable. File %s"), *Filename);
-
 	FArchive* FileReaderAr = IFileManager::Get().CreateFileReader(*Filename, FILEREAD_Silent);
 	checkf(FileReaderAr != NULL, TEXT("Attempted to load bulk data from an invalid filename '%s'."), *Filename);
 
@@ -805,7 +829,6 @@ void FUntypedBulkData::AsyncLoadBulkData()
 	delete FileReaderAr;
 }
 
-
 /*-----------------------------------------------------------------------------
 	Serialization.
 -----------------------------------------------------------------------------*/
@@ -817,7 +840,9 @@ void FUntypedBulkData::StartSerializingBulkData(FArchive& Ar, UObject* Owner, in
 
 	SerializeFuture = Async(EAsyncExecution::ThreadPool, [=]() 
 	{ 
-		AsyncLoadBulkData(); 
+		UE_CLOG(GEventDrivenLoaderEnabled, LogSerialization, Error, TEXT("Attempt to stream bulk data with EDL enabled. This is not desireable. File %s"), *Filename);
+		AsyncLoadBulkData();
+
 		return true;
 	});
 
@@ -931,7 +956,7 @@ void FUntypedBulkData::Serialize( FArchive& Ar, UObject* Owner, int32 Idx, bool
 			// and BulkDataSizeOnDisk need to be held as int64s, so set a flag indicating the new format.
 			if (Ar.IsSaving() && GetBulkDataSize() >= (1LL << 31))
 			{
-				BulkDataFlags |= BULKDATA_Size64Bit;
+				SetBulkDataFlags(BULKDATA_Size64Bit);
 			}
 			Ar << BulkDataFlags;
 		}
@@ -960,13 +985,13 @@ void FUntypedBulkData::Serialize( FArchive& Ar, UObject* Owner, int32 Idx, bool
 				// Bulk data that is being serialized via seekfree loading is single use only. This allows us 
 				// to free the memory as e.g. the bulk data won't be attached to an archive in the case of
 				// seek free loading.
-				BulkDataFlags |= BULKDATA_SingleUse;
+				SetBulkDataFlags(BULKDATA_SingleUse);
 			}
 
 			// Hacky fix for using cooked data in editor. Cooking sets BULKDATA_SingleUse for textures, but PIEing needs to keep bulk data around.
 			if (GIsEditor)
 			{
-				BulkDataFlags &= ~BULKDATA_SingleUse;
+				ClearBulkDataFlags(BULKDATA_SingleUse);
 			}
 
 			// Size on disk, which in the case of compression is != GetBulkDataSize()
@@ -988,7 +1013,7 @@ void FUntypedBulkData::Serialize( FArchive& Ar, UObject* Owner, int32 Idx, bool
 				uint16 DummyValue;
 				Ar << DummyValue;
 
-				BulkDataFlags &= ~BULKDATA_BadDataVersion;		
+				ClearBulkDataFlags(BULKDATA_BadDataVersion);	
 			}
 			
 			// determine whether the payload is stored inline or at the end of the file
@@ -1053,7 +1078,7 @@ void FUntypedBulkData::Serialize( FArchive& Ar, UObject* Owner, int32 Idx, bool
 							{
 								BulkDataOffsetInFile -= IFileManager::Get().FileSize(*Filename);
 								check(BulkDataOffsetInFile >= 0);
-								Filename = FPaths::GetBaseFilename(Filename, false) + TEXT(".uexp");
+								Filename = FPaths::GetBaseFilename(Filename, false) + BulkDataExt::Export;
 							}
 							bWasMapped = BulkData.MapFile(*Filename, BulkDataOffsetInFile, GetBulkDataSize());
 						}
@@ -1075,12 +1100,12 @@ void FUntypedBulkData::Serialize( FArchive& Ar, UObject* Owner, int32 Idx, bool
 				else if (BulkDataFlags & BULKDATA_DuplicateNonOptionalPayload)
 				{
 					// Load from optional payload instead if optional payload is available
-					FString OptionalFilename = FPaths::ChangeExtension(Filename, TEXT(".uptnl"));
+					FString OptionalFilename = FPathViews::ChangeExtension(Filename, BulkDataExt::Optional);
 					if (IFileManager::Get().FileExists(*OptionalFilename))
 					{
 						Filename = OptionalFilename;
 						Ar << BulkDataFlags;
-						BulkDataFlags |= BULKDATA_OptionalPayload;
+						SetBulkDataFlags(BULKDATA_OptionalPayload);
 
 						if (BulkDataFlags & BULKDATA_Size64Bit)
 						{
@@ -1116,22 +1141,22 @@ void FUntypedBulkData::Serialize( FArchive& Ar, UObject* Owner, int32 Idx, bool
 						int64 DummyBulkDataOffsetInFile;
 						Ar << DummyBulkDataOffsetInFile;
 
-						Filename = FPaths::ChangeExtension(Filename, TEXT(".ubulk"));
+						Filename = FPathViews::ChangeExtension(Filename, BulkDataExt::Default);
 					}
 				}
 				else if (BulkDataFlags & BULKDATA_OptionalPayload)
 				{
-					Filename = FPaths::ChangeExtension(Filename, TEXT(".uptnl"));
+					Filename = FPathViews::ChangeExtension(Filename, BulkDataExt::Optional);
 				}
 				else if (BulkDataFlags & BULKDATA_PayloadInSeperateFile)
 				{
 					if (BulkDataFlags & BULKDATA_MemoryMappedPayload)
 					{
-						Filename = FPaths::ChangeExtension(Filename, TEXT(".m.ubulk"));
+						Filename = FPathViews::ChangeExtension(Filename, BulkDataExt::MemoryMapped);
 					}
 					else
 					{
-						Filename = FPaths::ChangeExtension(Filename, TEXT(".ubulk"));
+						Filename = FPathViews::ChangeExtension(Filename, BulkDataExt::Default);
 					}
 
 
@@ -1179,7 +1204,7 @@ void FUntypedBulkData::Serialize( FArchive& Ar, UObject* Owner, int32 Idx, bool
 							{
 								BulkDataOffsetInFile -= IFileManager::Get().FileSize(*Filename);
 								check(BulkDataOffsetInFile >= 0);
-								Filename = FPaths::GetBaseFilename(Filename, false) + TEXT(".uexp");
+								Filename = FPaths::GetBaseFilename(Filename, false) + BulkDataExt::Export;
 							}
 
 							FArchive* TargetArchive = IFileManager::Get().CreateFileReader(*Filename);
@@ -1211,7 +1236,7 @@ void FUntypedBulkData::Serialize( FArchive& Ar, UObject* Owner, int32 Idx, bool
 		else if( Ar.IsSaving() )
 		{
 			// Remove single element serialization requirement before saving out bulk data flags.
-			BulkDataFlags &= ~BULKDATA_ForceSingleElementSerialization;
+			ClearBulkDataFlags(BULKDATA_ForceSingleElementSerialization);
 
 			// Make sure bulk data is loaded.
 			MakeSureBulkDataIsLoaded();
@@ -1258,7 +1283,7 @@ void FUntypedBulkData::Serialize( FArchive& Ar, UObject* Owner, int32 Idx, bool
 			if (!bStoreInline)
 			{
 				// set the flag indicating where the payload is stored
-				BulkDataFlags |= BULKDATA_PayloadAtEndOfFile;
+				SetBulkDataFlags(BULKDATA_PayloadAtEndOfFile);
 
 				// with no LinkerSave we have to store the data inline
 				check(LinkerSave != NULL);				
@@ -1324,7 +1349,7 @@ void FUntypedBulkData::Serialize( FArchive& Ar, UObject* Owner, int32 Idx, bool
 			else
 			{
 				// set the flag indicating where the payload is stored
-				BulkDataFlags &= ~BULKDATA_PayloadAtEndOfFile;
+				ClearBulkDataFlags(BULKDATA_PayloadAtEndOfFile);
 
 				int64 SavedBulkDataStartPos = Ar.Tell();
 
@@ -1433,16 +1458,17 @@ void FUntypedBulkData::StoreCompressedOnDisk( FName CompressionFormat )
 		if( CompressionFormat == NAME_None )
 		{
 			// clear all compression settings
-			BulkDataFlags &= ~BULKDATA_SerializeCompressed;
+			ClearBulkDataFlags(BULKDATA_SerializeCompressed);
 		}
 		else
 		{
 			// right BulkData only knows zlib
 			check(CompressionFormat == NAME_Zlib);
-			BulkDataFlags |= (CompressionFormat == NAME_Zlib) ? BULKDATA_SerializeCompressedZLIB : BULKDATA_None;
+			const uint32 FlagToSet = CompressionFormat == NAME_Zlib ? BULKDATA_SerializeCompressedZLIB : BULKDATA_None;
+			SetBulkDataFlags(FlagToSet);
 
 			// make sure we are not forcing the bulkdata to be stored inline if we use compression
-			BulkDataFlags &= ~BULKDATA_ForceInlinePayload;
+			ClearBulkDataFlags(BULKDATA_ForceInlinePayload);
 		}
 	}
 }
@@ -1620,7 +1646,7 @@ IBulkDataIORequest* FUntypedBulkData::CreateStreamingRequest(int64 OffsetInBulkD
 	if (GEventDrivenLoaderEnabled && (AdjustedFilename.EndsWith(TEXT(".uasset")) || AdjustedFilename.EndsWith(TEXT(".umap"))))
 	{
 		AdjustedBulkDataOffsetInFile -= IFileManager::Get().FileSize(*AdjustedFilename);
-		AdjustedFilename = FPaths::GetBaseFilename(AdjustedFilename, false) + TEXT(".uexp");
+		AdjustedFilename = FPaths::GetBaseFilename(AdjustedFilename, false) + BulkDataExt::Export;
 
 		UE_LOG(LogSerialization, Error, TEXT("Streaming from the .uexp file '%s' this MUST be in a ubulk instead for best performance."), *AdjustedFilename);
 	}
@@ -1710,6 +1736,14 @@ FBulkDataIORequest::FBulkDataIORequest(IAsyncReadFileHandle* InFileHandle)
 {
 }
 
+FBulkDataIORequest::FBulkDataIORequest(IAsyncReadFileHandle* InFileHandle, IAsyncReadRequest* InReadRequest, int64 BytesToRead)
+	: FileHandle(InFileHandle)
+	, ReadRequest(InReadRequest)
+	, Size(BytesToRead)
+{
+
+}
+
 FBulkDataIORequest::~FBulkDataIORequest()
 {
 	delete ReadRequest;
@@ -1852,11 +1886,11 @@ void FUntypedBulkData::LoadDataIntoMemory( void* Dest )
 		FString BulkDataFilename;
 		if (BulkDataFlags & BULKDATA_MemoryMappedPayload)
 		{
-			BulkDataFilename = FPaths::ChangeExtension(Filename, TEXT(".m.ubulk"));
+			BulkDataFilename = FPathViews::ChangeExtension(Filename, BulkDataExt::MemoryMapped);
 		}
 		else
 		{
-			BulkDataFilename = FPaths::ChangeExtension(Filename, TEXT(".ubulk"));
+			BulkDataFilename = FPathViews::ChangeExtension(Filename, BulkDataExt::Default);
 		}
 		BulkDataArchive = IFileManager::Get().CreateFileReader(*BulkDataFilename, FILEREAD_Silent);
 	}
@@ -1927,7 +1961,7 @@ void FUntypedBulkData::LoadDataIntoMemory( void* Dest )
 		// is it possible for streaming mips to be loaded in non streaming ways.
 		if (CVarTextureStreamingEnabled->GetValueOnAnyThread() != 0)
 		{
-			bool bIsBulkFile = Filename.EndsWith(TEXT(".ubulk"));
+			bool bIsBulkFile = Filename.EndsWith(BulkDataExt::Default);
 			UE_CLOG(GEventDrivenLoaderEnabled && bIsBulkFile && (IsInGameThread() || IsInAsyncLoadingThread()), LogSerialization, Error, TEXT("Attempt to sync load bulk data with EDL enabled (LoadDataIntoMemory). This is not desireable. File %s"), *Filename);
 		}
 #endif
@@ -1936,7 +1970,7 @@ void FUntypedBulkData::LoadDataIntoMemory( void* Dest )
 		{
 			BulkDataOffsetInFile -= IFileManager::Get().FileSize(*Filename);
 			check(BulkDataOffsetInFile >= 0);
-			Filename = FPaths::GetBaseFilename(Filename, false) + TEXT(".uexp");
+			Filename = FPaths::GetBaseFilename(Filename, false) + BulkDataExt::Export;
 		}
 
 		FArchive* Ar = IFileManager::Get().CreateFileReader(*Filename, FILEREAD_Silent);
diff --git a/Engine/Source/Runtime/CoreUObject/Private/Serialization/BulkData2.cpp b/Engine/Source/Runtime/CoreUObject/Private/Serialization/BulkData2.cpp
--- a/Engine/Source/Runtime/CoreUObject/Private/Serialization/BulkData2.cpp
+++ b/Engine/Source/Runtime/CoreUObject/Private/Serialization/BulkData2.cpp
@@ -6,12 +6,14 @@
 #include "HAL/PlatformFilemanager.h"
 #include "Misc/CommandLine.h"
 #include "Misc/Paths.h"
+#include "Misc/PathViews.h"
 #include "ProfilingDebugging/LoadTimeTracker.h"
 #include "Serialization/BulkData.h"
 #include "UObject/LinkerLoad.h"
 #include "UObject/Object.h"
 #include "UObject/Package.h"
 #include "IO/IoDispatcher.h"
+#include "Async/Async.h"
 #include "Async/MappedFileHandle.h"
 
 DEFINE_LOG_CATEGORY_STATIC(LogBulkDataRuntime, Log, All);
@@ -24,14 +26,20 @@ IMPLEMENT_TYPE_LAYOUT(FBulkDataBase);
 // If set to 0 then we will pretend that optional data does not exist, useful for testing.
 #define ALLOW_OPTIONAL_DATA 1
 
-namespace
+// Handy macro to validate FIoStatus return values
+#define CHECK_IOSTATUS( InIoStatus, InMethodName ) checkf(InIoStatus.IsOk(), TEXT("%s failed: %s"), InMethodName, *InIoStatus.ToString());
+
+namespace BulkDataExt
 {
 	// TODO: Maybe expose this and start using everywhere?
-	static const FString InlinedExt = TEXT(".uexp");			// Stored in the export data
-	static const FString DefaultExt = TEXT(".ubulk");			// Stored in a separate file
-	static const FString MemoryMappedExt = TEXT(".m.ubulk");	// Stored in a separate file aligned for memory mapping
-	static const FString OptionalExt = TEXT(".uptnl");			// Stored in a separate file that is optional
+	const FString Export		= TEXT(".uexp");	// Stored in the export data
+	const FString Default		= TEXT(".ubulk");	// Stored in a separate file
+	const FString MemoryMapped	= TEXT(".m.ubulk");	// Stored in a separate file aligned for memory mapping
+	const FString Optional		= TEXT(".uptnl");	// Stored in a separate file that is optional
+}
 
+namespace
+{
 	const uint16 InvalidBulkDataIndex = ~uint16(0);
 
 	FORCEINLINE bool IsIoDispatcherEnabled()
@@ -132,7 +140,7 @@ namespace FileTokenSystem
 
 	FRWLock TokenLock;
 
-	FileToken RegisterFileToken( const FName& PackageName, const FString& Filename, uint64 BulkDataOffsetInFile )
+	FBulkDataOrId::FileToken RegisterFileToken( const FName& PackageName, const FString& Filename, uint64 BulkDataOffsetInFile )
 	{
 		FWriteScopeLock LockForScope(TokenLock);
 
@@ -142,11 +150,10 @@ namespace FileTokenSystem
 		Data.PackageName = PackageName;
 		Data.BulkDataOffsetInFile = BulkDataOffsetInFile;
 
-		FileToken FileToken  = TokenData.Add(Data);
-		return FileToken;
+		return TokenData.Add(Data);
 	}
 
-	void UnregisterFileToken(FileToken ID)
+	void UnregisterFileToken(FBulkDataOrId::FileToken ID)
 	{
 		if (ID != FBulkDataBase::InvalidToken)
 		{
@@ -159,7 +166,7 @@ namespace FileTokenSystem
 		}
 	}
 
-	FileToken CopyFileToken(FileToken ID)
+	FBulkDataOrId::FileToken CopyFileToken(FBulkDataOrId::FileToken ID)
 	{
 		if (ID != FBulkDataBase::InvalidToken)
 		{
@@ -183,7 +190,7 @@ namespace FileTokenSystem
 		}
 	}
 
-	Data GetFileData(FileToken ID)
+	Data GetFileData(FBulkDataOrId::FileToken ID)
 	{
 		if (ID == FBulkDataBase::InvalidToken)
 		{
@@ -200,7 +207,7 @@ namespace FileTokenSystem
 		return Output;
 	}
 
-	FString GetFilename(FileToken ID)
+	FString GetFilename(FBulkDataOrId::FileToken ID)
 	{
 		if (ID == FBulkDataBase::InvalidToken)
 		{
@@ -211,7 +218,7 @@ namespace FileTokenSystem
 		return StringTable.Resolve(TokenData[ID].PackageName);
 	}
 
-	uint64 GetBulkDataOffset(FileToken ID)
+	uint64 GetBulkDataOffset(FBulkDataOrId::FileToken ID)
 	{
 		if (ID == FBulkDataBase::InvalidToken)
 		{
@@ -255,6 +262,11 @@ private:
 	}
 };
 
+// TODO: Currently shared between all FReadChunkIdRequest as the PS4/Pak implementation do but it would be
+// worth profiling on some different platforms to see if we lose more perf from the potential increase in 
+// locks vs the gain we get from not creating so many CriticalSections.
+static FCriticalSection FReadChunkIdRequestEvent;
+
 class FReadChunkIdRequest : public IAsyncReadRequest
 {
 public:
@@ -268,17 +280,25 @@ public:
 		{
 			Memory = (uint8*)FMemory::Malloc(InBytesToRead);
 		}
-
-		DoneEvent = FPlatformProcess::GetSynchEventFromPool(true);
-		DoneEvent->Reset();
 		
 		FIoReadOptions Options(InOffset, InBytesToRead);
 		Options.SetTargetVa(Memory);
 
 		auto OnRequestLoaded = [this](TIoStatusOr<FIoBuffer> Result)
 		{
-			SetComplete();
-			ReleaseEvent(true);
+			SetDataComplete();
+			
+			{
+				FScopeLock Lock(&FReadChunkIdRequestEvent);
+				bRequestOutstanding = false;
+
+				if (DoneEvent != nullptr)
+				{
+					DoneEvent->Trigger();
+				}
+
+				SetAllComplete();
+			}
 		};
 
 		FBulkDataBase::GetIoDispatcher()->ReadWithCallback(InChunkId, Options, OnRequestLoaded);
@@ -286,7 +306,8 @@ public:
 
 	virtual ~FReadChunkIdRequest()
 	{
-		ReleaseEvent(false);
+		// Make sure no other thread is waiting on this request
+		checkf(DoneEvent == nullptr, TEXT("A thread is still waiting on a FReadChunkIdRequest that is being destroyed!")); 
 
 		// Free memory if the request allocated it (although if the user accessed the memory after
 		// reading then they will have taken ownership of it anyway, and if they didn't access the
@@ -304,38 +325,56 @@ protected:
 
 	virtual void WaitCompletionImpl(float TimeLimitSeconds) override
 	{
-		if (!PollCompletion())
+		// Make sure no other thread is waiting on this request
+		checkf(DoneEvent == nullptr, TEXT("Multiple threads attempting to wait on the same FReadChunkIdRequest"));
+
 		{
-			uint32 TimeLimitMilliseconds = TimeLimitSeconds <= 0.0f ? (uint32)(TimeLimitSeconds * 1000.0f) : MAX_uint32;
+			FScopeLock Lock(&FReadChunkIdRequestEvent);
+			if (bRequestOutstanding)
+			{
+				check(DoneEvent == nullptr);
+				DoneEvent = FPlatformProcess::GetSynchEventFromPool(true);
+			}
+		}
+
+		if (DoneEvent != nullptr)
+		{
+			uint32 TimeLimitMilliseconds = TimeLimitSeconds <= 0.0f ? MAX_uint32 : (uint32)(TimeLimitSeconds * 1000.0f);
 			DoneEvent->Wait(TimeLimitMilliseconds);
+
+			FScopeLock Lock(&FReadChunkIdRequestEvent);
+			FPlatformProcess::ReturnSynchEventToPool(DoneEvent);
+			DoneEvent = nullptr;
 		}
+
+		// Make sure everything has completed
+		check(bRequestOutstanding == false);
+		check(PollCompletion() == true);
 	}
 
 	virtual void CancelImpl() override
 	{
 		bCanceled = true;
-		SetComplete();
-
-		DoneEvent->Trigger();
-		ReleaseEvent(true);
-	}
 
-	void ReleaseEvent(bool bShouldTrigger)
-	{
-		if (DoneEvent != nullptr)
 		{
-			if (bShouldTrigger)
+			FScopeLock Lock(&FReadChunkIdRequestEvent);
+			bRequestOutstanding = false;
+
+			if (DoneEvent != nullptr)
 			{
 				DoneEvent->Trigger();
 			}
 
-			FPlatformProcess::ReturnSynchEventToPool(DoneEvent);
-			DoneEvent = nullptr;
+			SetComplete();
 		}
 	}
 
+	/** The ChunkId that is being read. */
 	FIoChunkId ChunkId;
-	FEvent* DoneEvent;
+	/** Only actually gets created if WaitCompletion is called. */
+	FEvent* DoneEvent = nullptr;
+	/** True while the request is pending, true once it has either been completed or canceled. */
+	bool bRequestOutstanding = true;
 };
 
 class FAsyncReadChunkIdHandle : public IAsyncReadFileHandle
@@ -384,7 +423,8 @@ public:
 	{
 		for (const FIoChunkId& ChunkId : ChunkIDs)
 		{
-			RequestArray.Push({ChunkId, 0, (uint64)INDEX_NONE});
+			const uint64 Size = FBulkDataBase::GetIoDispatcher()->GetSizeForChunk(ChunkId).ConsumeValueOrDie();
+			RequestArray.Push({ChunkId, 0, Size });
 		}
 
 		if (InCompleteCallback != nullptr)
@@ -395,11 +435,7 @@ public:
 
 	virtual ~FBulkDataIoDispatcherRequest()
 	{
-		if (Task != nullptr)
-		{
-			Task->EnsureCompletion();
-			delete Task;
-		}
+		WaitCompletion(0.0f); // Wait for ever as we cannot leave outstanding requests
 
 		// Free the data is no caller has taken ownership of it and it was allocated by FBulkDataIoDispatcherRequest
 		if (UserSuppliedMemory == nullptr)
@@ -407,25 +443,57 @@ public:
 			FMemory::Free(DataResult);
 			DataResult = nullptr;
 		}
+
+		check(!IoBatch.IsValid()); // Should be freed by the callback!
 	}
 
 	void StartAsyncWork()
-	{
-		check(Task == nullptr);
-
-		bIsCompleted = false;
-		FPlatformMisc::MemoryBarrier();
+	{		
+		check(RequestArray.Num() > 0);
+		check(!IoBatch.IsValid());
 
-		if (FPlatformProcess::SupportsMultithreading())
+		auto Callback = [this](TIoStatusOr<FIoBuffer> Result)
 		{
-			Task = new FAsyncTask<FBulkDataIoDispatcherRequestWorker>(*this);
-			Task->StartBackgroundTask(GBackgroundPriorityThreadPool);
-		}
-		else
+			CHECK_IOSTATUS(Result.Status(), TEXT("FIoBatch::IssueWithCallback"));
+			FIoBuffer IoBuffer = Result.ConsumeValueOrDie();
+
+			SizeResult = IoBuffer.DataSize();
+
+			if (IoBuffer.IsMemoryOwned())
+			{
+				DataResult = IoBuffer.Release().ConsumeValueOrDie();
+			}
+			else
+			{
+				DataResult = IoBuffer.Data();
+			}
+
+			check(DataResult);
+
+			bIsCompleted = true;
+
+			FPlatformMisc::MemoryBarrier();
+
+			if (CompleteCallback)
+			{
+				CompleteCallback(bIsCanceled, this);
+			}
+
+			FBulkDataBase::GetIoDispatcher()->FreeBatch(IoBatch);
+		};
+
+		IoBatch = FBulkDataBase::GetIoDispatcher()->NewBatch();
+
+		for (Request& Request : RequestArray)
 		{
-			// No multi threading support so just process immediately
-			DoWork();
+			IoBatch.Read(Request.ChunkId, FIoReadOptions(Request.OffsetInBulkData, Request.BytesToRead));
 		}
+
+		FIoBatchReadOptions BatchReadOptions;
+		BatchReadOptions.SetTargetVa(UserSuppliedMemory);
+
+		FIoStatus Status = IoBatch.IssueWithCallback(BatchReadOptions, Callback);
+		CHECK_IOSTATUS(Status, TEXT("FIoBatch::IssueWithCallback"));
 	}
 
 	virtual bool PollCompletion() const override
@@ -495,112 +563,14 @@ public:
 		}
 	}
 
-private:
-
-	void DoWork()
-	{
-		FIoReadOptions IoOptions;
-
-		// TODO: FIoReadOptions is ignored at the moment
-		//IoOptions.SetRange(OffsetInBulkData, BytesToRead);
-		//IoOptions.SetTargetVa((uint64)UserSuppliedMemory);
-
-		FIoBatch NewBatch = FBulkDataBase::GetIoDispatcher()->NewBatch();
-		for (Request& Request : RequestArray)
-		{
-			Request.IoRequest = NewBatch.Read(Request.ChunkId, IoOptions);
-		}
-		NewBatch.Issue();
-
-		// Wait until the request is completed
-		while (!AreRequestsCompleted())
-		{
-			FPlatformProcess::Sleep(0.0f);
-		}
-
-		SizeResult = GetTotalBufferSize();
-
-		uint8* DstBuffer = UserSuppliedMemory != nullptr ? UserSuppliedMemory : (uint8*)FMemory::Malloc(SizeResult);
-		DataResult = DstBuffer;
-
-		for (const Request& Request : RequestArray)
-		{
-			TIoStatusOr<FIoBuffer> Result = Request.IoRequest.GetResult();
-			check(Result.IsOk());
-
-			const FIoBuffer& IoBuffer = Result.ValueOrDie();
-			uint64 BytesToRead = Request.BytesToRead == INDEX_NONE ? IoBuffer.DataSize() : Request.BytesToRead;
-		
-			FMemory::Memcpy(DstBuffer, IoBuffer.Data() + Request.OffsetInBulkData, BytesToRead);
-			DstBuffer += BytesToRead;
-		}
-		
-		FBulkDataBase::GetIoDispatcher()->FreeBatch(NewBatch);
-
-		bIsCompleted = true;
-		FPlatformMisc::MemoryBarrier();
-
-		if (CompleteCallback)
-		{
-			CompleteCallback(bIsCanceled, this);
-		}	
-	}
-
-	class FBulkDataIoDispatcherRequestWorker : public FNonAbandonableTask
-	{
-		FBulkDataIoDispatcherRequest& IoRequest;
-	public:
-		FBulkDataIoDispatcherRequestWorker(FBulkDataIoDispatcherRequest& InIoRequest)
-			: IoRequest(InIoRequest)
-		{
-		}
-
-		void DoWork()
-		{
-			IoRequest.DoWork();
-		}
-
-		FORCEINLINE TStatId GetStatId() const
-		{
-			RETURN_QUICK_DECLARE_CYCLE_STAT(FBulkDataIoDispatcherRequestWorker, STATGROUP_ThreadPoolAsyncTasks);
-		}
-	};
-
 private:
 	struct Request
 	{
 		FIoChunkId ChunkId;
 		uint64 OffsetInBulkData;
 		uint64 BytesToRead;
-
-		FIoRequest IoRequest;
 	};
 
-	bool AreRequestsCompleted() const
-	{
-		for (const Request& Req : RequestArray)
-		{
-			if (!Req.IoRequest.Status().IsCompleted())
-			{
-				return false;
-			}
-		}
-
-		return true;
-	}
-
-	int64 GetTotalBufferSize()  
-	{
-		uint64 Size = 0;
-		for (Request& Req : RequestArray)
-		{
-			TIoStatusOr<FIoBuffer> Result = Req.IoRequest.GetResult();
-			Size += Result.ValueOrDie().DataSize();
-		}
-
-		return Size;
-	}
-
 	TArray<Request, TInlineAllocator<8>> RequestArray;
 
 	FBulkDataIORequestCallBack CompleteCallback;
@@ -609,10 +579,10 @@ private:
 	uint8* DataResult = nullptr;
 	int64 SizeResult = 0;
 
-	// Task related members
-	FAsyncTask<FBulkDataIoDispatcherRequestWorker>* Task = nullptr;
 	bool bIsCompleted = false;
 	bool bIsCanceled = false;
+
+	FIoBatch IoBatch;
 };
 
 FBulkDataBase::FBulkDataBase(FBulkDataBase&& Other)
@@ -638,7 +608,15 @@ FBulkDataBase& FBulkDataBase::operator=(const FBulkDataBase& Other)
 
 	RemoveBulkData();
 
-	Data.Fallback.Token = FileTokenSystem::CopyFileToken(Other.Data.Fallback.Token);
+	if (Other.IsUsingIODispatcher())
+	{
+		Data.ChunkID = Other.Data.ChunkID;
+	}
+	else
+	{
+		Data.Fallback.BulkDataSize = Other.Data.Fallback.BulkDataSize;
+		Data.Fallback.Token = FileTokenSystem::CopyFileToken(Other.Data.Fallback.Token);
+	}
 
 	// Copy token
 	BulkDataFlags = Other.BulkDataFlags;
@@ -659,10 +637,21 @@ FBulkDataBase& FBulkDataBase::operator=(const FBulkDataBase& Other)
 	}
 	else
 	{
-		const int64 BulkDataSize = GetBulkDataSize();
-		FileTokenSystem::Data FileData = FileTokenSystem::GetFileData(Data.Fallback.Token);
-		const FString MemoryMappedFilename = ConvertFilenameFromFlags(FileData.PackageHeaderFilename);
-		MemoryMapBulkData(MemoryMappedFilename, FileData.BulkDataOffsetInFile, BulkDataSize);
+		// Note we don't need a fallback since thge original already managed the load, if we fail now then it
+		// is an actual error.
+		if (Other.IsUsingIODispatcher())
+		{
+			TIoStatusOr<FIoMappedRegion> Status = IoDispatcher->OpenMapped(Data.ChunkID, FIoReadOptions());
+			FIoMappedRegion MappedRegion = Status.ConsumeValueOrDie();
+			DataAllocation.SetMemoryMappedData(this, MappedRegion.MappedFileHandle, MappedRegion.MappedFileRegion);		
+		}
+		else
+		{
+			const int64 BulkDataSize = GetBulkDataSize();
+			FileTokenSystem::Data FileData = FileTokenSystem::GetFileData(Data.Fallback.Token);
+			const FString MemoryMappedFilename = ConvertFilenameFromFlags(FileData.PackageHeaderFilename);
+			MemoryMapBulkData(MemoryMappedFilename, FileData.BulkDataOffsetInFile, BulkDataSize);
+		}
 	}
 
 	return *this;
@@ -670,6 +659,8 @@ FBulkDataBase& FBulkDataBase::operator=(const FBulkDataBase& Other)
 
 FBulkDataBase::~FBulkDataBase()
 {
+	FlushAsyncLoading();
+	
 	check(LockStatus == LOCKSTATUS_Unlocked);
 
 	FreeData();
@@ -736,12 +727,14 @@ void FBulkDataBase::Serialize(FArchive& Ar, UObject* Owner, int32 /*Index*/, boo
 
 		if (!IsInlined() && bUseIoDispatcher)
 		{
-			const EIoChunkType Type = IsOptional() ? EIoChunkType::OptionalBulkData : EIoChunkType::BulkData;
+			const EIoChunkType Type =	IsOptional() ? EIoChunkType::OptionalBulkData : 
+										IsFileMemoryMapped() ? EIoChunkType::MemoryMappedBulkData : 
+										EIoChunkType::BulkData;
 
 			const int64 BulkDataID = BulkDataSize > 0 ? BulkDataOffsetInFile : TNumericLimits<uint64>::Max();
 			Data.ChunkID = CreateBulkdataChunkId(Package->GetPackageId().ToIndex(), BulkDataID, Type);
 
-			BulkDataFlags |= BULKDATA_UsesIoDispatcher; // Indicates that this BulkData should use the FIoChunkId rather than a filename
+			SetRuntimeBulkDataFlags(BULKDATA_UsesIoDispatcher); // Indicates that this BulkData should use the FIoChunkId rather than a filename
 		}
 		else
 		{
@@ -791,11 +784,27 @@ void FBulkDataBase::Serialize(FArchive& Ar, UObject* Owner, int32 /*Index*/, boo
 
 			if (bAttemptFileMapping)
 			{
-				check(Filename != nullptr);
-				FString MemoryMappedFilename = ConvertFilenameFromFlags(*Filename);
-				if (!MemoryMapBulkData(MemoryMappedFilename, BulkDataOffsetInFile, BulkDataSize))
+				if (bUseIoDispatcher)
 				{
-					bShouldForceLoad = true; // Signal we want to force the BulkData to load
+					TIoStatusOr<FIoMappedRegion> Status = IoDispatcher->OpenMapped(Data.ChunkID, FIoReadOptions());
+					if (Status.IsOk())
+					{
+						FIoMappedRegion MappedRegion = Status.ConsumeValueOrDie();
+						DataAllocation.SetMemoryMappedData(this, MappedRegion.MappedFileHandle, MappedRegion.MappedFileRegion);
+					}
+					else
+					{
+						bShouldForceLoad = true; // Signal we want to force the BulkData to load
+					}
+				}
+				else
+				{
+					check(Filename != nullptr);
+					FString MemoryMappedFilename = ConvertFilenameFromFlags(*Filename);
+					if (!MemoryMapBulkData(MemoryMappedFilename, BulkDataOffsetInFile, BulkDataSize))
+					{
+						bShouldForceLoad = true; // Signal we want to force the BulkData to load
+					}
 				}
 			}
 			else if (!Ar.IsAllowingLazyLoading() && !IsInSeperateFile())
@@ -834,7 +843,7 @@ void* FBulkDataBase::Lock(uint32 LockFlags)
 {
 	check(LockStatus == LOCKSTATUS_Unlocked);
 	
-	ForceBulkDataResident(); // Will load 
+	ForceBulkDataResident(); 	// If nothing is currently loaded then load from disk
 
 	if (LockFlags & LOCK_READ_WRITE)
 	{
@@ -905,6 +914,9 @@ void FBulkDataBase::GetCopy(void** DstBuffer, bool bDiscardInternalCopy)
 	check(LockStatus == LOCKSTATUS_Unlocked);
 	check(DstBuffer);
 
+	// Wait for anything that might be currently loading
+	FlushAsyncLoading();
+
 	UE_CLOG(IsDataMemoryMapped(), LogSerialization, Warning, TEXT("FBulkDataBase::GetCopy being called on a memory mapped BulkData object, call ::StealFileMapping instead!"));
 
 	if (*DstBuffer != nullptr)
@@ -954,7 +966,7 @@ void FBulkDataBase::SetBulkDataFlags(uint32 BulkDataFlagsToSet)
 	check(!CanLoadFromDisk());	// We only want to allow the editing of flags if the BulkData
 								// was dynamically created at runtime, not loaded off disk
 
-	BulkDataFlags |= BulkDataFlagsToSet;
+	BulkDataFlags = EBulkDataFlags(BulkDataFlags | BulkDataFlagsToSet);
 }
 
 void FBulkDataBase::ResetBulkDataFlags(uint32 BulkDataFlagsToSet)
@@ -962,7 +974,7 @@ void FBulkDataBase::ResetBulkDataFlags(uint32 BulkDataFlagsToSet)
 	check(!CanLoadFromDisk());	// We only want to allow the editing of flags if the BulkData
 								// was dynamically created at runtime, not loaded off disk
 
-	BulkDataFlags = BulkDataFlagsToSet;
+	BulkDataFlags = (EBulkDataFlags)BulkDataFlagsToSet;
 }
 
 void FBulkDataBase::ClearBulkDataFlags(uint32 BulkDataFlagsToClear)
@@ -970,19 +982,25 @@ void FBulkDataBase::ClearBulkDataFlags(uint32 BulkDataFlagsToClear)
 	check(!CanLoadFromDisk());	// We only want to allow the editing of flags if the BulkData
 								// was dynamically created at runtime, not loaded off disk
 
-	BulkDataFlags &= ~BulkDataFlagsToClear;
+	BulkDataFlags = EBulkDataFlags(BulkDataFlags & ~BulkDataFlagsToClear);
 }
 
 void FBulkDataBase::SetRuntimeBulkDataFlags(uint32 BulkDataFlagsToSet)
 {
-	check(BulkDataFlagsToSet == BULKDATA_UsesIoDispatcher || BulkDataFlagsToSet == BULKDATA_DataIsMemoryMapped);
-	BulkDataFlags |= BulkDataFlagsToSet;
+	check(	BulkDataFlagsToSet == BULKDATA_UsesIoDispatcher || 
+			BulkDataFlagsToSet == BULKDATA_DataIsMemoryMapped || 
+			BulkDataFlagsToSet == BULKDATA_HasAsyncReadPending);
+			
+	BulkDataFlags = EBulkDataFlags(BulkDataFlags | BulkDataFlagsToSet);
 }
 
 void FBulkDataBase::ClearRuntimeBulkDataFlags(uint32 BulkDataFlagsToClear)
 {
-	check(BulkDataFlagsToClear == BULKDATA_UsesIoDispatcher || BulkDataFlagsToClear == BULKDATA_DataIsMemoryMapped);
-	BulkDataFlags &= ~BulkDataFlagsToClear;
+	check(	BulkDataFlagsToClear == BULKDATA_UsesIoDispatcher ||
+			BulkDataFlagsToClear == BULKDATA_DataIsMemoryMapped ||
+			BulkDataFlagsToClear == BULKDATA_HasAsyncReadPending);
+
+	BulkDataFlags = EBulkDataFlags(BulkDataFlags & ~BulkDataFlagsToClear);
 }
 
 int64 FBulkDataBase::GetBulkDataSize() const
@@ -1074,7 +1092,7 @@ bool FBulkDataBase::IsSingleUse() const
 	return (BulkDataFlags & BULKDATA_SingleUse) != 0;
 }
 
-bool FBulkDataBase::IsMemoryMapped() const
+bool FBulkDataBase::IsFileMemoryMapped() const
 {
 	return (BulkDataFlags & BULKDATA_MemoryMappedPayload) != 0;
 }
@@ -1190,6 +1208,10 @@ IBulkDataIORequest* FBulkDataBase::CreateStreamingRequestForRange(const BulkData
 
 void FBulkDataBase::ForceBulkDataResident()
 {
+	// First wait for any async load requests to finish
+	FlushAsyncLoading();
+
+	// Then check if we actually need to load or not
 	if (!IsBulkDataLoaded())
 	{
 		void* DataBuffer = nullptr;
@@ -1218,10 +1240,61 @@ void FBulkDataBase::RemoveBulkData()
 		Data.Fallback.Token = InvalidToken;
 	}
 
-	BulkDataFlags = 0;
+	BulkDataFlags = BULKDATA_None;
 
 }
 
+bool FBulkDataBase::StartAsyncLoading()
+{
+	DECLARE_SCOPE_CYCLE_COUNTER(TEXT("FUntypedBulkData::StartAsyncLoading"), STAT_UBD_StartSerializingBulkData, STATGROUP_Memory);
+
+	if (!IsAsyncLoadingComplete())
+	{
+		return true; // Early out if an asynchronous load is already in progress.
+	}
+
+	if (IsBulkDataLoaded())
+	{
+		return false; // Early out if we do not need to actually load any data
+	}
+
+	if (!CanLoadFromDisk())
+	{
+		return false; // Early out if we cannot load from disk
+	}
+
+	check(LockStatus == LOCKSTATUS_Unlocked)
+
+	LockStatus = LOCKSTATUS_ReadWriteLock; // Bulkdata is effectively locked while streaming!
+
+	// Indicate that we have an async read in flight
+	SetRuntimeBulkDataFlags(BULKDATA_HasAsyncReadPending);
+	FPlatformMisc::MemoryBarrier();
+
+	AsyncCallback Callback = [this](TIoStatusOr<FIoBuffer> Result)
+	{
+		check(Result.IsOk());
+		FIoBuffer IoBuffer = Result.ConsumeValueOrDie();
+
+		check(!IoBuffer.IsMemoryOwned());
+		DataAllocation.SetData(this, IoBuffer.Data());
+
+		FPlatformMisc::MemoryBarrier();
+
+		ClearRuntimeBulkDataFlags(BULKDATA_HasAsyncReadPending);
+		LockStatus = LOCKSTATUS_Unlocked;
+	};
+
+	LoadDataAsynchronously(MoveTemp(Callback));
+
+	return true;
+}
+
+bool FBulkDataBase::IsAsyncLoadingComplete() const
+{
+	return (GetBulkDataFlags() & BULKDATA_HasAsyncReadPending) == 0;
+}
+
 int64 FBulkDataBase::GetBulkDataOffsetInFile() const
 {
 	if (!IsUsingIODispatcher())
@@ -1268,7 +1341,6 @@ bool FBulkDataBase::CanDiscardInternalData() const
 void FBulkDataBase::LoadDataDirectly(void** DstBuffer)
 {
 	DECLARE_SCOPE_CYCLE_COUNTER(TEXT("FBulkDataBase::LoadDataDirectly"), STAT_UBD_LoadDataDirectly, STATGROUP_Memory);
-	
 	if (!CanLoadFromDisk())
 	{
 		UE_LOG(LogSerialization, Warning, TEXT("Attempting to load a BulkData object that cannot be loaded from disk"));
@@ -1277,80 +1349,141 @@ void FBulkDataBase::LoadDataDirectly(void** DstBuffer)
 
 	if (!IsIoDispatcherEnabled())
 	{
-		const int64 BulkDataSize = GetBulkDataSize();
-		FileTokenSystem::Data FileData = FileTokenSystem::GetFileData(Data.Fallback.Token);
+		InternalLoadFromFileSystem(DstBuffer);
+	}
+	else if (IsUsingIODispatcher())
+	{
 
-		FString Filename; 
-		int64 Offset = FileData.BulkDataOffsetInFile;
+		InternalLoadFromIoStore(DstBuffer);
+	}
+	else
+	{
+		// Note that currently this shouldn't be reachable as we should early out due to the ::CanLoadFromDisk check at the start of the method
+		UE_LOG(LogSerialization, Error, TEXT("Attempting to reload inline BulkData when the IoDispatcher is enabled, this operation is not supported! (%d)"), IsInlined());
+	}
+}
 
-		// Fix up the Filename/Offset to work with streaming if EDL is enabled and the filename is still referencing a uasset or umap
-		if (IsInlined() && (FileData.PackageHeaderFilename.EndsWith(TEXT(".uasset")) || FileData.PackageHeaderFilename.EndsWith(TEXT(".umap"))))
-		{
-			Offset -= IFileManager::Get().FileSize(*FileData.PackageHeaderFilename);
-			Filename = FPaths::GetBaseFilename(FileData.PackageHeaderFilename, false) + TEXT(".uexp");
-		}
-		else
-		{
-			Filename = ConvertFilenameFromFlags(FileData.PackageHeaderFilename);
-		}
+void FBulkDataBase::LoadDataAsynchronously(AsyncCallback&& Callback)
+{
+	DECLARE_SCOPE_CYCLE_COUNTER(TEXT("FBulkDataBase::LoadDataDirectly"), STAT_UBD_LoadDataDirectly, STATGROUP_Memory);
+	if (!CanLoadFromDisk())
+	{
+		UE_LOG(LogSerialization, Warning, TEXT("Attempting to load a BulkData object that cannot be loaded from disk"));
+		return; // Early out if there is nothing to load anyway
+	}
 
-		// If the data is inlined then we already loaded is during ::Serialize, this warning should help track cases where data is being discarded then re-requested.
-		// Disabled at the moment
-		UE_CLOG(IsInlined(), LogSerialization, Warning, TEXT("Reloading inlined bulk data directly from disk, this is detrimental to loading performance. Filename: '%s'."), *Filename);
+	if (!IsIoDispatcherEnabled())
+	{
+		Async(EAsyncExecution::ThreadPool, [this, Callback]()
+			{
+				void* DataPtr = nullptr;
+				InternalLoadFromFileSystem(&DataPtr);
 
-		FArchive* Ar = IFileManager::Get().CreateFileReader(*Filename, FILEREAD_Silent);
-		checkf(Ar != nullptr, TEXT("Failed to open the file to load bulk data from. Filename: '%s'."), *Filename);
+				FIoBuffer Buffer(FIoBuffer::Wrap, DataPtr, GetBulkDataSize());
+				TIoStatusOr<FIoBuffer> Status(Buffer);
 
-		// Seek to the beginning of the bulk data in the file.
-		Ar->Seek(Offset);
+				Callback(Status);
+			});
+	}
+	else if (IsUsingIODispatcher())
+	{
+		void* DummyPointer = nullptr;
+		InternalLoadFromIoStoreAsync(&DummyPointer, MoveTemp(Callback));
+	}
+	else
+	{
+		// Note that currently this shouldn't be reachable as we should early out due to the ::CanLoadFromDisk check at the start of the method
+		UE_LOG(LogSerialization, Error, TEXT("Attempting to reload inline BulkData when the IoDispatcher is enabled, this operation is not supported!"));
+	}
+}
 
-		if (*DstBuffer == nullptr)
-		{
-			*DstBuffer = FMemory::Malloc(BulkDataSize, 0);
-		}
+void FBulkDataBase::InternalLoadFromFileSystem(void** DstBuffer)
+{
+	const int64 BulkDataSize = GetBulkDataSize();
+	FileTokenSystem::Data FileData = FileTokenSystem::GetFileData(Data.Fallback.Token);
 
-		SerializeBulkData(*Ar, *DstBuffer, BulkDataSize);
+	FString Filename;
+	int64 Offset = FileData.BulkDataOffsetInFile;
 
-		delete Ar;
+	// Fix up the Filename/Offset to work with streaming if EDL is enabled and the filename is still referencing a uasset or umap
+	if (IsInlined() && (FileData.PackageHeaderFilename.EndsWith(TEXT(".uasset")) || FileData.PackageHeaderFilename.EndsWith(TEXT(".umap"))))
+	{
+		Offset -= IFileManager::Get().FileSize(*FileData.PackageHeaderFilename);
+		Filename = FPaths::GetBaseFilename(FileData.PackageHeaderFilename, false) + BulkDataExt::Export;
 	}
-	else if (IsUsingIODispatcher())
+	else
 	{
-		// Allocate the buffer if needed
-		if (*DstBuffer == nullptr)
-		{
-			*DstBuffer = FMemory::Malloc(GetBulkDataSize(), 0);
-		}
+		Filename = ConvertFilenameFromFlags(FileData.PackageHeaderFilename);
+	}
 
-		// Set up our options (we only need to set the target)
-		FIoReadOptions Options;
-		Options.SetTargetVa(*DstBuffer);
-		
-		FIoBatch NewBatch = GetIoDispatcher()->NewBatch();
-		FIoRequest Request = NewBatch.Read(Data.ChunkID, Options);
+	// If the data is inlined then we already loaded is during ::Serialize, this warning should help track cases where data is being discarded then re-requested.
+	// Disabled at the moment
+	UE_CLOG(IsInlined(), LogSerialization, Warning, TEXT("Reloading inlined bulk data directly from disk, this is detrimental to loading performance. Filename: '%s'."), *Filename);
 
-		NewBatch.Issue();
-		NewBatch.Wait(); // Blocking wait until all requests in the batch are done
+	FArchive* Ar = IFileManager::Get().CreateFileReader(*Filename, FILEREAD_Silent);
+	checkf(Ar != nullptr, TEXT("Failed to open the file to load bulk data from. Filename: '%s'."), *Filename);
 
-		check(Request.IsOk());
+	// Seek to the beginning of the bulk data in the file.
+	Ar->Seek(Offset);
 
-		FBulkDataBase::GetIoDispatcher()->FreeBatch(NewBatch);
+	if (*DstBuffer == nullptr)
+	{
+		*DstBuffer = FMemory::Malloc(BulkDataSize, 0);
 	}
-	else
+
+	SerializeBulkData(*Ar, *DstBuffer, BulkDataSize);
+
+	delete Ar;
+}
+
+void FBulkDataBase::InternalLoadFromIoStore(void** DstBuffer)
+{
+	// Allocate the buffer if needed
+	if (*DstBuffer == nullptr)
 	{
-		// Note that currently this shouldn't be reachable as we should early out due to the ::CanLoadFromDisk check at the start of the method
-		UE_LOG(LogSerialization, Error, TEXT("Attempting to reload inline BulkData when the IoDispatcher is enabled, this operation is not supported! (%d)"), IsInlined());
+		*DstBuffer = FMemory::Malloc(GetBulkDataSize(), 0);
+	}
+
+	// Set up our options (we only need to set the target)
+	FIoReadOptions Options;
+	Options.SetTargetVa(*DstBuffer);
+
+	FIoBatch NewBatch = GetIoDispatcher()->NewBatch();
+	FIoRequest Request = NewBatch.Read(Data.ChunkID, Options);
+
+	NewBatch.Issue();
+	NewBatch.Wait(); // Blocking wait until all requests in the batch are done
+
+	check(Request.IsOk());
+
+	FBulkDataBase::GetIoDispatcher()->FreeBatch(NewBatch);
+}
+
+void FBulkDataBase::InternalLoadFromIoStoreAsync(void** DstBuffer, AsyncCallback&& Callback)
+{
+	// Allocate the buffer if needed
+	if (*DstBuffer == nullptr)
+	{
+		*DstBuffer = FMemory::Malloc(GetBulkDataSize(), 0);
 	}
+
+	// Set up our options (we only need to set the target)
+	FIoReadOptions Options;
+	Options.SetTargetVa(*DstBuffer);
+
+	GetIoDispatcher()->ReadWithCallback(Data.ChunkID, Options, MoveTemp(Callback));
 }
 
 void FBulkDataBase::ProcessDuplicateData(FArchive& Ar, const UPackage* Package, const FString* Filename, int64& InOutSizeOnDisk, int64& InOutOffsetInFile)
 {
 	// We need to load the optional bulkdata info as we might need to create a FIoChunkId based on it!
-	uint32 NewFlags;
+	EBulkDataFlags NewFlags;
 	int64 NewSizeOnDisk;
 	int64 NewOffset;
 
 	SerializeDuplicateData(Ar, NewFlags, NewSizeOnDisk, NewOffset);
 
+
 #if ALLOW_OPTIONAL_DATA
 	if (IsUsingIODispatcher())
 	{
@@ -1359,7 +1492,7 @@ void FBulkDataBase::ProcessDuplicateData(FArchive& Ar, const UPackage* Package,
 
 		if (IoDispatcher->DoesChunkExist(OptionalChunkId))
 		{
-			BulkDataFlags = NewFlags | BULKDATA_UsesIoDispatcher;
+			BulkDataFlags = EBulkDataFlags(NewFlags | BULKDATA_UsesIoDispatcher);
 			InOutSizeOnDisk = NewSizeOnDisk;
 			InOutOffsetInFile = NewOffset;
 
@@ -1385,7 +1518,7 @@ void FBulkDataBase::ProcessDuplicateData(FArchive& Ar, const UPackage* Package,
 #endif
 }
 
-void FBulkDataBase::SerializeDuplicateData(FArchive& Ar, uint32& OutBulkDataFlags, int64& OutBulkDataSizeOnDisk, int64& OutBulkDataOffsetInFile)
+void FBulkDataBase::SerializeDuplicateData(FArchive& Ar, EBulkDataFlags& OutBulkDataFlags, int64& OutBulkDataSizeOnDisk, int64& OutBulkDataOffsetInFile)
 {
 	Ar << OutBulkDataFlags;
 
@@ -1468,6 +1601,33 @@ bool FBulkDataBase::MemoryMapBulkData(const FString& Filename, int64 OffsetInBul
 	return true;
 }
 
+void FBulkDataBase::FlushAsyncLoading()
+{
+	if (!IsAsyncLoadingComplete())
+	{
+		DECLARE_SCOPE_CYCLE_COUNTER(TEXT(" FBulkDataBase::FlushAsyncLoading"), STAT_UBD_WaitForAsyncLoading, STATGROUP_Memory);
+
+#if NO_LOGGING
+		while (IsAsyncLoadingComplete() == false)
+		{
+			FPlatformProcess::Sleep(0);
+		}
+#else
+		uint64 StartTime = FPlatformTime::Cycles64();
+		while (IsAsyncLoadingComplete() == false)
+		{
+			if (FPlatformTime::ToMilliseconds64(FPlatformTime::Cycles64() - StartTime))
+			{
+				UE_LOG(LogSerialization, Warning, TEXT("Waiting for %s bulk data (%lld) to be loaded longer than 1000ms"), *GetFilename(), GetBulkDataSize());
+				StartTime = FPlatformTime::Cycles64(); // Reset so we spam the log every second or so that we are stalled!
+			}
+
+			FPlatformProcess::Sleep(0);
+		}
+#endif
+	}	
+}
+
 FString FBulkDataBase::ConvertFilenameFromFlags(const FString& Filename) const
 {	
 	if (IsOptional())
@@ -1475,7 +1635,7 @@ FString FBulkDataBase::ConvertFilenameFromFlags(const FString& Filename) const
 		// Optional data should be tested for first as we in theory can have data that would
 		// be marked as inline, also marked as optional and in this case we should treat it as
 		// optional data first.
-		return FPaths::ChangeExtension(Filename, OptionalExt);
+		return FPathViews::ChangeExtension(Filename, BulkDataExt::Optional);
 	}
 	else if (!IsInSeperateFile())
 	{
@@ -1483,15 +1643,15 @@ FString FBulkDataBase::ConvertFilenameFromFlags(const FString& Filename) const
 	}
 	else if (IsInlined())
 	{
-		return FPaths::ChangeExtension(Filename, InlinedExt);
+		return FPathViews::ChangeExtension(Filename, BulkDataExt::Export);
 	}
-	else if (IsMemoryMapped())
+	else if (IsFileMemoryMapped())
 	{
-		return FPaths::ChangeExtension(Filename, MemoryMappedExt);
+		return FPathViews::ChangeExtension(Filename, BulkDataExt::MemoryMapped);
 	}
 	else
 	{
-		return FPaths::ChangeExtension(Filename, DefaultExt);
+		return FPathViews::ChangeExtension(Filename, BulkDataExt::Default);
 	}
 }
 
@@ -1604,3 +1764,5 @@ void FBulkDataAllocation::Swap(FBulkDataBase* Owner, void** DstBuffer)
 		Owner->ClearRuntimeBulkDataFlags(BULKDATA_DataIsMemoryMapped);
 	}	
 }
+
+#undef CHECK_IOSTATUS
diff --git a/Engine/Source/Runtime/CoreUObject/Private/UObject/SavePackage.cpp b/Engine/Source/Runtime/CoreUObject/Private/UObject/SavePackage.cpp
--- a/Engine/Source/Runtime/CoreUObject/Private/UObject/SavePackage.cpp
+++ b/Engine/Source/Runtime/CoreUObject/Private/UObject/SavePackage.cpp
@@ -6255,9 +6255,22 @@ void SaveBulkData(FLinkerSave* Linker, const UPackage* InOuter, const TCHAR* Fil
 
 			if (SavePackageContext != nullptr && SavePackageContext->BulkDataManifest != nullptr)
 			{
-				const bool bIsOptional = (BulkDataStorageInfo.BulkDataFlags & BULKDATA_OptionalPayload) != 0;
-				const FPackageStoreBulkDataManifest::EBulkdataType Type = !bIsOptional ? FPackageStoreBulkDataManifest::EBulkdataType::Normal :
-																						 FPackageStoreBulkDataManifest::EBulkdataType::Optional;
+				auto BulkDataTypeFromFlags = [](uint32 BulkDataFlags) 
+				{
+					if (BulkDataFlags & BULKDATA_MemoryMappedPayload)
+					{
+						return FPackageStoreBulkDataManifest::EBulkdataType::MemoryMapped;
+					}
+
+					if (BulkDataFlags & BULKDATA_OptionalPayload)
+					{
+						return FPackageStoreBulkDataManifest::EBulkdataType::Optional;
+					}
+
+					return FPackageStoreBulkDataManifest::EBulkdataType::Normal;
+				};
+
+				const FPackageStoreBulkDataManifest::EBulkdataType Type = BulkDataTypeFromFlags(BulkDataStorageInfo.BulkDataFlags);
 
 				SavePackageContext->BulkDataManifest->AddFileAccess(Filename, Type, StoredBulkStartOffset, BulkStartOffset, SizeOnDisk);
 			}
diff --git a/Engine/Source/Runtime/CoreUObject/Public/Serialization/AsyncLoading2.h b/Engine/Source/Runtime/CoreUObject/Public/Serialization/AsyncLoading2.h
--- a/Engine/Source/Runtime/CoreUObject/Public/Serialization/AsyncLoading2.h
+++ b/Engine/Source/Runtime/CoreUObject/Public/Serialization/AsyncLoading2.h
@@ -8,12 +8,141 @@
 
 #include "CoreMinimal.h"
 #include "UObject/ObjectResource.h"
+#include "UObject/PackageId.h"
+#include "Serialization/Archive.h"
 
 class FArchive;
 class IAsyncPackageLoader;
 class FIoDispatcher;
 class IEDLBootNotificationManager;
 
+using FSourceToLocalizedPackageIdMap = TMap<FPackageId, FPackageId>;
+using FCulturePackageMap = TMap<FString, FSourceToLocalizedPackageIdMap>;
+
+struct FNameMapEntry
+{
+	uint32 Index;
+	uint32 Number;
+#if WITH_EDITOR
+	FName ToFName() const
+	{
+		return FName::CreateFromDisplayId(FNameEntryId::FromUnstableInt(Index), Number);
+	}
+#else
+	FName ToFName() const
+	{
+		return *reinterpret_cast<const FName*>(this);
+	}
+#endif
+
+	COREUOBJECT_API friend FArchive& operator<<(FArchive& Ar, FNameMapEntry& NameMapEntry);
+};
+
+struct FNameMapIndex
+{
+	uint32 NameEntryIndex;
+	uint32 Number;
+};
+
+class FPackageObjectIndex
+{
+	static constexpr uint32 IndexBits = 30;
+	static constexpr uint32 IndexMask = (1 << IndexBits) - 1;
+	static constexpr uint32 TypeMask = ~IndexMask;
+	static constexpr uint32 TypeShift = IndexBits;
+
+	uint32 TypeAndIndex = Null << TypeShift;
+
+public:
+	enum Type
+	{
+		ScriptImport,
+		PackageImport,
+		ImportTypeCount,
+		Export = ImportTypeCount,
+		Null,
+		TypeCount = Null,
+	};
+	static_assert((TypeCount - 1) <= (TypeMask >> TypeShift), "FPackageObjectIndex: Too many index types for TypeMask");
+
+	FPackageObjectIndex() = default;
+	inline explicit FPackageObjectIndex(Type InType, int32 InIndex) : TypeAndIndex((InType << TypeShift) | InIndex) {}
+
+	inline bool IsNull() const
+	{
+		return (TypeAndIndex & TypeMask) == (Null << TypeShift);
+	}
+
+	inline bool IsExport() const
+	{
+		return (TypeAndIndex & TypeMask) == (Export << TypeShift);
+	}
+
+	inline bool IsImport() const
+	{
+		return IsScriptImport() | IsPackageImport();
+	}
+
+	inline bool IsScriptImport() const
+	{
+		return (TypeAndIndex & TypeMask) == (ScriptImport << TypeShift);
+	}
+
+	inline bool IsPackageImport() const
+	{
+		return (TypeAndIndex & TypeMask) == (PackageImport << TypeShift);
+	}
+
+	inline uint32 ToExport() const
+	{
+		check(IsExport());
+		return TypeAndIndex & IndexMask;
+	}
+
+	inline uint32 ToScriptImport() const
+	{
+		check(IsScriptImport());
+		return TypeAndIndex & IndexMask;
+	}
+
+	inline uint32 ToPackageImport() const
+	{
+		check(IsScriptImport());
+		return TypeAndIndex & IndexMask;
+	}
+
+	inline Type GetType() const
+	{
+		return Type((TypeAndIndex & TypeMask) >> TypeShift);
+	}
+
+	inline int32 GetIndex() const
+	{
+		return TypeAndIndex & IndexMask;
+	}
+
+	inline bool operator==(FPackageObjectIndex Other) const
+	{
+		return TypeAndIndex == Other.TypeAndIndex;
+	}
+
+	inline bool operator!=(FPackageObjectIndex Other) const
+	{
+		return TypeAndIndex != Other.TypeAndIndex;
+	}
+
+	COREUOBJECT_API friend FArchive& operator<<(FArchive& Ar, FPackageObjectIndex& Value)
+	{
+		Ar << Value.TypeAndIndex;
+		return Ar;
+	}
+
+	inline friend uint32 GetTypeHash(const FPackageObjectIndex& Value)
+	{
+		return Value.TypeAndIndex;
+	}
+};
+
 /**
  * Event node.
  */
@@ -44,15 +173,13 @@ enum class EExportFilterFlags : uint8
 struct FPackageSummary
 {
 	uint32 PackageFlags;
+	uint32 CookedHeaderSize;
 	int32 NameMapOffset;
 	int32 ImportMapOffset;
 	int32 ExportMapOffset;
 	int32 ExportBundlesOffset;
 	int32 GraphDataOffset;
 	int32 GraphDataSize;
-	int32 BulkDataStartOffset;
-	int32 GlobalImportIndex;
-	int32 Pad;
 };
 
 /**
@@ -78,6 +205,41 @@ struct FExportBundleMetaEntry
 {
 	uint32 LoadOrder = ~uint32(0);
 	uint32 PayloadSize = ~uint32(0);
+
+	COREUOBJECT_API friend FArchive& operator<<(FArchive& Ar, FExportBundleMetaEntry& ExportBundleMetaEntry);
+};
+
+template<typename T>
+class TPackageStoreEntryCArrayView
+{
+	const uint32 ArrayNum = 0;
+	const uint32 OffsetToDataFromThis = 0;
+
+public:
+	inline uint32 Num() const						{ return ArrayNum; }
+
+	inline const T* Data() const					{ return (T*)((char*)this + OffsetToDataFromThis); }
+	inline T* Data()								{ return (T*)((char*)this + OffsetToDataFromThis); }
+
+	inline const T* begin() const					{ return Data(); }
+	inline T* begin()								{ return Data(); }
+
+	inline const T* end() const						{ return Data() + ArrayNum; }
+	inline T* end()									{ return Data() + ArrayNum; }
+
+	inline const T& operator[](uint32 Index) const	{ return Data()[Index]; }
+	inline T& operator[](uint32 Index)				{ return Data()[Index]; }
+};
+
+
+struct FPackageStoreEntry
+{
+	FNameMapEntry Name;
+	FPackageId SourcePackageId;
+	int32 ExportCount;
+	TPackageStoreEntryCArrayView<FPackageId> ImportedPackages;
+	TPackageStoreEntryCArrayView<FPackageObjectIndex> PublicExports;
+	TPackageStoreEntryCArrayView<FExportBundleMetaEntry> ExportBundles;
 };
 
 /**
@@ -91,18 +253,28 @@ struct FExportBundleHeader
 	COREUOBJECT_API friend FArchive& operator<<(FArchive& Ar, FExportBundleHeader& ExportBundleHeader);
 };
 
+struct FScriptObjectEntry
+{
+	FNameMapEntry ObjectName;
+	FPackageObjectIndex OuterIndex;
+	FPackageObjectIndex CDOClassIndex;
+
+	COREUOBJECT_API friend FArchive& operator<<(FArchive& Ar, FScriptObjectEntry& ScriptObjectEntry);
+};
+
 /**
  * Export map entry.
  */
 struct FExportMapEntry
 {
-	uint64 SerialSize;
-	int32 ObjectName[2];
-	FPackageIndex OuterIndex;
-	FPackageIndex ClassIndex;
-	FPackageIndex SuperIndex;
-	FPackageIndex TemplateIndex;
-	int32 GlobalImportIndex;
+	uint64 CookedSerialOffset;
+	uint64 CookedSerialSize;
+	FNameMapIndex ObjectName;
+	FPackageObjectIndex OuterIndex;
+	FPackageObjectIndex ClassIndex;
+	FPackageObjectIndex SuperIndex;
+	FPackageObjectIndex TemplateIndex;
+	FPackageObjectIndex GlobalImportIndex;
 	EObjectFlags ObjectFlags;
 	EExportFilterFlags FilterFlags;
 	uint8 Pad[7];
@@ -110,12 +282,19 @@ struct FExportMapEntry
 	COREUOBJECT_API friend FArchive& operator<<(FArchive& Ar, FExportMapEntry& ExportMapEntry);
 };
 
+#ifndef WITH_ASYNCLOADING2
+#define WITH_ASYNCLOADING2 !WITH_EDITORONLY_DATA
+#endif
+
+#if WITH_ASYNCLOADING2
+
 /**
  * Creates a new instance of the AsyncPackageLoader #2.
  *
  * @param InIoDispatcher				The I/O dispatcher.
- * @param InEDLBootNotificationManager	The EDL boot notification manager.
  *
  * @return The async package loader.
  */
-IAsyncPackageLoader* MakeAsyncPackageLoader2(FIoDispatcher& InIoDispatcher, IEDLBootNotificationManager& InEDLBootNotificationManager);
+IAsyncPackageLoader* MakeAsyncPackageLoader2(FIoDispatcher& InIoDispatcher);
+
+#endif
diff --git a/Engine/Source/Runtime/CoreUObject/Public/Serialization/AsyncPackageLoader.h b/Engine/Source/Runtime/CoreUObject/Public/Serialization/AsyncPackageLoader.h
--- a/Engine/Source/Runtime/CoreUObject/Public/Serialization/AsyncPackageLoader.h
+++ b/Engine/Source/Runtime/CoreUObject/Public/Serialization/AsyncPackageLoader.h
@@ -46,7 +46,7 @@ class IEDLBootNotificationManager
 public:
 	virtual ~IEDLBootNotificationManager() = default;
 
-	virtual bool AddWaitingPackage(void* Pkg, FName PackageName, FName ObjectName, FPackageIndex Import) = 0;
+	virtual bool AddWaitingPackage(void* Pkg, FName PackageName, FName ObjectName, FPackageIndex Import, bool bIgnoreMissingPackage) = 0;
 	virtual bool ConstructWaitingBootObjects() = 0;
 	virtual bool FireCompletedCompiledInImports(bool bFinalRun = false) = 0;
 	virtual bool IsWaitingForSomething() = 0;
@@ -197,3 +197,17 @@ extern COREUOBJECT_API uint32 GFlushAsyncLoadingCount;
 extern COREUOBJECT_API uint32 GSyncLoadCount;
 
 extern COREUOBJECT_API void ResetAsyncLoadingStats();
+
+// Time limit
+extern COREUOBJECT_API int32 GWarnIfTimeLimitExceeded;
+extern COREUOBJECT_API float GTimeLimitExceededMultiplier;
+extern COREUOBJECT_API float GTimeLimitExceededMinTime;
+
+void IsTimeLimitExceededPrint(
+	double InTickStartTime,
+	double CurrentTime,
+	double LastTestTime,
+	float InTimeLimit,
+	const TCHAR* InLastTypeOfWorkPerformed = nullptr,
+	UObject* InLastObjectWorkWasPerformedOn = nullptr);
+
diff --git a/Engine/Source/Runtime/CoreUObject/Public/Serialization/BulkData.h b/Engine/Source/Runtime/CoreUObject/Public/Serialization/BulkData.h
--- a/Engine/Source/Runtime/CoreUObject/Public/Serialization/BulkData.h
+++ b/Engine/Source/Runtime/CoreUObject/Public/Serialization/BulkData.h
@@ -8,6 +8,7 @@
 #include "Async/Future.h"
 #include "Async/AsyncFileHandle.h"
 #include "GenericPlatform/GenericPlatformFile.h"
+#include "BulkDataCommon.h"
 #include "BulkData2.h"
 
 #if WITH_EDITOR == 0 && WITH_EDITORONLY_DATA == 0 //Runtime
@@ -25,76 +26,6 @@
 	#define STREAMINGTOKEN_PARAM(param) 
 #endif
 
-/**
- * Flags serialized with the bulk data.
- */
-enum EBulkDataFlags
-{
-	/** Empty flag set.																*/
-	BULKDATA_None								= 0,
-	/** If set, payload is stored at the end of the file and not inline				*/
-	BULKDATA_PayloadAtEndOfFile					= 1<<0,
-	/** If set, payload should be [un]compressed using ZLIB during serialization.	*/
-	BULKDATA_SerializeCompressedZLIB			= 1<<1,
-	/** Force usage of SerializeElement over bulk serialization.					*/
-	BULKDATA_ForceSingleElementSerialization	= 1<<2,
-	/** Bulk data is only used once at runtime in the game.							*/
-	BULKDATA_SingleUse							= 1<<3,
-	/** Bulk data won't be used and doesn't need to be loaded						*/
-	BULKDATA_Unused								= 1<<5,
-	/** Forces the payload to be saved inline, regardless of its size				*/
-	BULKDATA_ForceInlinePayload					= 1<<6,
-	/** Flag to check if either compression mode is specified						*/
-	BULKDATA_SerializeCompressed				= (BULKDATA_SerializeCompressedZLIB),
-	/** Forces the payload to be always streamed, regardless of its size */
-	BULKDATA_ForceStreamPayload					= 1 << 7,
-	/** If set, payload is stored in a .upack file alongside the uasset				*/
-	BULKDATA_PayloadInSeperateFile				= 1 << 8,
-	/** DEPRECATED: If set, payload is compressed using platform specific bit window	*/
-	BULKDATA_SerializeCompressedBitWindow		= 1 << 9,
-	/** There is a new default to inline unless you opt out */
-	BULKDATA_Force_NOT_InlinePayload			= 1 << 10,
-	/** This payload is optional and may not be on device */
-	BULKDATA_OptionalPayload					= 1 << 11,
-	/** This payload will be memory mapped, this requires alignment, no compression etc. */
-	BULKDATA_MemoryMappedPayload				= 1 << 12,
-	/** Bulk data size is 64 bits long */
-	BULKDATA_Size64Bit							= 1 << 13,
-	/** Duplicate non-optional payload in optional bulk data. */
-	BULKDATA_DuplicateNonOptionalPayload		= 1 << 14,
-	/** Indicates that an old ID is present in the data, at some point when the DDCs are flushed we can remove this. */
-	BULKDATA_BadDataVersion						= 1 << 15,
-	
-	/* Runtime flags go below! */
-	
-	/** Assigned at runtime to indicate that the BulkData should be using the IoDispatcher when loading, not filepaths. */
-	BULKDATA_UsesIoDispatcher					= 1 << 16,
-	/** Assigned at runtime to indicate that the BulkData allocation is a memory mapped region of a file and not raw data */
-	BULKDATA_DataIsMemoryMapped					= 1 << 17,
-};
-
-/**
- * Enumeration for bulk data lock status.
- */
-enum EBulkDataLockStatus
-{
-	/** Unlocked array													*/
-	LOCKSTATUS_Unlocked							= 0,
-	/** Locked read-only												*/
-	LOCKSTATUS_ReadOnlyLock						= 1,
-	/** Locked read-write-realloc										*/
-	LOCKSTATUS_ReadWriteLock					= 2,
-};
-
-/**
- * Enumeration for bulk data lock behavior
- */
-enum EBulkDataLockFlags
-{
-	LOCK_READ_ONLY								= 1,
-	LOCK_READ_WRITE								= 2,
-};
-
 class IMappedFileHandle;
 class IMappedFileRegion;
 
@@ -107,8 +38,7 @@ class IMappedFileRegion;
  */
 struct COREUOBJECT_API FOwnedBulkDataPtr
 {
-
-	FOwnedBulkDataPtr(void* InAllocatedData)
+	explicit FOwnedBulkDataPtr(void* InAllocatedData)
 		: AllocatedData(InAllocatedData)
 		, MappedHandle(nullptr)
 		, MappedRegion(nullptr)
@@ -160,6 +90,8 @@ class FBulkDataIORequest : public IBulkDataIORequest
 {
 public:
 	FBulkDataIORequest(IAsyncReadFileHandle* InFileHandle);
+	FBulkDataIORequest(IAsyncReadFileHandle* InFileHandle, IAsyncReadRequest* InReadRequest, int64 BytesToRead);
+
 	virtual ~FBulkDataIORequest();
 
 	bool MakeReadRequest(int64 Offset, int64 BytesToRead, EAsyncIOPriorityAndFlags PriorityAndFlags, FBulkDataIORequestCallBack* CompleteCallback, uint8* UserSuppliedMemory);
@@ -636,6 +568,16 @@ public:
 	 * Forces the bulk data to be resident in memory and detaches the archive.
 	 */
 	void ForceBulkDataResident();
+
+	/** 
+	* Initiates a new asynchronous operation to load the dulkdata from disk assuming that it is not already
+	* loaded.
+	* Note that a new asynchronous loading operation will not be created if one is already in progress.
+	*
+	* @return True if an asynchronous loading operation is in progress by the time that the method returns
+	* and false if the data is already loaded or cannot be loaded from disk.
+	*/
+	bool StartAsyncLoading();
 	
 	/**
 	 * Sets whether we should store the data compressed on disk.
@@ -839,7 +781,7 @@ private:
 	-----------------------------------------------------------------------------*/
 
 	/** Serialized flags for bulk data																					*/
-	uint32					BulkDataFlags;
+	EBulkDataFlags			BulkDataFlags;
 	/** Alignment of bulk data																							*/
 	uint16					BulkDataAlignment;
 	/** Current lock status																								*/
diff --git a/Engine/Source/Runtime/CoreUObject/Public/Serialization/BulkData2.h b/Engine/Source/Runtime/CoreUObject/Public/Serialization/BulkData2.h
--- a/Engine/Source/Runtime/CoreUObject/Public/Serialization/BulkData2.h
+++ b/Engine/Source/Runtime/CoreUObject/Public/Serialization/BulkData2.h
@@ -3,6 +3,7 @@
 #pragma once
 
 #include "CoreMinimal.h"
+#include "BulkDataCommon.h"
 #include "BulkDataBuffer.h"
 #include "Async/AsyncFileHandle.h"
 #include "IO/IoDispatcher.h"
@@ -10,6 +11,7 @@
 struct FOwnedBulkDataPtr;
 class IMappedFileHandle;
 class IMappedFileRegion;
+class FBulkDataBase;
 
 /**
  * Represents an IO request from the BulkData streaming API.
@@ -31,9 +33,10 @@ public:
 	virtual void Cancel() = 0;
 };
 
-using FileToken = int32;
 struct FBulkDataOrId
 {
+	using FileToken = int32;
+
 	union
 	{
 		// Inline data or fallback path
@@ -50,6 +53,9 @@ struct FBulkDataOrId
 };
 DECLARE_INTRINSIC_TYPE_LAYOUT(FBulkDataOrId);
 
+// Declare this here rather than BulkDataCommon.h
+DECLARE_INTRINSIC_TYPE_LAYOUT(EBulkDataFlags);
+
 /**
  * This is a wrapper for the BulkData memory allocation so we can use a single pointer to either
  * reference a straight memory allocation or in the case that the BulkData object represents a 
@@ -62,7 +68,6 @@ DECLARE_INTRINSIC_TYPE_LAYOUT(FBulkDataOrId);
  * Note: We use a flag set in the owning BulkData object to tell us what storage type we are using 
  * so all accessors require that a pointer to the parent object be passed in.
  */
-class FBulkDataBase;
 class FBulkDataAllocation
 {
 public:
@@ -104,21 +109,29 @@ public:
 	using BulkDataRangeArray = TArray<FBulkDataBase*, TInlineAllocator<8>>;
 
 	static void				SetIoDispatcher(FIoDispatcher* InIoDispatcher) { IoDispatcher = InIoDispatcher; }
-	static FIoDispatcher* GetIoDispatcher() { return IoDispatcher; }
+	static FIoDispatcher*	GetIoDispatcher() { return IoDispatcher; }
 public:
-	static constexpr FileToken InvalidToken = INDEX_NONE;
-
-	FBulkDataBase(const FBulkDataBase& Other) { *this = Other; }
-	FBulkDataBase(FBulkDataBase&& Other);
-	FBulkDataBase& operator=(const FBulkDataBase& Other);
+	static constexpr FBulkDataOrId::FileToken InvalidToken = INDEX_NONE;
 
 	FBulkDataBase()
 	{ 
 		Data.Fallback.BulkDataSize = 0;
 		Data.Fallback.Token = InvalidToken;
 	}
-	~FBulkDataBase();
+	
+	FBulkDataBase(const FBulkDataBase& Other)
+	{
+		// Need some partial initialization of operator= will try to release the token!
+		Data.Fallback.BulkDataSize = 0;
+		Data.Fallback.Token = InvalidToken;
 
+		*this = Other;
+	}
+
+	FBulkDataBase(FBulkDataBase&& Other);
+	FBulkDataBase& operator=(const FBulkDataBase& Other);
+
+	~FBulkDataBase();
 protected:
 
 	void Serialize(FArchive& Ar, UObject* Owner, int32 Index, bool bAttemptFileMapping, int32 ElementSize);
@@ -165,10 +178,12 @@ public:
 	bool IsOptional() const;
 	bool IsInlined() const;
 	UE_DEPRECATED(4.25, "Use ::IsInSeperateFile() instead")
-		FORCEINLINE bool InSeperateFile() const { return IsInSeperateFile(); }
+	FORCEINLINE bool InSeperateFile() const { return IsInSeperateFile(); }
 	bool IsInSeperateFile() const;
 	bool IsSingleUse() const;
-	bool IsMemoryMapped() const;
+	UE_DEPRECATED(4.26, "Use ::IsFileMemoryMapped() instead")
+	bool IsMemoryMapped() const { return IsFileMemoryMapped(); }
+	bool IsFileMemoryMapped() const;
 	bool IsDataMemoryMapped() const;
 	bool IsUsingIODispatcher() const;
 
@@ -181,7 +196,16 @@ public:
 
 	void RemoveBulkData();
 
-	bool IsAsyncLoadingComplete() const { return true; }
+	/**
+	* Initiates a new asynchronous operation to load the dulkdata from disk assuming that it is not already
+	* loaded.
+	* Note that a new asynchronous loading operation will not be created if one is already in progress.
+	*
+	* @return True if an asynchronous loading operation is in progress by the time that the method returns
+	* and false if the data is already loaded or cannot be loaded from disk.
+	*/
+	bool StartAsyncLoading();
+	bool IsAsyncLoadingComplete() const;
 
 	// Added for compatibility with the older BulkData system
 	int64 GetBulkDataOffsetInFile() const;
@@ -205,10 +229,8 @@ private:
 	 */
 	bool CanDiscardInternalData() const;
 
-	void LoadDataDirectly(void** DstBuffer);
-
 	void ProcessDuplicateData(FArchive& Ar, const UPackage* Package, const FString* Filename, int64& InOutSizeOnDisk, int64& InOutOffsetInFile);
-	void SerializeDuplicateData(FArchive& Ar, uint32& OutBulkDataFlags, int64& OutBulkDataSizeOnDisk, int64& OutBulkDataOffsetInFile);
+	void SerializeDuplicateData(FArchive& Ar, EBulkDataFlags& OutBulkDataFlags, int64& OutBulkDataSizeOnDisk, int64& OutBulkDataOffsetInFile);
 	void SerializeBulkData(FArchive& Ar, void* DstBuffer, int64 DataLength);
 
 	bool MemoryMapBulkData(const FString& Filename, int64 OffsetInBulkData, int64 BytesToRead);
@@ -219,15 +241,27 @@ private:
 	FORCEINLINE void* GetDataBufferForWrite() const { return DataAllocation.GetAllocationForWrite(this); }
 	FORCEINLINE const void* GetDataBufferReadOnly() const { return DataAllocation.GetAllocationReadOnly(this); }
 
+	/** Blocking call that waits until any pending async load finishes */
+	void FlushAsyncLoading();
+	
 	FString ConvertFilenameFromFlags(const FString& Filename) const;
 
 private:
+	using AsyncCallback = TFunction<void(TIoStatusOr<FIoBuffer>)>;
+
+	void LoadDataDirectly(void** DstBuffer);
+	void LoadDataAsynchronously(AsyncCallback&& Callback);
+
+	// Used by LoadDataDirectly/LoadDataAsynchronously
+	void InternalLoadFromFileSystem(void** DstBuffer);
+	void InternalLoadFromIoStore(void** DstBuffer);
+	void InternalLoadFromIoStoreAsync(void** DstBuffer, AsyncCallback&& Callback);
 
 	static FIoDispatcher* IoDispatcher;
 
 	LAYOUT_FIELD(FBulkDataOrId, Data);
 	LAYOUT_FIELD(FBulkDataAllocation, DataAllocation);
-	LAYOUT_FIELD_INITIALIZED(uint32, BulkDataFlags, 0);
+	LAYOUT_FIELD_INITIALIZED(EBulkDataFlags, BulkDataFlags, EBulkDataFlags::BULKDATA_None);
 	LAYOUT_MUTABLE_FIELD_INITIALIZED(uint8, LockStatus, 0); // Mutable so that the read only lock can be const
 };
 
diff --git a/Engine/Source/Runtime/CoreUObject/Public/Serialization/BulkDataCommon.h b/Engine/Source/Runtime/CoreUObject/Public/Serialization/BulkDataCommon.h
new file mode 100644
--- /dev/null
+++ b/Engine/Source/Runtime/CoreUObject/Public/Serialization/BulkDataCommon.h
@@ -0,0 +1,96 @@
+// Copyright Epic Games, Inc. All Rights Reserved.
+
+#pragma once
+
+// NOTE: This file only needs to exist as long as we need to maintain the editor and runtime versions of Bulkdata.
+// Code common to both can be placed here.
+
+/**
+ * Flags serialized with the bulk data.
+ */
+enum EBulkDataFlags : uint32
+{
+	/** Empty flag set. */
+	BULKDATA_None = 0,
+	/** If set, payload is stored at the end of the file and not inline. */
+	BULKDATA_PayloadAtEndOfFile = 1 << 0,
+	/** If set, payload should be [un]compressed using ZLIB during serialization. */
+	BULKDATA_SerializeCompressedZLIB = 1 << 1,
+	/** Force usage of SerializeElement over bulk serialization. */
+	BULKDATA_ForceSingleElementSerialization = 1 << 2,
+	/** Bulk data is only used once at runtime in the game. */
+	BULKDATA_SingleUse = 1 << 3,
+	/** Bulk data won't be used and doesn't need to be loaded. */
+	BULKDATA_Unused = 1 << 5,
+	/** Forces the payload to be saved inline, regardless of its size. */
+	BULKDATA_ForceInlinePayload = 1 << 6,
+	/** Flag to check if either compression mode is specified. */
+	BULKDATA_SerializeCompressed = (BULKDATA_SerializeCompressedZLIB),
+	/** Forces the payload to be always streamed, regardless of its size. */
+	BULKDATA_ForceStreamPayload = 1 << 7,
+	/** If set, payload is stored in a .upack file alongside the uasset. */
+	BULKDATA_PayloadInSeperateFile = 1 << 8,
+	/** DEPRECATED: If set, payload is compressed using platform specific bit window. */
+	BULKDATA_SerializeCompressedBitWindow = 1 << 9,
+	/** There is a new default to inline unless you opt out. */
+	BULKDATA_Force_NOT_InlinePayload = 1 << 10,
+	/** This payload is optional and may not be on device. */
+	BULKDATA_OptionalPayload = 1 << 11,
+	/** This payload will be memory mapped, this requires alignment, no compression etc. */
+	BULKDATA_MemoryMappedPayload = 1 << 12,
+	/** Bulk data size is 64 bits long. */
+	BULKDATA_Size64Bit = 1 << 13,
+	/** Duplicate non-optional payload in optional bulk data. */
+	BULKDATA_DuplicateNonOptionalPayload = 1 << 14,
+	/** Indicates that an old ID is present in the data, at some point when the DDCs are flushed we can remove this. */
+	BULKDATA_BadDataVersion = 1 << 15,
+
+	/* Runtime only flags below this point! Note that they take the high bits in reverse order! */
+
+	/** Assigned at runtime to indicate that the BulkData should be using the IoDispatcher when loading, not filepaths. */
+	BULKDATA_UsesIoDispatcher = 1u << 31u,
+	/** Assigned at runtime to indicate that the BulkData allocation is a memory mapped region of a file and not raw data. */
+	BULKDATA_DataIsMemoryMapped = 1 << 30,
+	/** Assigned at runtime to indicate that the BulkData object has an async loading request in flight and will need to wait on it. */
+	BULKDATA_HasAsyncReadPending = 1 << 29
+};
+
+/**
+ * Allows FArchive to serialize EBulkDataFlags, this will not be required once EBulkDataFlags is promoted
+ * to be a enum class.
+ */
+FORCEINLINE FArchive& operator<<(FArchive& Ar, EBulkDataFlags& Flags)
+{
+	Ar << (uint32&)Flags;
+	return Ar;
+}
+
+/**
+ * Enumeration for bulk data lock status.
+ */
+enum EBulkDataLockStatus
+{
+	/** Unlocked array				*/
+	LOCKSTATUS_Unlocked = 0,
+	/** Locked read-only			*/
+	LOCKSTATUS_ReadOnlyLock = 1,
+	/** Locked read-write-realloc	*/
+	LOCKSTATUS_ReadWriteLock = 2,
+};
+
+/**
+ * Enumeration for bulk data lock behavior
+ */
+enum EBulkDataLockFlags
+{
+	LOCK_READ_ONLY = 1,
+	LOCK_READ_WRITE = 2,
+};
+
+namespace BulkDataExt
+{
+	extern const FString Export;		// Stored in the export data
+	extern const FString Default;		// Stored in a separate file
+	extern const FString MemoryMapped;	// Stored in a separate file aligned for memory mapping
+	extern const FString Optional;		// Stored in a separate file that is optional
+}
diff --git a/Engine/Source/Runtime/CoreUObject/Public/Serialization/BulkDataManifest.h b/Engine/Source/Runtime/CoreUObject/Public/Serialization/BulkDataManifest.h
--- a/Engine/Source/Runtime/CoreUObject/Public/Serialization/BulkDataManifest.h
+++ b/Engine/Source/Runtime/CoreUObject/Public/Serialization/BulkDataManifest.h
@@ -11,7 +11,8 @@ public:
 	enum class EBulkdataType : uint8
 	{
 		Normal = 0,
-		Optional
+		Optional,
+		MemoryMapped
 	};
 
 	COREUOBJECT_API FPackageStoreBulkDataManifest(const FString& ProjectPath);
diff --git a/Engine/Source/Runtime/CoreUObject/Public/UObject/PackageId.h b/Engine/Source/Runtime/CoreUObject/Public/UObject/PackageId.h
--- a/Engine/Source/Runtime/CoreUObject/Public/UObject/PackageId.h
+++ b/Engine/Source/Runtime/CoreUObject/Public/UObject/PackageId.h
@@ -38,6 +38,11 @@ public:
 		return Id;
 	}
 
+	inline bool operator<(FPackageId Other) const
+	{
+		return Id < Other.Id;
+	}
+
 	inline bool operator==(FPackageId Other) const
 	{
 		return Id == Other.Id;
diff --git a/Engine/Source/Runtime/Engine/Private/Texture2D.cpp b/Engine/Source/Runtime/Engine/Private/Texture2D.cpp
--- a/Engine/Source/Runtime/Engine/Private/Texture2D.cpp
+++ b/Engine/Source/Runtime/Engine/Private/Texture2D.cpp
@@ -370,7 +370,7 @@ void FTexture2DMipMap::FCompactByteBulkData::GetCopy(void** Dest, bool bDiscardI
 	}
 }
 
-IBulkDataIORequest* FTexture2DMipMap::FCompactByteBulkData::CreateStreamingRequest(FString Filename, int64 OffsetInBulkData, int64 BytesToRead, EAsyncIOPriorityAndFlags Priority, FAsyncFileCallBack* CompleteCallback, uint8* UserSuppliedMemory) const
+IBulkDataIORequest* FTexture2DMipMap::FCompactByteBulkData::CreateStreamingRequest(FString Filename, int64 OffsetInBulkData, int64 BytesToRead, EAsyncIOPriorityAndFlags Priority, FBulkDataIORequestCallBack* CompleteCallback, uint8* UserSuppliedMemory) const
 {
 	check(Filename.IsEmpty() == false);
 
@@ -396,14 +396,16 @@ IBulkDataIORequest* FTexture2DMipMap::FCompactByteBulkData::CreateStreamingReque
 
 	const int64 FinalOffsetInFile = GetBulkDataOffsetInFile() + OffsetInBulkData;
 
-	IAsyncReadRequest* ReadRequest = IORequestHandle->ReadRequest(FinalOffsetInFile, BytesToRead, Priority, CompleteCallback, UserSuppliedMemory);
-	if (ReadRequest != nullptr)
+	FBulkDataIORequest* Request = new FBulkDataIORequest(IORequestHandle);
+	if (Request->MakeReadRequest(FinalOffsetInFile, BytesToRead, Priority, CompleteCallback, UserSuppliedMemory))
 	{
-		return new FBulkDataIORequest(IORequestHandle, ReadRequest, BytesToRead);
+		return Request;
 	}
 	else
 	{
 		delete IORequestHandle;
+		delete Request;
+
 		return nullptr;
 	}
 }
diff --git a/Engine/Source/Runtime/Engine/Public/TextureResource.h b/Engine/Source/Runtime/Engine/Public/TextureResource.h
--- a/Engine/Source/Runtime/Engine/Public/TextureResource.h
+++ b/Engine/Source/Runtime/Engine/Public/TextureResource.h
@@ -86,7 +86,7 @@ struct FTexture2DMipMap
 		void* Realloc(int32 NumBytes);
 		void GetCopy(void** Dest, bool bDiscardInternalCopy = true);
 
-		IBulkDataIORequest* CreateStreamingRequest(FString Filename, int64 OffsetInBulkData, int64 BytesToRead, EAsyncIOPriorityAndFlags Priority, FAsyncFileCallBack* CompleteCallback, uint8* UserSuppliedMemory) const;
+		IBulkDataIORequest* CreateStreamingRequest(FString Filename, int64 OffsetInBulkData, int64 BytesToRead, EAsyncIOPriorityAndFlags Priority, FBulkDataIORequestCallBack* CompleteCallback, uint8* UserSuppliedMemory) const;
 
 		/**
 		 * FCompactByteBulkData doesn't support GetFilename. Use FTexturePlatformData::CachedPackageFileName
