AsyncLoading2: Gracefully handle missing chunks and imports
#rnx

diff --git a/Engine/Source/Runtime/CoreUObject/Private/Serialization/AsyncLoading2.cpp b/Engine/Source/Runtime/CoreUObject/Private/Serialization/AsyncLoading2.cpp
--- a/Engine/Source/Runtime/CoreUObject/Private/Serialization/AsyncLoading2.cpp
+++ b/Engine/Source/Runtime/CoreUObject/Private/Serialization/AsyncLoading2.cpp
@@ -228,8 +228,11 @@ struct FExportObject
 {
 	UObject* Object = nullptr;
 	bool bFiltered = false;
+	bool bExportLoadFailed = false;
 };
 
+char(*__kaboom)[sizeof(FExportObject)] = 1;
+
 struct FAsyncPackageDesc2
 {
 	// A unique request id for each external call to LoadPackage
@@ -2339,7 +2342,15 @@ void FAsyncLoadingThread2::StartBundleIoRequests()
 			ReadOptions,
 			[Package](TIoStatusOr<FIoBuffer> Result)
 		{
-			Package->IoBuffer = Result.ConsumeValueOrDie();
+			if (Result.IsOk())
+			{
+				Package->IoBuffer = Result.ConsumeValueOrDie();
+			}
+			else
+			{
+				UE_LOG(LogStreaming, Error, TEXT("Failed reading chunk for package %s [%s]"), *Package->Desc.NameToLoad.ToString(), *Result.Status().ToString());
+				Package->bLoadHasFailed = true;
+			}
 			Package->GetExportBundleNode(EEventLoadNode2::ExportBundle_Process, 0)->ReleaseBarrier();
 			Package->AsyncLoadingThread.WaitingForIoBundleCounter.Decrement();
 		});
@@ -2898,122 +2909,125 @@ EAsyncPackageState::Type FAsyncPackage2::Event_ProcessExportBundle(FAsyncPackage
 	};
 
 	check(ExportBundleIndex < Package->ExportBundleCount);
-
-	if (Package->AsyncPackageLoadingState == EAsyncPackageLoadingState2::WaitingForSummary)
+	
+	if (!Package->bLoadHasFailed)
 	{
-		TRACE_CPUPROFILER_EVENT_SCOPE(ProcessPackageSummary);
-
-		check(ExportBundleIndex == 0);
-		check(Package->ExportBundleEntryIndex == 0);
-
-		const uint8* PackageSummaryData = Package->IoBuffer.Data();
-		const FPackageSummary* PackageSummary = reinterpret_cast<const FPackageSummary*>(PackageSummaryData);
-		const uint8* GraphData = PackageSummaryData + PackageSummary->GraphDataOffset;
-		const uint64 PackageSummarySize = GraphData + PackageSummary->GraphDataSize - PackageSummaryData;
+		if (Package->AsyncPackageLoadingState == EAsyncPackageLoadingState2::WaitingForSummary)
+		{
+			TRACE_CPUPROFILER_EVENT_SCOPE(ProcessPackageSummary);
 
-		Package->CookedHeaderSize = PackageSummary->CookedHeaderSize;
-		Package->PackageNameMap = reinterpret_cast<const int32*>(PackageSummaryData + PackageSummary->NameMapOffset);
-		Package->ImportStore.ImportMap = reinterpret_cast<const FPackageObjectIndex*>(PackageSummaryData + PackageSummary->ImportMapOffset);
-		Package->ImportStore.ImportMapCount = (PackageSummary->ExportMapOffset - PackageSummary->ImportMapOffset) / sizeof(int32);
-		Package->ExportMap = reinterpret_cast<const FExportMapEntry*>(PackageSummaryData + PackageSummary->ExportMapOffset);
-		Package->ExportBundles = reinterpret_cast<const FExportBundleHeader*>(PackageSummaryData + PackageSummary->ExportBundlesOffset);
-		Package->ExportBundleEntries = reinterpret_cast<const FExportBundleEntry*>(Package->ExportBundles + Package->ExportBundleCount);
+			check(ExportBundleIndex == 0);
+			check(Package->ExportBundleEntryIndex == 0);
 
-		Package->CreateUPackage(PackageSummary);
-		Package->SetupSerializedArcs(GraphData, PackageSummary->GraphDataSize);
+			const uint8* PackageSummaryData = Package->IoBuffer.Data();
+			const FPackageSummary* PackageSummary = reinterpret_cast<const FPackageSummary*>(PackageSummaryData);
+			const uint8* GraphData = PackageSummaryData + PackageSummary->GraphDataOffset;
+			const uint64 PackageSummarySize = GraphData + PackageSummary->GraphDataSize - PackageSummaryData;
 
-		Package->AllExportDataPtr = PackageSummaryData + PackageSummarySize;
-		Package->CurrentExportDataPtr = Package->AllExportDataPtr;
-		Package->AsyncPackageLoadingState = EAsyncPackageLoadingState2::ProcessNewImportsAndExports;
+			Package->CookedHeaderSize = PackageSummary->CookedHeaderSize;
+			Package->PackageNameMap = reinterpret_cast<const int32*>(PackageSummaryData + PackageSummary->NameMapOffset);
+			Package->ImportStore.ImportMap = reinterpret_cast<const FPackageObjectIndex*>(PackageSummaryData + PackageSummary->ImportMapOffset);
+			Package->ImportStore.ImportMapCount = (PackageSummary->ExportMapOffset - PackageSummary->ImportMapOffset) / sizeof(int32);
+			Package->ExportMap = reinterpret_cast<const FExportMapEntry*>(PackageSummaryData + PackageSummary->ExportMapOffset);
+			Package->ExportBundles = reinterpret_cast<const FExportBundleHeader*>(PackageSummaryData + PackageSummary->ExportBundlesOffset);
+			Package->ExportBundleEntries = reinterpret_cast<const FExportBundleEntry*>(Package->ExportBundles + Package->ExportBundleCount);
 
-		TRACE_LOADTIME_PACKAGE_SUMMARY(Package, PackageSummarySize, Package->ImportStore.ImportMapCount, Package->ExportCount);
-	}
+			Package->CreateUPackage(PackageSummary);
+			Package->SetupSerializedArcs(GraphData, PackageSummary->GraphDataSize);
 
-	check(Package->AsyncPackageLoadingState == EAsyncPackageLoadingState2::ProcessNewImportsAndExports);
+			Package->AllExportDataPtr = PackageSummaryData + PackageSummarySize;
+			Package->CurrentExportDataPtr = Package->AllExportDataPtr;
+			Package->AsyncPackageLoadingState = EAsyncPackageLoadingState2::ProcessNewImportsAndExports;
 
-	const uint64 AllExportDataSize = Package->IoBuffer.DataSize() - (Package->AllExportDataPtr - Package->IoBuffer.Data());
-	FExportArchive Ar(Package->AllExportDataPtr, Package->CurrentExportDataPtr, AllExportDataSize);
-	{
-		Ar.SetUE4Ver(Package->LinkerRoot->LinkerPackageVersion);
-		Ar.SetLicenseeUE4Ver(Package->LinkerRoot->LinkerLicenseeVersion);
-		// Ar.SetEngineVer(Summary.SavedByEngineVersion); // very old versioning scheme
-		// Ar.SetCustomVersions(LinkerRoot->LinkerCustomVersion); // only if not cooking with -unversioned
-		Ar.SetUseUnversionedPropertySerialization(CanUseUnversionedPropertySerialization());
-		Ar.SetIsLoading(true);
-		Ar.SetIsPersistent(true);
-		if (Package->LinkerRoot->GetPackageFlags() & PKG_FilterEditorOnly)
-		{
-			Ar.SetFilterEditorOnly(true);
+			TRACE_LOADTIME_PACKAGE_SUMMARY(Package, PackageSummarySize, Package->ImportStore.ImportMapCount, Package->ExportCount);
 		}
-		Ar.ArAllowLazyLoading = true;
 
-		// FExportArchive special fields
-		Ar.CookedHeaderSize = Package->CookedHeaderSize;
-		Ar.PackageDesc = &Package->Desc;
-		Ar.PackageNameMap = Package->PackageNameMap;
-		Ar.GlobalNameMap = &Package->AsyncLoadingThread.GlobalNameMap.GetNameEntries();
-		Ar.ImportStore = &Package->ImportStore;
-		Ar.Exports = &Package->Exports;
-		Ar.ExportMap = Package->ExportMap;
-		Ar.ExportCount = Package->ExportCount;
-		Ar.ExternalReadDependencies = &Package->ExternalReadDependencies;
-	}
-	const FExportBundleHeader* ExportBundle = Package->ExportBundles + ExportBundleIndex;
-	const FExportBundleEntry* BundleEntries = Package->ExportBundleEntries + ExportBundle->FirstEntryIndex;
-	const FExportBundleEntry* BundleEntry = BundleEntries + Package->ExportBundleEntryIndex;
-	const FExportBundleEntry* BundleEntryEnd = BundleEntries + ExportBundle->EntryCount;
-	check(BundleEntry <= BundleEntryEnd);
-	while (BundleEntry < BundleEntryEnd)
-	{
-		if (FAsyncLoadingThreadState2::Get()->IsTimeLimitExceeded(TEXT("Event_ProcessExportBundle")))
-		{
-			return EAsyncPackageState::TimeOut;
-		}
-		const FExportMapEntry& ExportMapEntry = Package->ExportMap[BundleEntry->LocalExportIndex];
-		FExportObject& Export = Package->Exports[BundleEntry->LocalExportIndex];
-		Export.bFiltered = FilterExport(ExportMapEntry.FilterFlags);
+		check(Package->AsyncPackageLoadingState == EAsyncPackageLoadingState2::ProcessNewImportsAndExports);
 
-		if (BundleEntry->CommandType == FExportBundleEntry::ExportCommandType_Create)
-		{
-			if (!Export.bFiltered)
+		const uint64 AllExportDataSize = Package->IoBuffer.DataSize() - (Package->AllExportDataPtr - Package->IoBuffer.Data());
+		FExportArchive Ar(Package->AllExportDataPtr, Package->CurrentExportDataPtr, AllExportDataSize);
+		{
+			Ar.SetUE4Ver(Package->LinkerRoot->LinkerPackageVersion);
+			Ar.SetLicenseeUE4Ver(Package->LinkerRoot->LinkerLicenseeVersion);
+			// Ar.SetEngineVer(Summary.SavedByEngineVersion); // very old versioning scheme
+			// Ar.SetCustomVersions(LinkerRoot->LinkerCustomVersion); // only if not cooking with -unversioned
+			Ar.SetUseUnversionedPropertySerialization(CanUseUnversionedPropertySerialization());
+			Ar.SetIsLoading(true);
+			Ar.SetIsPersistent(true);
+			if (Package->LinkerRoot->GetPackageFlags() & PKG_FilterEditorOnly)
 			{
-				Package->EventDrivenCreateExport(BundleEntry->LocalExportIndex);
+				Ar.SetFilterEditorOnly(true);
 			}
-		}
-		else
-		{
-			check(BundleEntry->CommandType == FExportBundleEntry::ExportCommandType_Serialize);
-
-			const uint64 CookedSerialSize = ExportMapEntry.CookedSerialSize;
-			UObject* Object = Export.Object;
-
-			check(Package->CurrentExportDataPtr + CookedSerialSize <= Package->IoBuffer.Data() + Package->IoBuffer.DataSize());
-			check(Object || Export.bFiltered);
-
-			Ar.ExportBufferBegin(ExportMapEntry.CookedSerialOffset, ExportMapEntry.CookedSerialSize);
-			if (!Export.bFiltered && Object->HasAnyFlags(RF_NeedLoad))
+			Ar.ArAllowLazyLoading = true;
+
+			// FExportArchive special fields
+			Ar.CookedHeaderSize = Package->CookedHeaderSize;
+			Ar.PackageDesc = &Package->Desc;
+			Ar.PackageNameMap = Package->PackageNameMap;
+			Ar.GlobalNameMap = &Package->AsyncLoadingThread.GlobalNameMap.GetNameEntries();
+			Ar.ImportStore = &Package->ImportStore;
+			Ar.Exports = &Package->Exports;
+			Ar.ExportMap = Package->ExportMap;
+			Ar.ExportCount = Package->ExportCount;
+			Ar.ExternalReadDependencies = &Package->ExternalReadDependencies;
+		}
+		const FExportBundleHeader* ExportBundle = Package->ExportBundles + ExportBundleIndex;
+		const FExportBundleEntry* BundleEntries = Package->ExportBundleEntries + ExportBundle->FirstEntryIndex;
+		const FExportBundleEntry* BundleEntry = BundleEntries + Package->ExportBundleEntryIndex;
+		const FExportBundleEntry* BundleEntryEnd = BundleEntries + ExportBundle->EntryCount;
+		check(BundleEntry <= BundleEntryEnd);
+		while (BundleEntry < BundleEntryEnd)
+		{
+			if (FAsyncLoadingThreadState2::Get()->IsTimeLimitExceeded(TEXT("Event_ProcessExportBundle")))
 			{
-				TRACE_LOADTIME_SERIALIZE_EXPORT_SCOPE(Object, CookedSerialSize);
-				const int64 Pos = Ar.Tell();
-				check(CookedSerialSize <= uint64(Ar.TotalSize() - Pos));
+				return EAsyncPackageState::TimeOut;
+			}
+			const FExportMapEntry& ExportMapEntry = Package->ExportMap[BundleEntry->LocalExportIndex];
+			FExportObject& Export = Package->Exports[BundleEntry->LocalExportIndex];
+			Export.bFiltered = FilterExport(ExportMapEntry.FilterFlags);
 
-				Package->EventDrivenSerializeExport(BundleEntry->LocalExportIndex, Ar);
-				checkf(CookedSerialSize == uint64(Ar.Tell() - Pos), TEXT("Expect read size: %llu - Actual read size: %llu"), CookedSerialSize, uint64(Ar.Tell() - Pos));
+			if (BundleEntry->CommandType == FExportBundleEntry::ExportCommandType_Create)
+			{
+				if (!(Export.bFiltered | Export.bExportLoadFailed))
+				{
+					Package->EventDrivenCreateExport(BundleEntry->LocalExportIndex);
+				}
 			}
 			else
 			{
-				Ar.Skip(CookedSerialSize);
-			}
-			Ar.ExportBufferEnd();
+				check(BundleEntry->CommandType == FExportBundleEntry::ExportCommandType_Serialize);
 
-			check((Object && !Object->HasAnyFlags(RF_NeedLoad)) || Export.bFiltered);
+				const uint64 CookedSerialSize = ExportMapEntry.CookedSerialSize;
+				UObject* Object = Export.Object;
 
-			Package->CurrentExportDataPtr += CookedSerialSize;
+				check(Package->CurrentExportDataPtr + CookedSerialSize <= Package->IoBuffer.Data() + Package->IoBuffer.DataSize());
+				check(Object || Export.bFiltered || Export.bExportLoadFailed);
+
+				Ar.ExportBufferBegin(ExportMapEntry.CookedSerialOffset, ExportMapEntry.CookedSerialSize);
+				if (!(Export.bFiltered | Export.bExportLoadFailed) && Object->HasAnyFlags(RF_NeedLoad))
+				{
+					TRACE_LOADTIME_SERIALIZE_EXPORT_SCOPE(Object, CookedSerialSize);
+					const int64 Pos = Ar.Tell();
+					check(CookedSerialSize <= uint64(Ar.TotalSize() - Pos));
+
+					Package->EventDrivenSerializeExport(BundleEntry->LocalExportIndex, Ar);
+					checkf(CookedSerialSize == uint64(Ar.Tell() - Pos), TEXT("Expect read size: %llu - Actual read size: %llu"), CookedSerialSize, uint64(Ar.Tell() - Pos));
+				}
+				else
+				{
+					Ar.Skip(CookedSerialSize);
+				}
+				Ar.ExportBufferEnd();
+
+				check((Object && !Object->HasAnyFlags(RF_NeedLoad)) || Export.bFiltered || Export.bExportLoadFailed);
+
+				Package->CurrentExportDataPtr += CookedSerialSize;
+			}
+			++BundleEntry;
+			++Package->ExportBundleEntryIndex;
 		}
-		++BundleEntry;
-		++Package->ExportBundleEntryIndex;
 	}
-
+	
 	Package->ExportBundleEntryIndex = 0;
 
 	if (ExportBundleIndex + 1 < Package->ExportBundleCount)
@@ -3063,7 +3077,7 @@ UObject* FAsyncPackage2::EventDrivenIndexToObject(FPackageObjectIndex Index, boo
 	else if (Index.IsImport())
 	{
 		Result = ImportStore.FindOrGetImportObject(Index);
-		check(Result);
+		UE_CLOG(!Result, LogStreaming, Warning, TEXT("Missing import for package %s (ScriptImport: %d, Index: %d)"), *Desc.Name.ToString(), Index.IsScriptImport(), Index.IsScriptImport() ? Index.ToScriptImport() : Index.ToPackageImport());
 	}
 #if DO_CHECK
 	if (bCheckSerialized && !IsFullyLoadedObj(Result))
@@ -3081,7 +3095,7 @@ UObject* FAsyncPackage2::EventDrivenIndexToObject(FPackageObjectIndex Index, boo
 		{
 			UE_LOG(LogStreaming, Fatal, TEXT("Missing Dependency, request for %s but it was still has RF_NeedLoad."), *Linker->GetPathName(Index));
 		}*/
-		UE_LOG(LogStreaming, Fatal, TEXT("Missing Dependency"));
+		UE_LOG(LogStreaming, Warning, TEXT("Missing Dependency"));
 	}
 	if (Result)
 	{
@@ -3118,8 +3132,18 @@ void FAsyncPackage2::EventDrivenCreateExport(int32 LocalExportIndex)
 	UClass* LoadClass = Export.ClassIndex.IsNull() ? UClass::StaticClass() : CastEventDrivenIndexToObject<UClass>(Export.ClassIndex, true);
 	UObject* ThisParent = Export.OuterIndex.IsNull() ? LinkerRoot : EventDrivenIndexToObject(Export.OuterIndex, false);
 
-	checkf(LoadClass, TEXT("Could not find class object for %s in %s"), *ObjectName.ToString(), *Desc.NameToLoad.ToString());
-	checkf(ThisParent, TEXT("Could not find outer object for %s in %s"), *ObjectName.ToString(), *Desc.NameToLoad.ToString());
+	if (!LoadClass)
+	{
+		UE_LOG(LogStreaming, Error, TEXT("Could not find class object for %s in %s"), *ObjectName.ToString(), *Desc.NameToLoad.ToString());
+		Exports[LocalExportIndex].bExportLoadFailed = true;
+		return;
+	}
+	if (!ThisParent)
+	{
+		UE_LOG(LogStreaming, Error, TEXT("Could not find outer object for %s in %s"), *ObjectName.ToString(), *Desc.NameToLoad.ToString());
+		Exports[LocalExportIndex].bExportLoadFailed = true;
+		return;
+	}
 	check(!dynamic_cast<UObjectRedirector*>(ThisParent));
 
 	// Try to find existing object first as we cannot in-place replace objects, could have been created by other export in this package
@@ -3162,7 +3186,12 @@ void FAsyncPackage2::EventDrivenCreateExport(int32 LocalExportIndex)
 		// Find the Archetype object for the one we are loading.
 		check(!Export.TemplateIndex.IsNull());
 		UObject* Template = EventDrivenIndexToObject(Export.TemplateIndex, true);
-		checkf(Template, TEXT("Could not find template for %s in %s"), *ObjectName.ToString(), *Desc.NameToLoad.ToString());
+		if (!Template)
+		{
+			UE_LOG(LogStreaming, Error, TEXT("Could not find template for %s in %s"), *ObjectName.ToString(), *Desc.NameToLoad.ToString());
+			Exports[LocalExportIndex].bExportLoadFailed = true;
+			return;
+		}
 		// we also need to ensure that the template has set up any instances
 		Template->ConditionalPostLoadSubobjects();
 
@@ -3267,7 +3296,12 @@ void FAsyncPackage2::EventDrivenSerializeExport(int32 LocalExportIndex, FExportA
 		if (!Export.SuperIndex.IsNull())
 		{
 			UStruct* SuperStruct = CastEventDrivenIndexToObject<UStruct>(Export.SuperIndex, true);
-			checkf(SuperStruct, TEXT("Could not find SuperStruct for %s"), *Object->GetFullName());
+			if (!SuperStruct)
+			{
+				UE_LOG(LogStreaming, Fatal, TEXT("Could not find SuperStruct for %s"), *Object->GetFullName());
+				Exports[LocalExportIndex].bExportLoadFailed = true;
+				return;
+			}
 			Struct->SetSuperStruct(SuperStruct);
 			if (UClass* ClassObject = dynamic_cast<UClass*>(Object))
 			{
@@ -4627,6 +4661,11 @@ EAsyncPackageState::Type FAsyncPackage2::ProcessExternalReads(EExternalReadActio
  */
 EAsyncPackageState::Type FAsyncPackage2::PostLoadObjects()
 {
+	if (bLoadHasFailed)
+	{
+		return EAsyncPackageState::Complete;
+	}
+
 	LLM_SCOPE(ELLMTag::UObject);
 
 	SCOPED_LOADTIMER(PostLoadObjectsTime);
@@ -4643,7 +4682,7 @@ EAsyncPackageState::Type FAsyncPackage2::PostLoadObjects()
 	while (PostLoadIndex < ExportCount && !FAsyncLoadingThreadState2::Get()->IsTimeLimitExceeded(TEXT("PostLoadObject")))
 	{
 		const FExportObject& Export = Exports[PostLoadIndex++];
-		if (Export.bFiltered)
+		if (Export.bFiltered | Export.bExportLoadFailed)
 		{
 			continue;
 		}
@@ -4674,6 +4713,13 @@ EAsyncPackageState::Type FAsyncPackage2::PostLoadObjects()
 
 EAsyncPackageState::Type FAsyncPackage2::PostLoadDeferredObjects()
 {
+	if (bLoadHasFailed)
+	{
+		FSoftObjectPath::InvalidateTag();
+		FUniqueObjectGuid::InvalidateTag();
+		return EAsyncPackageState::Complete;
+	}
+
 	SCOPED_LOADTIMER(PostLoadDeferredObjectsTime);
 
 	FAsyncPackageScope2 PackageScope(this);
@@ -4689,7 +4735,7 @@ EAsyncPackageState::Type FAsyncPackage2::PostLoadDeferredObjects()
 		!FAsyncLoadingThreadState2::Get()->IsTimeLimitExceeded(TEXT("PostLoadDeferredObjects")))
 	{
 		const FExportObject& Export = Exports[DeferredPostLoadIndex++];
-		if (Export.bFiltered)
+		if (Export.bFiltered | Export.bExportLoadFailed)
 		{
 			continue;
 		}
@@ -4720,7 +4766,7 @@ EAsyncPackageState::Type FAsyncPackage2::PostLoadDeferredObjects()
 			!FAsyncLoadingThreadState2::Get()->IsTimeLimitExceeded(TEXT("PostLoadDeferredObjects")))
 		{
 			const FExportObject& Export = Exports[DeferredFinalizeIndex++];
-			if (Export.bFiltered)
+			if (Export.bFiltered | Export.bExportLoadFailed)
 			{
 				continue;
 			}
@@ -4783,7 +4829,7 @@ EAsyncPackageState::Type FAsyncPackage2::PostLoadDeferredObjects()
 			{
 				for (const FExportObject& Export : Exports)
 				{
-					if (!Export.bFiltered && Export.Object->CanBeClusterRoot())
+					if (!(Export.bFiltered | Export.bExportLoadFailed) && Export.Object->CanBeClusterRoot())
 					{
 						bHasClusterObjects = true;
 						break;
@@ -4807,7 +4853,7 @@ EAsyncPackageState::Type FAsyncPackage2::CreateClusters()
 	{
 		const FExportObject& Export = Exports[DeferredClusterIndex++];
 
-		if (!Export.bFiltered && Export.Object->CanBeClusterRoot())
+		if (!(Export.bFiltered | Export.bExportLoadFailed) && Export.Object->CanBeClusterRoot())
 		{
 			Export.Object->CreateCluster();
 		}
