Compile fix
#rnx

diff --git a/Engine/Source/Runtime/CoreUObject/Private/Serialization/AsyncLoading2.cpp b/Engine/Source/Runtime/CoreUObject/Private/Serialization/AsyncLoading2.cpp
--- a/Engine/Source/Runtime/CoreUObject/Private/Serialization/AsyncLoading2.cpp
+++ b/Engine/Source/Runtime/CoreUObject/Private/Serialization/AsyncLoading2.cpp
@@ -231,8 +231,6 @@ struct FExportObject
 	bool bExportLoadFailed = false;
 };
 
-char(*__kaboom)[sizeof(FExportObject)] = 1;
-
 struct FAsyncPackageDesc2
 {
 	// A unique request id for each external call to LoadPackage
