Submitted on behalf of Robert.Manuzsewski

#UE4 Fix load errors not being supressed for subclasses of BlueprintGeneratedClass

diff --git a/Engine/Source/Runtime/CoreUObject/Private/Blueprint/BlueprintSupport.cpp b/Engine/Source/Runtime/CoreUObject/Private/Blueprint/BlueprintSupport.cpp
--- a/Engine/Source/Runtime/CoreUObject/Private/Blueprint/BlueprintSupport.cpp
+++ b/Engine/Source/Runtime/CoreUObject/Private/Blueprint/BlueprintSupport.cpp
@@ -1055,6 +1055,17 @@ bool FLinkerLoad::DeferPotentialCircularImport(const int32 Index)
 }
 
 #if WITH_EDITOR
+/** Helper function find the actual class object given import class and package namme */
+static UClass* FindImportClass(FName ClassPackageName, FName ClassName)
+{
+	UClass* Class = nullptr;
+	UPackage* ClassPackage = Cast<UPackage>(StaticFindObjectFast(UPackage::StaticClass(), nullptr, ClassPackageName));
+	if (ClassPackage)
+	{
+		Class = Cast<UClass>(StaticFindObjectFast(UClass::StaticClass(), ClassPackage, ClassName));
+	}
+	return Class;
+}
 bool FLinkerLoad::IsSuppressableBlueprintImportError(int32 ImportIndex) const
 {
 	// We want to suppress any import errors that target a BlueprintGeneratedClass
@@ -1062,12 +1073,15 @@ bool FLinkerLoad::IsSuppressableBlueprintImportError(int32 ImportIndex) const
 	// without compiling. This should not be a problem because all Blueprints are
 	// compiled-on-load.
 	static const FName NAME_BlueprintGeneratedClass("BlueprintGeneratedClass");
-
+	static const FName NAME_EnginePackage("/Script/Engine");
+	UClass* BlueprintGeneratedClass = FindImportClass(NAME_EnginePackage, NAME_BlueprintGeneratedClass);
+	check(BlueprintGeneratedClass);
 	// We will look at each outer of the Import to see if any of them are a BPGC
 	while (ImportMap.IsValidIndex(ImportIndex))
 	{
 		const FObjectImport& TestImport = ImportMap[ImportIndex];
-		if (TestImport.ClassName == NAME_BlueprintGeneratedClass)
+		UClass* ImportClass = FindImportClass(TestImport.ClassPackage, TestImport.ClassName);
+		if (ImportClass && ImportClass->IsChildOf(BlueprintGeneratedClass))
 		{
 			// The import is a BPGC, suppress errors
 			return true;
@@ -1076,9 +1090,13 @@ bool FLinkerLoad::IsSuppressableBlueprintImportError(int32 ImportIndex) const
 		// Check if this is a BP CDO, if so our class will be in the import table
 		for (const FObjectImport& PotentialBPClass : ImportMap)
 		{
-			if (PotentialBPClass.ObjectName == TestImport.ClassName && PotentialBPClass.ClassName == NAME_BlueprintGeneratedClass)
+			if (PotentialBPClass.ObjectName == TestImport.ClassName)
 			{
-				return true;
+				UClass* PotentialBPClassClass = FindImportClass(PotentialBPClass.ClassName, PotentialBPClass.ClassPackage);
+				if (PotentialBPClassClass && PotentialBPClassClass->IsChildOf(BlueprintGeneratedClass))
+				{
+					return true;
+				}
 			}
 		}
 
