Fix unguarded multithreaded access to FCoreRedirects::KnownMissing functions.
#rnx

diff --git a/Engine/Source/Runtime/CoreUObject/Private/UObject/CoreRedirects.cpp b/Engine/Source/Runtime/CoreUObject/Private/UObject/CoreRedirects.cpp
--- a/Engine/Source/Runtime/CoreUObject/Private/UObject/CoreRedirects.cpp
+++ b/Engine/Source/Runtime/CoreUObject/Private/UObject/CoreRedirects.cpp
@@ -10,6 +10,7 @@
 #include "Misc/AutomationTest.h"
 #include "Misc/ConfigCacheIni.h"
 #include "Misc/PackageName.h"
+#include "Misc/ScopeRWLock.h"
 #include "Serialization/DeferredMessageLog.h"
 #include "Templates/Casts.h"
 
@@ -434,6 +435,7 @@ FCoreRedirectObjectName FCoreRedirect::RedirectName(const FCoreRedirectObjectNam
 bool FCoreRedirects::bInitialized;
 TMap<FName, ECoreRedirectFlags> FCoreRedirects::ConfigKeyMap;
 TMap<ECoreRedirectFlags, FCoreRedirects::FRedirectNameMap> FCoreRedirects::RedirectTypeMap;
+FRWLock FCoreRedirects::KnownMissingLock;
 
 bool FCoreRedirects::RedirectNameAndValues(ECoreRedirectFlags Type, const FCoreRedirectObjectName& OldObjectName, FCoreRedirectObjectName& NewObjectName, const FCoreRedirect** FoundValueRedirect)
 {
@@ -587,6 +589,8 @@ bool FCoreRedirects::FindPreviousNames(ECoreRedirectFlags SearchFlags, const FCo
 bool FCoreRedirects::IsKnownMissing(ECoreRedirectFlags Type, const FCoreRedirectObjectName& ObjectName)
 {
 	TArray<const FCoreRedirect*> FoundRedirects;
+
+	FRWScopeLock ScopeLock(KnownMissingLock, FRWScopeLockType::SLT_ReadOnly);
 	return FCoreRedirects::GetMatchingRedirects(Type | ECoreRedirectFlags::Category_Removed, ObjectName, FoundRedirects);
 }
 
@@ -595,6 +599,8 @@ bool FCoreRedirects::AddKnownMissing(ECoreRedirectFlags Type, const FCoreRedirec
 	check((Channel & ~ECoreRedirectFlags::Option_MissingLoad) == ECoreRedirectFlags::None);
 	TArray<FCoreRedirect> NewRedirects;
 	NewRedirects.Emplace(Type | ECoreRedirectFlags::Category_Removed | Channel, ObjectName, FCoreRedirectObjectName());
+
+	FRWScopeLock ScopeLock(KnownMissingLock, FRWScopeLockType::SLT_Write);
 	return AddRedirectList(NewRedirects, TEXT("AddKnownMissing"));
 }
 
@@ -603,6 +609,8 @@ bool FCoreRedirects::RemoveKnownMissing(ECoreRedirectFlags Type, const FCoreRedi
 	check((Channel & ~ECoreRedirectFlags::Option_MissingLoad) == ECoreRedirectFlags::None);
 	TArray<FCoreRedirect> RedirectsToRemove;
 	RedirectsToRemove.Emplace(Type | ECoreRedirectFlags::Category_Removed | Channel, ObjectName, FCoreRedirectObjectName());
+
+	FRWScopeLock ScopeLock(KnownMissingLock, FRWScopeLockType::SLT_Write);
 	return RemoveRedirectList(RedirectsToRemove, TEXT("RemoveKnownMissing"));
 }
 
@@ -610,6 +618,8 @@ void FCoreRedirects::ClearKnownMissing(ECoreRedirectFlags Type, ECoreRedirectFla
 {
 	check((Channel & ~ECoreRedirectFlags::Option_MissingLoad) == ECoreRedirectFlags::None);
 	ECoreRedirectFlags RedirectFlags = Type | ECoreRedirectFlags::Category_Removed | Channel;
+
+	FRWScopeLock ScopeLock(KnownMissingLock, FRWScopeLockType::SLT_Write);
 	RedirectTypeMap.Remove(RedirectFlags);
 }
 
diff --git a/Engine/Source/Runtime/CoreUObject/Public/UObject/CoreRedirects.h b/Engine/Source/Runtime/CoreUObject/Public/UObject/CoreRedirects.h
--- a/Engine/Source/Runtime/CoreUObject/Public/UObject/CoreRedirects.h
+++ b/Engine/Source/Runtime/CoreUObject/Public/UObject/CoreRedirects.h
@@ -8,6 +8,7 @@
 
 #include "CoreMinimal.h"
 #include "Misc/EnumClassFlags.h"
+#include "HAL/CriticalSection.h"
 
 class IPakFile;
 
@@ -306,4 +307,10 @@ private:
 
 	/** Map from name of thing being mapped to full list. List must be filtered further */
 	static TMap<ECoreRedirectFlags, FRedirectNameMap> RedirectTypeMap;
+
+	/**
+	 * Lock to protect multithreaded access to *KnownMissing functions, which can be called from the async loading threads. 
+	 * TODO: The KnownMissing functions use RedirectTypeMap, which is unguarded; there is race condition vulnerability if asyncloading thread is active before all categories are added to RedirectTypeMap.
+	 */
+	static FRWLock KnownMissingLock;
 };
