Remove debug assert that verifies that the new optimized StaticFindObject behaves exactly the same as the old algorithm

The old algorithm would match if the outermost name ended with the supplied outermost name. The new algorithm does exact matching of the outermost name, which makes more sense and was likely the intent of the old algorithm.

diff --git a/Engine/Source/Runtime/CoreUObject/Private/UObject/UObjectHash.cpp b/Engine/Source/Runtime/CoreUObject/Private/UObject/UObjectHash.cpp
--- a/Engine/Source/Runtime/CoreUObject/Private/UObject/UObjectHash.cpp
+++ b/Engine/Source/Runtime/CoreUObject/Private/UObject/UObjectHash.cpp
@@ -615,8 +615,6 @@ UObject* StaticFindObjectFastInternalThreadSafe(FUObjectHashTables& ThreadHash,
 			{
 				UObject* Object = (UObject*)*It;
 
-				checkSlow(Object->GetFName() != SearchPath.Inner || SearchPath.MatchOuterNames(Object->GetOuter()) == Object->GetPathName().EndsWith(ObjectName.ToString()));
-
 				if
 					((Object->GetFName() == SearchPath.Inner)
 
