Creating new UObjects while Collecting Garbage will now result in an assert.

diff --git a/Engine/Source/Runtime/CoreUObject/Private/UObject/UObjectGlobals.cpp b/Engine/Source/Runtime/CoreUObject/Private/UObject/UObjectGlobals.cpp
--- a/Engine/Source/Runtime/CoreUObject/Private/UObject/UObjectGlobals.cpp
+++ b/Engine/Source/Runtime/CoreUObject/Private/UObject/UObjectGlobals.cpp
@@ -2282,7 +2282,8 @@ UObject* StaticAllocateObject
 	check(GIsEditor || bCreatingCDO || !InClass->HasAnyClassFlags(CLASS_Abstract)); // this is a warning in the editor, otherwise it is illegal to create an abstract class, except the CDO
 	check(InOuter || (InClass == UPackage::StaticClass() && InName != NAME_None)); // only packages can not have an outer, and they must be named explicitly
 	check(bCreatingCDO || !InOuter || InOuter->IsA(InClass->ClassWithin));
-
+	checkf(!IsGarbageCollecting(), TEXT("Unable to create new object: %s %s.%s. Creating UObjects while Collecting Garbage is not allowed!"),
+		*GetNameSafe(InClass), *GetPathNameSafe(InOuter), *InName.ToString());
 
 	if (bCreatingCDO)
 	{
