Chaning how FFieldPaths (TFieldPaths) are internally stored and serialized. Now instead of storing the full path (from the outermost UPackage to the innermost FProperty) FFieldPaths will store the owner UStruct reference and a short path to its property (or inner property of its property). The same applies to serialization: both the owner UStruct reference and only the short path will be serialized.
Serializing the owner reference will enable proper handling of any owner struct package renames as well as greatly increase the initial field (property) resolution performance.

#tests Basic editor functionality tests related to blueprint editing, editor quick smoke tests, cooked windows client and server, cooked build bot soak tests

diff --git a/Engine/Source/Editor/PropertyEditor/Private/PropertyNode.cpp b/Engine/Source/Editor/PropertyEditor/Private/PropertyNode.cpp
--- a/Engine/Source/Editor/PropertyEditor/Private/PropertyNode.cpp
+++ b/Engine/Source/Editor/PropertyEditor/Private/PropertyNode.cpp
@@ -733,45 +733,51 @@ FPropertyAccess::Result FPropertyNode::GetPropertyValueText(FText& OutText, cons
 	if (ValueAddress != nullptr)
 	{
 		const FProperty* PropertyPtr = GetProperty();
-
-		if (PropertyPtr->IsA(FTextProperty::StaticClass()))
+		if (PropertyPtr)
 		{
-			OutText = CastField<FTextProperty>(PropertyPtr)->GetPropertyValue(ValueAddress);
-		}
-		else
-		{
-			FString ExportedTextString;
-			FPropertyTextUtilities::PropertyToTextHelper(ExportedTextString, this, PropertyPtr, ValueAddress, PPF_PropertyWindow);
-
-			UEnum* Enum = nullptr;
-			int64 EnumValue = 0;
-			if (const FByteProperty* ByteProperty = CastField<FByteProperty>(PropertyPtr))
+			if (PropertyPtr->IsA(FTextProperty::StaticClass()))
 			{
-				Enum = ByteProperty->Enum;
-				EnumValue = ByteProperty->GetPropertyValue(ValueAddress);
+				OutText = CastField<FTextProperty>(PropertyPtr)->GetPropertyValue(ValueAddress);
 			}
-			else if (const FEnumProperty* EnumProperty = CastField<FEnumProperty>(PropertyPtr))
+			else
 			{
-				Enum = EnumProperty->GetEnum();
-				EnumValue = EnumProperty->GetUnderlyingProperty()->GetSignedIntPropertyValue(ValueAddress);
-			}
+				FString ExportedTextString;
+				FPropertyTextUtilities::PropertyToTextHelper(ExportedTextString, this, PropertyPtr, ValueAddress, PPF_PropertyWindow);
 
-			if (Enum)
-			{
-				if (Enum->IsValidEnumValue(EnumValue))
+				UEnum* Enum = nullptr;
+				int64 EnumValue = 0;
+				if (const FByteProperty* ByteProperty = CastField<FByteProperty>(PropertyPtr))
 				{
-					// Text form is always display name
-					OutText = Enum->GetDisplayNameTextByValue(EnumValue);
+					Enum = ByteProperty->Enum;
+					EnumValue = ByteProperty->GetPropertyValue(ValueAddress);
+				}
+				else if (const FEnumProperty* EnumProperty = CastField<FEnumProperty>(PropertyPtr))
+				{
+					Enum = EnumProperty->GetEnum();
+					EnumValue = EnumProperty->GetUnderlyingProperty()->GetSignedIntPropertyValue(ValueAddress);
+				}
+
+				if (Enum)
+				{
+					if (Enum->IsValidEnumValue(EnumValue))
+					{
+						// Text form is always display name
+						OutText = Enum->GetDisplayNameTextByValue(EnumValue);
+					}
+					else
+					{
+						Result = FPropertyAccess::Fail;
+					}
 				}
 				else
 				{
-					Result = FPropertyAccess::Fail;
+					OutText = FText::FromString(ExportedTextString);
 				}
 			}
-			else
-			{
-				OutText = FText::FromString(ExportedTextString);
-			}
+		}
+		else
+		{
+			Result = FPropertyAccess::Fail;
 		}
 	}
 
diff --git a/Engine/Source/Runtime/Core/Public/UObject/FortniteMainBranchObjectVersion.h b/Engine/Source/Runtime/Core/Public/UObject/FortniteMainBranchObjectVersion.h
--- a/Engine/Source/Runtime/Core/Public/UObject/FortniteMainBranchObjectVersion.h
+++ b/Engine/Source/Runtime/Core/Public/UObject/FortniteMainBranchObjectVersion.h
@@ -105,6 +105,9 @@ struct CORE_API FFortniteMainBranchObjectVersion
 		// Ocean collision component has become dynamic
 		MakeOceanCollisionTransient,
 
+		// FFieldPath will serialize the owner struct reference and only a short path to its property
+		FFieldPathOwnerSerialization,
+
 		// -----<new versions can be added above this line>-------------------------------------------------
 		VersionPlusOne,
 		LatestVersion = VersionPlusOne - 1
diff --git a/Engine/Source/Runtime/CoreUObject/Private/UObject/Class.cpp b/Engine/Source/Runtime/CoreUObject/Private/UObject/Class.cpp
--- a/Engine/Source/Runtime/CoreUObject/Private/UObject/Class.cpp
+++ b/Engine/Source/Runtime/CoreUObject/Private/UObject/Class.cpp
@@ -1967,7 +1967,7 @@ void UStruct::PostLoad()
 				FField** TargetScriptPropertyPtr = (FField**)(Script.GetData() + MissingProperty.Value);
 				*TargetScriptPropertyPtr = ResolvedProperty;
 			}
-			else if (!MissingProperty.Key.IsEmpty())
+			else if (!MissingProperty.Key.IsPathToFieldEmpty())
 			{
 				UE_LOG(LogClass, Warning, TEXT("Failed to resolve bytecode referenced field from path: %s when loading %s"), *MissingProperty.Key.ToString(), *GetFullName());
 			}
diff --git a/Engine/Source/Runtime/CoreUObject/Private/UObject/FieldPath.cpp b/Engine/Source/Runtime/CoreUObject/Private/UObject/FieldPath.cpp
--- a/Engine/Source/Runtime/CoreUObject/Private/UObject/FieldPath.cpp
+++ b/Engine/Source/Runtime/CoreUObject/Private/UObject/FieldPath.cpp
@@ -8,10 +8,68 @@
 #include "UObject/UnrealType.h"
 #include "UObject/Package.h"
 #include "UObject/UObjectArray.h"
+#include "UObject/FortniteMainBranchObjectVersion.h"
+#include "UObject/LinkerLoad.h"
+#include "UObject/UObjectThreadContext.h"
+#include "UObject/PropertyHelper.h"
+#include "Algo/Reverse.h"
+
+/** Helper function used for logging the path to property */
+static FString PathToString(const TArray<FName>& InPath)
+{
+	// See FFieldPath::Generate for formatting specifics
+	// Generate path from the outermost (last item) to the property (first item)
+
+	// Initilize with the number of delimiters we're going to add
+	int32 PathLength = InPath.Num() ? (InPath.Num() - 1) : 0; 
+	for (FName PathSegment : InPath)
+	{
+		PathLength += PathSegment.GetStringLength();
+	}
+	FString Result;
+	if (PathLength)
+	{
+		// Allocate once and construct the path string
+		Result.Reserve(PathLength);
+
+		// Handle nativized blueprint oddities 
+		FString ObjName = InPath.Last().ToString();
+		if (ObjName.StartsWith(UDynamicClass::GetTempPackagePrefix(), ESearchCase::IgnoreCase))
+		{
+			ObjName.RemoveFromStart(UDynamicClass::GetTempPackagePrefix(), ESearchCase::IgnoreCase);
+		}
+		Result += ObjName;
+
+		if (InPath.Num() > 1)
+		{
+			// Path may be either a full path from the package to the property or just from the outer property to the inner property
+			// If it's a full path, the delimiter char for the next object (owner struct) will be a '.' otherwise we assume it's all subobjects
+			if (Result.Len() && Result[0] == '/')
+			{
+				Result += '.';
+			}
+			else
+			{
+				Result += SUBOBJECT_DELIMITER_CHAR;
+			}
+
+			// Add the remainder of the path (it's all subobjects from this point)
+			for (int32 PathIndex = InPath.Num() - 2; PathIndex >= 0; --PathIndex)
+			{
+				Result += InPath[PathIndex].ToString();
+				if (PathIndex > 0)
+				{
+					Result += SUBOBJECT_DELIMITER_CHAR;
+				}
+			}
+		}
+	}
+
+	return Result;
+}
 
 #if WITH_EDITORONLY_DATA
 FFieldPath::FFieldPath(UField* InField, const FName& InPropertyTypeName)
-	: InitialFieldClass(nullptr)
 {
 	if (InField)
 	{
@@ -24,35 +82,42 @@ FFieldPath::FFieldPath(UField* InField, const FName& InPropertyTypeName)
 
 void FFieldPath::Generate(FField* InField)
 {
-	Path.Empty();
-
+	Reset();
 	if (InField)
 	{
-		// Add names from the innermost to the outermost
+		// Add field names from the innermost to the outermost, stop at the owner struct
+		UStruct* Owner = InField->GetOwnerStruct();
+		check(Owner); //  A field that has no owner is not allowed in FFieldPath
 		for (FFieldVariant Iter(InField); Iter.IsValid(); Iter = Iter.GetOwnerVariant())
 		{
-			Path.Add(Iter.GetFName());
+			UStruct* MaybeOwner = Iter.Get<UStruct>();
+			if (MaybeOwner == Owner)
+			{
+				break;
+			}
+			else
+			{
+				Path.Add(Iter.GetFName());
+			}
 		}
-		UStruct* OwnerStruct = CastChecked<UStruct>(InField->GetOwnerUObject());
-		ResolvedOwner = OwnerStruct;
+		ResolvedOwner = Owner;
+		ResolvedField = InField;
 #if WITH_EDITORONLY_DATA
-		FieldPathSerialNumber = OwnerStruct->FieldPathSerialNumber;
+		FieldPathSerialNumber = Owner->FieldPathSerialNumber;
 		InitialFieldClass = InField->GetClass();
 #endif // WITH_EDITORONLY_DATA
 	}
-	else
-	{
-		ClearCachedField();
-	}
 }
 
 void FFieldPath::Generate(const TCHAR* InFieldPathString)
 {
 	// Expected format is: FullPackageName.Subobject[:Subobject:...]:FieldName
 	check(InFieldPathString);
-	
-	Path.Empty();
+	Reset();
+
+	TArray<FName> Result;
 	{
+		
 		TCHAR NameBuffer[NAME_SIZE];
 		int32 NameIndex = 0;
 
@@ -64,7 +129,7 @@ void FFieldPath::Generate(const TCHAR* InFieldPathString)
 				NameBuffer[NameIndex] = '\0';
 				if (NameIndex > 0)
 				{
-					Path.Add(NameBuffer);
+					Result.Add(NameBuffer);
 					NameIndex = 0;
 				}
 				if (*InFieldPathString == '\0')
@@ -80,25 +145,88 @@ void FFieldPath::Generate(const TCHAR* InFieldPathString)
 		}
 	}
 
-	if (Path.Num() > 1)
+	if (Result.Num() > 1)
+	{
+		// Reverse the order so that it's innermost to outermost
+		Algo::Reverse(Result);
+	}
+	else if (Result.Num() == 1 && Result[0] == NAME_None)
 	{
-		// Reverse the order
-		for (int32 NameIndex = 0; NameIndex < (Path.Num() / 2); ++NameIndex)
+		Result.Empty();
+	}
+
+	Path = MoveTemp(Result);
+	ResolveField();
+}
+
+UStruct* FFieldPath::ConvertFromFullPath(FLinkerLoad* InLinker)
+{
+	// First try the StaticFindObject approach
+	UStruct* Owner = TryToResolveOwnerFromStruct();
+	if (!Owner)
+	{
+		// UClass::Serialize() unhashes the class currently being serialized so SFO will not work on it
+		// If possible try to find the owner through the current serialize context
+		if (InLinker)
 		{
-			Swap<FName>(Path[NameIndex], Path[Path.Num() - NameIndex - 1]);
+			Owner = TryToResolveOwnerFromLinker(InLinker);
 		}
 	}
+	// It's possible the full path points to a renamed asset
+	UE_CLOG(!Owner && Path.Num(), LogProperty, Verbose, TEXT("Failed resolve owner when converting from full property path \"%s\""), *PathToString(Path));
+	return Owner;
 }
 
-FField* FFieldPath::TryToResolvePath(UStruct* InCurrentStruct, UStruct** OutOwner, FFieldPath::EPathResolveType InResolveType /*= FFieldPath::UseStructIfOuterNotFound*/) const
+/** Helper function that checks if two paths have identical trailing sequences */
+static bool HasCommonTrailingSequence(const TArray<FName>& PathA, const TArray<FName>& PathB)
 {
-	FField* Result = nullptr;
+	bool bTrailingSeuenceIdentical = true;
+	const int32 MaxNumToCheck = FMath::Min(PathA.Num(), PathB.Num());
+	for (int32 PathIndex = 0; PathIndex < MaxNumToCheck; ++PathIndex)
+	{
+		if (PathA[PathA.Num() - PathIndex - 1] != PathB[PathB.Num() - PathIndex - 1])
+		{
+			bTrailingSeuenceIdentical = false;
+			break;
+		}
+	}
+	return bTrailingSeuenceIdentical;
+}
+
+UStruct* FFieldPath::TryToResolveOwnerFromLinker(FLinkerLoad* InLinker) const
+{
+	UStruct* OwnerStruct = nullptr;
+	check(InLinker);
+	FUObjectSerializeContext* Context = InLinker->GetSerializeContext();
+	if (Context && Context->SerializedObject && Context->SerializedObject->IsA<UStruct>())
+	{
+		TArray<FName> StructPath;
+		for (UObject* Obj = Context->SerializedObject; Obj; Obj = Obj->GetOuter())
+		{
+			StructPath.Add(Obj->GetFName());
+		}
+		// Check if our Path contains StructToPath (the assumption is that Path has more elements than struct path) otherwise we know the struct can't hold our property
+		if (StructPath.Num() < Path.Num() && HasCommonTrailingSequence(StructPath, Path))
+		{
+			int32 OwnerPathIndex = Path.Num() - StructPath.Num();
+			OwnerStruct = CastChecked<UStruct>(Context->SerializedObject);
+			ResolvedOwner = OwnerStruct;
+
+			// Remove portion of the path responsible for storing the owner path
+			FFieldPath* MutableThis = const_cast<FFieldPath*>(this);
+			MutableThis->Path.RemoveAt(OwnerPathIndex, Path.Num() - OwnerPathIndex);
+		}
+	}
+	return OwnerStruct;
+}
 
+UStruct* FFieldPath::TryToResolveOwnerFromStruct(UStruct* InCurrentStruct /*= nullptr*/, EPathResolveType InResolveType /*= FFieldPath::UseStructIfOuterNotFound*/) const
+{
 	// Resolve from the outermost to the innermost UObject
 	UObject* LastOuter = nullptr;
-	int32 PathIndex = Path.Num() - 1;
-	for (; PathIndex > 0; --PathIndex)
-	{				
+	int32 LastOuterIndex = -1;
+	for (int32 PathIndex = Path.Num() - 1; PathIndex > 0; --PathIndex)
+	{
 		UObject* Outer = StaticFindObjectFast(UObject::StaticClass(), LastOuter, Path[PathIndex]);
 
 		if (InCurrentStruct && PathIndex == (Path.Num() - 1))
@@ -106,23 +234,44 @@ FField* FFieldPath::TryToResolvePath(UStruct* InCurrentStruct, UStruct** OutOwne
 			UObject* CurrentOutermost = InCurrentStruct->GetOutermost();
 
 			if ((InResolveType == FFieldPath::UseStructIfOuterNotFound && !Outer) || // Outer is not found so try to use the provided struct Outer
-			    (InResolveType == FFieldPath::UseStructAlways && CurrentOutermost != Outer) // Prioritize the provided struct Outer over the resolved one
-			   )
+				(InResolveType == FFieldPath::UseStructAlways && CurrentOutermost != Outer) // Prioritize the provided struct Outer over the resolved one
+				)
 			{
 				Outer = CurrentOutermost;
-				// If we don't update the path then after a GC when this needs resolving we would resolve back to the unrenamed class package
-				FFieldPath* MutableThis = const_cast<FFieldPath*>(this);
-				MutableThis->Path[PathIndex] = Outer->GetFName();
 			}
 		}
 		if (!Outer)
 		{
 			break;
 		}
-		LastOuter = Outer;				
+		LastOuterIndex = PathIndex;
+		LastOuter = Outer;
+	}
+	UStruct* Owner = Cast<UStruct>(LastOuter);
+	if (Owner)
+	{
+		ResolvedOwner = Owner;
+
+		// Remove portion of the path responsible for storing the owner path
+		FFieldPath* MutableThis = const_cast<FFieldPath*>(this);
+		MutableThis->Path.RemoveAt(LastOuterIndex, MutableThis->Path.Num() - LastOuterIndex);
+	}
+	return Owner;
+}
+
+FField* FFieldPath::TryToResolvePath(UStruct* InCurrentStruct, FFieldPath::EPathResolveType InResolveType /*= FFieldPath::UseStructIfOuterNotFound*/) const
+{	
+	FField* Result = nullptr;
+	UStruct* Owner = ResolvedOwner.Get();
+	if (!Owner)
+	{
+		// We're probably dealing with an old path format where the Path array contained the full path to the field
+		Owner = TryToResolveOwnerFromStruct(InCurrentStruct, InResolveType);
 	}
-	if (UStruct* Owner = Cast<UStruct>(LastOuter))
+	// At this point the owner should've been fully resolved
+	if (Owner && Path.Num())
 	{
+		int32 PathIndex = Path.Num() - 1;
 		check(PathIndex <= 1);
 		Result = FindFProperty<FField>(Owner, Path[PathIndex]);
 		if (Result)
@@ -132,55 +281,123 @@ FField* FFieldPath::TryToResolvePath(UStruct* InCurrentStruct, UStruct** OutOwne
 				// Nested property
 				Result = Result->GetInnerFieldByName(Path[0]);
 			}
-			if (OutOwner)
-			{
-				*OutOwner = Owner;
-			}
 		}
 	}
-
 	return Result;
 }
 
 FString FFieldPath::ToString() const
 {
 	FString Result;
-	// See FFieldPath::Generate for formatting specifics
-	// Generate path from the outermost (last item) to the property (first item)
-	for (int32 PathIndex = Path.Num() - 1; PathIndex >= 0; --PathIndex)
+	if (UStruct* Owner = ResolvedOwner.Get())
 	{
-		// @todo: this should be handled by some kind of a flag passed to this function
-		FString ObjName = Path[PathIndex].ToString();
-		if (PathIndex == (Path.Num() - 1) && ObjName.StartsWith(UDynamicClass::GetTempPackagePrefix(), ESearchCase::IgnoreCase))
+		if (ResolvedField)
 		{
-			ObjName.RemoveFromStart(UDynamicClass::GetTempPackagePrefix(), ESearchCase::IgnoreCase);
+			Result = ResolvedField->GetPathName();
 		}
-		Result += ObjName;
-		if (PathIndex > 0)
+		else
 		{
-			// Separator between the package name (last item) and the class (asset object - the second to last item is a '.', oterwise use SUBOBJECT_DELIMITER_CHAR)
-			if (PathIndex == (Path.Num() - 1))
-			{
-				Result += '.';
-			}
-			else
+			Result = Owner->GetPathName();
+			Result += SUBOBJECT_DELIMITER_CHAR;
+			Result += PathToString(Path);
+		}
+	}
+	else
+	{
+		// Revert back to old path format where the package and UStruct owner were also specified
+		Result = PathToString(Path);
+	}
+	return Result;
+}
+
+FArchive& operator<<(FArchive& Ar, FFieldPath& InOutPropertyPath)
+{
+	Ar.UsingCustomVersion(FFortniteMainBranchObjectVersion::GUID);
+
+	if (Ar.IsSaving())
+	{		
+		UStruct* Owner = InOutPropertyPath.ResolvedOwner.Get();
+		if (!Owner)
+		{
+			// If there's no owner, make sure we don't serialize potentially unresolved path from the actual field
+			// because if we don't save the owner, we won't be able to resolve it anyway.
+			// Possible scenario: the owner was GC'd and the path is no longer valid
+			TArray<FName> EmptyPath;
+			Ar << EmptyPath;
+			UE_CLOG(InOutPropertyPath.Path.Num(), LogProperty, Verbose, TEXT("Null owner but property path is not empty when saving \"%s\""), *PathToString(InOutPropertyPath.Path));
+		}
+		else
+		{
+			Ar << InOutPropertyPath.Path;
+		}
+		Ar << Owner;
+		checkf(Owner == InOutPropertyPath.ResolvedOwner.Get(), TEXT("FFieldPath owner has changed when saving, this is not allowed (Path: \"%s\", new owner: \"%s\")"),
+			*InOutPropertyPath.ToString(), *GetPathNameSafe(Owner));
+	}
+	else
+	{		
+		Ar << InOutPropertyPath.Path;
+		// The old serialization format could save 'None' paths, they should be just empty
+		if (InOutPropertyPath.Path.Num() == 1 && InOutPropertyPath.Path[0] == NAME_None)
+		{
+			InOutPropertyPath.Path.Empty();
+		}
+		if (Ar.CustomVer(FFortniteMainBranchObjectVersion::GUID) >= FFortniteMainBranchObjectVersion::FFieldPathOwnerSerialization)
+		{
+			UStruct* SerializedOwner = InOutPropertyPath.ResolvedOwner.Get();
+			Ar << SerializedOwner;
+			InOutPropertyPath.ResolvedOwner = SerializedOwner;
+			if (!SerializedOwner)
 			{
-				Result += SUBOBJECT_DELIMITER_CHAR;
+				UE_CLOG(InOutPropertyPath.Path.Num(), LogProperty, Verbose, TEXT("Serialized null owner for property \"%s\""), *PathToString(InOutPropertyPath.Path));
+				// At this point it makes no sense to keep the remainder of the path around as it will produce unnecessary warnings
+				// Possible scenario: owning struct was not cooked or deleted and the asset was not loaded
+				InOutPropertyPath.Path.Empty();
 			}
 		}
+		else if (InOutPropertyPath.Path.Num())
+		{
+			// Convert from the old format: resolve the owner now and remove its path from the field path
+			FLinkerLoad* Linker = Cast<FLinkerLoad>(Ar.GetLinker());
+			UStruct* Owner = InOutPropertyPath.ConvertFromFullPath(Linker);
+
+			InOutPropertyPath.ResolvedOwner = Owner;
+
+			// This usually happens when the old path format serialized a path and then the owner struct's package got renamed or moved
+			// There's code to handle that in bot UClass and UAnimBlueprintGeneratedClass
+			UE_CLOG(!Owner, LogProperty, Verbose, TEXT("Failed to resolve property owner from Path \"%s\""), *PathToString(InOutPropertyPath.Path));
+		}
+		else
+		{
+			InOutPropertyPath.ResolvedOwner = nullptr;
+		}
+
+		if (!Ar.IsObjectReferenceCollector())
+		{
+			InOutPropertyPath.ClearCachedField();
+		}
 	}
-	return Result;
+
+	return Ar;
 }
 
 #if WITH_EDITORONLY_DATA
 void FFieldPath::GenerateFromUField(UField* InField)
 {
-	Path.Empty();
-	ClearCachedField();
+	Reset();
 	for (UObject* Obj = InField; Obj; Obj = Obj->GetOuter())
 	{
-		Path.Add(Obj->GetFName());
-	}	
+		UStruct* MaybeOwner = Cast<UStruct>(Obj);
+		if (MaybeOwner)
+		{
+			ResolvedOwner = MaybeOwner;
+			break;
+		}
+		else
+		{
+			Path.Add(Obj->GetFName());
+		}
+	}
 }
 
 bool FFieldPath::IsFieldPathSerialNumberIdentical(UStruct* InStruct) const
diff --git a/Engine/Source/Runtime/CoreUObject/Private/UObject/FieldPathProperty.cpp b/Engine/Source/Runtime/CoreUObject/Private/UObject/FieldPathProperty.cpp
--- a/Engine/Source/Runtime/CoreUObject/Private/UObject/FieldPathProperty.cpp
+++ b/Engine/Source/Runtime/CoreUObject/Private/UObject/FieldPathProperty.cpp
@@ -34,21 +34,28 @@ EConvertFromTypeResult FFieldPathProperty::ConvertFromType(const FPropertyTag& T
 		FLinker* Linker = UnderlyingArchive.GetLinker();
 		check(Linker);
 
+		FFieldPath ConvertedValue;
+
 		FPackageIndex Index;
 		UnderlyingArchive << Index;
 
-		FString PropertyPathName;
-		if (Index.IsImport())
+		bool bExport = Index.IsExport();
+		while (!Index.IsNull())
+		{
+			const FObjectResource& Res = Linker->ImpExp(Index);
+			ConvertedValue.Path.Add(Res.ObjectName);
+			Index = Res.OuterIndex;
+		}
+		if (bExport)
 		{
-			PropertyPathName = Linker->GetImportPathName(Index);
+			check(Linker->LinkerRoot);
+			ConvertedValue.Path.Add(Linker->LinkerRoot->GetFName());
 		}
-		else if (Index.IsExport())
+		if (ConvertedValue.Path.Num())
 		{
-			PropertyPathName = Linker->GetExportPathName(Index);
+			ConvertedValue.ConvertFromFullPath(Cast<FLinkerLoad>(Linker));
 		}
 
-		FFieldPath ConvertedValue;
-		ConvertedValue.Generate(*PropertyPathName);
 		SetPropertyValue_InContainer(Data, ConvertedValue, Tag.ArrayIndex);
 		return EConvertFromTypeResult::Converted;
 	}
@@ -62,10 +69,10 @@ bool FFieldPathProperty::Identical( const void* A, const void* B, uint32 PortFla
 	if (B)
 	{
 		const FFieldPath ValueB = GetPropertyValue(B);
-		return ValueA.IsPathIdentical(ValueB);
+		return ValueA == ValueB;
 	}
 
-	return ValueA.IsEmpty();
+	return !ValueA.GetTyped(FField::StaticClass());
 }
 
 void FFieldPathProperty::SerializeItem(FStructuredArchive::FSlot Slot, void* Value, void const* Defaults) const
@@ -182,4 +189,9 @@ FString FFieldPathProperty::GetCPPType(FString* ExtendedTypeText, uint32 CPPExpo
 	return TEXT("TFieldPath");
 }
 
+bool FFieldPathProperty::SupportsNetSharedSerialization() const
+{
+	return false;
+}
+
 #include "UObject/DefineUPropertyMacros.h"
\ No newline at end of file
diff --git a/Engine/Source/Runtime/CoreUObject/Public/UObject/FieldPath.h b/Engine/Source/Runtime/CoreUObject/Public/UObject/FieldPath.h
--- a/Engine/Source/Runtime/CoreUObject/Public/UObject/FieldPath.h
+++ b/Engine/Source/Runtime/CoreUObject/Public/UObject/FieldPath.h
@@ -15,6 +15,7 @@
 
 class UStruct;
 class UField;
+class FLinkerLoad;
 
 struct COREUOBJECT_API FFieldPath
 {
@@ -26,6 +27,9 @@ struct COREUOBJECT_API FFieldPath
 	template<class T>
 	friend struct TWeakFieldPtr;
 
+	// FFieldPathProperty needs access to ConvertFromFullPath
+	friend class FFieldPathProperty;
+
 protected:
 
 	/* Determines the behavior when resolving stored path */
@@ -77,7 +81,6 @@ protected:
 		InitialFieldClass = nullptr;
 		FieldPathSerialNumber = 0;
 #endif // WITH_EDITORONLY_DATA
-		ResolvedOwner.Reset();
 	}
 
 private:
@@ -100,6 +103,30 @@ private:
 		ResolvedOwner.Reset();
 	}
 
+	/**
+	 * Tries to resolve the field owner
+	 * @param InCurrentStruct Struct that's trying to resolve this field path
+	 * @param InResolveType Type of the resolve operation
+	 * @return Resolved owner struct
+	 */
+	UStruct* TryToResolveOwnerFromStruct(UStruct* InCurrentStruct = nullptr, EPathResolveType InResolveType = FFieldPath::UseStructIfOuterNotFound) const;
+	
+
+	/**
+	 * Tries to resolve the field owner
+	 * @param InLinker the current linker load serializing this field path
+	 * @return Resolved owner struct
+	 */
+	UStruct* TryToResolveOwnerFromLinker(FLinkerLoad* InLinker) const;
+
+
+	/**
+	 * Tries to convert the full path stored in this FFieldPath to the new format (Owner reference + path to the field)
+	 * @param InLinker the current linker load serializing this field path
+	 * @return Resulved owner struct
+	 */
+	UStruct* ConvertFromFullPath(FLinkerLoad* InLinker);
+
 public:
 
 	FFieldPath() = default;
@@ -129,7 +156,7 @@ public:
 	 * @param OutOwnerIndex ObjectIndex of the Owner UObject
 	 * @return Resolved field or null
 	 */
-	FField* TryToResolvePath(UStruct* InCurrentStruct = nullptr, UStruct** OutOwner = nullptr, EPathResolveType InResolveType = FFieldPath::UseStructIfOuterNotFound) const;
+	FField* TryToResolvePath(UStruct* InCurrentStruct, EPathResolveType InResolveType = FFieldPath::UseStructIfOuterNotFound) const;
 
 	/**
 	 * Tries to resolve the path and caches the result
@@ -138,8 +165,7 @@ public:
 	 */
 	FORCEINLINE void ResolveField(FFieldClass* ExpectedClass = FField::StaticClass(), UStruct* InCurrentStruct = nullptr, EPathResolveType InResolveType = FFieldPath::UseStructIfOuterNotFound) const
 	{
-		UStruct* FoundOwner = nullptr;
-		FField* FoundField = TryToResolvePath(InCurrentStruct, &FoundOwner, InResolveType);
+		FField* FoundField = TryToResolvePath(InCurrentStruct, InResolveType);
 		if (FoundField && FoundField->IsA(ExpectedClass) 
 #if WITH_EDITORONLY_DATA
 			&& (!InitialFieldClass || FoundField->IsA(InitialFieldClass))
@@ -147,14 +173,14 @@ public:
 			)
 		{
 			ResolvedField = FoundField;
-			ResolvedOwner = FoundOwner;
 #if WITH_EDITORONLY_DATA
 			if (!InitialFieldClass)
 			{
 				InitialFieldClass = FoundField->GetClass();
 			}
-			check(FoundOwner);
-			FieldPathSerialNumber = GetFieldPathSerialNumber(FoundOwner);
+			UStruct* Owner = ResolvedOwner.Get();
+			check(Owner);
+			FieldPathSerialNumber = GetFieldPathSerialNumber(Owner);
 #endif // WITH_EDITORONLY_DATA
 		}
 		else
@@ -163,8 +189,26 @@ public:
 		}
 	}
 
-	/** Returns true if the field path is empty */
-	inline bool IsEmpty() const
+	/**
+	 * Gets the field represented by this FFieldPath
+	 * @param ExpectedType Expected type of the resolved field
+	 * @param InCurrentStruct Struct that's trying to resolve this field path
+	 * @return Field represented by this FFieldPath or null if it couldn't be resolved
+	 */
+	FORCEINLINE FField* GetTyped(FFieldClass* ExpectedType, UStruct* InCurrentStruct = nullptr) const
+	{
+		if (NeedsResolving() && Path.Num())
+		{
+			ResolveField(ExpectedType, InCurrentStruct, FFieldPath::UseStructIfOuterNotFound);
+		}
+		return ResolvedField;
+	}
+
+	/** 
+	 * Returns true if the field path is empty (does not test if the owner is valid) 
+	 * This is usually used to verify if the reason behind this field being unresolved is because the owner is missing or the property couldn't be found.
+	 **/
+	inline bool IsPathToFieldEmpty() const
 	{
 		return !Path.Num();
 	}
@@ -175,7 +219,11 @@ public:
 	**/
 	inline bool IsStale() const
 	{
-		return ResolvedField && (TryToResolvePath() != ResolvedField);
+		return ResolvedField && (!ResolvedOwner.IsValid()
+#if WITH_EDITORONLY_DATA
+			|| !IsFieldPathSerialNumberIdentical(ResolvedOwner.Get())
+#endif // WITH_EDITORONLY_DATA
+			);
 	}
 
 	/**
@@ -184,24 +232,33 @@ public:
 	inline void Reset()
 	{
 		ClearCachedField();
+		ResolvedOwner.Reset();
 		Path.Empty();
 	}
 
-	inline bool IsPathIdentical(const FFieldPath& InOther) const
+	FORCEINLINE bool operator==(const FFieldPath& InOther) const
 	{
-		return Path == InOther.Path;
+		return ResolvedOwner == InOther.ResolvedOwner && Path == InOther.Path;
+	}
+
+	FORCEINLINE bool operator!=(const FFieldPath& InOther) const
+	{
+		return ResolvedOwner != InOther.ResolvedOwner || Path != InOther.Path;
 	}
 
 	FString ToString() const;
 
-	friend FArchive& operator<<(FArchive& Ar, FFieldPath& InOutPropertyPath)
+	friend FArchive& operator<<(FArchive& Ar, FFieldPath& InOutPropertyPath);
+
+	/** Hash function. */
+	FORCEINLINE friend uint32 GetTypeHash(const FFieldPath& InPropertyPath)
 	{
-		Ar << InOutPropertyPath.Path;
-		if (Ar.IsLoading())
+		uint32 HashValue = 0;
+		for (const FName PathSegment : InPropertyPath.Path)
 		{
-			InOutPropertyPath.ClearCachedField();
+			HashValue = HashCombine(HashValue, GetTypeHash(PathSegment));
 		}
-		return Ar;
+		return HashValue;
 	}
 };
 
@@ -315,17 +372,13 @@ public:
 	}
 
 	/**
-	* Dereference the weak pointer
-	* @param bEvenIfPendingKill, if this is true, pendingkill objects are considered valid
-	* @return NULL if this object is gone or the weak pointer was NULL, otherwise a valid uobject pointer
-	**/
+	 * Gets the field represented by this TFieldPath
+	 * @param InCurrentStruct Struct that's trying to resolve this field path
+	 * @return Field represented by this FFieldPath or null if it couldn't be resolved
+	 */
 	FORCEINLINE PropertyType* Get(UStruct* InCurrentStruct = nullptr) const
 	{
-		if (NeedsResolving() && Path.Num())
-		{
-			ResolveField(PropertyType::StaticClass(), InCurrentStruct, FFieldPath::UseStructIfOuterNotFound);
-		}
-		return static_cast<PropertyType*>(ResolvedField);
+		return (PropertyType*)GetTyped(PropertyType::StaticClass(), InCurrentStruct);
 	}
 
 	FORCEINLINE PropertyType* ResolveWithRenamedStructPackage(UStruct* InCurrentStruct)
@@ -351,21 +404,6 @@ public:
 		return Get();
 	}
 
-	/** Hash function. */
-	FORCEINLINE friend uint32 GetTypeHash(const TFieldPath& InPropertyPath)
-	{
-		uint32 HashValue = 0;
-		if (InPropertyPath.Path.Num())
-		{
-			HashValue = GetTypeHash(InPropertyPath.Path[0]);
-			for (int32 PathIndex = 1; PathIndex < InPropertyPath.Path.Num(); ++PathIndex)
-			{
-				HashValue = HashCombine(HashValue, GetTypeHash(InPropertyPath.Path[PathIndex]));
-			}
-		}
-		return HashValue;
-	}
-
 	/**
 	* Compare weak pointers for equality
 	* @param Other weak pointer to compare to
@@ -376,7 +414,7 @@ public:
 		static_assert(TPointerIsConvertibleFromTo<OtherPropertyType, FField>::Value, "TFieldPath can only be compared with FField types");
 		static_assert(TPointerIsConvertibleFromTo<PropertyType, OtherPropertyType>::Value, "Unable to compare TFieldPath with raw pointer - types are incompatible");
 
-		return Path == Other.Path;
+		return FFieldPath::operator==(Other);
 	}
 
 	/**
@@ -389,7 +427,7 @@ public:
 		static_assert(TPointerIsConvertibleFromTo<OtherPropertyType, FField>::Value, "TFieldPath can only be compared with FField types");
 		static_assert(TPointerIsConvertibleFromTo<PropertyType, OtherPropertyType>::Value, "Unable to compare TFieldPath with raw pointer - types are incompatible");
 
-		return Path != Other.Path;
+		return FFieldPath::operator!=(Other);
 	}
 
 	/**
diff --git a/Engine/Source/Runtime/CoreUObject/Public/UObject/FieldPathProperty.h b/Engine/Source/Runtime/CoreUObject/Public/UObject/FieldPathProperty.h
--- a/Engine/Source/Runtime/CoreUObject/Public/UObject/FieldPathProperty.h
+++ b/Engine/Source/Runtime/CoreUObject/Public/UObject/FieldPathProperty.h
@@ -52,5 +52,6 @@ public:
 	virtual void Serialize(FArchive& Ar) override;
 	virtual bool ContainsObjectReference(TArray<const FStructProperty*>& EncounteredStructProps) const override;
 	virtual void EmitReferenceInfo(UClass& OwnerClass, int32 BaseOffset, TArray<const FStructProperty*>& EncounteredStructProps) override;
+	virtual bool SupportsNetSharedSerialization() const override;
 	// End of FProperty interface
 };
\ No newline at end of file
diff --git a/Engine/Source/Runtime/CoreUObject/Public/UObject/PropertyProxyArchive.h b/Engine/Source/Runtime/CoreUObject/Public/UObject/PropertyProxyArchive.h
--- a/Engine/Source/Runtime/CoreUObject/Public/UObject/PropertyProxyArchive.h
+++ b/Engine/Source/Runtime/CoreUObject/Public/UObject/PropertyProxyArchive.h
@@ -47,7 +47,7 @@ public:
 		if (IsLoading())
 		{
 			Value = PropertyPath.Get(Container);
-			if (!Value && !PropertyPath.IsEmpty())
+			if (!Value && !PropertyPath.IsPathToFieldEmpty())
 			{
 				// Store the field path for deferred resolving
 				UnresolvedProperties.Add(TPair<TFieldPath<FField>, int32>(PropertyPath, BytecodeIndex));
diff --git a/Engine/Source/Runtime/CoreUObject/Public/UObject/WeakFieldPtr.h b/Engine/Source/Runtime/CoreUObject/Public/UObject/WeakFieldPtr.h
--- a/Engine/Source/Runtime/CoreUObject/Public/UObject/WeakFieldPtr.h
+++ b/Engine/Source/Runtime/CoreUObject/Public/UObject/WeakFieldPtr.h
@@ -178,7 +178,7 @@ public:
 	**/
 	FORCEINLINE bool IsValid(bool bEvenIfPendingKill, bool bThreadsafeTest = false) const
 	{
-		return Owner.IsValid(bEvenIfPendingKill, bThreadsafeTest) && Field.TryToResolvePath();
+		return Owner.IsValid(bEvenIfPendingKill, bThreadsafeTest) && Field.Get();
 	}
 
 	/**
@@ -187,7 +187,7 @@ public:
 	*/
 	FORCEINLINE bool IsValid(/*bool bEvenIfPendingKill = false, bool bThreadsafeTest = false*/) const
 	{
-		return Owner.IsValid() && Field.TryToResolvePath();
+		return Owner.IsValid() && Field.Get();
 	}
 
 	/**
diff --git a/Engine/Source/Runtime/Engine/Private/Animation/AnimBlueprintGeneratedClass.cpp b/Engine/Source/Runtime/Engine/Private/Animation/AnimBlueprintGeneratedClass.cpp
--- a/Engine/Source/Runtime/Engine/Private/Animation/AnimBlueprintGeneratedClass.cpp
+++ b/Engine/Source/Runtime/Engine/Private/Animation/AnimBlueprintGeneratedClass.cpp
@@ -293,6 +293,7 @@ void UAnimBlueprintGeneratedClass::Link(FArchive& Ar, bool bRelinkExistingProper
 	{
 		// handle potential renames of the class package
 		Handler.ValueHandlerNodeProperty.ResolveWithRenamedStructPackage(this);
+		check(Handler.ValueHandlerNodeProperty.Get() || Handler.ValueHandlerNodeProperty.IsPathToFieldEmpty());
 	}
 #endif // WITH_EDITORONLY_DATA
 
