Adding deprecated support for Cast and CastChecked functions for FProperties to minimize compilation errors when upgrading to the new engine version

diff --git a/Engine/Source/Runtime/CoreUObject/Public/Templates/Casts.h b/Engine/Source/Runtime/CoreUObject/Public/Templates/Casts.h
--- a/Engine/Source/Runtime/CoreUObject/Public/Templates/Casts.h
+++ b/Engine/Source/Runtime/CoreUObject/Public/Templates/Casts.h
@@ -123,6 +123,18 @@ struct TCastImpl
 	{
 		return (To*)Src;
 	}
+
+	UE_DEPRECATED(4.25, "Cast<>() and CastChecked<>() should not be used with FProperties. Use CastField<>() or CastFieldChecked<>() instead.")
+	FORCEINLINE static To* DoCast( FField* Src )
+	{
+		return Src && Src->IsA<To>() ? (To*)Src : nullptr;
+	}
+
+	UE_DEPRECATED(4.25, "Cast<>() and CastChecked<>() should not be used with FProperties. Use CastField<>() or CastFieldChecked<>() instead.")
+	FORCEINLINE static To* DoCastCheckedWithoutTypeCheck( FField* Src )
+	{
+		return (To*)Src;
+	}
 };
 
 template <typename From, typename To>
@@ -137,6 +149,18 @@ struct TCastImpl<From, To, ECastType::UObjectToUObject>
 	{
 		return (To*)Src;
 	}
+
+	UE_DEPRECATED(4.25, "Cast<>() and CastChecked<>() should not be used with FProperties. Use CastField<>() or CastFieldChecked<>() instead.")
+	FORCEINLINE static To* DoCast( FField* Src )
+	{
+		return Src && Src->IsA<To>() ? (To*)Src : nullptr;
+	}
+
+	UE_DEPRECATED(4.25, "Cast<>() and CastChecked<>() should not be used with FProperties. Use CastField<>() or CastFieldChecked<>() instead.")
+	FORCEINLINE static To* DoCastCheckedWithoutTypeCheck( FField* Src )
+	{
+		return (To*)Src;
+	}
 };
 
 template <typename From, typename To>
@@ -205,6 +229,25 @@ FORCEINLINE T* ExactCast( UObject* Src )
 
 #if DO_CHECK
 
+	// Helper function to get the full name for UObjects and UInterfaces
+	template <
+		typename T,
+		typename TEnableIf<!TIsDerivedFrom<T, FField>::Value>::Type* = nullptr
+	>
+	FString GetFullNameForCastLogError(T* InObjectOrInterface)
+	{
+		return Cast<UObject>(InObjectOrInterface)->GetFullName();
+	}
+	// And a special version for FFields
+	template <
+		typename T,
+		typename TEnableIf<TIsDerivedFrom<T, FField>::Value>::Type* = nullptr
+	>
+	FString GetFullNameForCastLogError(T* InField)
+	{
+		return GetFullNameSafe(InField);
+	}
+
 	template <typename To, typename From>
 	FUNCTION_NON_NULL_RETURN_START
 		To* CastChecked(From* Src)
@@ -218,7 +261,7 @@ FORCEINLINE T* ExactCast( UObject* Src )
 		To* Result = Cast<To>(Src);
 		if (!Result)
 		{
-			CastLogError(*Cast<UObject>(Src)->GetFullName(), *GetTypeName<To>());
+			CastLogError(*GetFullNameForCastLogError(Src), *GetTypeName<To>());
 		}
 
 		return Result;
@@ -232,7 +275,7 @@ FORCEINLINE T* ExactCast( UObject* Src )
 			To* Result = Cast<To>(Src);
 			if (!Result)
 			{
-				CastLogError(*Cast<UObject>(Src)->GetFullName(), *GetTypeName<To>());
+				CastLogError(*GetFullNameForCastLogError(Src), *GetTypeName<To>());
 			}
 
 			return Result;
