Fixing CIS (another fallout from 12847659)

diff --git a/Engine/Source/Runtime/CoreUObject/Public/UObject/FieldPath.h b/Engine/Source/Runtime/CoreUObject/Public/UObject/FieldPath.h
--- a/Engine/Source/Runtime/CoreUObject/Public/UObject/FieldPath.h
+++ b/Engine/Source/Runtime/CoreUObject/Public/UObject/FieldPath.h
@@ -248,7 +248,7 @@ public:
 
 	FString ToString() const;
 
-	friend FArchive& operator<<(FArchive& Ar, FFieldPath& InOutPropertyPath);
+	COREUOBJECT_API friend FArchive& operator<<(FArchive& Ar, FFieldPath& InOutPropertyPath);
 
 	/** Hash function. */
 	FORCEINLINE friend uint32 GetTypeHash(const FFieldPath& InPropertyPath)
