FPropertyDelegate now lists the SignatureFunction in its 'Referenced Objects', TRASHCLASS is no marked as CLASS_NewerVersionExists, reload package now checks outer chain for out of date classes

diff --git a/Engine/Source/Editor/KismetCompiler/Private/KismetCompiler.cpp b/Engine/Source/Editor/KismetCompiler/Private/KismetCompiler.cpp
--- a/Engine/Source/Editor/KismetCompiler/Private/KismetCompiler.cpp
+++ b/Engine/Source/Editor/KismetCompiler/Private/KismetCompiler.cpp
@@ -247,7 +247,7 @@ void FKismetCompilerContext::CleanAndSanitizeClass(UBlueprintGeneratedClass* Cla
 	}
 	TransientClass->ClassAddReferencedObjects = ParentClass->AddReferencedObjects;
 	TransientClass->ClassGeneratedBy = Blueprint;
-	TransientClass->ClassFlags |= CLASS_CompiledFromBlueprint;
+	TransientClass->ClassFlags |= CLASS_CompiledFromBlueprint|CLASS_NewerVersionExists;
 
 	SetNewClass( ClassToClean );
 	InOldCDO = ClassToClean->ClassDefaultObject; // we don't need to create the CDO at this point
diff --git a/Engine/Source/Runtime/CoreUObject/Private/UObject/PackageReload.cpp b/Engine/Source/Runtime/CoreUObject/Private/UObject/PackageReload.cpp
--- a/Engine/Source/Runtime/CoreUObject/Private/UObject/PackageReload.cpp
+++ b/Engine/Source/Runtime/CoreUObject/Private/UObject/PackageReload.cpp
@@ -580,9 +580,16 @@ void ReloadPackages(const TArrayView<FReloadPackageData>& InPackagesToReload, TA
 
 				// Mutating the old versions of classes can result in us replacing the SuperStruct pointer, which results
 				// in class layout change and subsequently crashes because instances will not match this new class layout:
-				if(UClass* AsClass = Cast<UClass>(PotentialReferencer))
+				UClass* AsClass = Cast<UClass>(PotentialReferencer);
+				if (!AsClass)
 				{
-					if(AsClass->HasAnyClassFlags(CLASS_NewerVersionExists) || AsClass->HasAnyFlags(RF_NewerVersionExists))
+					AsClass = PotentialReferencer->GetTypedOuter<UClass>();
+				}
+
+				if(AsClass)
+				{
+					if( AsClass->HasAnyClassFlags(CLASS_NewerVersionExists) || 
+						AsClass->HasAnyFlags(RF_NewerVersionExists))
 					{
 						continue;
 					}
diff --git a/Engine/Source/Runtime/CoreUObject/Private/UObject/PropertyDelegate.cpp b/Engine/Source/Runtime/CoreUObject/Private/UObject/PropertyDelegate.cpp
--- a/Engine/Source/Runtime/CoreUObject/Private/UObject/PropertyDelegate.cpp
+++ b/Engine/Source/Runtime/CoreUObject/Private/UObject/PropertyDelegate.cpp
@@ -171,6 +171,17 @@ void FDelegateProperty::Serialize( FArchive& Ar )
 #endif // USE_CIRCULAR_DEPENDENCY_LOAD_DEFERRING
 }
 
+void FDelegateProperty::AddReferencedObjects(FReferenceCollector& Collector)
+{
+#if USE_CIRCULAR_DEPENDENCY_LOAD_DEFERRING
+	if (SignatureFunction && !SignatureFunction->IsA<ULinkerPlaceholderFunction>())
+#endif//USE_CIRCULAR_DEPENDENCY_LOAD_DEFERRING
+	{
+		Collector.AddReferencedObject(SignatureFunction);
+	}
+	Super::AddReferencedObjects(Collector);
+}
+
 bool FDelegateProperty::SameType(const FProperty* Other) const
 {
 	return Super::SameType(Other) && (SignatureFunction == ((FDelegateProperty*)Other)->SignatureFunction);
diff --git a/Engine/Source/Runtime/CoreUObject/Public/UObject/UnrealType.h b/Engine/Source/Runtime/CoreUObject/Public/UObject/UnrealType.h
--- a/Engine/Source/Runtime/CoreUObject/Public/UObject/UnrealType.h
+++ b/Engine/Source/Runtime/CoreUObject/Public/UObject/UnrealType.h
@@ -4808,6 +4808,7 @@ public:
 
 	// UObject interface
 	virtual void Serialize( FArchive& Ar ) override;
+	virtual void AddReferencedObjects(FReferenceCollector& Collector) override;
 	virtual void BeginDestroy() override;
 	// End of UObject interface
 
