Fixed a bug causing duplicated non optional payloads to never be loaded if the optional pakfiles are not mounted, even though the data exists in both the .ubulk and .uptnl file.

- When checking for the existence of the optional file (.uptnl) via filenames rather then FIoChunkId we were actually checking for the non optional (.ubulk) file which would always be true.
 - Caused by using ::ConvertFilenameFromFlags which checks with the original flags, not the new flags.
- This would cause the check in BulkData to always be true so we would report to the texture streaming system that the BulkData was optional data.
- Later when it came time to stream in the mip, the BulkData system would check is the optional data existed on disk and be unable to find it as the optional pakfiles are not mounted.

- Since FBulkDataBase::ProcessDuplicateData is already special case, we can just change the file extension to BulkDataExt::Optional rather than call ::ConvertFilenameFromFlags

diff --git a/Engine/Source/Runtime/CoreUObject/Private/Serialization/BulkData2.cpp b/Engine/Source/Runtime/CoreUObject/Private/Serialization/BulkData2.cpp
--- a/Engine/Source/Runtime/CoreUObject/Private/Serialization/BulkData2.cpp
+++ b/Engine/Source/Runtime/CoreUObject/Private/Serialization/BulkData2.cpp
@@ -1483,7 +1483,6 @@ void FBulkDataBase::ProcessDuplicateData(FArchive& Ar, const UPackage* Package,
 
 	SerializeDuplicateData(Ar, NewFlags, NewSizeOnDisk, NewOffset);
 
-
 #if ALLOW_OPTIONAL_DATA
 	if (IsUsingIODispatcher())
 	{
@@ -1502,7 +1501,7 @@ void FBulkDataBase::ProcessDuplicateData(FArchive& Ar, const UPackage* Package,
 	else
 	{
 		check(Filename != nullptr);
-		const FString OptionalDataFilename = ConvertFilenameFromFlags(*Filename);
+		const FString OptionalDataFilename = FPathViews::ChangeExtension(*Filename, BulkDataExt::Optional);
 
 		if (IFileManager::Get().FileExists(*OptionalDataFilename))
 		{
