Compile fix
#rnx

diff --git a/Engine/Source/Runtime/CoreUObject/Private/Serialization/AsyncLoading2.cpp b/Engine/Source/Runtime/CoreUObject/Private/Serialization/AsyncLoading2.cpp
--- a/Engine/Source/Runtime/CoreUObject/Private/Serialization/AsyncLoading2.cpp
+++ b/Engine/Source/Runtime/CoreUObject/Private/Serialization/AsyncLoading2.cpp
@@ -3130,7 +3130,7 @@ EAsyncPackageState::Type FAsyncPackage2::Event_ProcessExportBundle(FAsyncPackage
 			Ar.CookedHeaderSize = Package->CookedHeaderSize;
 			Ar.PackageDesc = &Package->Desc;
 			Ar.PackageNameMap = Package->PackageNameMap;
-			Ar.GlobalNameMap = &Package->AsyncLoadingThread.GlobalNameMap.GetNameEntries();
+			Ar.NameMap = &Package->NameMap->GetNameEntries();
 			Ar.ImportStore = &Package->ImportStore;
 			Ar.Exports = &Package->Exports;
 			Ar.ExportMap = Package->ExportMap;
