AsyncLoading2: Fix incorrect check in ToPackageImport

diff --git a/Engine/Source/Runtime/CoreUObject/Public/Serialization/AsyncLoading2.h b/Engine/Source/Runtime/CoreUObject/Public/Serialization/AsyncLoading2.h
--- a/Engine/Source/Runtime/CoreUObject/Public/Serialization/AsyncLoading2.h
+++ b/Engine/Source/Runtime/CoreUObject/Public/Serialization/AsyncLoading2.h
@@ -171,7 +171,7 @@ public:
 
 	inline uint32 ToPackageImport() const
 	{
-		check(IsScriptImport());
+		check(IsPackageImport());
 		return TypeAndIndex & IndexMask;
 	}
 
