Fixing FFieldPath::ToString() not removing nativized BP temp package prefix after changes from 12847659

diff --git a/Engine/Source/Runtime/CoreUObject/Private/UObject/FieldPath.cpp b/Engine/Source/Runtime/CoreUObject/Private/UObject/FieldPath.cpp
--- a/Engine/Source/Runtime/CoreUObject/Private/UObject/FieldPath.cpp
+++ b/Engine/Source/Runtime/CoreUObject/Private/UObject/FieldPath.cpp
@@ -31,14 +31,7 @@ static FString PathToString(const TArray<FName>& InPath)
 	{
 		// Allocate once and construct the path string
 		Result.Reserve(PathLength);
-
-		// Handle nativized blueprint oddities 
-		FString ObjName = InPath.Last().ToString();
-		if (ObjName.StartsWith(UDynamicClass::GetTempPackagePrefix(), ESearchCase::IgnoreCase))
-		{
-			ObjName.RemoveFromStart(UDynamicClass::GetTempPackagePrefix(), ESearchCase::IgnoreCase);
-		}
-		Result += ObjName;
+		Result += InPath.Last().ToString();
 
 		if (InPath.Num() > 1)
 		{
@@ -307,6 +300,13 @@ FString FFieldPath::ToString() const
 		// Revert back to old path format where the package and UStruct owner were also specified
 		Result = PathToString(Path);
 	}
+
+	// Nativized BP support
+	if (Result.StartsWith(UDynamicClass::GetTempPackagePrefix(), ESearchCase::IgnoreCase))
+	{
+		Result.RemoveFromStart(UDynamicClass::GetTempPackagePrefix(), ESearchCase::IgnoreCase);
+	}
+
 	return Result;
 }
 
