Fix lock condition check when reallocating BulkData.

The condition got incorrectly set when improving the error message.

#rnx

diff --git a/Engine/Source/Runtime/CoreUObject/Private/Serialization/BulkData2.cpp b/Engine/Source/Runtime/CoreUObject/Private/Serialization/BulkData2.cpp
--- a/Engine/Source/Runtime/CoreUObject/Private/Serialization/BulkData2.cpp
+++ b/Engine/Source/Runtime/CoreUObject/Private/Serialization/BulkData2.cpp
@@ -895,7 +895,7 @@ void* FBulkDataBase::Realloc(int64 SizeInBytes)
 {
 	DECLARE_SCOPE_CYCLE_COUNTER(TEXT("FBulkDataBase::Realloc"), STAT_UBD_Realloc, STATGROUP_Memory);
 
-	checkf(LockStatus == LOCKSTATUS_Unlocked, TEXT("Attempting to modify a BulkData object that is locked"));
+	checkf(LockStatus == LOCKSTATUS_ReadWriteLock, TEXT("BulkData must be locked for 'write' before reallocating!"));
 	checkf(!CanLoadFromDisk(), TEXT("Cannot re-allocate a FBulkDataBase object that represents a file on disk!"));
 
 	// We might want to consider this a valid use case if anyone can come up with one?
