When cooking, references to editor-only FProperties will be nulled out

#tests cooked windows client and server, cooked build bot soak tests

diff --git a/Engine/Source/Runtime/CoreUObject/Private/UObject/FieldPath.cpp b/Engine/Source/Runtime/CoreUObject/Private/UObject/FieldPath.cpp
--- a/Engine/Source/Runtime/CoreUObject/Private/UObject/FieldPath.cpp
+++ b/Engine/Source/Runtime/CoreUObject/Private/UObject/FieldPath.cpp
@@ -317,20 +317,31 @@ FArchive& operator<<(FArchive& Ar, FFieldPath& InOutPropertyPath)
 	if (Ar.IsSaving())
 	{		
 		UStruct* Owner = InOutPropertyPath.ResolvedOwner.Get();
-		if (!Owner)
+		bool bFilterEditorOnlyProperty = false;
+		if (Owner && Ar.IsFilterEditorOnly())
+		{
+			FProperty* ResolvedProperty = CastField<FProperty>(InOutPropertyPath.GetTyped(FProperty::StaticClass()));
+			if (ResolvedProperty && ResolvedProperty->HasAnyPropertyFlags(CPF_EditorOnly))
+			{
+				bFilterEditorOnlyProperty = true;
+			}
+		}
+		if (!Owner || bFilterEditorOnlyProperty)
 		{
 			// If there's no owner, make sure we don't serialize potentially unresolved path from the actual field
 			// because if we don't save the owner, we won't be able to resolve it anyway.
 			// Possible scenario: the owner was GC'd and the path is no longer valid
 			TArray<FName> EmptyPath;
+			UStruct* NullOwner = nullptr;
 			Ar << EmptyPath;
+			Ar << NullOwner;
 			UE_CLOG(InOutPropertyPath.Path.Num(), LogProperty, Verbose, TEXT("Null owner but property path is not empty when saving \"%s\""), *PathToString(InOutPropertyPath.Path));
 		}
 		else
 		{
 			Ar << InOutPropertyPath.Path;
+			Ar << Owner;
 		}
-		Ar << Owner;
 		checkf(Owner == InOutPropertyPath.ResolvedOwner.Get(), TEXT("FFieldPath owner has changed when saving, this is not allowed (Path: \"%s\", new owner: \"%s\")"),
 			*InOutPropertyPath.ToString(), *GetPathNameSafe(Owner));
 	}
diff --git a/Engine/Source/Runtime/CoreUObject/Public/UObject/PropertyProxyArchive.h b/Engine/Source/Runtime/CoreUObject/Public/UObject/PropertyProxyArchive.h
--- a/Engine/Source/Runtime/CoreUObject/Public/UObject/PropertyProxyArchive.h
+++ b/Engine/Source/Runtime/CoreUObject/Public/UObject/PropertyProxyArchive.h
@@ -32,6 +32,7 @@ public:
 		, BytecodeIndex(InBytecodeIndex)
 		, Container(InContainer)
 	{
+		ArIsFilterEditorOnly = InnerArchive.ArIsFilterEditorOnly;
 	}
 	virtual FArchive& operator<<(FField*& Value) override
 	{
