Fixed problem where profiler scope names would sometimes be inaccurate.
#RNX

diff --git a/Engine/Source/Runtime/CoreUObject/Private/UObject/UObjectBase.cpp b/Engine/Source/Runtime/CoreUObject/Private/UObject/UObjectBase.cpp
--- a/Engine/Source/Runtime/CoreUObject/Private/UObject/UObjectBase.cpp
+++ b/Engine/Source/Runtime/CoreUObject/Private/UObject/UObjectBase.cpp
@@ -204,7 +204,9 @@ void UObjectBase::AddObject(FName InName, EInternalObjectFlags InSetInternalFlag
  */
 void UObjectBase::LowLevelRename(FName NewName,UObject *NewOuter)
 {
-	STAT(((UObject*)this)->ResetStatID()); // reset the stat id since this thing now has a different name
+#if STATS || ENABLE_STATNAMEDEVENTS_UOBJECT
+	((UObject*)this)->ResetStatID(); // reset the stat id since this thing now has a different name
+#endif
 	UnhashObject(this);
 	check(InternalIndex >= 0);
 	NamePrivate = NewName;
@@ -217,7 +219,9 @@ void UObjectBase::LowLevelRename(FName NewName,UObject *NewOuter)
 
 void UObjectBase::SetClass(UClass* NewClass)
 {
-	STAT(((UObject*)this)->ResetStatID()); // reset the stat id since this thing now has a different name
+#if STATS || ENABLE_STATNAMEDEVENTS_UOBJECT
+	((UObject*)this)->ResetStatID(); // reset the stat id since this thing now has a different name
+#endif
 
 	UnhashObject(this);
 #if USE_UBER_GRAPH_PERSISTENT_FRAME
diff --git a/Engine/Source/Runtime/CoreUObject/Public/UObject/UObjectBaseUtility.h b/Engine/Source/Runtime/CoreUObject/Public/UObject/UObjectBaseUtility.h
--- a/Engine/Source/Runtime/CoreUObject/Public/UObject/UObjectBaseUtility.h
+++ b/Engine/Source/Runtime/CoreUObject/Public/UObject/UObjectBaseUtility.h
@@ -616,6 +616,9 @@ public:
 	FORCEINLINE void ResetStatID()
 	{
 		GUObjectArray.IndexToObject(InternalIndex)->StatID = TStatId();
+#if ENABLE_STATNAMEDEVENTS_UOBJECT
+		GUObjectArray.IndexToObject(InternalIndex)->StatIDStringStorage = nullptr;
+#endif
 	}
 #endif
 	/**
