FFieldPaths should no longer re-resolve themselves each time a blueprint is destroy and only re-resolve if the blueprint that owns the property they resolve to is destroyed / or recompiled.
Changing how FFieldPaths (TFieldPaths) are internally stored and serialized. Now instead of storing the full path (from the outermost UPackage to the innermost FProperty) FFieldPaths will store the owner UStruct reference and a short path to its property (or inner property of its property). The same applies to serialization: both the owner UStruct reference and only the short path will be serialized.
Serializing the owner reference will enable proper handling of any owner struct package renames as well as greatly increase the initial field (property) resolution performance.

diff --git a/Engine/Source/Runtime/Core/Public/UObject/ReleaseObjectVersion.h b/Engine/Source/Runtime/Core/Public/UObject/ReleaseObjectVersion.h
--- a/Engine/Source/Runtime/Core/Public/UObject/ReleaseObjectVersion.h
+++ b/Engine/Source/Runtime/Core/Public/UObject/ReleaseObjectVersion.h
@@ -99,6 +99,8 @@ struct CORE_API FReleaseObjectVersion
 		// custom collision type
 		CustomImplicitCollisionType,
 
+		// FFieldPath will serialize the owner struct reference and only a short path to its property
+		FFieldPathOwnerSerialization,
 
 		// -----<new versions can be added above this line>-------------------------------------------------
 		VersionPlusOne,
diff --git a/Engine/Source/Runtime/CoreUObject/Private/UObject/FieldPath.cpp b/Engine/Source/Runtime/CoreUObject/Private/UObject/FieldPath.cpp
--- a/Engine/Source/Runtime/CoreUObject/Private/UObject/FieldPath.cpp
+++ b/Engine/Source/Runtime/CoreUObject/Private/UObject/FieldPath.cpp
@@ -9,6 +9,7 @@
 #include "UObject/Package.h"
 #include "UObject/UObjectArray.h"
 #include "UObject/FortniteMainBranchObjectVersion.h"
+#include "UObject/ReleaseObjectVersion.h"
 #include "UObject/LinkerLoad.h"
 #include "UObject/UObjectThreadContext.h"
 #include "UObject/PropertyHelper.h"
@@ -313,6 +314,7 @@ FString FFieldPath::ToString() const
 FArchive& operator<<(FArchive& Ar, FFieldPath& InOutPropertyPath)
 {
 	Ar.UsingCustomVersion(FFortniteMainBranchObjectVersion::GUID);
+	Ar.UsingCustomVersion(FReleaseObjectVersion::GUID);
 
 	if (Ar.IsSaving())
 	{		
@@ -353,7 +355,7 @@ FArchive& operator<<(FArchive& Ar, FFieldPath& InOutPropertyPath)
 		{
 			InOutPropertyPath.Path.Empty();
 		}
-		if (Ar.CustomVer(FFortniteMainBranchObjectVersion::GUID) >= FFortniteMainBranchObjectVersion::FFieldPathOwnerSerialization)
+		if (Ar.CustomVer(FFortniteMainBranchObjectVersion::GUID) >= FFortniteMainBranchObjectVersion::FFieldPathOwnerSerialization || Ar.CustomVer(FReleaseObjectVersion::GUID) >= FReleaseObjectVersion::FFieldPathOwnerSerialization)
 		{
 			UStruct* SerializedOwner = InOutPropertyPath.ResolvedOwner.Get();
 			Ar << SerializedOwner;
@@ -420,5 +422,4 @@ int32 FFieldPath::GetFieldPathSerialNumber(UStruct* InStruct) const
 {
 	return InStruct->FieldPathSerialNumber;
 }
-
 #endif // WITH_EDITORONLY_DATA
