Replace ContentMountPointCriticalSection with an FRWLock to allow for concurrent Reads

diff --git a/Engine/Source/Runtime/CoreUObject/Private/Misc/PackageName.cpp b/Engine/Source/Runtime/CoreUObject/Private/Misc/PackageName.cpp
--- a/Engine/Source/Runtime/CoreUObject/Private/Misc/PackageName.cpp
+++ b/Engine/Source/Runtime/CoreUObject/Private/Misc/PackageName.cpp
@@ -22,7 +22,7 @@
 #include "HAL/ThreadHeartBeat.h"
 #include "Misc/AutomationTest.h"
 #include "Misc/PathViews.h"
-#include "Misc/ScopeLock.h"
+#include "Misc/ScopeRWLock.h"
 #include "Misc/StringBuilder.h"
 #include "IO/IoDispatcher.h"
 #include "ProfilingDebugging/LoadTimeTracker.h"
@@ -33,7 +33,7 @@ FString FPackageName::AssetPackageExtension = TEXT(".uasset");
 FString FPackageName::MapPackageExtension = TEXT(".umap");
 FString FPackageName::TextAssetPackageExtension = TEXT(".utxt");
 FString FPackageName::TextMapPackageExtension = TEXT(".utxtmap");
-static FCriticalSection ContentMountPointCriticalSection;
+static FRWLock ContentMountPointCriticalSection;
 
 /** Event that is triggered when a new content path is mounted */
 FPackageName::FOnContentPathMountedEvent FPackageName::OnContentPathMountedEvent;
@@ -198,7 +198,7 @@ struct FLongPackagePathsSingleton
 		OutRoots.Add(GameRootPath);
 
 		{
-			FScopeLock ScopeLock(&ContentMountPointCriticalSection);
+			FReadScopeLock ScopeLock(ContentMountPointCriticalSection);
 			OutRoots += MountPointRootPaths;
 		}
 
@@ -249,7 +249,7 @@ struct FLongPackagePathsSingleton
 
 		FPathPair Pair(RootPath, RelativeContentPath);
 		{
-			FScopeLock ScopeLock(&ContentMountPointCriticalSection);
+			FWriteScopeLock ScopeLock(ContentMountPointCriticalSection);
 			ContentRootToPath.Insert(Pair, 0);
 			ContentPathToRoot.Insert(Pair, 0);
 			MountPointRootPaths.Add(RootPath);
@@ -273,7 +273,7 @@ struct FLongPackagePathsSingleton
 
 		bool bFirePathDismountedDelegate = false;
 		{
-			FScopeLock ScopeLock(&ContentMountPointCriticalSection);
+			FWriteScopeLock ScopeLock(ContentMountPointCriticalSection);
 			if ( MountPointRootPaths.Remove(RootPath) > 0 )
 			{
 				FPathPair Pair(RootPath, RelativeContentPath);
@@ -321,7 +321,7 @@ private:
 		GameExtraPathRebased   = RebasedGameDir / TEXT("Extra/");
 		GameSavedPathRebased   = RebasedGameDir / TEXT("Saved/");
 		
-		FScopeLock ScopeLock(&ContentMountPointCriticalSection);
+		FWriteScopeLock ScopeLock(ContentMountPointCriticalSection);
 
 		ContentPathToRoot.Empty(12);
 		ContentPathToRoot.Emplace(EngineRootPath, EngineContentPath);
@@ -371,7 +371,7 @@ void FPackageName::InternalFilenameToLongPackageName(FStringView InFilename, FSt
 	// Convert to relative path if it's not already a long package name
 	bool bIsValidLongPackageName = false;
 	{
-		FScopeLock ScopeLock(&ContentMountPointCriticalSection);
+		FReadScopeLock ScopeLock(ContentMountPointCriticalSection);
 		for (const auto& Pair : Paths.ContentRootToPath)
 		{
 			if (Filename.StartsWith(Pair.RootPath))
@@ -400,7 +400,7 @@ void FPackageName::InternalFilenameToLongPackageName(FStringView InFilename, FSt
 	FStringView Result = FPathViews::GetBaseFilenameWithPath(Filename);
 
 	{
-		FScopeLock ScopeLock(&ContentMountPointCriticalSection);
+		FReadScopeLock ScopeLock(ContentMountPointCriticalSection);
 		for (const auto& Pair : Paths.ContentPathToRoot)
 		{
 			if (Result.StartsWith(Pair.ContentPath))
@@ -481,7 +481,7 @@ FString FPackageName::FilenameToLongPackageName(const FString& InFilename)
 bool FPackageName::TryConvertLongPackageNameToFilename(const FString& InLongPackageName, FString& OutFilename, const FString& InExtension)
 {
 	const auto& Paths = FLongPackagePathsSingleton::Get();
-	FScopeLock ScopeLock(&ContentMountPointCriticalSection);
+	FReadScopeLock ScopeLock(ContentMountPointCriticalSection);
 	for (const auto& Pair : Paths.ContentRootToPath)
 	{
 		if (InLongPackageName.StartsWith(Pair.RootPath))
@@ -498,7 +498,7 @@ bool FPackageName::TryConvertLongPackageNameToFilename(const FString& InLongPack
 bool FPackageName::ConvertRootPathToContentPath( const FString& RootPath, FString& OutContentPath)
 {
 	const auto& Paths = FLongPackagePathsSingleton::Get();
-	FScopeLock ScopeLock(&ContentMountPointCriticalSection);
+	FReadScopeLock ScopeLock(ContentMountPointCriticalSection);
 	for (const auto& Pair : Paths.ContentRootToPath)
 	{
 		if (RootPath.StartsWith(Pair.RootPath))
@@ -900,7 +900,7 @@ bool FPackageName::FixPackageNameCase(FString& LongPackageName, FStringView Exte
 {
 	// Find the matching long package root
 	const FLongPackagePathsSingleton& Paths = FLongPackagePathsSingleton::Get();
-	FScopeLock ScopeLock(&ContentMountPointCriticalSection);
+	FReadScopeLock ScopeLock(ContentMountPointCriticalSection);
 	for (const FPathPair& Pair : Paths.ContentRootToPath)
 	{
 		if (LongPackageName.StartsWith(Pair.RootPath))
