Added UPS package bit to future-proof package header format

diff --git a/Engine/Source/Runtime/CoreUObject/Private/Serialization/AsyncLoading2.cpp b/Engine/Source/Runtime/CoreUObject/Private/Serialization/AsyncLoading2.cpp
--- a/Engine/Source/Runtime/CoreUObject/Private/Serialization/AsyncLoading2.cpp
+++ b/Engine/Source/Runtime/CoreUObject/Private/Serialization/AsyncLoading2.cpp
@@ -3117,7 +3117,7 @@ EAsyncPackageState::Type FAsyncPackage2::Event_ProcessExportBundle(FAsyncPackage
 			Ar.SetLicenseeUE4Ver(Package->LinkerRoot->LinkerLicenseeVersion);
 			// Ar.SetEngineVer(Summary.SavedByEngineVersion); // very old versioning scheme
 			// Ar.SetCustomVersions(LinkerRoot->LinkerCustomVersion); // only if not cooking with -unversioned
-			Ar.SetUseUnversionedPropertySerialization(CanUseUnversionedPropertySerialization());
+			Ar.SetUseUnversionedPropertySerialization((Package->LinkerRoot->GetPackageFlags() & PKG_UnversionedProperties) != 0);
 			Ar.SetIsLoading(true);
 			Ar.SetIsPersistent(true);
 			if (Package->LinkerRoot->GetPackageFlags() & PKG_FilterEditorOnly)
diff --git a/Engine/Source/Runtime/CoreUObject/Private/Serialization/UnversionedPropertySerialization.cpp b/Engine/Source/Runtime/CoreUObject/Private/Serialization/UnversionedPropertySerialization.cpp
--- a/Engine/Source/Runtime/CoreUObject/Private/Serialization/UnversionedPropertySerialization.cpp
+++ b/Engine/Source/Runtime/CoreUObject/Private/Serialization/UnversionedPropertySerialization.cpp
@@ -699,7 +699,7 @@ private:
 	}
 };
 
-bool CanUseUnversionedPropertySerialization()
+static bool CanUseUnversionedPropertySerialization()
 {
 	bool bTemp;
 	static bool bAllow = GConfig->GetBool(TEXT("Core.System"), TEXT("CanUseUnversionedPropertySerialization"), bTemp, GEngineIni) && bTemp;
diff --git a/Engine/Source/Runtime/CoreUObject/Private/Serialization/UnversionedPropertySerialization.h b/Engine/Source/Runtime/CoreUObject/Private/Serialization/UnversionedPropertySerialization.h
--- a/Engine/Source/Runtime/CoreUObject/Private/Serialization/UnversionedPropertySerialization.h
+++ b/Engine/Source/Runtime/CoreUObject/Private/Serialization/UnversionedPropertySerialization.h
@@ -4,11 +4,6 @@
 
 #include "UObject/UnrealType.h"
 
-struct FUnversionedStructSchema;
-
-// Check if unversioned property serialization is configured to be used on current platform
-bool CanUseUnversionedPropertySerialization();
-
 // Check if unversioned property serialization is configured to be used on target platform
 bool CanUseUnversionedPropertySerialization(const ITargetPlatform* Target);
 
diff --git a/Engine/Source/Runtime/CoreUObject/Private/UObject/LinkerLoad.cpp b/Engine/Source/Runtime/CoreUObject/Private/UObject/LinkerLoad.cpp
--- a/Engine/Source/Runtime/CoreUObject/Private/UObject/LinkerLoad.cpp
+++ b/Engine/Source/Runtime/CoreUObject/Private/UObject/LinkerLoad.cpp
@@ -1169,7 +1169,7 @@ FLinkerLoad::ELinkerStatus FLinkerLoad::SerializePackageFileSummaryInternal()
 	}
 
 	// Set desired property tag format
-	bool bUseUnversionedProperties = Summary.bUnversioned && CanUseUnversionedPropertySerialization();
+	bool bUseUnversionedProperties = (Summary.PackageFlags & PKG_UnversionedProperties) != 0;
 	SetUseUnversionedPropertySerialization(bUseUnversionedProperties);
 	Loader->SetUseUnversionedPropertySerialization(bUseUnversionedProperties);
 
diff --git a/Engine/Source/Runtime/CoreUObject/Private/UObject/SavePackage.cpp b/Engine/Source/Runtime/CoreUObject/Private/UObject/SavePackage.cpp
--- a/Engine/Source/Runtime/CoreUObject/Private/UObject/SavePackage.cpp
+++ b/Engine/Source/Runtime/CoreUObject/Private/UObject/SavePackage.cpp
@@ -3770,6 +3770,11 @@ FSavePackageResultStruct UPackage::Save(UPackage* InOuter, UObject* Base, EObjec
 				bool bUseUnversionedProperties = bSaveUnversioned && CanUseUnversionedPropertySerialization(TargetPlatform);
 				Linker->SetUseUnversionedPropertySerialization(bUseUnversionedProperties);
 				Linker->Saver->SetUseUnversionedPropertySerialization(bUseUnversionedProperties);
+				if (bUseUnversionedProperties)
+				{
+					Linker->Summary.PackageFlags |= PKG_UnversionedProperties;
+					Linker->LinkerRoot->SetPackageFlags(PKG_UnversionedProperties);
+				}
 
 				// Make sure the package has the same version as the linker
 				InOuter->LinkerPackageVersion = Linker->UE4Ver();
diff --git a/Engine/Source/Runtime/CoreUObject/Public/UObject/ObjectMacros.h b/Engine/Source/Runtime/CoreUObject/Public/UObject/ObjectMacros.h
--- a/Engine/Source/Runtime/CoreUObject/Public/UObject/ObjectMacros.h
+++ b/Engine/Source/Runtime/CoreUObject/Public/UObject/ObjectMacros.h
@@ -114,7 +114,7 @@ enum EPackageFlags
 //	PKG_Unused						= 0x00000400,
 //	PKG_Unused						= 0x00000800,
 //	PKG_Unused						= 0x00001000,
-//	PKG_Unused						= 0x00002000,
+	PKG_UnversionedProperties		= 0x00002000,   ///< Uses unversioned property serialization instead of versioned tagged property serialization
 	PKG_ContainsMapData				= 0x00004000,   ///< Contains map data (UObjects only referenced by a single ULevel) but is stored in a different package
 //	PKG_Unused						= 0x00008000,
 	PKG_Compiling					= 0x00010000,	///< package is currently being compiled
