Temporarily disabling GC optimization because some script generated classes pretend to be native classes

diff --git a/Engine/Source/Runtime/CoreUObject/Private/UObject/GarbageCollection.cpp b/Engine/Source/Runtime/CoreUObject/Private/UObject/GarbageCollection.cpp
--- a/Engine/Source/Runtime/CoreUObject/Private/UObject/GarbageCollection.cpp
+++ b/Engine/Source/Runtime/CoreUObject/Private/UObject/GarbageCollection.cpp
@@ -2617,8 +2617,8 @@ void UClass::AssembleReferenceTokenStream(bool bForce)
 
 		{
 			check(ClassAddReferencedObjects != NULL);
-			const bool bKeepOuter = GetFName() != NAME_Package;
-			const bool bKeepClass = !(GetClassFlags() & CLASS_Native) || IsA(UDynamicClass::StaticClass());
+			const bool bKeepOuter = true;// GetFName() != NAME_Package;
+			const bool bKeepClass = true;// !(GetClassFlags() & CLASS_Native) || IsA(UDynamicClass::StaticClass());
 
 			ClassAddReferencedObjectsType AddReferencedObjectsFn = nullptr;
 #if !WITH_EDITOR
