Compile out Garbage Collection, Core Redirect and SavePackage related code in UHT since they contribute greatly to the overall build time. Brings down a clean build of QAGame client win64 to under three minutes from just under four minutes

diff --git a/Engine/Source/Programs/UnrealHeaderTool/UnrealHeaderTool.Target.cs b/Engine/Source/Programs/UnrealHeaderTool/UnrealHeaderTool.Target.cs
--- a/Engine/Source/Programs/UnrealHeaderTool/UnrealHeaderTool.Target.cs
+++ b/Engine/Source/Programs/UnrealHeaderTool/UnrealHeaderTool.Target.cs
@@ -23,7 +23,7 @@ public class UnrealHeaderToolTarget : TargetRules
 		bCompileAgainstEngine = false;
 		bCompileAgainstApplicationCore = false;
 
-		// Force execption handling across all modules.
+		// Force exception handling across all modules.
 		bForceEnableExceptions = true;
 
 		// Plugin support
@@ -36,5 +36,10 @@ public class UnrealHeaderToolTarget : TargetRules
         GlobalDefinitions.Add("HACK_HEADER_GENERATOR=1");
         GlobalDefinitions.Add("FNAME_WRITE_PROTECT_PAGES=0");
         GlobalDefinitions.Add("USE_LOCALIZED_PACKAGE_CACHE=0");
+
+		// We disable compilation of GC related code since it takes a very long time (up to 30 seconds on typical hardware at the time of writing)
+		GlobalDefinitions.Add("UE_WITH_GC=0");
+		GlobalDefinitions.Add("UE_WITH_SAVEPACKAGE=0");
+		GlobalDefinitions.Add("UE_WITH_CORE_REDIRECTS=0");
 	}
 }
diff --git a/Engine/Source/Runtime/CoreUObject/Private/UObject/CoreRedirects.cpp b/Engine/Source/Runtime/CoreUObject/Private/UObject/CoreRedirects.cpp
--- a/Engine/Source/Runtime/CoreUObject/Private/UObject/CoreRedirects.cpp
+++ b/Engine/Source/Runtime/CoreUObject/Private/UObject/CoreRedirects.cpp
@@ -14,6 +14,10 @@
 #include "Serialization/DeferredMessageLog.h"
 #include "Templates/Casts.h"
 
+#if !defined UE_WITH_CORE_REDIRECTS
+#	define UE_WITH_CORE_REDIRECTS 1
+#endif
+
 FCoreRedirectObjectName::FCoreRedirectObjectName(const FString& InString)
 {
 	if (!ExpandNames(InString, ObjectName, OuterName, PackageName))
@@ -1100,6 +1104,8 @@ ECoreRedirectFlags FCoreRedirects::GetFlagsForTypeClass(UClass *TypeClass)
 
 // We want to only load these redirects in editor builds, but Matinee needs them at runtime still 
 
+#if UE_WITH_CORE_REDIRECTS
+
 PRAGMA_DISABLE_OPTIMIZATION
 
 // The compiler doesn't like having a massive string table in a single function so split it up
@@ -1912,4 +1918,9 @@ void FCoreRedirects::RegisterNativeRedirects()
 	// 4.10 and later are in baseengine.ini
 
 	AddRedirectList(Redirects, TEXT("RegisterNativeRedirects"));
-}
\ No newline at end of file
+}
+#else
+void FCoreRedirects::RegisterNativeRedirects()
+{
+}
+#endif // UE_WITH_CORE_REDIRECTS
diff --git a/Engine/Source/Runtime/CoreUObject/Private/UObject/GarbageCollection.cpp b/Engine/Source/Runtime/CoreUObject/Private/UObject/GarbageCollection.cpp
--- a/Engine/Source/Runtime/CoreUObject/Private/UObject/GarbageCollection.cpp
+++ b/Engine/Source/Runtime/CoreUObject/Private/UObject/GarbageCollection.cpp
@@ -535,6 +535,8 @@ void ShutdownGarbageCollection()
 /**
 * Handles UObject references found by TFastReferenceCollector
 */
+
+#if UE_WITH_GC
 template <bool bParallel, bool bWithClusters>
 class FGCReferenceProcessor
 {
@@ -908,7 +910,6 @@ public:
 	}
 };
 
-
 template <bool bParallel, bool bWithClusters>
 FGCCollector<bParallel, bWithClusters>::FGCCollector(FGCReferenceProcessor<bParallel, bWithClusters>& InProcessor, FGCArrayStruct& InObjectArrayStruct)
 		: ReferenceProcessor(InProcessor)
@@ -947,6 +948,7 @@ void FGCCollector<bParallel, bWithClusters>::HandleObjectReferences(UObject** In
 			InternalHandleObjectReference(Object, InReferencingObject, InReferencingProperty);
 		}
 }
+#endif	// UE_WITH_GC
 
 /*----------------------------------------------------------------------------
 	FReferenceFinder.
@@ -1020,6 +1022,9 @@ void FReferenceFinder::HandleObjectReference( UObject*& InObject, const UObject*
  * is used to deal with object references from types that aren't supported by the reflectable type system.
  * interface doesn't make sense to implement for.
  */
+
+#if UE_WITH_GC
+
 class FRealtimeGC : public FGarbageCollectionTracer
 {
 	typedef void(FRealtimeGC::*MarkObjectsFn)(TArray<UObject*>&, const EObjectFlags);
@@ -1302,6 +1307,7 @@ public:
 		(this->*ReachabilityAnalysisFunctions[GetGCFunctionIndex(!bForceSingleThreaded, bWithClusters)])(ArrayStruct);
 	}
 };
+#endif // UE_WITH_GC
 
 // Allow parallel GC to be overridden to single threaded via console command.
 static int32 GAllowParallelGC = 1;
@@ -1309,7 +1315,7 @@ static int32 GAllowParallelGC = 1;
 static FAutoConsoleVariableRef CVarAllowParallelGC(
 	TEXT("gc.AllowParallelGC"),
 	GAllowParallelGC,
-	TEXT("sed to control parallel GC."),
+	TEXT("Used to control parallel GC."),
 	ECVF_Default
 );
 
@@ -1376,6 +1382,9 @@ static bool IncrementalDestroyGarbage(bool bUseTimeLimit, float TimeLimit);
  */
 void IncrementalPurgeGarbage(bool bUseTimeLimit, float TimeLimit)
 {
+#if !UE_WITH_GC
+	return;
+#else
 	SCOPED_NAMED_EVENT(IncrementalPurgeGarbage, FColor::Red);
 	DECLARE_SCOPE_CYCLE_COUNTER(TEXT("IncrementalPurgeGarbage"), STAT_IncrementalPurgeGarbage, STATGROUP_GC);
 	CSV_SCOPED_TIMING_STAT_EXCLUSIVE(GarbageCollection);
@@ -1435,8 +1444,10 @@ void IncrementalPurgeGarbage(bool bUseTimeLimit, float TimeLimit)
 	{
 		bCompleted = IncrementalDestroyGarbage(bUseTimeLimit, TimeLimit);
 	}
+#endif // !UE_WITH_GC
 }
 
+#if UE_WITH_GC
 bool IncrementalDestroyGarbage(bool bUseTimeLimit, float TimeLimit)
 {
 	const bool bMultithreadedPurge = !ShouldForceSingleThreadedGC() && GMultithreadedDestructionEnabled;
@@ -1734,6 +1745,7 @@ bool IncrementalDestroyGarbage(bool bUseTimeLimit, float TimeLimit)
 
 	return bCompleted;
 }
+#endif // UE_WITH_GC
 
 /**
  * Returns whether an incremental purge is still pending/ in progress.
@@ -1856,6 +1868,9 @@ void GatherUnreachableObjects(bool bForceSingleThreaded)
  */
 void CollectGarbageInternal(EObjectFlags KeepFlags, bool bPerformFullPurge)
 {
+#if !UE_WITH_GC
+	return;
+#else
 	if (GIsInitialLoad)
 	{
 		// During initial load classes may not yet have their GC token streams assembled
@@ -1999,6 +2014,7 @@ void CollectGarbageInternal(EObjectFlags KeepFlags, bool bPerformFullPurge)
 	FCoreUObjectDelegates::GetPostGarbageCollect().Broadcast();
 
 	STAT_ADD_CUSTOMMESSAGE_NAME( STAT_NamedMarker, TEXT( "GarbageCollection - End" ) );
+#endif	// UE_WITH_GC
 }
 
 bool IsIncrementalUnhashPending()
diff --git a/Engine/Source/Runtime/CoreUObject/Private/UObject/GarbageCollectionVerification.h b/Engine/Source/Runtime/CoreUObject/Private/UObject/GarbageCollectionVerification.h
--- a/Engine/Source/Runtime/CoreUObject/Private/UObject/GarbageCollectionVerification.h
+++ b/Engine/Source/Runtime/CoreUObject/Private/UObject/GarbageCollectionVerification.h
@@ -12,7 +12,7 @@
 #include "UObject/Class.h"
 
 // UE_BUILD_SHIPPING has GShouldVerifyGCAssumptions=false by default
-#define VERIFY_DISREGARD_GC_ASSUMPTIONS			!(UE_BUILD_SHIPPING || UE_BUILD_TEST)
+#define VERIFY_DISREGARD_GC_ASSUMPTIONS			!(UE_BUILD_SHIPPING || UE_BUILD_TEST || !UE_WITH_GC)
 
 #if VERIFY_DISREGARD_GC_ASSUMPTIONS
 
diff --git a/Engine/Source/Runtime/CoreUObject/Private/UObject/SavePackage.cpp b/Engine/Source/Runtime/CoreUObject/Private/UObject/SavePackage.cpp
--- a/Engine/Source/Runtime/CoreUObject/Private/UObject/SavePackage.cpp
+++ b/Engine/Source/Runtime/CoreUObject/Private/UObject/SavePackage.cpp
@@ -1,6 +1,8 @@
 // Copyright Epic Games, Inc. All Rights Reserved.
 
 #include "UObject/SavePackage.h"
+
+#if UE_WITH_SAVEPACKAGE
 #include "CoreMinimal.h"
 #include "Misc/MessageDialog.h"
 #include "HAL/FileManager.h"
@@ -6472,3 +6474,4 @@ FSavePackageContext::~FSavePackageContext()
 	delete PackageStoreWriter;
 	delete BulkDataManifest;
 }
+#endif	// UE_WITH_SAVEPACKAGE
diff --git a/Engine/Source/Runtime/CoreUObject/Public/UObject/GarbageCollection.h b/Engine/Source/Runtime/CoreUObject/Public/UObject/GarbageCollection.h
--- a/Engine/Source/Runtime/CoreUObject/Public/UObject/GarbageCollection.h
+++ b/Engine/Source/Runtime/CoreUObject/Public/UObject/GarbageCollection.h
@@ -12,6 +12,10 @@
 #include "Serialization/ArchiveUObject.h"
 #include "HAL/ThreadSafeBool.h"
 
+#if !defined(UE_WITH_GC)
+#	define UE_WITH_GC	1
+#endif
+
 /** Context sensitive keep flags for garbage collection */
 #define GARBAGE_COLLECTION_KEEPFLAGS	(GIsEditor ? RF_Standalone : RF_NoFlags)
 
diff --git a/Engine/Source/Runtime/CoreUObject/Public/UObject/SavePackage.h b/Engine/Source/Runtime/CoreUObject/Public/UObject/SavePackage.h
--- a/Engine/Source/Runtime/CoreUObject/Public/UObject/SavePackage.h
+++ b/Engine/Source/Runtime/CoreUObject/Public/UObject/SavePackage.h
@@ -7,6 +7,10 @@
 #include "Containers/Map.h"
 #include "UObject/NameTypes.h"
 
+#if !defined(UE_WITH_SAVEPACKAGE)
+#	define UE_WITH_SAVEPACKAGE 1
+#endif
+
 class FArchive;
 class FIoBuffer;
 class FLinkerLoad;
