Added #include of GarbageCollection.h to GarbageCollectionVerification.h to ensure we compile consistently in non-unity builds

diff --git a/Engine/Source/Runtime/CoreUObject/Private/UObject/GarbageCollectionVerification.h b/Engine/Source/Runtime/CoreUObject/Private/UObject/GarbageCollectionVerification.h
--- a/Engine/Source/Runtime/CoreUObject/Private/UObject/GarbageCollectionVerification.h
+++ b/Engine/Source/Runtime/CoreUObject/Private/UObject/GarbageCollectionVerification.h
@@ -8,6 +8,7 @@
 
 #include "CoreMinimal.h"
 #include "Stats/Stats.h"
+#include "UObject/GarbageCollection.h"	// Needed for UE_WITH_GC definition
 #include "UObject/UObjectGlobals.h"
 #include "UObject/Class.h"
 
