Don't allocate and free memory constantly for the PostInitProperties routing validation
#rnx

diff --git a/Engine/Source/Runtime/CoreUObject/Private/UObject/UObjectGlobals.cpp b/Engine/Source/Runtime/CoreUObject/Private/UObject/UObjectGlobals.cpp
--- a/Engine/Source/Runtime/CoreUObject/Private/UObject/UObjectGlobals.cpp
+++ b/Engine/Source/Runtime/CoreUObject/Private/UObject/UObjectGlobals.cpp
@@ -2798,7 +2798,7 @@ void FObjectInitializer::PostConstructInit()
 	Class->PostInitInstance(Obj);
 
 #if !(UE_BUILD_SHIPPING || UE_BUILD_TEST)
-	if (!FUObjectThreadContext::Get().PostInitPropertiesCheck.Num() || (FUObjectThreadContext::Get().PostInitPropertiesCheck.Pop() != Obj))
+	if (!FUObjectThreadContext::Get().PostInitPropertiesCheck.Num() || (FUObjectThreadContext::Get().PostInitPropertiesCheck.Pop(false) != Obj))
 	{
 		UE_LOG(LogUObjectGlobals, Fatal, TEXT("%s failed to route PostInitProperties. Call Super::PostInitProperties() in %s::PostInitProperties()."), *Obj->GetClass()->GetName(), *Obj->GetClass()->GetName());
 	}
diff --git a/Engine/Source/Runtime/CoreUObject/Public/UObject/UObjectThreadContext.h b/Engine/Source/Runtime/CoreUObject/Public/UObject/UObjectThreadContext.h
--- a/Engine/Source/Runtime/CoreUObject/Public/UObject/UObjectThreadContext.h
+++ b/Engine/Source/Runtime/CoreUObject/Public/UObject/UObjectThreadContext.h
@@ -81,7 +81,7 @@ public:
 
 #if !(UE_BUILD_SHIPPING || UE_BUILD_TEST)
 	/** Stack to ensure that PostInitProperties is routed through Super:: calls. **/
-	TArray<UObject*> PostInitPropertiesCheck;
+	TArray<UObject*, TInlineAllocator<16>> PostInitPropertiesCheck;
 	/** Used to verify that the Super::PostLoad chain is intact.			*/
 	TArray<UObject*, TInlineAllocator<16> > DebugPostLoad;
 #endif
