Add FindMetaData function that returns a const FString* to the metadata value for the given key rather than having to do Has then Get or assume value does not exist based on empty string from Get.
Added Move versions of UMetadata::SetObjectValues and FField::SetMetaData

diff --git a/Engine/Source/Editor/UnrealEd/Private/Fbx/FbxStaticMeshImport.cpp b/Engine/Source/Editor/UnrealEd/Private/Fbx/FbxStaticMeshImport.cpp
--- a/Engine/Source/Editor/UnrealEd/Private/Fbx/FbxStaticMeshImport.cpp
+++ b/Engine/Source/Editor/UnrealEd/Private/Fbx/FbxStaticMeshImport.cpp
@@ -116,7 +116,7 @@ struct FRestoreReimportData
 			//Package was recreate so restore the metadata
 			UMetaData* PackageMetaData = DupObject->GetOutermost()->GetMetaData();
 			checkSlow(PackageMetaData);
-			PackageMetaData->SetObjectValues(DupObject, ObjectMetaData);
+			PackageMetaData->SetObjectValues(DupObject, MoveTemp(ObjectMetaData));
 		}
 		//Since all loaded package are add to root, we have to remove the dup from the root
 		DupObject->RemoveFromRoot();
diff --git a/Engine/Source/Runtime/CoreUObject/Private/UObject/Class.cpp b/Engine/Source/Runtime/CoreUObject/Private/UObject/Class.cpp
--- a/Engine/Source/Runtime/CoreUObject/Private/UObject/Class.cpp
+++ b/Engine/Source/Runtime/CoreUObject/Private/UObject/Class.cpp
@@ -395,7 +395,7 @@ void UField::FormatNativeToolTip(FString& ToolTipString, bool bRemoveExtraSectio
  * @param Key The key to lookup in the metadata
  * @return true if there is a (possibly blank) value associated with this key
  */
-bool UField::HasMetaData(const TCHAR* Key) const
+const FString* UField::FindMetaData(const TCHAR* Key) const
 {
 	UPackage* Package = GetOutermost();
 	check(Package);
@@ -403,12 +403,10 @@ bool UField::HasMetaData(const TCHAR* Key) const
 	UMetaData* MetaData = Package->GetMetaData();
 	check(MetaData);
 
-	bool bHasMetaData = MetaData->HasValue(this, Key);
-		
-	return bHasMetaData;
+	return MetaData->FindValue(this, Key);
 }
 
-bool UField::HasMetaData(const FName& Key) const
+const FString* UField::FindMetaData(const FName& Key) const
 {
 	UPackage* Package = GetOutermost();
 	check(Package);
@@ -416,9 +414,7 @@ bool UField::HasMetaData(const FName& Key) const
 	UMetaData* MetaData = Package->GetMetaData();
 	check(MetaData);
 
-	bool bHasMetaData = MetaData->HasValue(this, Key);
-
-	return bHasMetaData;
+	return MetaData->FindValue(this, Key);
 }
 
 /**
@@ -457,9 +453,9 @@ FText UField::GetMetaDataText(const TCHAR* MetaDataKey, const FString Localizati
 {
 	FString DefaultMetaData;
 
-	if( HasMetaData( MetaDataKey ))
+	if(const FString* FoundMetaData = FindMetaData( MetaDataKey ))
 	{
-		DefaultMetaData = GetMetaData(MetaDataKey);
+		DefaultMetaData = *FoundMetaData;
 	}
 
 	// If attempting to grab the DisplayName metadata, we must correct the source string and output it as a DisplayString for lookup
@@ -484,9 +480,9 @@ FText UField::GetMetaDataText(const FName& MetaDataKey, const FString Localizati
 {
 	FString DefaultMetaData;
 
-	if( HasMetaData( MetaDataKey ))
+	if (const FString* FoundMetaData = FindMetaData( MetaDataKey ))
 	{
-		DefaultMetaData = GetMetaData(MetaDataKey);
+		DefaultMetaData = *FoundMetaData;
 	}
 
 	// If attempting to grab the DisplayName metadata, we must correct the source string and output it as a DisplayString for lookup
@@ -2055,11 +2051,11 @@ bool UStruct::GetStringMetaDataHierarchical(const FName& Key, FString* OutValue)
 {
 	for (const UStruct* TestStruct = this; TestStruct != nullptr; TestStruct = TestStruct->GetSuperStruct())
 	{
-		if (TestStruct->HasMetaData(Key))
+		if (const FString* FoundMetaData = TestStruct->FindMetaData(Key))
 		{
 			if (OutValue != nullptr)
 			{
-				*OutValue = TestStruct->GetMetaData(Key);
+				*OutValue = *FoundMetaData;
 			}
 
 			return true;
@@ -5266,20 +5262,18 @@ const FString UClass::GetConfigName() const
 void UClass::GetHideFunctions(TArray<FString>& OutHideFunctions) const
 {
 	static const FName NAME_HideFunctions(TEXT("HideFunctions"));
-	if (HasMetaData(NAME_HideFunctions))
+	if (const FString* HideFunctions = FindMetaData(NAME_HideFunctions))
 	{
-		const FString& HideFunctions = GetMetaData(NAME_HideFunctions);
-		HideFunctions.ParseIntoArray(OutHideFunctions, TEXT(" "), true);
+		HideFunctions->ParseIntoArray(OutHideFunctions, TEXT(" "), true);
 	}
 }
 
 bool UClass::IsFunctionHidden(const TCHAR* InFunction) const
 {
 	static const FName NAME_HideFunctions(TEXT("HideFunctions"));
-	if (HasMetaData(NAME_HideFunctions))
+	if (const FString* HideFunctions = FindMetaData(NAME_HideFunctions))
 	{
-		const FString& HideFunctions = GetMetaData(NAME_HideFunctions);
-		return !!FCString::StrfindDelim(*HideFunctions, InFunction, TEXT(" "));
+		return !!FCString::StrfindDelim(**HideFunctions, InFunction, TEXT(" "));
 	}
 	return false;
 }
@@ -5287,20 +5281,18 @@ bool UClass::IsFunctionHidden(const TCHAR* InFunction) const
 void UClass::GetAutoExpandCategories(TArray<FString>& OutAutoExpandCategories) const
 {
 	static const FName NAME_AutoExpandCategories(TEXT("AutoExpandCategories"));
-	if (HasMetaData(NAME_AutoExpandCategories))
+	if (const FString* AutoExpandCategories = FindMetaData(NAME_AutoExpandCategories))
 	{
-		const FString& AutoExpandCategories = GetMetaData(NAME_AutoExpandCategories);
-		AutoExpandCategories.ParseIntoArray(OutAutoExpandCategories, TEXT(" "), true);
+		AutoExpandCategories->ParseIntoArray(OutAutoExpandCategories, TEXT(" "), true);
 	}
 }
 
 bool UClass::IsAutoExpandCategory(const TCHAR* InCategory) const
 {
 	static const FName NAME_AutoExpandCategories(TEXT("AutoExpandCategories"));
-	if (HasMetaData(NAME_AutoExpandCategories))
+	if (const FString* AutoExpandCategories = FindMetaData(NAME_AutoExpandCategories))
 	{
-		const FString& AutoExpandCategories = GetMetaData(NAME_AutoExpandCategories);
-		return !!FCString::StrfindDelim(*AutoExpandCategories, InCategory, TEXT(" "));
+		return !!FCString::StrfindDelim(**AutoExpandCategories, InCategory, TEXT(" "));
 	}
 	return false;
 }
@@ -5308,20 +5300,18 @@ bool UClass::IsAutoExpandCategory(const TCHAR* InCategory) const
 void UClass::GetAutoCollapseCategories(TArray<FString>& OutAutoCollapseCategories) const
 {
 	static const FName NAME_AutoCollapseCategories(TEXT("AutoCollapseCategories"));
-	if (HasMetaData(NAME_AutoCollapseCategories))
+	if (const FString* AutoCollapseCategories = FindMetaData(NAME_AutoCollapseCategories))
 	{
-		const FString& AutoCollapseCategories = GetMetaData(NAME_AutoCollapseCategories);
-		AutoCollapseCategories.ParseIntoArray(OutAutoCollapseCategories, TEXT(" "), true);
+		AutoCollapseCategories->ParseIntoArray(OutAutoCollapseCategories, TEXT(" "), true);
 	}
 }
 
 bool UClass::IsAutoCollapseCategory(const TCHAR* InCategory) const
 {
 	static const FName NAME_AutoCollapseCategories(TEXT("AutoCollapseCategories"));
-	if (HasMetaData(NAME_AutoCollapseCategories))
+	if (const FString* AutoCollapseCategories = FindMetaData(NAME_AutoCollapseCategories))
 	{
-		const FString& AutoCollapseCategories = GetMetaData(NAME_AutoCollapseCategories);
-		return !!FCString::StrfindDelim(*AutoCollapseCategories, InCategory, TEXT(" "));
+		return !!FCString::StrfindDelim(**AutoCollapseCategories, InCategory, TEXT(" "));
 	}
 	return false;
 }
@@ -5329,20 +5319,18 @@ bool UClass::IsAutoCollapseCategory(const TCHAR* InCategory) const
 void UClass::GetClassGroupNames(TArray<FString>& OutClassGroupNames) const
 {
 	static const FName NAME_ClassGroupNames(TEXT("ClassGroupNames"));
-	if (HasMetaData(NAME_ClassGroupNames))
+	if (const FString* ClassGroupNames = FindMetaData(NAME_ClassGroupNames))
 	{
-		const FString& ClassGroupNames = GetMetaData(NAME_ClassGroupNames);
-		ClassGroupNames.ParseIntoArray(OutClassGroupNames, TEXT(" "), true);
+		ClassGroupNames->ParseIntoArray(OutClassGroupNames, TEXT(" "), true);
 	}
 }
 
 bool UClass::IsClassGroupName(const TCHAR* InGroupName) const
 {
 	static const FName NAME_ClassGroupNames(TEXT("ClassGroupNames"));
-	if (HasMetaData(NAME_ClassGroupNames))
+	if (const FString* ClassGroupNames = FindMetaData(NAME_ClassGroupNames))
 	{
-		const FString& ClassGroupNames = GetMetaData(NAME_ClassGroupNames);
-		return !!FCString::StrfindDelim(*ClassGroupNames, InGroupName, TEXT(" "));
+		return !!FCString::StrfindDelim(**ClassGroupNames, InGroupName, TEXT(" "));
 	}
 	return false;
 }
diff --git a/Engine/Source/Runtime/CoreUObject/Private/UObject/Field.cpp b/Engine/Source/Runtime/CoreUObject/Private/UObject/Field.cpp
--- a/Engine/Source/Runtime/CoreUObject/Private/UObject/Field.cpp
+++ b/Engine/Source/Runtime/CoreUObject/Private/UObject/Field.cpp
@@ -690,9 +690,9 @@ FText FField::GetDisplayNameText() const
 	const FString Key = GetFullGroupName(false);
 
 	FString NativeDisplayName;
-	if (HasMetaData(NAME_DisplayName))
+	if (const FString* FoundMetaData = FindMetaData(NAME_DisplayName))
 	{
-		NativeDisplayName = GetMetaData(NAME_DisplayName);
+		NativeDisplayName = *FoundMetaData;
 	}
 	else
 	{
@@ -760,20 +760,14 @@ FText FField::GetToolTipText(bool bShortTooltip) const
 	return LocalizedToolTip;
 }
 
-/**
-* Determines if the property has any metadata associated with the key
-*
-* @param Key The key to lookup in the metadata
-* @return true if there is a (possibly blank) value associated with this key
-*/
-bool FField::HasMetaData(const TCHAR* Key) const
+const FString* FField::FindMetaData(const TCHAR* Key) const
 {
-	return HasMetaData(FName(Key, FNAME_Find));
+	return FindMetaData(FName(Key, FNAME_Find));
 }
 
-bool FField::HasMetaData(const FName& Key) const
+const FString* FField::FindMetaData(const FName& Key) const
 {
-	return MetaDataMap && MetaDataMap->Contains(Key);
+	return (MetaDataMap ? MetaDataMap->Find(Key) : nullptr);
 }
 
 /**
@@ -815,9 +809,9 @@ const FText FField::GetMetaDataText(const TCHAR* MetaDataKey, const FString Loca
 {
 	FString DefaultMetaData;
 
-	if (HasMetaData(MetaDataKey))
+	if (const FString* FoundMetaData = FindMetaData(MetaDataKey))
 	{
-		DefaultMetaData = GetMetaData(MetaDataKey);
+		DefaultMetaData = *FoundMetaData;
 	}
 
 	// If attempting to grab the DisplayName metadata, we must correct the source string and output it as a DisplayString for lookup
@@ -843,9 +837,9 @@ const FText FField::GetMetaDataText(const FName& MetaDataKey, const FString Loca
 {
 	FString DefaultMetaData;
 
-	if (HasMetaData(MetaDataKey))
+	if (const FString* FoundMetaData = FindMetaData(MetaDataKey))
 	{
-		DefaultMetaData = GetMetaData(MetaDataKey);
+		DefaultMetaData = *FoundMetaData;
 	}
 
 	// If attempting to grab the DisplayName metadata, we must correct the source string and output it as a DisplayString for lookup
@@ -875,17 +869,27 @@ const FText FField::GetMetaDataText(const FName& MetaDataKey, const FString Loca
 */
 void FField::SetMetaData(const TCHAR* Key, const TCHAR* InValue)
 {
-	SetMetaData(FName(Key), InValue);
+	SetMetaData(FName(Key), FString(InValue));
+}
+
+void FField::SetMetaData(const TCHAR* Key, FString&& InValue)
+{
+	SetMetaData(FName(Key), MoveTemp(InValue));
 }
 
 void FField::SetMetaData(const FName& Key, const TCHAR* InValue)
+{
+	SetMetaData(Key, FString(InValue));
+}
+
+void FField::SetMetaData(const FName& Key, FString&& InValue)
 {
 	check(Key != NAME_None);
 	if (!MetaDataMap)
 	{
 		MetaDataMap = new TMap<FName, FString>();
 	}
-	MetaDataMap->Add(Key, InValue);
+	MetaDataMap->Add(Key, MoveTemp(InValue));
 }
 
 UClass* FField::GetClassMetaData(const TCHAR* Key) const
diff --git a/Engine/Source/Runtime/CoreUObject/Private/UObject/MetaData.cpp b/Engine/Source/Runtime/CoreUObject/Private/UObject/MetaData.cpp
--- a/Engine/Source/Runtime/CoreUObject/Private/UObject/MetaData.cpp
+++ b/Engine/Source/Runtime/CoreUObject/Private/UObject/MetaData.cpp
@@ -277,13 +277,7 @@ const FString& UMetaData::GetValue(const UObject* Object, const TCHAR* Key)
 	return GetValue(Object, FName(Key, FNAME_Find));
 }
 
-/**
- * Return whether or not the Key is in the meta data
- * @param Object the object to lookup the metadata for
- * @param Key The key to query for existence
- * @return true if found
- */
-bool UMetaData::HasValue(const UObject* Object, FName Key)
+const FString* UMetaData::FindValue(const UObject* Object, FName Key)
 {
 	// every key needs to be valid
 	if (Key == NAME_None)
@@ -301,20 +295,14 @@ bool UMetaData::HasValue(const UObject* Object, FName Key)
 	}
 
 	// if we had the map, see if we had the key
-	return ObjectValues->Find(Key) != NULL;
+	return ObjectValues->Find(Key);
 }
 
-/**
- * Return whether or not the Key is in the meta data
- * @param Object the object to lookup the metadata for
- * @param Key The key to query for existence
- * @return true if found
- */
-bool UMetaData::HasValue(const UObject* Object, const TCHAR* Key)
+const FString* UMetaData::FindValue(const UObject* Object, const TCHAR* Key)
 {
 	// only find names, don't bother creating a name if it's not already there
 	// (HasValue will return false if Key is NAME_None)
-	return HasValue(Object, FName(Key, FNAME_Find));
+	return FindValue(Object, FName(Key, FNAME_Find));
 }
 
 /**
@@ -337,6 +325,16 @@ void UMetaData::SetObjectValues(const UObject* Object, const TMap<FName, FString
 	ObjectMetaDataMap.Add(const_cast<UObject*>(Object), ObjectValues);
 }
 
+/**
+ * Set the key/value pair in the Property's metadata
+ * @param Object the object to set the metadata for
+ * @Values The metadata key/value pairs
+ */
+void UMetaData::SetObjectValues(const UObject* Object, TMap<FName, FString>&& ObjectValues)
+{
+	ObjectMetaDataMap.Add(const_cast<UObject*>(Object), MoveTemp(ObjectValues));
+}
+
 /**
  * Set the key/value pair in the Property's metadata
  * @param Object the object to set the metadata for
diff --git a/Engine/Source/Runtime/CoreUObject/Public/UObject/Field.h b/Engine/Source/Runtime/CoreUObject/Public/UObject/Field.h
--- a/Engine/Source/Runtime/CoreUObject/Public/UObject/Field.h
+++ b/Engine/Source/Runtime/CoreUObject/Public/UObject/Field.h
@@ -667,8 +667,17 @@ public:
 	* @param Key The key to lookup in the metadata
 	* @return true if there is a (possibly blank) value associated with this key
 	*/
-	bool HasMetaData(const TCHAR* Key) const;
-	bool HasMetaData(const FName& Key) const;
+	bool HasMetaData(const TCHAR* Key) const { return FindMetaData(Key) != nullptr; }
+	bool HasMetaData(const FName& Key) const { return FindMetaData(Key) != nullptr; }
+
+	/**
+	* Find the metadata value associated with the key
+	*
+	* @param Key The key to lookup in the metadata
+	* @return The value associated with the key if it exists, null otherwise
+	*/
+	const FString* FindMetaData(const TCHAR* Key) const;
+	const FString* FindMetaData(const FName& Key) const;
 
 	/**
 	* Find the metadata value associated with the key
@@ -699,6 +708,9 @@ public:
 	void SetMetaData(const TCHAR* Key, const TCHAR* InValue);
 	void SetMetaData(const FName& Key, const TCHAR* InValue);
 
+	void SetMetaData(const TCHAR* Key, FString&& InValue);
+	void SetMetaData(const FName& Key, FString&& InValue);
+
 	/**
 	* Find the metadata value associated with the key
 	* and return bool
diff --git a/Engine/Source/Runtime/CoreUObject/Public/UObject/MetaData.h b/Engine/Source/Runtime/CoreUObject/Public/UObject/MetaData.h
--- a/Engine/Source/Runtime/CoreUObject/Public/UObject/MetaData.h
+++ b/Engine/Source/Runtime/CoreUObject/Public/UObject/MetaData.h
@@ -55,7 +55,10 @@ public:
 	 * @param Key The key to query for existence
 	 * @return true if found
 	 */
-	bool HasValue(const UObject* Object, const TCHAR* Key);
+	bool HasValue(const UObject* Object, const TCHAR* Key)
+	{
+		return FindValue(Object, Key) != nullptr;
+	}
 
 	/**
 	 * Return whether or not the Key is in the meta data
@@ -64,6 +67,25 @@ public:
 	 * @return true if found
 	 */
 	bool HasValue(const UObject* Object, FName Key);
+	{
+		return FindValue(Object, Key) != nullptr;
+	}
+
+	/**
+	 * Returns the value for a the given key if it exists, null otherwise
+	 * @param Object the object to lookup the metadata for
+	 * @param Key The key to query for existence
+	 * @return true if found
+	 */
+	const FString* FindValue(const UObject* Object, const TCHAR* Key);
+
+	/**
+	 * Returns the value for a the given key if it exists, null otherwise
+	 * @param Object the object to lookup the metadata for
+	 * @param Key The key to query for existence
+	 * @return true if found
+	 */
+	const FString* FindValue(const UObject* Object, FName Key);
 
 	/**
 	 * Is there any metadata for this property?
@@ -79,6 +101,13 @@ public:
 	 */
 	void SetObjectValues(const UObject* Object, const TMap<FName, FString>& Values);
 
+	/**
+	 * Set the key/value pair in the Property's metadata
+	 * @param Object the object to set the metadata for
+	 * @Values The metadata key/value pairs
+	 */
+	void SetObjectValues(const UObject* Object, TMap<FName, FString>&& Values);
+
 	/**
 	 * Set the key/value pair in the Object's metadata
 	 * @param Object the object to set the metadata for
