Missed file for CL# 13136175
#rnx

diff --git a/Engine/Source/Runtime/CoreUObject/Public/UObject/Class.h b/Engine/Source/Runtime/CoreUObject/Public/UObject/Class.h
--- a/Engine/Source/Runtime/CoreUObject/Public/UObject/Class.h
+++ b/Engine/Source/Runtime/CoreUObject/Public/UObject/Class.h
@@ -128,8 +128,17 @@ class COREUOBJECT_API UField : public UObject
 	 * @param Key The key to lookup in the metadata
 	 * @return true if there is a (possibly blank) value associated with this key
 	 */
-	bool HasMetaData(const TCHAR* Key) const;
-	bool HasMetaData(const FName& Key) const;
+	bool HasMetaData(const TCHAR* Key) const { return FindMetaData(Key) != nullptr; }
+	bool HasMetaData(const FName& Key) const { return FindMetaData(Key) != nullptr; }
+
+	/**
+	 * Find the metadata value associated with the key
+	 *
+	 * @param Key The key to lookup in the metadata
+	 * @return The value associated with the key if exists, null otherwise
+	 */
+	const FString* FindMetaData(const TCHAR* Key) const;
+	const FString* FindMetaData(const FName& Key) const;
 
 	/**
 	 * Find the metadata value associated with the key
