Remove stray ;
#rnx

diff --git a/Engine/Source/Runtime/CoreUObject/Public/UObject/MetaData.h b/Engine/Source/Runtime/CoreUObject/Public/UObject/MetaData.h
--- a/Engine/Source/Runtime/CoreUObject/Public/UObject/MetaData.h
+++ b/Engine/Source/Runtime/CoreUObject/Public/UObject/MetaData.h
@@ -66,7 +66,7 @@ public:
 	 * @param Key The key to query for existence
 	 * @return true if found
 	 */
-	bool HasValue(const UObject* Object, FName Key);
+	bool HasValue(const UObject* Object, FName Key)
 	{
 		return FindValue(Object, Key) != nullptr;
 	}
