Fix return of incorrect type
#rnx

diff --git a/Engine/Source/Runtime/CoreUObject/Private/UObject/MetaData.cpp b/Engine/Source/Runtime/CoreUObject/Private/UObject/MetaData.cpp
--- a/Engine/Source/Runtime/CoreUObject/Private/UObject/MetaData.cpp
+++ b/Engine/Source/Runtime/CoreUObject/Private/UObject/MetaData.cpp
@@ -282,16 +282,16 @@ const FString* UMetaData::FindValue(const UObject* Object, FName Key)
 	// every key needs to be valid
 	if (Key == NAME_None)
 	{
-		return false;
+		return nullptr;
 	}
 
 	// look up the existing map if we have it
 	TMap<FName, FString>* ObjectValues = ObjectMetaDataMap.Find(Object);
 
 	// if not, return false
-	if (ObjectValues == NULL)
+	if (ObjectValues == nullptr)
 	{
-		return false;
+		return nullptr;
 	}
 
 	// if we had the map, see if we had the key
