Proper fix to workaround the fact that some *GeneratedClasses mark themselves as CLASS_Native. Token stream generation code will now skip class refernces for classes marked as EInternalObjectFlags::Native which is also used by GC to skip marking objects as unreachable.

diff --git a/Engine/Source/Runtime/CoreUObject/Private/UObject/GarbageCollection.cpp b/Engine/Source/Runtime/CoreUObject/Private/UObject/GarbageCollection.cpp
--- a/Engine/Source/Runtime/CoreUObject/Private/UObject/GarbageCollection.cpp
+++ b/Engine/Source/Runtime/CoreUObject/Private/UObject/GarbageCollection.cpp
@@ -2633,8 +2633,8 @@ void UClass::AssembleReferenceTokenStream(bool bForce)
 
 		{
 			check(ClassAddReferencedObjects != NULL);
-			const bool bKeepOuter = true;// GetFName() != NAME_Package;
-			const bool bKeepClass = true;// !(GetClassFlags() & CLASS_Native) || IsA(UDynamicClass::StaticClass());
+			const bool bKeepOuter = GetFName() != NAME_Package;
+			const bool bKeepClass = !(HasAnyInternalFlags(EInternalObjectFlags::Native) || IsA(UDynamicClass::StaticClass());
 
 			ClassAddReferencedObjectsType AddReferencedObjectsFn = nullptr;
 #if !WITH_EDITOR
