Fixed botched merge

diff --git a/Engine/Source/Runtime/CoreUObject/Private/UObject/GarbageCollection.cpp b/Engine/Source/Runtime/CoreUObject/Private/UObject/GarbageCollection.cpp
--- a/Engine/Source/Runtime/CoreUObject/Private/UObject/GarbageCollection.cpp
+++ b/Engine/Source/Runtime/CoreUObject/Private/UObject/GarbageCollection.cpp
@@ -2634,7 +2634,7 @@ void UClass::AssembleReferenceTokenStream(bool bForce)
 		{
 			check(ClassAddReferencedObjects != NULL);
 			const bool bKeepOuter = GetFName() != NAME_Package;
-			const bool bKeepClass = !(HasAnyInternalFlags(EInternalObjectFlags::Native) || IsA(UDynamicClass::StaticClass());
+			const bool bKeepClass = !HasAnyInternalFlags(EInternalObjectFlags::Native) || IsA(UDynamicClass::StaticClass());
 
 			ClassAddReferencedObjectsType AddReferencedObjectsFn = nullptr;
 #if !WITH_EDITOR
