AsyncLoading2: Replace GlobalImports with PublicExports.
Move RefCount tracking for imports from import object level to package level.
Only track currently loaded public exports, and their packages.
Replace import object weak pointers with public export raw pointers (uobject index).
Add new NotifyUnreachableObjects GC hook to get information about all objects to be destroyed.
When a public export is destroyed, clean up its raw pointer, and mark its package as not loaded.
When a package is destroyed, clean up its ref count tracking data.

diff --git a/Engine/Source/Runtime/CoreUObject/Private/Serialization/AsyncLoading2.cpp b/Engine/Source/Runtime/CoreUObject/Private/Serialization/AsyncLoading2.cpp
--- a/Engine/Source/Runtime/CoreUObject/Private/Serialization/AsyncLoading2.cpp
+++ b/Engine/Source/Runtime/CoreUObject/Private/Serialization/AsyncLoading2.cpp
@@ -310,89 +310,6 @@ struct FAsyncPackageDesc2
 
 using FExportObjects = TArray<FExportObject>;
 
-class FGlobalImport
-{
-private:
-	union
-	{
-		UObject* Object = nullptr;
-		struct
-		{
-			int32 ObjectIndex;
-			int32 SerialNumber;
-		} WeakPointer;
-	};
-
-	TAtomic<int32> RefCount { 0 };
-	bool bIsWeakPointer = false;
-
-public:
-	inline void AddRef()
-	{
-		++RefCount;
-	}
-
-	inline void ReleaseRef()
-	{
-		--RefCount;
-	}
-
-	inline int32 GetRefCount()
-	{
-		return RefCount;
-	}
-
-	UObject* GetObject()
-	{
-		if (bIsWeakPointer)
-		{
-			InternalMakeRaw();
-		}
-		return Object;
-	}
-
-	UObject* GetObjectIfRawPointer()
-	{
-		return bIsWeakPointer ? nullptr : Object;
-	}
-
-	void SetObject(UObject* InObject)
-	{
-#if DO_CHECK
-		if (bIsWeakPointer)
-		{
-			InternalMakeRaw();
-		}
-		check(!Object || Object == InObject);
-#endif
-		Object = InObject;
-		bIsWeakPointer = false;
-	}
-
-	void MakeWeak()
-	{
-		check(RefCount == 0);
-		check(!bIsWeakPointer);
-		int32 ObjectIndex = GUObjectArray.ObjectToIndex((UObjectBase*)Object);
-		int32 SerialNumber = GUObjectArray.AllocateSerialNumber(ObjectIndex);
-		WeakPointer.ObjectIndex = ObjectIndex;
-		WeakPointer.SerialNumber = SerialNumber;
-		bIsWeakPointer = true;
-	}
-
-private:
-	void InternalMakeRaw()
-	{
-		check(bIsWeakPointer);
-		const bool bEvenIfPendingKill = false;
-		FUObjectItem* ObjectItem = GUObjectArray.IndexToValidObject(WeakPointer.ObjectIndex, bEvenIfPendingKill);
-		int32 ActualSerialNumber = GUObjectArray.GetSerialNumber(WeakPointer.ObjectIndex);
-		Object = (ObjectItem && ActualSerialNumber == WeakPointer.SerialNumber) ? (UObject*)ObjectItem->Object : nullptr;
-		check(!Object || (!Object->IsUnreachable() && !Object->IsPendingKill()));
-		bIsWeakPointer = false;
-	}
-};
-
 class FNameMap
 {
 public:
@@ -451,10 +368,17 @@ private:
 	FMappedName::EType NameMapType = FMappedName::EType::Global;
 };
 
+struct FPublicExportRef
+{
+	int32 ObjectIndex = 0;
+	FPackageId PackageId; // for fast clear of package is loaded during GC
+};
+
 struct FGlobalImportStore
 {
 	TArray<UObject*> ScriptObjects;
-	TArray<FGlobalImport> PublicExports;
+	TMap<FPackageObjectIndex, FPublicExportRef> PublicExportToObjectIndex;
+	TMap<int32, FPackageObjectIndex> ObjectIndexToPublicExport;
 	TArray<FScriptObjectEntry> ScriptObjectEntries;
 
 	inline FName GetName(FPackageObjectIndex GlobalIndex)
@@ -466,11 +390,34 @@ struct FGlobalImportStore
 			: NAME_None;
 	}
 
-	inline FGlobalImport& GetImport(FPackageObjectIndex GlobalIndex)
+	FPackageId RemovePublicExport(UObject* InObject)
+	{
+		FPackageId PackageId;
+		int32 ObjectIndex = GUObjectArray.ObjectToIndex(InObject);
+		FPackageObjectIndex* GlobalIndex = ObjectIndexToPublicExport.Find(ObjectIndex);
+		if (GlobalIndex)
+		{
+			PackageId = PublicExportToObjectIndex.FindRef(*GlobalIndex).PackageId;
+			PublicExportToObjectIndex.Remove(*GlobalIndex);
+			ObjectIndexToPublicExport.Remove(ObjectIndex);
+		}
+		return PackageId;
+	}
+
+	inline UObject* GetPublicExportObject(FPackageObjectIndex GlobalIndex)
 	{
 		check(GlobalIndex.IsPackageImport());
-		FGlobalImport& Import = PublicExports[GlobalIndex.GetIndex()];
-		return Import;
+		FPublicExportRef* PublicExport = PublicExportToObjectIndex.Find(GlobalIndex);
+		if (!PublicExport)
+		{
+			return nullptr;
+		}
+
+		FUObjectItem* ObjectItem = GUObjectArray.IndexToObject(PublicExport->ObjectIndex);
+		check(ObjectItem);
+		UObject* Object = static_cast<UObject*>(ObjectItem->Object);
+		check(Object && !Object->IsUnreachable());
+		return Object;
 	}
 
 	inline UObject* GetImportObject(FPackageObjectIndex GlobalIndex)
@@ -479,12 +426,7 @@ struct FGlobalImportStore
 		{
 			return ScriptObjects[GlobalIndex.GetIndex()];
 		}
-		if (GlobalIndex.IsPackageImport())
-		{
-			return PublicExports[GlobalIndex.GetIndex()].GetObject();
-		}
-		check(false);
-		return nullptr;
+		return GetPublicExportObject(GlobalIndex);
 	}
 
 	UObject* FindScriptImportObjectFromIndex(int32 ScriptImportIndex);
@@ -525,7 +467,7 @@ struct FGlobalImportStore
 		return Object;
 	}
 
-	void StoreGlobalObject(FPackageObjectIndex GlobalIndex, UObject* InObject)
+	void StoreGlobalObject(FPackageId PackageId, FPackageObjectIndex GlobalIndex, UObject* InObject)
 	{
 		if (GlobalIndex.IsScriptImport())
 		{
@@ -535,8 +477,9 @@ struct FGlobalImportStore
 		}
 		else if (GlobalIndex.IsPackageImport())
 		{
-			FGlobalImport& Import = PublicExports[GlobalIndex.GetIndex()];
-			Import.SetObject(InObject);
+			int32 ObjectIndex = GUObjectArray.ObjectToIndex(InObject);
+			PublicExportToObjectIndex.Add(GlobalIndex, {ObjectIndex, PackageId});
+			ObjectIndexToPublicExport.Add(ObjectIndex, GlobalIndex);
 		}
 		else
 		{
@@ -544,14 +487,106 @@ struct FGlobalImportStore
 		}
 	}
 
-	/** Reference tracking for GC management */
-	TArray<UObject*> KeepAliveObjects;
-	bool bNeedToHandleGarbageCollect = false;
+	void FindAllScriptObjects();
+};
 
-	void OnPreGarbageCollect(bool bInIsLoadingPackages);
-	void OnPostGarbageCollect();
+class FLoadedPackageStore
+{
+private:
+	TMap<FPackageId, int32> RefCounts;
+	TBitArray<>	IsLoaded;
 
-	void FindAllScriptObjects();
+public:
+	void Init(int32 PackageCount)
+	{
+		IsLoaded.Init(false, PackageCount);
+	}
+
+	int32 NumTracked()
+	{
+		return RefCounts.Num();
+	}
+
+	inline void Remove(FPackageId PackageId)
+	{
+#if DO_CHECK
+		int32* RefCount = RefCounts.Find(PackageId);
+		check(RefCount && *RefCount == 0);
+		ClearAllPublicExportsLoaded(PackageId);
+#endif
+		RefCounts.Remove(PackageId);
+	}
+
+	inline bool AddRef(FPackageId PackageId)
+	{
+		int32& RefCount = RefCounts.FindOrAdd(PackageId);
+		++RefCount;
+		
+		// is this the first new reference to a fully loaded package?
+		return RefCount == 1 && IsLoaded[PackageId.ToIndex()];
+	}
+
+	inline bool ReleaseRef(FPackageId PackageId)
+	{
+		int32* RefCountPtr = RefCounts.Find(PackageId);
+		check(RefCountPtr);
+		int32& RefCount = *RefCountPtr;
+		check(RefCount > 0);
+		--RefCount;
+
+		check(IsLoaded[PackageId.ToIndex()]);
+		// is this the last reference to a package?
+		return RefCount == 0;
+	}
+
+	inline int32 GetRefCount(FPackageId PackageId)
+	{
+		int32* RefCount = RefCounts.Find(PackageId);
+		return RefCount ? *RefCount : 0;
+	}
+
+	inline bool AreAllPublicExportsLoaded(FPackageId PackageId)
+	{
+		return IsLoaded[PackageId.ToIndex()];
+	}
+
+	inline void SetAllPublicExportsLoaded(FPackageId PackageId)
+	{
+		IsLoaded[PackageId.ToIndex()] = true;
+	}
+
+	inline void ClearAllPublicExportsLoaded(FPackageId PackageId)
+	{
+		IsLoaded[PackageId.ToIndex()] = false;
+	}
+
+#if ALT2_VERIFY_ASYNC_FLAGS
+	void VerifyLoadedPackages()
+	{
+		for (TPair<FPackageId, int32>& Pair : RefCounts)
+		{
+			FPackageId& PackageId = Pair.Key;
+			int32& RefCount = Pair.Value;
+			UE_CLOG(RefCount > 0, LogStreaming, Warning,
+				TEXT("PackageId '%d' with ref count %d should not have a ref count now")
+				TEXT(", or this check is incorrectly reached during active loading."),
+				PackageId.ToIndex(),
+				RefCount);
+		}
+		for (int32 I = 0; I < IsLoaded.Num(); ++I)
+		{
+			if (IsLoaded[I])
+			{
+				FPackageId PackageId = FPackageId::FromIndex(I);
+				UE_CLOG(!RefCounts.Contains(PackageId), LogStreaming, Warning,
+					TEXT("PackageId '%d' is marked as loaded, then it should also exist with RefCount 0 in the map")
+					TEXT(", or this check is incorrectly reached during active loading."),
+					PackageId.ToIndex());
+				;
+			}
+		}
+	}
+#endif
 };
 
 class FPackageStore
@@ -579,6 +614,7 @@ public:
 	FSourceToLocalizedPackageIdMap* LocalizedPackageIdMap = nullptr;
 
 	FGlobalImportStore ImportStore;
+	FLoadedPackageStore LoadedPackageStore;
 	FPackageStoreEntry* StoreEntries = nullptr;
 	int32 PackageCount = 0;
 	int32 ScriptArcsCount = 0;
@@ -622,9 +658,9 @@ public:
 			TRACE_CPUPROFILER_EVENT_SCOPE(LoadPackageStoreTocFixup);
 
 			int32 PackageByteCount = 0;
-			int32 PackageImportCount = 0;
+			int32 PublicExportCount = 0;
 			GlobalMetaAr << PackageCount;
-			GlobalMetaAr << PackageImportCount;
+			GlobalMetaAr << PublicExportCount;
 			GlobalMetaAr << PackageByteCount;
 
 			StoreEntries = reinterpret_cast<FPackageStoreEntry*>(FMemory::Malloc(PackageByteCount));
@@ -632,7 +668,9 @@ public:
 			// In-place loading
 			GlobalMetaAr.Serialize(StoreEntries, PackageByteCount);
 
-			ImportStore.PublicExports.AddDefaulted(PackageImportCount);
+			LoadedPackageStore.Init(PackageCount);
+			ImportStore.PublicExportToObjectIndex.Reserve(FMath::Min(32768, PublicExportCount));
+			ImportStore.ObjectIndexToPublicExport.Reserve(FMath::Min(32768, PublicExportCount));
 
 			// add 10% slack for temp package names
 			{
@@ -880,6 +918,23 @@ public:
 		return ImportStore;
 	}
 
+	void RemovePackage(UPackage* Package)
+	{
+		FPackageId PackageId = FindPackageId(Package->GetFName());
+		check(PackageId.IsValid());
+		LoadedPackageStore.Remove(PackageId);
+	}
+
+	void RemovePublicExport(UObject* Object)
+	{
+		FPackageId PackageId = ImportStore.RemovePublicExport(Object);
+		if (PackageId.IsValid())
+		{
+			check(IsSerializedPackageId(PackageId));
+			LoadedPackageStore.ClearAllPublicExportsLoaded(PackageId);
+		}
+	}
+
 	inline FPackageId FindPackageId(FName Name)
 	{
 		FScopeLock Lock(&PackageNameToPackageIdCritical);
@@ -900,6 +955,11 @@ public:
 		return NewId;
 	}
 
+	inline bool IsSerializedPackageId(FPackageId PackageId)
+	{
+		return PackageId.ToIndex() < (uint32)PackageCount;
+	}
+
 	inline FPackageStoreEntry* GetGlobalStoreEntries()
 	{
 		return StoreEntries;
@@ -923,24 +983,23 @@ struct FPackageImportStore
 	const FPackageStoreEntry& StoreEntry;
 	const FPackageObjectIndex* ImportMap = nullptr;
 	int32 ImportMapCount = 0;
+	FPackageId PackageId;
 
-	FPackageImportStore(FPackageStore& InGlobalPackageStore, const FPackageStoreEntry& InStoreEntry)
+	FPackageImportStore(FPackageStore& InGlobalPackageStore, const FPackageStoreEntry& InStoreEntry, FPackageId InPackageId)
 		: GlobalPackageStore(InGlobalPackageStore)
 		, GlobalImportStore(GlobalPackageStore.ImportStore)
 		, StoreEntry(InStoreEntry)
+		, PackageId(InPackageId)
 	{
 		TRACE_CPUPROFILER_EVENT_SCOPE(NewPackageImportStore);
-		AddGlobalImportObjectReferences();
+		AddPackageReferences();
 	}
 
 	~FPackageImportStore()
 	{
+		TRACE_CPUPROFILER_EVENT_SCOPE(DeletePackageImportStore);
 		check(!ImportMap);
-	}
-
-	inline UObject* GetImportObject(FPackageObjectIndex GlobalIndex)
-	{
-		return GlobalImportStore.GetImportObject(GlobalIndex);
+		ReleasePackageReferences();
 	}
 
 	inline bool IsValidLocalImportIndex(FPackageIndex LocalIndex)
@@ -956,8 +1015,16 @@ struct FPackageImportStore
 		const int32 LocalImportIndex = LocalIndex.ToImport();
 		check(LocalImportIndex < ImportMapCount);
 		const FPackageObjectIndex GlobalIndex = ImportMap[LocalIndex.ToImport()];
-		check(GlobalIndex.IsImport());
-		return GlobalImportStore.FindOrGetImportObject(GlobalIndex);
+		UObject* Object = nullptr;
+		if (GlobalIndex.IsImport())
+		{
+			Object = GlobalImportStore.FindOrGetImportObject(GlobalIndex);
+		}
+		else
+		{
+			check(GlobalIndex.IsNull());
+		}
+		return Object;
 	}
 
 	inline UObject* FindOrGetImportObject(FPackageObjectIndex GlobalIndex)
@@ -979,48 +1046,72 @@ struct FPackageImportStore
 		return FName();
 	}
 
-	void StoreGlobalObject(FPackageObjectIndex GlobalIndex, UObject* InObject)
+	void StoreGlobalObject(FPackageId InPackageId, FPackageObjectIndex InGlobalIndex, UObject* InObject)
 	{
-		GlobalImportStore.StoreGlobalObject(GlobalIndex, InObject);
+		GlobalImportStore.StoreGlobalObject(InPackageId, InGlobalIndex, InObject);
 	}
 
 	void ClearReferences()
 	{
-		ReleaseGlobalImportObjectReferences();
+		ReleasePackageReferences();
 	}
 
 private:
-	void AddGlobalImportObjectReferences()
+	void AddAsyncFlags(FPackageId InPackageId)
+	{
+		const FPackageStoreEntry& ImportEntry = GlobalPackageStore.GetStoreEntry(InPackageId);
+		for (const FPackageObjectIndex& GlobalIndex : ImportEntry.PublicExports)
+		{
+			UObject* Object = GlobalImportStore.GetPublicExportObject(GlobalIndex);
+			check(Object && !Object->HasAnyInternalFlags(EInternalObjectFlags::Async));
+			Object->SetInternalFlags(EInternalObjectFlags::Async);
+		}
+	}
+
+	void ClearAsyncFlags(FPackageId InPackageId)
+	{
+		const FPackageStoreEntry& ImportEntry = GlobalPackageStore.GetStoreEntry(InPackageId);
+		for (const FPackageObjectIndex& GlobalIndex : ImportEntry.PublicExports)
+		{
+			UObject* Object = GlobalImportStore.GetPublicExportObject(GlobalIndex);
+			check(Object && Object->HasAnyInternalFlags(EInternalObjectFlags::Async));
+			Object->AtomicallyClearInternalFlags(EInternalObjectFlags::Async);
+		}
+	}
+
+	void AddPackageReferences()
 	{
-		for (const FPackageId& GlobalPackageId : StoreEntry.ImportedPackages)
+		for (const FPackageId& ImportedPackageId : StoreEntry.ImportedPackages)
 		{
-			const FPackageStoreEntry& ImportEntry = GlobalPackageStore.GetStoreEntry(GlobalPackageId);
-			for (const FPackageObjectIndex& GlobalImportIndex : ImportEntry.PublicExports)
+			if (GlobalPackageStore.LoadedPackageStore.AddRef(ImportedPackageId))
 			{
-				GlobalImportStore.GetImport(GlobalImportIndex).AddRef();
+				AddAsyncFlags(ImportedPackageId);
 			}
 		}
-		// Add ref counts to own imports to speed up OnPreGarbageCollect
-		for (const FPackageObjectIndex& GlobalImportIndex : StoreEntry.PublicExports)
+		if (GlobalPackageStore.LoadedPackageStore.AddRef(PackageId))
 		{
-			GlobalImportStore.GetImport(GlobalImportIndex).AddRef();
+			// should only happen if someone from outside call LoadPackage with an already loaded package
+			// this could be detected already in CreatePackagesFromQueue,
+			// but requires:
+			// - queuing up package callbacks
+			// - handling request ids properly
+			// - calling AddAsyncFlags(PackageId); (now this is done from create/serialize in the async package)
 		}
 	}
 
-	void ReleaseGlobalImportObjectReferences()
+	void ReleasePackageReferences()
 	{
-		for (const FPackageId& GlobalPackageId : StoreEntry.ImportedPackages)
+		for (const FPackageId& ImportedPackageId : StoreEntry.ImportedPackages)
 		{
-			const FPackageStoreEntry& ImportEntry = GlobalPackageStore.GetStoreEntry(GlobalPackageId);
-			for (const FPackageObjectIndex& GlobalImportIndex : ImportEntry.PublicExports)
+			if (GlobalPackageStore.LoadedPackageStore.ReleaseRef(ImportedPackageId))
 			{
-				GlobalImportStore.GetImport(GlobalImportIndex).ReleaseRef();
+				ClearAsyncFlags(ImportedPackageId);
 			}
 		}
-		// Release ref counts to own imports to speed up OnPreGarbageCollect
-		for (const FPackageObjectIndex& GlobalImportIndex : StoreEntry.PublicExports)
+		// clear own reference, and possible all async flags if no remaining ref count
+		if (GlobalPackageStore.LoadedPackageStore.ReleaseRef(PackageId))
 		{
-			GlobalImportStore.GetImport(GlobalImportIndex).ReleaseRef();
+			ClearAsyncFlags(PackageId);
 		}
 	}
 };
@@ -1608,7 +1699,7 @@ struct FAsyncPackage2
 		AddOwnedObject(Object, bForceAdd);
 		if (bForceAdd || IsInGameThread())
 		{
-			check(!Object->HasAnyInternalFlags(EInternalObjectFlags::Async));
+			check(bForceAdd || !Object->HasAnyInternalFlags(EInternalObjectFlags::Async));
 			Object->SetInternalFlags(EInternalObjectFlags::Async);
 		}
 		check(Object->HasAnyInternalFlags(EInternalObjectFlags::Async));
@@ -1989,10 +2080,6 @@ public:
 	/** Start the async loading thread */
 	virtual void StartThread() override;
 
-	/** [GC] Management of global import objects */
-	void OnPreGarbageCollect();
-	void OnPostGarbageCollect();
-
 	/** [EDL] Event queue */
 	FZenaphore AltZenaphore;
 	TArray<FZenaphore> WorkerZenaphores;
@@ -2074,6 +2161,8 @@ public:
 
 	virtual void NotifyConstructedDuringAsyncLoading(UObject* Object, bool bSubObject) override;
 
+	virtual void NotifyUnreachableObjects(const TArrayView<FUObjectItem*>& UnreachableObjects) override;
+
 	virtual void FireCompletedCompiledInImport(void* AsyncPackage, FPackageIndex Import) override;
 
 	/**
@@ -2252,7 +2341,7 @@ private:
 
 	FAsyncPackage2* CreateAsyncPackage(const FAsyncPackageDesc2& Desc)
 	{
-		if (Desc.PackageIdToLoad.ToIndex() < (uint32)GlobalPackageStore.PackageCount)
+		if (GlobalPackageStore.IsSerializedPackageId(Desc.PackageIdToLoad))
 		{
 			return new FAsyncPackage2(Desc, *this, GraphAllocator, EventSpecs.GetData());
 		}
@@ -2337,9 +2426,8 @@ void FAsyncLoadingThread2::InitializeLoading()
 
 	AsyncThreadReady.Increment();
 
-	UE_LOG(LogStreaming, Display, TEXT("AsyncLoading2 - Initialized: Packages: %d, PublicExports: %d, ScriptObjects %d, FNames: %d"),
+	UE_LOG(LogStreaming, Display, TEXT("AsyncLoading2 - Initialized: Packages: %d, Script objects %d, FNames: %d"),
 		GlobalPackageStore.PackageCount,
-		GlobalPackageStore.ImportStore.PublicExports.Num(),
 		GlobalPackageStore.ImportStore.ScriptObjects.Num(),
 		GlobalNameMap.GetNameEntries().Num());
 }
@@ -2935,7 +3023,7 @@ void FGlobalImportStore::FindAllScriptObjects()
 			}
 			else
 			{
-				UE_LOG(LogStreaming, Warning,
+				UE_LOG(LogStreaming, Log,
 					TEXT("AsyncLoading2 - Failed to find import script object after initial load: %s (%d) in outer %s (%d)"),
 					*MinimalNameToName(ScriptObjectEntries[GlobalImportIndex].ObjectName).ToString(),
 					GlobalImportIndex,
@@ -2956,74 +3044,64 @@ void FAsyncPackage2::ImportPackagesRecursive()
 	}
 	bHasImportedPackagesRecursive = true;
 
-	FPackageStore& GlobalPackageStore = AsyncLoadingThread.GlobalPackageStore;
-	int32 ImportedPackageCount = StoreEntry.ImportedPackages.Num();
+	const int32 ImportedPackageCount = StoreEntry.ImportedPackages.Num();
 	if (!ImportedPackageCount)
 	{
 		return;
 	}
 
+	FPackageStore& GlobalPackageStore = AsyncLoadingThread.GlobalPackageStore;
 	for (const FPackageId& ImportedPackageId : StoreEntry.ImportedPackages)
 	{
-		// AreAllImportsInImportPackageAlreadyLoaded?
-		const FPackageStoreEntry& ImportEntry = GlobalPackageStore.GetStoreEntry(ImportedPackageId);
+		if (GlobalPackageStore.LoadedPackageStore.AreAllPublicExportsLoaded(ImportedPackageId))
+		{
+			continue;
+		}
+
+		const FPackageStoreEntry& ImportedPackageEntry = GlobalPackageStore.GetStoreEntry(ImportedPackageId);
+		const FPackageId PackageIdToLoad = ImportedPackageId;
 #if DO_CHECK
-		UE_CLOG(!FMappedName::IsResolvedToMinimalName(ImportEntry.Name), LogStreaming, Warning,
+		UE_CLOG(!FMappedName::IsResolvedToMinimalName(ImportedPackageEntry.Name), LogStreaming, Warning,
 			TEXT("Package '%s' is trying to import non mounted package ID '%d'"),
 			*Desc.Name.ToString(),
 			ImportedPackageId.ToIndex());
 #endif
-		bool bNeedToLoadPackage = false;
-		for (const FPackageObjectIndex& GlobalImportIndex : ImportEntry.PublicExports)
-		{
-			UObject* ImportedObject = ImportStore.GetImportObject(GlobalImportIndex);
-			if (!ImportedObject || !ImportedObject->HasAllFlags(RF_WasLoaded | RF_LoadCompleted))
+		const FName NameToLoad = MinimalNameToName(ImportedPackageEntry.Name);
+		const FName SourceName =
+			ImportedPackageEntry.SourcePackageId.IsValid() ?
+			MinimalNameToName(GlobalPackageStore.StoreEntries[ImportedPackageEntry.SourcePackageId.ToIndex()].Name) :
+			NameToLoad;
+
+		FAsyncPackageDesc2 PackageDesc(INDEX_NONE, PackageIdToLoad, PackageIdToLoad, SourceName, NameToLoad);
+		bool bInserted;
+		FAsyncPackage2* ImportedPackage = AsyncLoadingThread.FindOrInsertPackage(&PackageDesc, bInserted);
+		if (ImportedPackage)
+		{
+			TRACE_LOADTIME_ASYNC_PACKAGE_IMPORT_DEPENDENCY(this, ImportedPackage);
+			if (bInserted)
 			{
-				bNeedToLoadPackage = true;
-				break;
+				UE_ASYNC_PACKAGE_LOG(Verbose, PackageDesc, TEXT("ImportPackages: AddPackage"),
+					TEXT("Start loading imported package."));
 			}
-		}
-
-		if (bNeedToLoadPackage)
-		{
-			const FPackageId PackageIdToLoad = ImportedPackageId;
-			const FName NameToLoad = MinimalNameToName(ImportEntry.Name);
-			const FName SourceName =
-				ImportEntry.SourcePackageId.IsValid() ?
-				MinimalNameToName(GlobalPackageStore.StoreEntries[ImportEntry.SourcePackageId.ToIndex()].Name) :
-				NameToLoad;
-
-			FAsyncPackageDesc2 PackageDesc(INDEX_NONE, PackageIdToLoad, PackageIdToLoad, SourceName, NameToLoad);
-			bool bInserted;
-			FAsyncPackage2* ImportedPackage = AsyncLoadingThread.FindOrInsertPackage(&PackageDesc, bInserted);
-			if (ImportedPackage)
+			else
 			{
-				TRACE_LOADTIME_ASYNC_PACKAGE_IMPORT_DEPENDENCY(this, ImportedPackage);
-				if (bInserted)
-				{
-					UE_ASYNC_PACKAGE_LOG(Verbose, PackageDesc, TEXT("ImportPackages: AddPackage"),
-						TEXT("Start loading imported package."));
-				}
-				else
-				{
-					UE_ASYNC_PACKAGE_LOG_VERBOSE(VeryVerbose, PackageDesc, TEXT("ImportPackages: UpdatePackage"),
-						TEXT("Imported package is already being loaded."));
-				}
-				ImportedPackage->AddRef();
-				ImportedAsyncPackages.Reserve(ImportedPackageCount);
-				ImportedAsyncPackages.Add(ImportedPackage);
-				if (bInserted)
-				{
-					ImportedPackage->ImportPackagesRecursive();
-					ImportedPackage->StartLoading();
-				}
+				UE_ASYNC_PACKAGE_LOG_VERBOSE(VeryVerbose, PackageDesc, TEXT("ImportPackages: UpdatePackage"),
+					TEXT("Imported package is already being loaded."));
 			}
-			else
+			ImportedPackage->AddRef();
+			ImportedAsyncPackages.Reserve(ImportedPackageCount);
+			ImportedAsyncPackages.Add(ImportedPackage);
+			if (bInserted)
 			{
-				UE_ASYNC_PACKAGE_LOG(Error, PackageDesc, TEXT("ImportPackages: SkipPackage"),
-					TEXT("Skipping unknown imported package, but this should not happen here"));
+				ImportedPackage->ImportPackagesRecursive();
+				ImportedPackage->StartLoading();
 			}
 		}
+		else
+		{
+			UE_ASYNC_PACKAGE_LOG(Error, PackageDesc, TEXT("ImportPackages: SkipPackage"),
+				TEXT("Skipping unknown imported package, but this should not happen here"));
+		}
 	}
 	UE_ASYNC_PACKAGE_LOG_VERBOSE(VeryVerbose, Desc, TEXT("ImportPackages: ImportsDone"),
 		TEXT("All imported packages are now being loaded."));
@@ -3439,7 +3517,11 @@ void FAsyncPackage2::EventDrivenCreateExport(int32 LocalExportIndex)
 	}
 
 	check(Object);
-	ImportStore.StoreGlobalObject(Export.GlobalImportIndex, Object);
+	// don't store temp packages, nothing is importing them
+	if (AsyncLoadingThread.GlobalPackageStore.IsSerializedPackageId(Desc.PackageId))
+	{
+		ImportStore.StoreGlobalObject(Desc.PackageId, Export.GlobalImportIndex, Object);
+	}
 }
 
 void FAsyncPackage2::EventDrivenSerializeExport(int32 LocalExportIndex, FExportArchive& Ar)
@@ -3521,6 +3603,13 @@ EAsyncPackageState::Type FAsyncPackage2::Event_ExportsDone(FAsyncPackage2* Packa
 {
 	TRACE_CPUPROFILER_EVENT_SCOPE(Event_ExportsDone);
 
+	// don't store temp packages, nothing is importing them
+	if (Package->AsyncLoadingThread.GlobalPackageStore.IsSerializedPackageId(Package->Desc.PackageId))
+	{
+		// for serialized packages, store the package to load, since that is what is imported
+		Package->AsyncLoadingThread.GlobalPackageStore.LoadedPackageStore.SetAllPublicExportsLoaded(Package->Desc.PackageIdToLoad);
+	}
+
 	Package->GetNode(EEventLoadNode2::Package_PostLoad)->ReleaseBarrier();
 	return EAsyncPackageState::Complete;
 }
@@ -3583,6 +3672,7 @@ EAsyncPackageState::Type FAsyncPackage2::Event_PostLoad(FAsyncPackage2* Package,
 EAsyncPackageState::Type FAsyncPackage2::Event_Delete(FAsyncPackage2* Package, int32)
 {
 	TRACE_CPUPROFILER_EVENT_SCOPE(Event_Delete);
+	UE_ASYNC_PACKAGE_LOG(Verbose, Package->Desc, TEXT("AsyncThread: Deleted"), TEXT("Package deleted."));
 	delete Package;
 	return EAsyncPackageState::Complete;
 }
@@ -3656,8 +3746,6 @@ EAsyncPackageState::Type FAsyncLoadingThread2::ProcessAsyncLoadingFromGameThread
 	{
 		do 
 		{
-			GlobalPackageStore.GetGlobalImportStore().bNeedToHandleGarbageCollect |= IsAsyncLoadingPackages();
-
 			ThreadState.ProcessDeferredFrees();
 
 			if (bNeedsHeartbeatTick && (++LoopIterations) % 32 == 31)
@@ -4047,9 +4135,6 @@ void FAsyncLoadingThread2::StartThread()
 	// Make sure the GC sync object is created before we start the thread (apparently this can happen before we call InitUObject())
 	FGCCSyncObject::Create();
 
-	FCoreUObjectDelegates::GetPreGarbageCollectDelegate().AddRaw(this, &FAsyncLoadingThread2::OnPreGarbageCollect);
-	FCoreUObjectDelegates::GetPostGarbageCollect().AddRaw(this, &FAsyncLoadingThread2::OnPostGarbageCollect);
-
 	if (!FAsyncLoadingThreadSettings::Get().bAsyncLoadingThreadEnabled)
 	{
 		GlobalPackageStore.FinalizeInitialLoad();
@@ -4170,7 +4255,6 @@ uint32 FAsyncLoadingThread2::Run()
 		}
 		else
 		{
-			GlobalPackageStore.GetGlobalImportStore().bNeedToHandleGarbageCollect |= IsAsyncLoadingPackages();
 			bool bDidSomething = false;
 			{
 				FGCScopeGuard GCGuard;
@@ -4409,94 +4493,56 @@ static void VerifyLoadFlagsWhenFinishedLoading()
 	UE_LOG(LogStreaming, Log, TEXT("Verified load flags when finished active loading."));
 }
 
-void FGlobalImportStore::OnPreGarbageCollect(bool bInIsLoadingPackages)
+void FAsyncLoadingThread2::NotifyUnreachableObjects(const TArrayView<FUObjectItem*>& UnreachableObjects)
 {
-	if (!bNeedToHandleGarbageCollect && !bInIsLoadingPackages)
+	TRACE_CPUPROFILER_EVENT_SCOPE(NotifyUnreachableObjects);
+
+	if (GExitPurge)
 	{
 		return;
 	}
-	bNeedToHandleGarbageCollect = bInIsLoadingPackages;
 
-	int32 NumWeak = 0;
-	for (FGlobalImport& GlobalImport : PublicExports)
-	{
-		UObject* Object = GlobalImport.GetObjectIfRawPointer();
-		if (!Object)
-		{
-			continue;
-		}
+	const int32 PublicExportCount = GlobalPackageStore.GetGlobalImportStore().PublicExportToObjectIndex.Num();
+	const int32 LoadedPackageCount = GlobalPackageStore.LoadedPackageStore.NumTracked();
+	int32 RemovedExportCount = 0;
+	int32 RemovedPackageCount = 0;
 
-		if (GlobalImport.GetRefCount() > 0)
+	for (FUObjectItem* ObjectItem : UnreachableObjects)
+	{
+		UObject* Object = static_cast<UObject*>(ObjectItem->Object);
+		check(Object);
+		if (Object->HasAllFlags(RF_WasLoaded | RF_Public))
 		{
-			// Import objects in packages currently being loaded already have the Async flag set.
-			// They will never be destroyed during GC, and the object pointers are safe to keep.
-			if (!Object->HasAnyInternalFlags(EInternalObjectFlags::Async))
+			if (Object->GetOuter())
 			{
-				// Mark object to be kept alive during GC
-				Object->SetInternalFlags(EInternalObjectFlags::Async);
-				KeepAliveObjects.Add(Object);
+				TRACE_CPUPROFILER_EVENT_SCOPE(PackageStoreRemovePublicExport);
+				GlobalPackageStore.RemovePublicExport(Object);
+				++RemovedExportCount;
+			}
+			else
+			{
+				TRACE_CPUPROFILER_EVENT_SCOPE(PackageStoreRemovePackage);
+				UPackage* Package = static_cast<UPackage*>(Object);
+				GlobalPackageStore.RemovePackage(Package);
+				++RemovedPackageCount;
 			}
-		}
-		else
-		{
-			// Convert object pointer to weak since object may get destroyed during GC
-			check(!Object->HasAnyInternalFlags(EInternalObjectFlags::Async));
-			GlobalImport.MakeWeak();
-			++NumWeak;
 		}
 	}
 
-	if (!bInIsLoadingPackages)
-	{
-		check(KeepAliveObjects.Num() == 0);
-	}
+	UE_LOG(LogStreaming, Display,
+		TEXT("FAsyncLoadingThread2::NotifyUnreachableObjects - Purged %d serialized public exports and %d serialized packages, ")
+		TEXT("%d -> %d tracked public exports, %d -> %d tracked serialized packages"),
+		RemovedExportCount, RemovedPackageCount,
+		PublicExportCount, GlobalPackageStore.GetGlobalImportStore().PublicExportToObjectIndex.Num(),
+		LoadedPackageCount, GlobalPackageStore.LoadedPackageStore.NumTracked());
 
 #if ALT2_VERIFY_ASYNC_FLAGS
-	if (!bInIsLoadingPackages)
+	if (!IsAsyncLoadingPackages())
 	{
-		for (FGlobalImport& GlobalImport : PublicExports)
-		{
-			check(GlobalImport.GetObjectIfRawPointer() == nullptr);
-			check(GlobalImport.GetRefCount() == 0);
-		}
+		GlobalPackageStore.LoadedPackageStore.VerifyLoadedPackages();
 		VerifyLoadFlagsWhenFinishedLoading();
 	}
 #endif
-
-	UE_LOG(LogStreaming, Display, TEXT("FGlobalImportStore::OnPreGarbageCollect - Marked %d objects to keep, made %d object pointers weak"),
-		KeepAliveObjects.Num(), NumWeak);
-}
-
-void FGlobalImportStore::OnPostGarbageCollect()
-{
-	if (KeepAliveObjects.Num() == 0)
-	{
-		return;
-	}
-	check(bNeedToHandleGarbageCollect);
-
-	for (UObject* Object : KeepAliveObjects)
-	{
-		Object->ClearInternalFlags(EInternalObjectFlags::Async);
-	}
-
-	const int32 UnmarkedCount = KeepAliveObjects.Num();
-	KeepAliveObjects.Reset();
-	UE_LOG(LogStreaming, Log, TEXT("FGlobalImportStore::UpdateGlobalImportsPostGC - Unmarked %d objects"),
-		UnmarkedCount);
-}
-
-void FAsyncLoadingThread2::OnPreGarbageCollect()
-{
-	TRACE_CPUPROFILER_EVENT_SCOPE(AltPreGC);
-	const bool bIsAsyncLoadingPackages = IsAsyncLoadingPackages();
-	GlobalPackageStore.GetGlobalImportStore().OnPreGarbageCollect(bIsAsyncLoadingPackages);
-}
-
-void FAsyncLoadingThread2::OnPostGarbageCollect()
-{
-	TRACE_CPUPROFILER_EVENT_SCOPE(AltPostGC);
-	GlobalPackageStore.GetGlobalImportStore().OnPostGarbageCollect();
 }
 
 /**
@@ -4559,7 +4605,7 @@ FAsyncPackage2::FAsyncPackage2(
 , LoadPercentage(0)
 , AsyncLoadingThread(InAsyncLoadingThread)
 , GraphAllocator(InGraphAllocator)
-, ImportStore(AsyncLoadingThread.GlobalPackageStore, StoreEntry)
+, ImportStore(AsyncLoadingThread.GlobalPackageStore, StoreEntry, Desc.PackageIdToLoad)
 , AsyncPackageLoadingState(EAsyncPackageLoadingState2::NewPackage)
 , bAllExportsSerialized(false)
 {
@@ -4650,29 +4696,42 @@ void FAsyncPackage2::ClearImportedPackages()
 	}
 	ImportedAsyncPackages.Empty();
 #endif
-	ImportStore.ClearReferences();
 }
 
 void FAsyncPackage2::ClearOwnedObjects()
 {
+	TRACE_CPUPROFILER_EVENT_SCOPE(ClearOwnedObjects);
 	for (UObject* Object : OwnedObjects)
 	{
 		const EObjectFlags Flags = Object->GetFlags();
 		const EInternalObjectFlags InternalFlags = Object->GetInternalFlags();
+		const bool bAsyncLoading = !!(InternalFlags & EInternalObjectFlags::AsyncLoading);
+		const bool bAsync = !!(InternalFlags & EInternalObjectFlags::Async);
+		const bool bPrivate = !(Flags & RF_Public);
+		const bool bPackage = !Object->GetOuter();
+
+		// the async flag of all public export objects are handled by FGlobalImportStore
+		const bool bShouldClearAsync = bAsync && (bPackage || bPrivate);
+
 		EInternalObjectFlags InternalFlagsToClear = EInternalObjectFlags::None;
 
 		check(!(Flags & (RF_NeedPostLoad | RF_NeedPostLoadSubobjects)));
-		if (!!(InternalFlags & EInternalObjectFlags::AsyncLoading))
+
+		if (bAsyncLoading)
 		{
 			check(!(Flags & RF_WasLoaded));
 			InternalFlagsToClear |= EInternalObjectFlags::AsyncLoading;
 		}
 
-		if (!!(InternalFlags & EInternalObjectFlags::Async))
+		if (bShouldClearAsync)
 		{
 			InternalFlagsToClear |= EInternalObjectFlags::Async;
 		}
-		Object->ClearInternalFlags(InternalFlagsToClear);
+
+		if (!!InternalFlagsToClear)
+		{
+			Object->ClearInternalFlags(InternalFlagsToClear);
+		}
 	}
 	OwnedObjects.Empty();
 }
diff --git a/Engine/Source/Runtime/CoreUObject/Private/Serialization/AsyncLoadingThread.h b/Engine/Source/Runtime/CoreUObject/Private/Serialization/AsyncLoadingThread.h
--- a/Engine/Source/Runtime/CoreUObject/Private/Serialization/AsyncLoadingThread.h
+++ b/Engine/Source/Runtime/CoreUObject/Private/Serialization/AsyncLoadingThread.h
@@ -282,6 +282,8 @@ public:
 
 	void NotifyConstructedDuringAsyncLoading(UObject* Object, bool bSubObject) override;
 
+	void NotifyUnreachableObjects(const TArrayView<FUObjectItem*>& UnreachableObjects) override {};
+
 	void FireCompletedCompiledInImport(void* AsyncPacakge, FPackageIndex Import) override;
 
 	/** [EDL] Event queue */
diff --git a/Engine/Source/Runtime/CoreUObject/Private/Serialization/AsyncPackageLoader.cpp b/Engine/Source/Runtime/CoreUObject/Private/Serialization/AsyncPackageLoader.cpp
--- a/Engine/Source/Runtime/CoreUObject/Private/Serialization/AsyncPackageLoader.cpp
+++ b/Engine/Source/Runtime/CoreUObject/Private/Serialization/AsyncPackageLoader.cpp
@@ -668,6 +668,11 @@ void NotifyConstructedDuringAsyncLoading(UObject* Object, bool bSubObject)
 	GetAsyncPackageLoader().NotifyConstructedDuringAsyncLoading(Object, bSubObject);
 }
 
+void NotifyUnreachableObjects(const TArrayView<FUObjectItem*>& UnreachableObjects)
+{
+	GetAsyncPackageLoader().NotifyUnreachableObjects(UnreachableObjects);
+}
+
 double GFlushAsyncLoadingTime = 0.0;
 uint32 GFlushAsyncLoadingCount = 0;
 uint32 GSyncLoadCount = 0;
diff --git a/Engine/Source/Runtime/CoreUObject/Private/UObject/GarbageCollection.cpp b/Engine/Source/Runtime/CoreUObject/Private/UObject/GarbageCollection.cpp
--- a/Engine/Source/Runtime/CoreUObject/Private/UObject/GarbageCollection.cpp
+++ b/Engine/Source/Runtime/CoreUObject/Private/UObject/GarbageCollection.cpp
@@ -74,6 +74,11 @@ static const uint32 GNumTokensPerPointer = sizeof(void*) / sizeof(uint32); //-V5
 
 FThreadSafeBool GIsGarbageCollecting(false);
 
+/**
+* Call back into the async loading code to inform of the destruction of serialized objects
+*/
+void NotifyUnreachableObjects(const TArrayView<FUObjectItem*>& UnreachableObjects);
+
 /** Locks all UObject hash tables when performing GC */
 class FGCScopeLock
 {
@@ -1981,6 +1986,7 @@ void CollectGarbageInternal(EObjectFlags KeepFlags, bool bPerformFullPurge)
 			FGCArrayPool::Get().ClearWeakReferences(bPerformFullPurge);
 
 			GatherUnreachableObjects(bForceSingleThreadedGC);
+			NotifyUnreachableObjects(GUnreachableObjects);
 
 			if (bPerformFullPurge || !GIncrementalBeginDestroyEnabled)
 			{
diff --git a/Engine/Source/Runtime/CoreUObject/Public/Serialization/AsyncPackageLoader.h b/Engine/Source/Runtime/CoreUObject/Public/Serialization/AsyncPackageLoader.h
--- a/Engine/Source/Runtime/CoreUObject/Public/Serialization/AsyncPackageLoader.h
+++ b/Engine/Source/Runtime/CoreUObject/Public/Serialization/AsyncPackageLoader.h
@@ -188,6 +188,8 @@ public:
 
 	virtual void NotifyConstructedDuringAsyncLoading(UObject* Object, bool bSubObject) = 0;
 
+	virtual void NotifyUnreachableObjects(const TArrayView<FUObjectItem*>& UnreachableObjects) = 0;
+
 	virtual void FireCompletedCompiledInImport(void* AsyncPackage, FPackageIndex Import) = 0;
 };
 
