Add some more logging when we call FAsyncLoadingThread::FlushLoading, the number of packages currently queued.
ben.woodhouse, serge.bernier

diff --git a/Engine/Source/Runtime/CoreUObject/Private/Serialization/AsyncLoading.cpp b/Engine/Source/Runtime/CoreUObject/Private/Serialization/AsyncLoading.cpp
--- a/Engine/Source/Runtime/CoreUObject/Private/Serialization/AsyncLoading.cpp
+++ b/Engine/Source/Runtime/CoreUObject/Private/Serialization/AsyncLoading.cpp
@@ -6910,16 +6910,15 @@ void FAsyncLoadingThread::FlushLoading(int32 PackageID)
 			static uint64 LastFrameNumber = -1;
 			if (LastFrameNumber != GFrameNumber)
 			{
-				UE_LOG(LogStreaming, Display, TEXT("Flushing async loaders."));
+				UE_LOG(LogStreaming, Display, TEXT("FAsyncLoadingThread::FlushLoading called, %d QueuedPackagesCount, %d ExistingAsyncPackagesCount"), GetQueuedPackagesCount(), GetExistingAsyncPackagesCount());
 				LastFrameNumber = GFrameNumber;
 			}
 			else
 			{
-				UE_LOG(LogStreaming, Log, TEXT("Flushing async loaders."));
+				UE_LOG(LogStreaming, Log, TEXT("FAsyncLoadingThread::FlushLoading called, %d QueuedPackagesCount, %d ExistingAsyncPackagesCount"), GetQueuedPackagesCount(), GetExistingAsyncPackagesCount());
 			}
 		}
 #endif
-
 		double StartTime = FPlatformTime::Seconds();
 
 		// Flush async loaders by not using a time limit. Needed for e.g. garbage collection.
