AsyncLoading2: Don't assume all public exports always exists, some may have been filtered out by game type.

diff --git a/Engine/Source/Runtime/CoreUObject/Private/Serialization/AsyncLoading2.cpp b/Engine/Source/Runtime/CoreUObject/Private/Serialization/AsyncLoading2.cpp
--- a/Engine/Source/Runtime/CoreUObject/Private/Serialization/AsyncLoading2.cpp
+++ b/Engine/Source/Runtime/CoreUObject/Private/Serialization/AsyncLoading2.cpp
@@ -1063,8 +1063,11 @@ private:
 		for (const FPackageObjectIndex& GlobalIndex : ImportEntry.PublicExports)
 		{
 			UObject* Object = GlobalImportStore.GetPublicExportObject(GlobalIndex);
-			check(Object && !Object->HasAnyInternalFlags(EInternalObjectFlags::Async));
-			Object->SetInternalFlags(EInternalObjectFlags::Async);
+			if (Object)
+			{
+				check(!Object->HasAnyInternalFlags(EInternalObjectFlags::Async));
+				Object->SetInternalFlags(EInternalObjectFlags::Async);
+			}
 		}
 	}
 
@@ -1074,8 +1077,11 @@ private:
 		for (const FPackageObjectIndex& GlobalIndex : ImportEntry.PublicExports)
 		{
 			UObject* Object = GlobalImportStore.GetPublicExportObject(GlobalIndex);
-			check(Object && Object->HasAnyInternalFlags(EInternalObjectFlags::Async));
-			Object->AtomicallyClearInternalFlags(EInternalObjectFlags::Async);
+			if (Object)
+			{
+				check(Object->HasAnyInternalFlags(EInternalObjectFlags::Async));
+				Object->AtomicallyClearInternalFlags(EInternalObjectFlags::Async);
+			}
 		}
 	}
 
