- Fixing a TMap comparison in CoreRediects to use order independent comparison, instead of Legacy, which is not very tolerant

diff --git a/Engine/Source/Runtime/CoreUObject/Private/UObject/CoreRedirects.cpp b/Engine/Source/Runtime/CoreUObject/Private/UObject/CoreRedirects.cpp
--- a/Engine/Source/Runtime/CoreUObject/Private/UObject/CoreRedirects.cpp
+++ b/Engine/Source/Runtime/CoreUObject/Private/UObject/CoreRedirects.cpp
@@ -520,7 +520,7 @@ bool FCoreRedirects::RedirectNameAndValues(ECoreRedirectFlags Type, const FCoreR
 				{
 					if (*FoundValueRedirect)
 					{
-						if (!LegacyCompareEqual((*FoundValueRedirect)->ValueChanges, Redirect->ValueChanges))
+						if ((*FoundValueRedirect)->ValueChanges.OrderIndependentCompareEqual(Redirect->ValueChanges) == false)
 						{
 							UE_LOG(LogLinker, Error, TEXT("RedirectNameAndValues(%s) found multiple conflicting value redirects, %s and %s!"), *OldObjectName.ToString(), *(*FoundValueRedirect)->OldName.ToString(), *Redirect->OldName.ToString());
 						}
