Fixed FLinker::GetExportPathName to return the proper object path if the export had an import in its outer chain
Removed the check for the public flag on an object when printing log about a missing class object. this was preventing proper warnings about actor which its BP class asset couldn't be found on disk. added the missing class name.

diff --git a/Engine/Source/Runtime/CoreUObject/Private/UObject/Linker.cpp b/Engine/Source/Runtime/CoreUObject/Private/UObject/Linker.cpp
--- a/Engine/Source/Runtime/CoreUObject/Private/UObject/Linker.cpp
+++ b/Engine/Source/Runtime/CoreUObject/Private/UObject/Linker.cpp
@@ -265,16 +265,18 @@ FString FLinker::GetExportPathName(int32 ExportIndex, const TCHAR* FakeRoot,bool
 	FString Result;
 
 	bool bForcedExport = false;
+	bool bHasOuterImport = false;
 	for ( FPackageIndex LinkerIndex = FPackageIndex::FromExport(ExportIndex); !LinkerIndex.IsNull(); LinkerIndex = ImpExp(LinkerIndex).OuterIndex )
 	{ 
+		bHasOuterImport |= LinkerIndex.IsImport();
 		const FObjectResource& Resource = ImpExp(LinkerIndex);
 
 		// don't append a dot in the first iteration
 		if ( Result.Len() > 0 )
 		{
 			// if this export is not a UPackage but this export's Outer is a UPackage, we need to use subobject notation
-			if ((Resource.OuterIndex.IsNull() || GetExportClassName(Resource.OuterIndex) == NAME_Package)
-			  && GetExportClassName(LinkerIndex) != NAME_Package)
+			if ((Resource.OuterIndex.IsNull() || GetClassName(Resource.OuterIndex) == NAME_Package)
+			  && GetClassName(LinkerIndex) != NAME_Package)
 			{
 				Result = FString(SUBOBJECT_DELIMITER) + Result;
 			}
@@ -287,7 +289,9 @@ FString FLinker::GetExportPathName(int32 ExportIndex, const TCHAR* FakeRoot,bool
 		bForcedExport = bForcedExport || (LinkerIndex.IsExport() ? Exp(LinkerIndex).bForcedExport : false);
 	}
 
-	if ( bForcedExport && FakeRoot == nullptr && bResolveForcedExports )
+	if ((bForcedExport && FakeRoot == nullptr && bResolveForcedExports) ||
+		// if the export we are building the path of has an import in its outer chain, no need to append the LinkerRoot path
+		bHasOuterImport )
 	{
 		// Result already contains the correct path name for this export
 		return Result;
diff --git a/Engine/Source/Runtime/CoreUObject/Private/UObject/LinkerLoad.cpp b/Engine/Source/Runtime/CoreUObject/Private/UObject/LinkerLoad.cpp
--- a/Engine/Source/Runtime/CoreUObject/Private/UObject/LinkerLoad.cpp
+++ b/Engine/Source/Runtime/CoreUObject/Private/UObject/LinkerLoad.cpp
@@ -4286,8 +4286,9 @@ UObject* FLinkerLoad::CreateExport( int32 Index )
 			Export.bExportLoadFailed = true;
 
 			FString OuterName = Export.OuterIndex.IsNull() ? LinkerRoot->GetFullName() : GetFullImpExpName(Export.OuterIndex);
-			UE_CLOG(Export.ObjectFlags & EObjectFlags::RF_Public, LogLinker, Warning, TEXT("Unable to load %s with outer %s because its class does not exist"), *Export.ObjectName.ToString(), *OuterName);
-			return NULL;
+			FString ClassName = GetClassName(Export.ThisIndex).ToString();
+			UE_LOG(LogLinker, Warning, TEXT("Unable to load %s with outer %s because its class (%s) does not exist"), *Export.ObjectName.ToString(), *OuterName, *ClassName);
+			return nullptr;
 		}
 
 		if (Export.DynamicType == FObjectExport::EDynamicType::ClassDefaultObject)
@@ -4301,7 +4302,7 @@ UObject* FLinkerLoad::CreateExport( int32 Index )
 			else
 			{
 				UE_LOG(LogLinker, Warning, TEXT("CreateExport: Failed to create CDO %s because class is not found"), *Export.ObjectName.ToString());
-				return NULL;
+				return nullptr;
 			}
 		}
 
@@ -4309,7 +4310,7 @@ UObject* FLinkerLoad::CreateExport( int32 Index )
 		// NULL (None) active class redirect.
 		if( !LoadClass && Export.ObjectName.IsNone() && Export.ClassIndex.IsNull() && !Export.OldClassName.IsNone() )
 		{
-			return NULL;
+			return nullptr;
 		}
 #endif
 		if( !LoadClass )
@@ -4327,7 +4328,7 @@ UObject* FLinkerLoad::CreateExport( int32 Index )
 			// otherwise, return NULL and let the calling code determine what to do
 			FString OuterName = Export.OuterIndex.IsNull() ? LinkerRoot->GetFullName() : GetFullImpExpName(Export.OuterIndex);
 			UE_LOG(LogLinker, Warning, TEXT("CreateExport: Failed to load Outer for resource because its class is a redirector '%s': %s"), *Export.ObjectName.ToString(), *OuterName);
-			return NULL;
+			return nullptr;
 		}
 
 		check(LoadClass);
