Downgrade object linker assert to a log warning temporarily until we get OFPA name collision sorted out

diff --git a/Engine/Source/Runtime/CoreUObject/Private/UObject/UObjectLinker.cpp b/Engine/Source/Runtime/CoreUObject/Private/UObject/UObjectLinker.cpp
--- a/Engine/Source/Runtime/CoreUObject/Private/UObject/UObjectLinker.cpp
+++ b/Engine/Source/Runtime/CoreUObject/Private/UObject/UObjectLinker.cpp
@@ -108,7 +108,12 @@ void UObject::SetLinker( FLinkerLoad* LinkerLoad, int32 LinkerIndex, bool bShoul
 	// Detach from existing linker.
 	if( Existing.Linker && bShouldDetachExisting )
 	{
-		checkf(!HasAnyFlags(RF_NeedLoad|RF_NeedPostLoad), TEXT("Detaching from existing linker for %s while object %s needs loaded"), *Existing.Linker->GetArchiveName(), *GetFullName());
+		// HACK-FH: temporarily downgrading assert to a warning until OFPA name collision problem are resolved for good.
+		UE_CLOG(HasAnyFlags(RF_NeedLoad|RF_NeedPostLoad), LogUObjectLinker, Warning,
+			TEXT("Detaching from existing linker for %s while object %s needs loading from linker %s."),
+			*Existing.Linker->GetArchiveName(),
+			*GetFullName(),
+			LinkerLoad ? *LinkerLoad->Filename : TEXT(""));
 		check(Existing.Linker->ExportMap[Existing.LinkerIndex].Object!=nullptr);
 		check(Existing.Linker->ExportMap[Existing.LinkerIndex].Object==this);
 		Existing.Linker->ExportMap[Existing.LinkerIndex].ResetObject();
