Modify FObjectInitializer to allow nested subobjects to be overriden
Fix subobject overriding in the case of a required subobject that has been both overriden for class and marked do not create. Previously it was ending up with the default subobject type rather than the valid overriden value.
Create an EngineTest that validates all expectations for subobjects and nested subobject overrides

diff --git a/Engine/Source/Runtime/CoreUObject/Private/UObject/UObjectGlobals.cpp b/Engine/Source/Runtime/CoreUObject/Private/UObject/UObjectGlobals.cpp
--- a/Engine/Source/Runtime/CoreUObject/Private/UObject/UObjectGlobals.cpp
+++ b/Engine/Source/Runtime/CoreUObject/Private/UObject/UObjectGlobals.cpp
@@ -1020,7 +1020,7 @@ public:
 			{
 				// get the calls debug callstack and 
 				FString DebugStackString;
-				for (const auto& DebugData : DebugDataStack)
+				for (const FName& DebugData : DebugDataStack)
 				{
 					DebugStackString += DebugData.ToString();
 					DebugStackString += TEXT("->");
@@ -2590,33 +2590,38 @@ FObjectInitializer::FObjectInitializer()
 	, ObjectArchetype(nullptr)
 	, bCopyTransientsFromClassDefaults(false)
 	, bShouldInitializePropsFromArchetype(false)
-	, bSubobjectClassInitializationAllowed(true)
 	, InstanceGraph(nullptr)
-	, LastConstructedObject(nullptr)
-#if USE_CIRCULAR_DEPENDENCY_LOAD_DEFERRING
-	, bIsDeferredInitializer(false)
-#endif // USE_CIRCULAR_DEPENDENCY_LOAD_DEFERRING
 {
-	FUObjectThreadContext& ThreadContext = FUObjectThreadContext::Get();
-	// Mark we're in the constructor now.	
-	ThreadContext.IsInConstructor++;
-	LastConstructedObject = ThreadContext.ConstructedObject;
-	ThreadContext.ConstructedObject = Obj;
-	ThreadContext.PushInitializer(this);
+	Construct_Internal();
 }	
 
+FObjectInitializer::FObjectInitializer(UObject* InObj, const FStaticConstructObjectParameters& StaticConstructParams)
+	: Obj(InObj)
+	, ObjectArchetype(StaticConstructParams.Template)
+	, bCopyTransientsFromClassDefaults(StaticConstructParams.bCopyTransientsFromClassDefaults)
+	, bShouldInitializePropsFromArchetype(true)
+	, InstanceGraph(StaticConstructParams.InstanceGraph)
+{
+	if (StaticConstructParams.SubobjectOverrides)
+	{
+		SubobjectOverrides = *StaticConstructParams.SubobjectOverrides;
+	}
+
+	Construct_Internal();
+}
+
 FObjectInitializer::FObjectInitializer(UObject* InObj, UObject* InObjectArchetype, bool bInCopyTransientsFromClassDefaults, bool bInShouldInitializeProps, struct FObjectInstancingGraph* InInstanceGraph)
 	: Obj(InObj)
 	, ObjectArchetype(InObjectArchetype)
 	  // if the SubobjectRoot NULL, then we want to copy the transients from the template, otherwise we are doing a duplicate and we want to copy the transients from the class defaults
 	, bCopyTransientsFromClassDefaults(bInCopyTransientsFromClassDefaults)
 	, bShouldInitializePropsFromArchetype(bInShouldInitializeProps)
-	, bSubobjectClassInitializationAllowed(true)
 	, InstanceGraph(InInstanceGraph)
-	, LastConstructedObject(nullptr)
-#if USE_CIRCULAR_DEPENDENCY_LOAD_DEFERRING
-	, bIsDeferredInitializer(false)
-#endif // USE_CIRCULAR_DEPENDENCY_LOAD_DEFERRING
+{
+	Construct_Internal();
+}
+
+void FObjectInitializer::Construct_Internal()
 {
 	FUObjectThreadContext& ThreadContext = FUObjectThreadContext::Get();
 	// Mark we're in the constructor now.
@@ -2921,7 +2926,7 @@ void FObjectInitializer::InstanceSubobjects(UClass* Class, bool bNeedInstancing,
 		UseInstancingGraph->AddNewObject(Obj, ObjectArchetype);
 	}
 	// Add any default subobjects
-	for (auto& SubobjectInit : ComponentInits.SubobjectInits)
+	for (const FSubobjectsToInit::FSubobjectInit& SubobjectInit : ComponentInits.SubobjectInits)
 	{
 		UseInstancingGraph->AddNewObject(SubobjectInit.Subobject, SubobjectInit.Template);
 	}
@@ -2936,7 +2941,7 @@ void FObjectInitializer::InstanceSubobjects(UClass* Class, bool bNeedInstancing,
 		for (int32 Index = 0; Index < ComponentInits.SubobjectInits.Num(); Index++)
 		{
 			UObject* Subobject = ComponentInits.SubobjectInits[Index].Subobject;
-			UObject* Template = ComponentInits	.SubobjectInits[Index].Template;
+			UObject* Template = ComponentInits.SubobjectInits[Index].Template;
 
 #if USE_CIRCULAR_DEPENDENCY_LOAD_DEFERRING
 			if ( !Subobject->HasAnyFlags(RF_NeedLoad) || bIsDeferredInitializer )
@@ -3060,51 +3065,161 @@ void FObjectInitializer::InitProperties(UObject* Obj, UClass* DefaultsClass, UOb
 	}
 }
 
-/**  Add an override, make sure it is legal **/
-void FObjectInitializer::FOverrides::Add(FName InComponentName, UClass* InComponentClass, FObjectInitializer const& ObjectInitializer)
+void FObjectInitializer::FOverrides::Add(FName InComponentName, UClass* InComponentClass, const TArrayView<const FName>* InFullComponentPath)
 {
+	auto GetSubobjectPath = [InComponentName, InFullComponentPath]()
+	{
+		if (InFullComponentPath)
+		{
+			FString SubobjectPath;
+			for (FName SubobjectName : *InFullComponentPath)
+			{
+				SubobjectPath += (SubobjectPath.IsEmpty() ? TEXT("") : TEXT("."));
+				SubobjectPath += SubobjectName.ToString();
+			}
+			return SubobjectPath;
+		}
+
+		return InComponentName.ToString();
+	};
+
 	const int32 Index = Find(InComponentName);
 	if (Index == INDEX_NONE)
 	{
-		Overrides.Emplace(FOverride(InComponentName, InComponentClass));
+		FOverride& Override = Overrides.Emplace_GetRef(FOverride(InComponentName));
+		Override.ComponentClass = InComponentClass;
+		Override.bDoNotCreate = (InComponentClass == nullptr);
 	}
-	else if (InComponentClass && Overrides[Index].ComponentClass)
+	else if (InComponentClass)
 	{
-		// if a base class is asking for an override, the existing override (which we are going to use) had better be derived
-		if (!IsLegalOverride(Overrides[Index].ComponentClass, InComponentClass))
+		if (Overrides[Index].ComponentClass)
 		{
-			UE_LOG(LogUObjectGlobals, Error, TEXT("%s is not a legal override for component %s because it does not derive from %s. Will use %s when constructing component."),
-				*Overrides[Index].ComponentClass->GetFullName(), *InComponentName.ToString(), *InComponentClass->GetFullName(), *InComponentClass->GetFullName());
+			// if a base class is asking for an override, the existing override (which we are going to use) had better be derived
+			if (!IsLegalOverride(Overrides[Index].ComponentClass, InComponentClass))
+			{
+				UE_LOG(LogUObjectGlobals, Error, TEXT("%s is not a legal override for component %s because it does not derive from %s. Will use %s when constructing component."),
+					*Overrides[Index].ComponentClass->GetFullName(), *GetSubobjectPath(), *InComponentClass->GetFullName(), *InComponentClass->GetFullName());
 
+				Overrides[Index].ComponentClass = InComponentClass;
+			}
+		}
+		else
+		{
+			// if the existing recorded component class is null then we could either have suboverrides in which case we still want to use the class,
+			// or it could be marked do not create, but since the base class may create it as non-optional we still want to record the class
 			Overrides[Index].ComponentClass = InComponentClass;
 		}
 	}
+	else
+	{
+		// Warn about existing overrides but the parent marking it DoNotCreate
+		// Note that even if we report an error, these overrides may still get used if the component is created as non-optional
+		if (Overrides[Index].ComponentClass)
+		{
+			UE_LOG(LogUObjectGlobals, Error, TEXT("%s is not a legal override for component %s because a parent class is marking it do not create."),
+				*Overrides[Index].ComponentClass->GetFullName(), *GetSubobjectPath());
+		}
+		if (Overrides[Index].SubOverrides)
+		{
+			UE_LOG(LogUObjectGlobals, Error, TEXT("Component %s has recorded nested subobject overrides, but won't be created because a parent class is marking it do not create."),
+				*GetSubobjectPath());
+		}
+
+		Overrides[Index].bDoNotCreate = true;
+	}
+}
+
+void FObjectInitializer::FOverrides::Add(FStringView InComponentPath, UClass* InComponentClass)
+{
+	TArray<FName> ComponentPath;
+
+	int32 PeriodIndex;
+	while (InComponentPath.FindChar(TEXT('.'), PeriodIndex))
+	{
+		ComponentPath.Add(FName(PeriodIndex, InComponentPath.GetData()));
+		InComponentPath.RemovePrefix(PeriodIndex+1);
+	}
+	ComponentPath.Add(FName(InComponentPath.Len(), InComponentPath.GetData()));
+
+	TArrayView<const FName> PathArrayView(ComponentPath);
+	Add(PathArrayView, InComponentClass, &PathArrayView);
+}
+
+void FObjectInitializer::FOverrides::Add(TArrayView<const FName> InComponentPath, UClass* InComponentClass, const TArrayView<const FName>* InFullComponentPath)
+{
+	if (InComponentPath.Num() > 1)
+	{
+		const FName ComponentName = InComponentPath[0];
+		int32 Index = Find(ComponentName);
+		if (Index == INDEX_NONE)
+		{
+			Index = Overrides.Emplace(FOverride(ComponentName));
+		}
+		if (!Overrides[Index].SubOverrides)
+		{
+			Overrides[Index].SubOverrides = MakeUnique<FOverrides>();
+		}
+
+		Overrides[Index].SubOverrides->Add(InComponentPath.Slice(1, InComponentPath.Num()-1), InComponentClass, (InFullComponentPath ? InFullComponentPath : &InComponentPath));
+	}
+	else
+	{
+		Add(InComponentPath[0], InComponentClass, (InFullComponentPath ? InFullComponentPath : &InComponentPath));
+	}
 }
 
 /**  Retrieve an override, or TClassToConstructByDefault::StaticClass or nullptr if this was removed by a derived class **/
-UClass* FObjectInitializer::FOverrides::Get(FName InComponentName, UClass* ReturnType, UClass* ClassToConstructByDefault, FObjectInitializer const& ObjectInitializer) const
+FObjectInitializer::FOverrides::FOverrideDetails FObjectInitializer::FOverrides::Get(FName InComponentName, UClass* ReturnType, UClass* ClassToConstructByDefault, bool bOptional) const
 {
+	FOverrideDetails Result;
+
 	const int32 Index = Find(InComponentName);
-	UClass* BaseComponentClass = ClassToConstructByDefault;
 	if (Index == INDEX_NONE)
 	{
-		return BaseComponentClass; // no override so just do what the base class wanted
+		Result.Class = ClassToConstructByDefault; // no override so just do what the base class wanted
+		Result.SubOverrides = nullptr;
+	}
+	else if (Overrides[Index].bDoNotCreate && bOptional)
+	{
+		Result.Class = nullptr;   // the override is of nullptr, which means "don't create this component"
+		Result.SubOverrides = nullptr; // and if we're not creating this component also don't need sub-overrides
 	}
 	else if (Overrides[Index].ComponentClass)
 	{
 		if (IsLegalOverride(Overrides[Index].ComponentClass, ReturnType)) // if THE base class is asking for a T, the existing override (which we are going to use) had better be derived
 		{
-			return Overrides[Index].ComponentClass; // the override is of an acceptable class, so use it
+			Result.Class = Overrides[Index].ComponentClass; // the override is of an acceptable class, so use it
+
+			if (Overrides[Index].bDoNotCreate)
+			{
+				UE_LOG(LogUObjectGlobals, Error, TEXT("Ignored DoNotCreateDefaultSubobject for %s as it's marked as required. Creating %s."), *InComponentName.ToString(), *Result.Class->GetName());
+			}
 		}
 		else
 		{
+			if (Overrides[Index].bDoNotCreate)
+			{
+				UE_LOG(LogUObjectGlobals, Error, TEXT("Ignored DoNotCreateDefaultSubobject for %s as it's marked as required. Creating %s."), *InComponentName.ToString(), *ClassToConstructByDefault->GetName());
+			}
 			UE_LOG(LogUObjectGlobals, Error, TEXT("%s is not a legal override for component %s because it does not derive from %s. Using %s to construct component."),
 				*Overrides[Index].ComponentClass->GetFullName(), *InComponentName.ToString(), *ReturnType->GetFullName(), *ClassToConstructByDefault->GetFullName());
 
-			return ClassToConstructByDefault;
+			Result.Class = ClassToConstructByDefault;
+		}
+		Result.SubOverrides = Overrides[Index].SubOverrides.Get();
+	}
+	else
+	{
+		if (Overrides[Index].bDoNotCreate)
+		{
+			UE_LOG(LogUObjectGlobals, Error, TEXT("Ignored DoNotCreateDefaultSubobject for %s as it's marked as required. Creating %s."), *InComponentName.ToString(), *ClassToConstructByDefault->GetName());
 		}
+
+		Result.Class = ClassToConstructByDefault; // Only sub-overrides were overriden, so use the base class' desire
+		Result.SubOverrides = Overrides[Index].SubOverrides.Get();
 	}
-	return nullptr;  // the override is of nullptr, which means "don't create this component"
+
+	return Result;  
 }
 bool FObjectInitializer::FOverrides::IsLegalOverride(const UClass* DerivedComponentClass, const UClass* BaseComponentClass)
 {
@@ -3121,6 +3236,29 @@ void FObjectInitializer::AssertIfSubobjectSetupIsNotAllowed(const FName Subobjec
 		TEXT("%s.%s: Subobject class setup is only allowed in base class constructor call (in the initialization list)"), Obj ? *Obj->GetFullName() : TEXT("NULL"), *SubobjectName.GetPlainNameString());
 }
 
+void FObjectInitializer::AssertIfSubobjectSetupIsNotAllowed(const FStringView SubobjectName) const
+{
+	UE_CLOG(!bSubobjectClassInitializationAllowed, LogUObjectGlobals, Fatal,
+		TEXT("%s.%.*s: Subobject class setup is only allowed in base class constructor call (in the initialization list)"), Obj ? *Obj->GetFullName() : TEXT("NULL"), SubobjectName.Len(), SubobjectName.GetData());
+}
+
+void FObjectInitializer::AssertIfSubobjectSetupIsNotAllowed(TArrayView<const FName> SubobjectNames) const
+{
+	auto MakeSubobjectPath = [&SubobjectNames]()
+	{
+		FString SubobjectPath;
+		for (FName SubobjectName : SubobjectNames)
+		{
+			SubobjectPath += (SubobjectPath.IsEmpty() ? TEXT("") : TEXT("."));
+			SubobjectPath += SubobjectName.ToString();
+		}
+		return SubobjectPath;
+	};
+
+	UE_CLOG(!bSubobjectClassInitializationAllowed, LogUObjectGlobals, Fatal,
+		TEXT("%s.%s: Subobject class setup is only allowed in base class constructor call (in the initialization list)"), Obj ? *Obj->GetFullName() : TEXT("NULL"), *MakeSubobjectPath());
+}
+
 #if DO_CHECK
 void CheckIsClassChildOf_Internal(const UClass* Parent, const UClass* Child)
 {
@@ -3192,7 +3330,7 @@ UObject* StaticConstructObject_Internal(const FStaticConstructObjectParameters&
 	if (!bRecycledSubobject)
 	{		
 		STAT(FScopeCycleCounterUObject ConstructorScope(InClass, GET_STATID(STAT_ConstructObject)));
-		(*InClass->ClassConstructor)( FObjectInitializer(Result, InTemplate, Params.bCopyTransientsFromClassDefaults, true, Params.InstanceGraph) );
+		(*InClass->ClassConstructor)(FObjectInitializer(Result, Params));
 	}
 	
 	if( GIsEditor && GUndo && (InFlags & RF_Transactional) && !(InFlags & RF_NeedLoad) && !InClass->IsChildOf(UField::StaticClass()) )
@@ -3745,19 +3883,15 @@ UScriptStruct* GetFallbackStruct()
 
 UObject* FObjectInitializer::CreateDefaultSubobject(UObject* Outer, FName SubobjectFName, UClass* ReturnType, UClass* ClassToCreateByDefault, bool bIsRequired, bool bIsTransient) const
 {
-	UE_CLOG(!FUObjectThreadContext::Get().IsInConstructor, LogClass, Fatal, TEXT("Subobjects cannot be created outside of UObject constructors. UObject constructing subobjects cannot be created using new or placement new operator."));
+	UE_CLOG(!FUObjectThreadContext::Get().IsInConstructor, LogUObjectGlobals, Fatal, TEXT("Subobjects cannot be created outside of UObject constructors. UObject constructing subobjects cannot be created using new or placement new operator."));
 	if (SubobjectFName == NAME_None)
 	{
-		UE_LOG(LogClass, Fatal, TEXT("Illegal default subobject name: %s"), *SubobjectFName.ToString());
+		UE_LOG(LogUObjectGlobals, Fatal, TEXT("Illegal default subobject name: %s"), *SubobjectFName.ToString());
 	}
 
-	UObject* Result = NULL;
-	UClass* OverrideClass = ComponentOverrides.Get(SubobjectFName, ReturnType, ClassToCreateByDefault, *this);
-	if (!OverrideClass && bIsRequired)
-	{
-		OverrideClass = ClassToCreateByDefault;
-		UE_LOG(LogClass, Warning, TEXT("Ignored DoNotCreateDefaultSubobject for %s as it's marked as required. Creating %s."), *SubobjectFName.ToString(), *OverrideClass->GetName());
-	}
+	UObject* Result = nullptr;
+	FOverrides::FOverrideDetails ComponentOverride = SubobjectOverrides.Get(SubobjectFName, ReturnType, ClassToCreateByDefault, !bIsRequired);
+	UClass* OverrideClass = ComponentOverride.Class;
 	if (OverrideClass)
 	{
 		check(OverrideClass->IsChildOf(ReturnType));
@@ -3767,7 +3901,7 @@ UObject* FObjectInitializer::CreateDefaultSubobject(UObject* Outer, FName Subobj
 			// Attempts to create an abstract class will return null. If it is not optional or the owning class is not also abstract report a warning.
 			if (!bIsRequired && !Outer->GetClass()->HasAnyClassFlags(CLASS_Abstract))
 			{
-				UE_LOG(LogClass, Warning, TEXT("Required default subobject %s not created as requested class %s is abstract. Returning null."), *SubobjectFName.ToString(), *OverrideClass->GetName());
+				UE_LOG(LogUObjectGlobals, Warning, TEXT("Required default subobject %s not created as requested class %s is abstract. Returning null."), *SubobjectFName.ToString(), *OverrideClass->GetName());
 			}
 		}
 		else
@@ -3788,7 +3922,7 @@ UObject* FObjectInitializer::CreateDefaultSubobject(UObject* Outer, FName Subobj
 			// We only need to check the name as ConstructObject would fail anyway if an object of the same name but different class already existed.
 			if (ConstructedSubobjects.Find(SubobjectFName) != INDEX_NONE)
 			{
-				UE_LOG(LogClass, Fatal, TEXT("Default subobject %s %s already exists for %s."), *OverrideClass->GetName(), *SubobjectFName.ToString(), *Outer->GetFullName());
+				UE_LOG(LogUObjectGlobals, Fatal, TEXT("Default subobject %s %s already exists for %s."), *OverrideClass->GetName(), *SubobjectFName.ToString(), *Outer->GetFullName());
 			}
 			else
 			{
@@ -3799,6 +3933,7 @@ UObject* FObjectInitializer::CreateDefaultSubobject(UObject* Outer, FName Subobj
 			Params.Outer = Outer;
 			Params.Name = SubobjectFName;
 			Params.SetFlags = SubobjectFlags;
+			Params.SubobjectOverrides = ComponentOverride.SubOverrides;
 
 			Result = StaticConstructObject_Internal(Params);
 			if (!bIsTransient && (bOwnerArchetypeIsNotNative || bOwnerTemplateIsNotCDO))
diff --git a/Engine/Source/Runtime/CoreUObject/Public/UObject/UObjectGlobals.h b/Engine/Source/Runtime/CoreUObject/Public/UObject/UObjectGlobals.h
--- a/Engine/Source/Runtime/CoreUObject/Public/UObject/UObjectGlobals.h
+++ b/Engine/Source/Runtime/CoreUObject/Public/UObject/UObjectGlobals.h
@@ -13,6 +13,7 @@
 #include "UObject/LinkerInstancingContext.h"
 #include "Misc/OutputDeviceRedirector.h"
 #include "Containers/ArrayView.h"
+#include "Containers/StringView.h"
 #include "Templates/Function.h"
 #include "Templates/IsArrayOrRefOfType.h"
 #include "Serialization/ArchiveUObject.h"
@@ -724,6 +725,14 @@ public:
 	 */
 	FObjectInitializer(UObject* InObj, UObject* InObjectArchetype, bool bInCopyTransientsFromClassDefaults, bool bInShouldInitializeProps, struct FObjectInstancingGraph* InInstanceGraph = nullptr);
 
+	/** Special constructor for static construct object internal that passes along the params block directly */
+	FObjectInitializer(UObject* InObj, const FStaticConstructObjectParameters& StaticConstructParams);
+
+private:
+	/** Helper for the common behaviors in the constructors */
+	void Construct_Internal();
+
+public:
 	~FObjectInitializer();
 
 	/** 
@@ -830,10 +839,10 @@ public:
 	 * @param	SubobjectName	name of the new component or subobject
 	 * @param	Class			The class to use for the specified subobject or component.
 	 */
-	FObjectInitializer const& SetDefaultSubobjectClass(FName SubobjectName, UClass* Class) const
+	const FObjectInitializer& SetDefaultSubobjectClass(FName SubobjectName, UClass* Class) const
 	{
 		AssertIfSubobjectSetupIsNotAllowed(SubobjectName);
-		ComponentOverrides.Add(SubobjectName, Class, *this);
+		SubobjectOverrides.Add(SubobjectName, Class);
 		return *this;
 	}
 
@@ -842,7 +851,7 @@ public:
 	 * @param	SubobjectName	name of the new component or subobject
 	 */
 	template<class T>
-	FObjectInitializer const& SetDefaultSubobjectClass(FName SubobjectName) const
+	const FObjectInitializer& SetDefaultSubobjectClass(FName SubobjectName) const
 	{
 		return SetDefaultSubobjectClass(SubobjectName, T::StaticClass());
 	}
@@ -851,10 +860,76 @@ public:
 	 * Indicates that a base class should not create a component
 	 * @param	SubobjectName	name of the new component or subobject to not create
 	 */
-	FObjectInitializer const& DoNotCreateDefaultSubobject(FName SubobjectName) const
+	const FObjectInitializer& DoNotCreateDefaultSubobject(FName SubobjectName) const
 	{
 		AssertIfSubobjectSetupIsNotAllowed(SubobjectName);
-		ComponentOverrides.Add(SubobjectName, nullptr, *this);
+		SubobjectOverrides.Add(SubobjectName, nullptr);
+		return *this;
+	}
+
+	/**
+	 * Sets the class to use for a subobject defined in a nested subobject, the class must be a subclass of the class used when calling CreateDefaultSubobject.
+	 * @param	SubobjectName	path to the new component or subobject
+	 * @param	Class			The class to use for the specified subobject or component.
+	 */
+	const FObjectInitializer& SetNestedDefaultSubobjectClass(FStringView SubobjectName, UClass* Class) const
+	{
+		AssertIfSubobjectSetupIsNotAllowed(SubobjectName);
+		SubobjectOverrides.Add(SubobjectName, Class);
+		return *this;
+	}
+
+	/**
+	 * Sets the class to use for a subobject defined in a nested subobject, the class must be a subclass of the class used when calling CreateDefaultSubobject.
+	 * @param	SubobjectName	path to the new component or subobject
+	 * @param	Class			The class to use for the specified subobject or component.
+	 */
+	const FObjectInitializer& SetNestedDefaultSubobjectClass(TArrayView<const FName> SubobjectNames, UClass* Class) const
+	{
+		AssertIfSubobjectSetupIsNotAllowed(SubobjectNames);
+		SubobjectOverrides.Add(SubobjectNames, Class);
+		return *this;
+	}
+
+	/**
+	 * Sets the class to use for a subobject defined in a nested subobject, the class must be a subclass of the class used when calling CreateDefaultSubobject.
+	 * @param	SubobjectName	path to the new component or subobject
+	 */
+	template<class T>
+	const FObjectInitializer& SetNestedDefaultSubobjectClass(FStringView SubobjectName) const
+	{
+		return SetNestedDefaultSubobjectClass(SubobjectName, T::StaticClass());
+	}
+
+	/**
+	 * Sets the class to use for a subobject defined in a nested subobject, the class must be a subclass of the class used when calling CreateDefaultSubobject.
+	 * @param	SubobjectName	path to the new component or subobject
+	 */
+	template<class T>
+	const FObjectInitializer& SetNestedDefaultSubobjectClass(TArrayView<const FName> SubobjectNames) const
+	{
+		return SetNestedDefaultSubobjectClass(SubobjectNames, T::StaticClass());
+	}
+
+	/**
+	 * Indicates that a subobject should not create a component if created using CreateOptionalDefaultSubobject
+	 * @param	SubobjectName	name of the new component or subobject to not create
+	 */
+	const FObjectInitializer& DoNotCreateNestedDefaultSubobject(FStringView SubobjectName) const
+	{
+		AssertIfSubobjectSetupIsNotAllowed(SubobjectName);
+		SubobjectOverrides.Add(SubobjectName, nullptr);
+		return *this;
+	}
+
+	/**
+	 * Indicates that a subobject should not create a component if created using CreateOptionalDefaultSubobject
+	 * @param	SubobjectName	name of the new component or subobject to not create
+	 */
+	const FObjectInitializer& DoNotCreateNestedDefaultSubobject(TArrayView<const FName> SubobjectNames) const
+	{
+		AssertIfSubobjectSetupIsNotAllowed(SubobjectNames);
+		SubobjectOverrides.Add(SubobjectNames, nullptr);
 		return *this;
 	}
 
@@ -927,10 +1002,22 @@ private:
 	struct FOverrides
 	{
 		/**  Add an override, make sure it is legal **/
-		void Add(FName InComponentName, UClass* InComponentClass, FObjectInitializer const& ObjectInitializer);
+		void Add(FName InComponentName, UClass* InComponentClass, const TArrayView<const FName>* FullPath = nullptr);
+
+		/**  Add a potentially nested override, make sure it is legal **/
+		void Add(FStringView InComponentPath, UClass* InComponentClass);
+
+		/**  Add a potentially nested override, make sure it is legal **/
+		void Add(TArrayView<const FName> InComponentPath, UClass* InComponentClass, const TArrayView<const FName>* FullPath = nullptr);
+
+		struct FOverrideDetails
+		{
+			UClass* Class = nullptr;
+			FOverrides* SubOverrides = nullptr;
+		};
 
 		/**  Retrieve an override, or TClassToConstructByDefault::StaticClass or nullptr if this was removed by a derived class **/
-		UClass* Get(FName InComponentName, UClass* ReturnType, UClass* ClassToConstructByDefault, FObjectInitializer const& ObjectInitializer) const;
+		FOverrideDetails Get(FName InComponentName, UClass* ReturnType, UClass* ClassToConstructByDefault, bool bOptional) const;
 
 	private:
 		static bool IsLegalOverride(const UClass* DerivedComponentClass, const UClass* BaseComponentClass);
@@ -950,13 +1037,31 @@ private:
 		/**  Element of the override array **/
 		struct FOverride
 		{
-			FName	ComponentName;
-			UClass* ComponentClass;
-			FOverride(FName InComponentName, UClass *InComponentClass)
+			FName ComponentName;
+			UClass* ComponentClass = nullptr;
+			TUniquePtr<FOverrides> SubOverrides;
+			bool bDoNotCreate = false;
+			FOverride(FName InComponentName)
 				: ComponentName(InComponentName)
-				, ComponentClass(InComponentClass)
 			{
 			}
+
+			FOverride& operator=(const FOverride& Other)
+			{
+				ComponentName = Other.ComponentName;
+				ComponentClass = Other.ComponentClass;
+				SubOverrides = (Other.SubOverrides ? MakeUnique<FOverrides>(*Other.SubOverrides) : nullptr);
+				bDoNotCreate = Other.bDoNotCreate;
+				return *this;
+			}
+
+			FOverride(const FOverride& Other)
+			{
+				*this = Other;
+			}
+
+			FOverride(FOverride&&) = default;
+			FOverride& operator=(FOverride&&) = default;
 		};
 		/**  The override array **/
 		TArray<FOverride, TInlineAllocator<8> > Overrides;
@@ -991,6 +1096,12 @@ private:
 	/** Asserts if SetDefaultSubobjectClass or DoNotCreateOptionalDefaultSuobject are called inside of the constructor body */
 	void AssertIfSubobjectSetupIsNotAllowed(const FName SubobjectName) const;
 
+	/** Asserts if SetDefaultSubobjectClass or DoNotCreateOptionalDefaultSuobject are called inside of the constructor body */
+	void AssertIfSubobjectSetupIsNotAllowed(FStringView SubobjectName) const;
+
+	/** Asserts if SetDefaultSubobjectClass or DoNotCreateOptionalDefaultSuobject are called inside of the constructor body */
+	void AssertIfSubobjectSetupIsNotAllowed(TArrayView<const FName> SubobjectNames) const;
+
 	/**  object to initialize, from static allocate object, after construction **/
 	UObject* Obj;
 	/**  object to copy properties from **/
@@ -1000,11 +1111,15 @@ private:
 	/**  If true, initialize the properties **/
 	bool bShouldInitializePropsFromArchetype;
 	/**  Only true until ObjectInitializer has not reached the base UObject class */
-	bool bSubobjectClassInitializationAllowed;
+	bool bSubobjectClassInitializationAllowed = true;
+#if USE_CIRCULAR_DEPENDENCY_LOAD_DEFERRING
+	/**  */
+	bool bIsDeferredInitializer = false;
+#endif // USE_CIRCULAR_DEPENDENCY_LOAD_DEFERRING
 	/**  Instance graph **/
 	struct FObjectInstancingGraph* InstanceGraph;
 	/**  List of component classes to override from derived classes **/
-	mutable FOverrides ComponentOverrides;
+	mutable FOverrides SubobjectOverrides;
 	/**  List of component classes to initialize after the C++ constructors **/
 	mutable FSubobjectsToInit ComponentInits;
 #if !UE_BUILD_SHIPPING
@@ -1012,12 +1127,9 @@ private:
 	mutable TArray<FName, TInlineAllocator<8>> ConstructedSubobjects;
 #endif
 	/**  Previously constructed object in the callstack */
-	UObject* LastConstructedObject;
+	UObject* LastConstructedObject = nullptr;
 
-#if USE_CIRCULAR_DEPENDENCY_LOAD_DEFERRING
-	/**  */
-	bool bIsDeferredInitializer : 1;
-#endif // USE_CIRCULAR_DEPENDENCY_LOAD_DEFERRING
+	friend struct FStaticConstructObjectParameters;
 };
 
 /**
@@ -1059,7 +1171,14 @@ struct FStaticConstructObjectParameters
 	/** Assign an external Package to the created object if non-null */
 	UPackage* ExternalPackage = nullptr;
 
+private:
+	FObjectInitializer::FOverrides* SubobjectOverrides = nullptr;
+
+public:
+
 	COREUOBJECT_API FStaticConstructObjectParameters(const UClass* InClass);
+
+	friend FObjectInitializer;
 };
 
 /**
