GetPackageLinker - reduce complexity by removing deprecated code paths and merging duplicate blocks.
The new function is equivalent for all cases except the cases below (which were already deprecated):

Case) An InOuter was provided with no InLongPackageName, and the InOuter had no existing linker, and the InOuter is a Memory-only package, and a file exists on disk with the same path.
Change) The previous version would have called CreateLinker with the filename from disk; the new version will silently skip the CreateLinker call and return nullptr.

Case) Load_AllowDll package flag was set, and an InOuter was provided with no InLongPackageName, and the InOuter had no existing linker, and the file did not exist on disk.
Change) The previous version would have silently returned nullptr; the new version will issue a warning and return nullptr.

Case) Load_NoWarn or Load_Quiet was provided, and an InOuter was provided with no InLongPackageName, and the InOuter had no existing linker, and the file did not exist on disk.
Change) The previous version would have issued a warning and returned nullptr; the new version will silently return nullptr.

Case) An InOuter was provided, and an InLongPackageName was provided, and the InOuter had no existing linker, and no package existed for the (possibly redirected) InLongPackageName.
Change) The previous version would have created an unused package with name InLongPackageName, and then called CreateLinker on the InOuter; the new version will not create the unused package.

Case) An InOutLoadContext is provided that is not equal to FUObjectThreadContext::Get().GetSerializeContext() and we reach the call to CreateLinker.
Change) The previous version would have passed FUObjectThreadContext::Get().GetSerializeContext() to CreateLinker; the new version will pass InOutLoadContext.
Note) This case does not currently exist in any callsite in the Engine; they all pass nullptr or FUObjectThreadContext::Get().GetSerializeContext() for the InOutLoadContext.

Case) An InOuter was not provided, and a CompatabilityGuid was provided, and a package and linker already exists in memory for the given name, and the guid on disk for the packagename matches the CompatabilityGuid, but the guid from the linker in memory does not.
Change) The previous version would have returned nullptr; the new version will ignore the mismatch and return the linker in memory.

#rnx

diff --git a/Engine/Source/Runtime/CoreUObject/Private/UObject/Linker.cpp b/Engine/Source/Runtime/CoreUObject/Private/UObject/Linker.cpp
--- a/Engine/Source/Runtime/CoreUObject/Private/UObject/Linker.cpp
+++ b/Engine/Source/Runtime/CoreUObject/Private/UObject/Linker.cpp
@@ -463,19 +463,15 @@ void DeleteLoader(FLinkerLoad* Loader)
 	FLinkerManager::Get().RemoveLinker(Loader);
 }
 
-static void LogGetPackageLinkerError(FArchive* LinkerArchive, FUObjectSerializeContext* LoadContext, const TCHAR* InFilename, const FText& InErrorMessage, UObject* InOuter, uint32 LoadFlags)
+static void LogGetPackageLinkerError(FUObjectSerializeContext* LoadContext, const TCHAR* InFilename, const FText& InErrorMessage, UObject* InOuter, uint32 LoadFlags)
 {
 	static FName NAME_LoadErrors("LoadErrors");
 	struct Local
 	{
 		/** Helper function to output more detailed error info if available */
-		static void OutputErrorDetail(FArchive* InLinkerArchive, FUObjectSerializeContext* InLoadContext, const FName& LogName)
+		static void OutputErrorDetail(FUObjectSerializeContext* InLoadContext, const FName& LogName)
 		{
 			FUObjectSerializeContext* LoadContextToReport = InLoadContext;
-			if (!LoadContextToReport && InLinkerArchive)
-			{
-				LoadContextToReport = InLinkerArchive->GetSerializeContext();
-			}
 			if (LoadContextToReport && LoadContextToReport->SerializedObject && LoadContextToReport->SerializedImportLinker)
 			{
 				FMessageLog LoadErrors(LogName);
@@ -485,13 +481,6 @@ static void LogGetPackageLinkerError(FArchive* LinkerArchive, FUObjectSerializeC
 				Message->AddToken(FAssetNameToken::Create(LoadContextToReport->SerializedImportLinker->GetImportPathName(LoadContextToReport->SerializedImportIndex)));
 				Message->AddToken(FTextToken::Create(LOCTEXT("FailedLoad_Referenced", "Referenced by")));
 				Message->AddToken(FUObjectToken::Create(LoadContextToReport->SerializedObject));
-				FProperty* SerializedProperty = InLinkerArchive ? InLinkerArchive->GetSerializedProperty() : nullptr;
-				if (SerializedProperty != nullptr)
-				{
-					FString PropertyPathName = SerializedProperty->GetPathName();
-					Message->AddToken(FTextToken::Create(LOCTEXT("FailedLoad_Property", "Property")));
-					Message->AddToken(FAssetNameToken::Create(PropertyPathName, FText::FromString( PropertyPathName) ) );
-				}
 			}
 		}
 	};
@@ -557,7 +546,7 @@ static void LogGetPackageLinkerError(FArchive* LinkerArchive, FUObjectSerializeC
 				Message->AddToken(FAssetNameToken::Create(OuterPackageName));
 			}
 
-			Local::OutputErrorDetail(LinkerArchive, LoadContext, NAME_LoadErrors);
+			Local::OutputErrorDetail(LoadContext, NAME_LoadErrors);
 		}
 	}
 	else
@@ -588,7 +577,7 @@ static void LogGetPackageLinkerError(FArchive* LinkerArchive, FUObjectSerializeC
 				SET_WARN_COLOR(COLOR_RED);
 				UE_LOG(LogLinker, Warning, TEXT("%s"), *FullErrorMessage.ToString());
 				CLEAR_WARN_COLOR();
-				Local::OutputErrorDetail(LinkerArchive, LoadContext, NAME_LoadErrors);
+				Local::OutputErrorDetail(LoadContext, NAME_LoadErrors);
 			}
 		}
 	}
@@ -655,12 +644,9 @@ FLinkerLoad* GetPackageLinker
 )
 {
 	FUObjectSerializeContext* InExistingContext = InOutLoadContext ? *InOutLoadContext : nullptr;
-	// See if there is already a linker for this package.
-	FLinkerLoad* Result = FLinkerLoad::FindExistingLinkerForPackage(InOuter);
 
-	// Try to load the linker.
 	// See if the linker is already loaded.
-	if (Result)
+	if (FLinkerLoad* Result = FLinkerLoad::FindExistingLinkerForPackage(InOuter))
 	{
 		if (InExistingContext && Result->GetSerializeContext() && Result->GetSerializeContext() != InExistingContext)
 		{
@@ -672,56 +658,32 @@ FLinkerLoad* GetPackageLinker
 		return Result;
 	}
 
-	UPackage* CreatedPackage = nullptr;
-
-	FString NewFilename;
-	if( !InLongPackageName )
+	FString PackageNameToCreate;
+	UPackage* TargetPackage = nullptr;
+	if (!InLongPackageName)
 	{
 		// Resolve filename from package name.
-		if( !InOuter )
+		if (!InOuter)
 		{
 			// try to recover from this instead of throwing, it seems recoverable just by doing this
-			LogGetPackageLinkerError(Result, InExistingContext, InLongPackageName, LOCTEXT("PackageResolveFailed", "Can't resolve asset name"), InOuter, LoadFlags);
+			LogGetPackageLinkerError(InExistingContext, InLongPackageName, LOCTEXT("PackageResolveFailed", "Can't resolve asset name"), InOuter, LoadFlags);
 			return nullptr;
 		}
-	
-		// Allow delegates to resolve this package
-		FString PackageNameToCreate = InOuter->GetName();
+		PackageNameToCreate = InOuter->GetName();
+		TargetPackage = InOuter;
 
 		// Process any package redirects
 		{
 			const FCoreRedirectObjectName NewPackageName = FCoreRedirects::GetRedirectedName(ECoreRedirectFlags::Type_Package, FCoreRedirectObjectName(NAME_None, NAME_None, *PackageNameToCreate));
 			NewPackageName.PackageName.ToString(PackageNameToCreate);
 		}
-
-		// The editor must not redirect packages for localization. We also shouldn't redirect script or in-memory packages.
-		FString PackageNameToLoad = PackageNameToCreate;
-		if (!(GIsEditor || InOuter->HasAnyPackageFlags(PKG_InMemoryOnly) || FPackageName::IsScriptPackage(PackageNameToLoad)))
-		{
-			PackageNameToLoad = FPackageName::GetDelegateResolvedPackagePath(PackageNameToLoad);
-			PackageNameToLoad = FPackageName::GetLocalizedPackagePath(PackageNameToLoad);
-		}
-
-		// Verify that the file exists.
-		const bool DoesPackageExist = DoesPackageExistForGetPackageLinker(PackageNameToLoad, CompatibleGuid, NewFilename);
-		if ( !DoesPackageExist )
-		{
-			// In memory-only packages have no linker and this is ok.
-			if (!(LoadFlags & LOAD_AllowDll) && !InOuter->HasAnyPackageFlags(PKG_InMemoryOnly) && !FLinkerLoad::IsKnownMissingPackage(InOuter->GetFName()))
-			{
-				LogGetPackageLinkerError(Result, InExistingContext, InLongPackageName, LOCTEXT("PackageNotFoundShort", "Can't find file."), InOuter, LoadFlags);
-			}
-
-			return nullptr;
-		}
 	}
 	else
 	{
-		FString PackageNameToCreate;
 		if (!FPackageName::TryConvertFilenameToLongPackageName(InLongPackageName, PackageNameToCreate))
 		{
 			// try to recover from this instead of throwing, it seems recoverable just by doing this
-			LogGetPackageLinkerError(Result, InExistingContext, InLongPackageName, LOCTEXT("PackageResolveFailed", "Can't resolve asset name"), InOuter, LoadFlags);
+			LogGetPackageLinkerError(InExistingContext, InLongPackageName, LOCTEXT("PackageResolveFailed", "Can't resolve asset name"), InOuter, LoadFlags);
 			return nullptr;
 		}
 
@@ -731,131 +693,108 @@ FLinkerLoad* GetPackageLinker
 			NewPackageName.PackageName.ToString(PackageNameToCreate);
 		}
 
-		// The editor must not redirect packages for localization. We also shouldn't redirect script packages.
-		FString PackageNameToLoad = PackageNameToCreate;
-		if (!(GIsEditor || FPackageName::IsScriptPackage(PackageNameToLoad)))
+		if (InOuter)
 		{
-			// Allow delegates to resolve this path
-			PackageNameToLoad = FPackageName::GetDelegateResolvedPackagePath(PackageNameToLoad);
-			PackageNameToLoad = FPackageName::GetLocalizedPackagePath(PackageNameToLoad);
+			TargetPackage = InOuter;
 		}
-
-		UPackage* ExistingPackage = FindObject<UPackage>(nullptr, *PackageNameToCreate);
-		if (ExistingPackage)
+		else
 		{
-			if (!ExistingPackage->GetOuter() && ExistingPackage->HasAnyPackageFlags(PKG_InMemoryOnly))
+			TargetPackage = FindObject<UPackage>(nullptr, *PackageNameToCreate);
+			if (TargetPackage && TargetPackage->GetOuter() != nullptr)
 			{
-				// This is a memory-only in package and so it has no linker and this is ok.
-				return nullptr;
+				TargetPackage = nullptr;
 			}
 		}
+	}
 
-		// Verify that the file exists.
-		const bool DoesPackageExist = DoesPackageExistForGetPackageLinker(PackageNameToLoad, CompatibleGuid, NewFilename);
-		if( !DoesPackageExist )
-		{
-			// Issue a warning if the caller didn't request nowar/quiet, and the package isn't marked as known to be missing.
-			bool IssueWarning = (LoadFlags & (LOAD_NoWarn | LOAD_Quiet)) == 0 && !FLinkerLoad::IsKnownMissingPackage(InLongPackageName);
+	if (TargetPackage && TargetPackage->HasAnyPackageFlags(PKG_InMemoryOnly))
+	{
+		// This is a memory-only in package and so it has no linker and this is ok.
+		return nullptr;
+	}
 
-			if (IssueWarning)
-			{
-				// try to recover from this instead of throwing, it seems recoverable just by doing this
-				LogGetPackageLinkerError(Result, InExistingContext, InLongPackageName, LOCTEXT("FileNotFoundShort", "Can't find file."), InOuter, LoadFlags);
-			}
-			return nullptr;
-		}
-		UPackage* FilenamePkg = ExistingPackage;
-		if (!FilenamePkg)
+	// The editor must not redirect packages for localization. We also shouldn't redirect script or in-memory packages (in-memory packages exited earlier so we don't need to check here).
+	FString PackageNameToLoad = PackageNameToCreate;
+	if (!(GIsEditor || FPackageName::IsScriptPackage(PackageNameToLoad)))
+	{
+		// Allow delegates to resolve the path
+		PackageNameToLoad = FPackageName::GetDelegateResolvedPackagePath(PackageNameToLoad);
+		PackageNameToLoad = FPackageName::GetLocalizedPackagePath(PackageNameToLoad);
+	}
+
+	// Verify that the file exists.
+	FString NewFilename;
+	const bool bDoesPackageExist = DoesPackageExistForGetPackageLinker(PackageNameToLoad, CompatibleGuid, NewFilename);
+	if (!bDoesPackageExist)
+	{
+		// Issue a warning if the caller didn't request nowarn/quiet, and the package isn't marked as known to be missing.
+		bool bIssueWarning = (LoadFlags & (LOAD_NoWarn | LOAD_Quiet)) == 0 && !FLinkerLoad::IsKnownMissingPackage(InLongPackageName);
+
+		if (bIssueWarning)
 		{
+			// try to recover from this instead of throwing, it seems recoverable just by doing this
+			LogGetPackageLinkerError(InExistingContext, InLongPackageName, LOCTEXT("FileNotFoundShort", "Can't find file."), InOuter, LoadFlags);
+		}
+		return nullptr;
+	}
+
+	UPackage* CreatedPackage = nullptr;
+	if (!TargetPackage)
+	{
 #if WITH_EDITORONLY_DATA
 		// Make sure the package name matches the name on disk
 		FPackageName::FixPackageNameCase(PackageNameToCreate, FPathViews::GetExtension(NewFilename));
 #endif
-			// Create the package with the provided long package name.
-			CreatedPackage = CreatePackage(*PackageNameToCreate);
-			FilenamePkg = CreatedPackage;
-		}
-
-		if (FilenamePkg && FilenamePkg != ExistingPackage && (LoadFlags & LOAD_PackageForPIE))
+		// Create the package with the provided long package name.
+		CreatedPackage = CreatePackage(*PackageNameToCreate);
+		if (!CreatedPackage)
 		{
-			check(FilenamePkg);
-			FilenamePkg->SetPackageFlags(PKG_PlayInEditor);
-		}
-
-		// If no package specified, use package from file.
-		if (!InOuter)
-		{
-			if( !FilenamePkg )
-			{
-				LogGetPackageLinkerError(Result, InExistingContext, InLongPackageName, LOCTEXT("FilenameToPackageShort", "Can't convert filename to asset name"), InOuter, LoadFlags);
-				return nullptr;
-			}
-			InOuter = FilenamePkg;
-			Result = FLinkerLoad::FindExistingLinkerForPackage(InOuter);
+			LogGetPackageLinkerError(InExistingContext, InLongPackageName, LOCTEXT("FilenameToPackageShort", "Can't convert filename to asset name"), InOuter, LoadFlags);
+			return nullptr;
 		}
-		else if (InOuter != FilenamePkg && FLinkerLoad::FindExistingLinkerForPackage(InOuter)) //!!should be tested and validated in new UnrealEd
+		if (LoadFlags & LOAD_PackageForPIE)
 		{
-			// Loading a new file into an existing package, so reset the loader.
-			//UE_LOG(LogLinker, Log,  TEXT("New File, Existing Package (%s, %s)"), *InOuter->GetFullName(), *FilenamePkg->GetFullName() );
-			ResetLoaders( InOuter );
+			CreatedPackage->SetPackageFlags(PKG_PlayInEditor);
 		}
+		TargetPackage = CreatedPackage;
 	}
-
-#if 0
-	// Make sure the package is accessible in the sandbox.
-	if( Sandbox && !Sandbox->SupportsPackage(InOuter) )
-	{
-		LogGetPackageLinkerError(Result, InExistingContext, InLongPackageName, LOCTEXT("SandboxShort", "Asset is not accessible in this sandbox"), InOuter, LoadFlags);
-		return nullptr;
-	}
-#endif
-
-	// Create new linker.
-	if( !Result )
-	{
-		// we will already have found the filename above
-		check(NewFilename.Len() > 0);
-		TRefCountPtr<FUObjectSerializeContext> LoadContext(FUObjectThreadContext::Get().GetSerializeContext());
-		Result = FLinkerLoad::CreateLinker(LoadContext, InOuter, *NewFilename, LoadFlags, InReaderOverride, ImportLinker ? &ImportLinker->GetInstancingContext() : InstancingContext);
-	}
-	else if (InExistingContext)
+	if (InOuter != TargetPackage)
 	{
-		if ((Result->GetSerializeContext() && Result->GetSerializeContext()->HasStartedLoading() && InExistingContext->GetBeginLoadCount() == 1) ||
-			  (IsInAsyncLoadingThread() && Result->GetSerializeContext()))
+		if (FLinkerLoad* Result = FLinkerLoad::FindExistingLinkerForPackage(TargetPackage))
 		{
-			// Use the context associated with the linker because it has already started loading objects (or we're in ALT where each package needs its own context)
-			*InOutLoadContext = Result->GetSerializeContext();
-		}
-		else
-		{
-			if (Result->GetSerializeContext() && Result->GetSerializeContext() != InExistingContext)
+			if (InExistingContext)
 			{
-				// Make sure the objects already loaded with the context associated with the existing linker
-				// are copied to the context provided for this function call to make sure they all get loaded ASAP
-				InExistingContext->AddUniqueLoadedObjects(Result->GetSerializeContext()->PRIVATE_GetObjectsLoadedInternalUseOnly());
+				if ((Result->GetSerializeContext() && Result->GetSerializeContext()->HasStartedLoading() && InExistingContext->GetBeginLoadCount() == 1) ||
+					(IsInAsyncLoadingThread() && Result->GetSerializeContext()))
+				{
+					// Use the context associated with the linker because it has already started loading objects (or we're in ALT where each package needs its own context)
+					*InOutLoadContext = Result->GetSerializeContext();
+				}
+				else
+				{
+					if (Result->GetSerializeContext() && Result->GetSerializeContext() != InExistingContext)
+					{
+						// Make sure the objects already loaded with the context associated with the existing linker
+						// are copied to the context provided for this function call to make sure they all get loaded ASAP
+						InExistingContext->AddUniqueLoadedObjects(Result->GetSerializeContext()->PRIVATE_GetObjectsLoadedInternalUseOnly());
+					}
+					// Replace the linker context with the one passed into this function
+					Result->SetSerializeContext(InExistingContext);
+				}
 			}
-			// Replace the linker context with the one passed into this function
-			Result->SetSerializeContext(InExistingContext);
+			return Result;
 		}
 	}
 
-	if ( !Result && CreatedPackage )
-	{
-		// kill it with fire
-		CreatedPackage->MarkPendingKill();
-		/*static int32 FailedPackageLoadIndex = 0;
-		++FailedPackageLoadIndex;
-		FString FailedLinkerLoad = FString::Printf(TEXT("/Temp/FailedLinker_%s_%d"), *NewFilename, FailedPackageLoadIndex);
-		
-		CreatedPackage->Rename(*FailedLinkerLoad);*/
-	}
+	// Create new linker.
+	// we will already have found the filename above
+	check(NewFilename.Len() > 0);
+	TRefCountPtr<FUObjectSerializeContext> LoadContext(InExistingContext ? InExistingContext : FUObjectThreadContext::Get().GetSerializeContext());
+	FLinkerLoad* Result = FLinkerLoad::CreateLinker(LoadContext, TargetPackage, *NewFilename, LoadFlags, InReaderOverride, ImportLinker ? &ImportLinker->GetInstancingContext() : InstancingContext);
 
-	// Verify compatibility.
-	if (Result && CompatibleGuid && Result->Summary.Guid != *CompatibleGuid)
+	if (!Result && CreatedPackage)
 	{
-		// This should never fire, because FindPackageFile should never return an incompatible file
-		LogGetPackageLinkerError(Result, InExistingContext, InLongPackageName, LOCTEXT("PackageVersionShort", "Asset version mismatch"), InOuter, LoadFlags);
-		return nullptr;
+		CreatedPackage->MarkPendingKill();
 	}
 
 	return Result;
diff --git a/Engine/Source/Runtime/CoreUObject/Public/UObject/ObjectMacros.h b/Engine/Source/Runtime/CoreUObject/Public/UObject/ObjectMacros.h
--- a/Engine/Source/Runtime/CoreUObject/Public/UObject/ObjectMacros.h
+++ b/Engine/Source/Runtime/CoreUObject/Public/UObject/ObjectMacros.h
@@ -61,7 +61,7 @@ enum ELoadFlags
 	LOAD_EditorOnly					= 0x00000004,	///< Load for editor-only purposes and by editor-only code
 	LOAD_ResolvingDeferredExports	= 0x00000008,	///< Denotes that we should not defer export loading (as we're resolving them)
 	LOAD_Verify						= 0x00000010,	///< Only verify existance; don't actually load.
-	LOAD_AllowDll					= 0x00000020,	///< Allow plain DLLs.
+//	LOAD_Unused						= 0x00000020,	///< Allow plain DLLs.
 //	LOAD_Unused						= 0x00000040
 	LOAD_NoVerify					= 0x00000080,   ///< Don't verify imports yet.
 	LOAD_IsVerifying				= 0x00000100,	///< Is verifying imports
