Add some useful Insights trace

diff --git a/Engine/Plugins/FX/Niagara/Source/Niagara/Private/NiagaraSystem.cpp b/Engine/Plugins/FX/Niagara/Source/Niagara/Private/NiagaraSystem.cpp
--- a/Engine/Plugins/FX/Niagara/Source/Niagara/Private/NiagaraSystem.cpp
+++ b/Engine/Plugins/FX/Niagara/Source/Niagara/Private/NiagaraSystem.cpp
@@ -1873,6 +1873,8 @@ void UNiagaraSystem::InitEmitterDataSetCompiledData(FNiagaraDataSetCompiledData&
 
 bool UNiagaraSystem::RequestCompile(bool bForce, FNiagaraSystemUpdateContext* OptionalUpdateContext)
 {
+	TRACE_CPUPROFILER_EVENT_SCOPE(UNiagaraSystem::RequestCompile)
+
 	// We remove emitters and scripts on dedicated servers, so skip further work.
 	const bool bIsDedicatedServer = !GIsClient && GIsServer;
 	if (bIsDedicatedServer)
diff --git a/Engine/Source/Developer/DesktopPlatform/Private/DesktopPlatformBase.cpp b/Engine/Source/Developer/DesktopPlatform/Private/DesktopPlatformBase.cpp
--- a/Engine/Source/Developer/DesktopPlatform/Private/DesktopPlatformBase.cpp
+++ b/Engine/Source/Developer/DesktopPlatform/Private/DesktopPlatformBase.cpp
@@ -660,6 +660,7 @@ bool FDesktopPlatformBase::IsUnrealBuildToolAvailable()
 
 bool FDesktopPlatformBase::InvokeUnrealBuildToolSync(const FString& InCmdLineParams, FOutputDevice &Ar, bool bSkipBuildUBT, int32& OutReturnCode, FString& OutProcOutput)
 {
+	TRACE_CPUPROFILER_EVENT_SCOPE(FDesktopPlatformBase::InvokeUnrealBuildToolSync);
 	void* PipeRead = nullptr;
 	void* PipeWrite = nullptr;
 
diff --git a/Engine/Source/Developer/ShaderPreprocessor/Private/ShaderPreprocessor.cpp b/Engine/Source/Developer/ShaderPreprocessor/Private/ShaderPreprocessor.cpp
--- a/Engine/Source/Developer/ShaderPreprocessor/Private/ShaderPreprocessor.cpp
+++ b/Engine/Source/Developer/ShaderPreprocessor/Private/ShaderPreprocessor.cpp
@@ -222,6 +222,8 @@ bool PreprocessShader(
 	EDumpShaderDefines DefinesPolicy
 	)
 {
+	TRACE_CPUPROFILER_EVENT_SCOPE(PreprocessShader);
+
 	// Skip the cache system and directly load the file path (used for debugging)
 	if (ShaderInput.bSkipPreprocessedCache)
 	{
diff --git a/Engine/Source/Developer/TargetPlatform/Private/TargetPlatformManagerModule.cpp b/Engine/Source/Developer/TargetPlatform/Private/TargetPlatformManagerModule.cpp
--- a/Engine/Source/Developer/TargetPlatform/Private/TargetPlatformManagerModule.cpp
+++ b/Engine/Source/Developer/TargetPlatform/Private/TargetPlatformManagerModule.cpp
@@ -1024,10 +1024,14 @@ protected:
 		TSharedPtr<FMonitoredProcess> UBTProcess = MakeShareable(new FMonitoredProcess(CmdExe, CommandLine, true));
 		UBTProcess->OnOutput().BindStatic(&FTargetPlatformManagerModule::OnStatusOutput);
 		SDKStatusMessage = TEXT("");
-		UBTProcess->Launch();
-		while(UBTProcess->Update())
+
 		{
-			FPlatformProcess::Sleep(0.01f);
+			TRACE_CPUPROFILER_EVENT_SCOPE(WaitUntilUBTStarted);
+			UBTProcess->Launch();
+			while(UBTProcess->Update())
+			{
+				FPlatformProcess::Sleep(0.01f);
+			}
 		}
 
 		TArray<FString> PlatArray;
diff --git a/Engine/Source/Editor/Kismet/Private/BlueprintActionMenuBuilder.cpp b/Engine/Source/Editor/Kismet/Private/BlueprintActionMenuBuilder.cpp
--- a/Engine/Source/Editor/Kismet/Private/BlueprintActionMenuBuilder.cpp
+++ b/Engine/Source/Editor/Kismet/Private/BlueprintActionMenuBuilder.cpp
@@ -508,6 +508,8 @@ void FBlueprintActionMenuBuilder::AddMenuSection(FBlueprintActionFilter const& F
 //------------------------------------------------------------------------------
 void FBlueprintActionMenuBuilder::RebuildActionList()
 {
+	TRACE_CPUPROFILER_EVENT_SCOPE(FBlueprintActionMenuBuilder::RebuildActionList);
+
 	using namespace FBlueprintActionMenuBuilderImpl;
 
 	FGraphActionListBuilderBase::Empty();
diff --git a/Engine/Source/Editor/Kismet/Private/BlueprintCompilationManager.cpp b/Engine/Source/Editor/Kismet/Private/BlueprintCompilationManager.cpp
--- a/Engine/Source/Editor/Kismet/Private/BlueprintCompilationManager.cpp
+++ b/Engine/Source/Editor/Kismet/Private/BlueprintCompilationManager.cpp
@@ -518,6 +518,7 @@ namespace SkelReinstUtils
 
 void FBlueprintCompilationManagerImpl::FlushCompilationQueueImpl(bool bSuppressBroadcastCompiled, TArray<UBlueprint*>* BlueprintsCompiled, TArray<UBlueprint*>* BlueprintsCompiledOrSkeletonCompiled, FUObjectSerializeContext* InLoadContext)
 {
+	TRACE_CPUPROFILER_EVENT_SCOPE(FBlueprintCompilationManager::FlushCompilationQueue);
 	DECLARE_SCOPE_HIERARCHICAL_COUNTER_FUNC()
 
 	TGuardValue<bool> GuardTemplateNameFlag(GCompilingBlueprint, true);
diff --git a/Engine/Source/Editor/MainFrame/Private/MainFrameModule.cpp b/Engine/Source/Editor/MainFrame/Private/MainFrameModule.cpp
--- a/Engine/Source/Editor/MainFrame/Private/MainFrameModule.cpp
+++ b/Engine/Source/Editor/MainFrame/Private/MainFrameModule.cpp
@@ -72,6 +72,8 @@ const FText StaticGetApplicationTitle( const bool bIncludeGameName )
 
 void FMainFrameModule::CreateDefaultMainFrame( const bool bStartImmersive, const bool bStartPIE )
 {
+	TRACE_CPUPROFILER_EVENT_SCOPE(FMainFrameModule::CreateDefaultMainFrame);
+
 	if (!IsWindowInitialized())
 	{
 		FRootWindowLocation DefaultWindowLocation;
diff --git a/Engine/Source/Editor/UnrealEd/Private/EditorServer.cpp b/Engine/Source/Editor/UnrealEd/Private/EditorServer.cpp
--- a/Engine/Source/Editor/UnrealEd/Private/EditorServer.cpp
+++ b/Engine/Source/Editor/UnrealEd/Private/EditorServer.cpp
@@ -2418,6 +2418,8 @@ bool UEditorEngine::PackageIsAMapFile( const TCHAR* PackageFilename, FText& OutN
 
 bool UEditorEngine::Map_Load(const TCHAR* Str, FOutputDevice& Ar)
 {
+	TRACE_CPUPROFILER_EVENT_SCOPE(UEditorEngine::Map_Load);
+
 	auto FindWorldInPackageOrFollowRedirector = [](UPackage*& InOutPackage)
 	{
 		UWorld* RetVal = nullptr;
diff --git a/Engine/Source/Editor/UnrealEd/Private/PlayLevel.cpp b/Engine/Source/Editor/UnrealEd/Private/PlayLevel.cpp
--- a/Engine/Source/Editor/UnrealEd/Private/PlayLevel.cpp
+++ b/Engine/Source/Editor/UnrealEd/Private/PlayLevel.cpp
@@ -2409,6 +2409,8 @@ void UEditorEngine::RemapGamepadControllerIdForPIE(class UGameViewportClient* In
 
 void UEditorEngine::StartPlayInEditorSession(FRequestPlaySessionParams& InRequestParams)
 {
+	TRACE_CPUPROFILER_EVENT_SCOPE(UEditorEngine::StartPlayInEditorSession);
+
 	// This reflects that the user has tried to launch a PIE session, but it may still
 	// create one-or-more new processes depending on multiplayer settings.
 	check(InRequestParams.SessionDestination == EPlaySessionDestinationType::InProcess);
diff --git a/Engine/Source/Editor/UnrealEd/Private/UnrealEdMisc.cpp b/Engine/Source/Editor/UnrealEd/Private/UnrealEdMisc.cpp
--- a/Engine/Source/Editor/UnrealEd/Private/UnrealEdMisc.cpp
+++ b/Engine/Source/Editor/UnrealEd/Private/UnrealEdMisc.cpp
@@ -535,6 +535,8 @@ void FUnrealEdMisc::InitEngineAnalytics()
 {
 	if ( FEngineAnalytics::IsAvailable() )
 	{
+		TRACE_CPUPROFILER_EVENT_SCOPE(FUnrealEdMisc::InitEngineAnalytics);
+
 		IAnalyticsProvider& EngineAnalytics = FEngineAnalytics::GetProvider();
 
 		// Send analytics about sample projects
diff --git a/Engine/Source/Runtime/ApplicationCore/Private/Windows/WindowsPlatformApplicationMisc.cpp b/Engine/Source/Runtime/ApplicationCore/Private/Windows/WindowsPlatformApplicationMisc.cpp
--- a/Engine/Source/Runtime/ApplicationCore/Private/Windows/WindowsPlatformApplicationMisc.cpp
+++ b/Engine/Source/Runtime/ApplicationCore/Private/Windows/WindowsPlatformApplicationMisc.cpp
@@ -100,6 +100,7 @@ int32 FWindowsPlatformApplicationMisc::GetAppIcon()
 
 static void WinPumpMessages()
 {
+	TRACE_CPUPROFILER_EVENT_SCOPE(WinPumpMessages);
 	{
 		MSG Msg;
 		while( PeekMessage(&Msg,NULL,0,0,PM_REMOVE) )
diff --git a/Engine/Source/Runtime/ApplicationCore/Private/Windows/WindowsPlatformSplash.cpp b/Engine/Source/Runtime/ApplicationCore/Private/Windows/WindowsPlatformSplash.cpp
--- a/Engine/Source/Runtime/ApplicationCore/Private/Windows/WindowsPlatformSplash.cpp
+++ b/Engine/Source/Runtime/ApplicationCore/Private/Windows/WindowsPlatformSplash.cpp
@@ -697,6 +697,7 @@ void FWindowsPlatformSplash::Hide()
 {
 	if(GSplashScreenThread)
 	{
+		TRACE_CPUPROFILER_EVENT_SCOPE(FWindowsPlatformSplash::Hide);
 		if(GSplashScreenWnd)
 		{
 			// Send message to splash screen window to destroy itself
diff --git a/Engine/Source/Runtime/Core/Private/Internationalization/Internationalization.cpp b/Engine/Source/Runtime/Core/Private/Internationalization/Internationalization.cpp
--- a/Engine/Source/Runtime/Core/Private/Internationalization/Internationalization.cpp
+++ b/Engine/Source/Runtime/Core/Private/Internationalization/Internationalization.cpp
@@ -392,6 +392,7 @@ FString& FInternationalization::Leetify(FString& SourceString)
 
 void FInternationalization::LoadAllCultureData()
 {
+	TRACE_CPUPROFILER_EVENT_SCOPE(FInternationalization::LoadAllCultureData);
 	Implementation->LoadAllCultureData();
 }
 
diff --git a/Engine/Source/Runtime/Core/Private/Modules/ModuleManager.cpp b/Engine/Source/Runtime/Core/Private/Modules/ModuleManager.cpp
--- a/Engine/Source/Runtime/Core/Private/Modules/ModuleManager.cpp
+++ b/Engine/Source/Runtime/Core/Private/Modules/ModuleManager.cpp
@@ -376,6 +376,7 @@ IModuleInterface* FModuleManager::LoadModuleWithFailureReason(const FName InModu
 		}
 	}
 
+	TRACE_CPUPROFILER_EVENT_SCOPE_TEXT(*FString::Printf(TEXT("LoadModule %s"), *InModuleName.ToString()));
 	DECLARE_SCOPE_CYCLE_COUNTER(TEXT("Module Load"), STAT_ModuleLoad, STATGROUP_LoadTime);
 #if	STATS
 	// This is fine here, we only load a handful of modules.
@@ -521,6 +522,8 @@ IModuleInterface* FModuleManager::LoadModuleWithFailureReason(const FName InModu
 
 						if ( ModuleInfo->Module.IsValid() )
 						{
+							TRACE_CPUPROFILER_EVENT_SCOPE_TEXT(*(InModuleName.ToString() + TEXT("::StartupModule")));
+
 							// Startup the module
 							ModuleInfo->Module->StartupModule();
 							// The module might try to load other dependent modules in StartupModule. In this case, we want those modules shut down AFTER this one because we may still depend on the module at shutdown.
diff --git a/Engine/Source/Runtime/Core/Public/Async/ParallelFor.h b/Engine/Source/Runtime/Core/Public/Async/ParallelFor.h
--- a/Engine/Source/Runtime/Core/Public/Async/ParallelFor.h
+++ b/Engine/Source/Runtime/Core/Public/Async/ParallelFor.h
@@ -143,6 +143,8 @@ namespace ParallelForImpl
 	template<typename FunctionType>
 	inline bool TParallelForData<FunctionType>::Process(int32 TasksToSpawn, TSharedRef<TParallelForData<FunctionType>, ESPMode::ThreadSafe>& Data, ENamedThreads::Type InDesiredThread, bool bMaster)
 	{
+		TRACE_CPUPROFILER_EVENT_SCOPE(TParallelForData::Process)
+
 		int32 MaybeTasksLeft = Num - IndexToDo.GetValue();
 		if (TasksToSpawn && MaybeTasksLeft > 0)
 		{
diff --git a/Engine/Source/Runtime/CoreUObject/Private/Misc/PackageName.cpp b/Engine/Source/Runtime/CoreUObject/Private/Misc/PackageName.cpp
--- a/Engine/Source/Runtime/CoreUObject/Private/Misc/PackageName.cpp
+++ b/Engine/Source/Runtime/CoreUObject/Private/Misc/PackageName.cpp
@@ -1051,6 +1051,7 @@ bool FPackageName::DoesPackageExist(const FString& LongPackageName, const FGuid*
 
 bool FPackageName::SearchForPackageOnDisk(const FString& PackageName, FString* OutLongPackageName, FString* OutFilename)
 {
+	TRACE_CPUPROFILER_EVENT_SCOPE(FPackageName::SearchForPackageOnDisk);
 	DECLARE_SCOPE_CYCLE_COUNTER(TEXT("FPackageName::SearchForPackageOnDisk"), STAT_PackageName_SearchForPackageOnDisk, STATGROUP_LoadTime);
 	
 	// This function may take a long time to complete, so suspend heartbeat measure while we're here
diff --git a/Engine/Source/Runtime/CoreUObject/Private/UObject/LinkerLoad.cpp b/Engine/Source/Runtime/CoreUObject/Private/UObject/LinkerLoad.cpp
--- a/Engine/Source/Runtime/CoreUObject/Private/UObject/LinkerLoad.cpp
+++ b/Engine/Source/Runtime/CoreUObject/Private/UObject/LinkerLoad.cpp
@@ -4044,6 +4044,10 @@ void FLinkerLoad::Preload( UObject* Object )
 						else
 #endif
 						{
+#if WITH_EDITOR
+							TRACE_CPUPROFILER_EVENT_SCOPE_TEXT(*(Object->GetClass()->GetName() + TEXT("_Serialize")));
+							TRACE_CPUPROFILER_EVENT_SCOPE_TEXT(*Object->GetFullName());
+#endif
 							Object->Serialize(*this);
 						}
 
diff --git a/Engine/Source/Runtime/CoreUObject/Private/UObject/UObjectHash.cpp b/Engine/Source/Runtime/CoreUObject/Private/UObject/UObjectHash.cpp
--- a/Engine/Source/Runtime/CoreUObject/Private/UObject/UObjectHash.cpp
+++ b/Engine/Source/Runtime/CoreUObject/Private/UObject/UObjectHash.cpp
@@ -1157,6 +1157,8 @@ void GetObjectsOfClass(const UClass* ClassToLookFor, TArray<UObject *>& Results,
 
 FORCEINLINE void ForEachObjectOfClasses_Implementation(FUObjectHashTables& ThreadHash, TArrayView<const UClass*> ClassesToLookFor, TFunctionRef<void(UObject*)> Operation, EObjectFlags ExcludeFlags /*= RF_ClassDefaultObject*/, EInternalObjectFlags ExclusionInternalFlags /*= EInternalObjectFlags::None*/)
 {
+	TRACE_CPUPROFILER_EVENT_SCOPE(ForEachObjectOfClasses_Implementation);
+
 	// We don't want to return any objects that are currently being background loaded unless we're using the object iterator during async loading.
 	ExclusionInternalFlags |= EInternalObjectFlags::Unreachable;
 	if (!IsInAsyncLoadingThread())
@@ -1183,6 +1185,8 @@ FORCEINLINE void ForEachObjectOfClasses_Implementation(FUObjectHashTables& Threa
 
 void ForEachObjectOfClass(const UClass* ClassToLookFor, TFunctionRef<void(UObject*)> Operation, bool bIncludeDerivedClasses, EObjectFlags ExclusionFlags, EInternalObjectFlags ExclusionInternalFlags)
 {
+	TRACE_CPUPROFILER_EVENT_SCOPE(ForEachObjectOfClass);
+
 	// Most classes searched for have around 10 subclasses, some have hundreds
 	TArray<const UClass*, TInlineAllocator<16>> ClassesToSearch;
 	ClassesToSearch.Add(ClassToLookFor);
diff --git a/Engine/Source/Runtime/Engine/Private/ShaderCompiler/ShaderCompiler.cpp b/Engine/Source/Runtime/Engine/Private/ShaderCompiler/ShaderCompiler.cpp
--- a/Engine/Source/Runtime/Engine/Private/ShaderCompiler/ShaderCompiler.cpp
+++ b/Engine/Source/Runtime/Engine/Private/ShaderCompiler/ShaderCompiler.cpp
@@ -3159,6 +3159,8 @@ void FShaderCompilingManager::FinishAllCompilation()
 
 void FShaderCompilingManager::ProcessAsyncResults(bool bLimitExecutionTime, bool bBlockOnGlobalShaderCompletion)
 {
+	TRACE_CPUPROFILER_EVENT_SCOPE(FShaderCompilingManager::ProcessAsyncResults)
+
 	COOK_STAT(FScopedDurationTimer Timer(ShaderCompilerCookStats::ProcessAsyncResultsTimeSec));
 	check(IsInGameThread());
 	if (bAllowAsynchronousShaderCompiling)
diff --git a/Engine/Source/Runtime/Engine/Private/SoundWave.cpp b/Engine/Source/Runtime/Engine/Private/SoundWave.cpp
--- a/Engine/Source/Runtime/Engine/Private/SoundWave.cpp
+++ b/Engine/Source/Runtime/Engine/Private/SoundWave.cpp
@@ -762,6 +762,7 @@ FByteBulkData* USoundWave::GetCompressedData(FName Format, const FPlatformAudioC
 		return nullptr;
 	}
 
+	TRACE_CPUPROFILER_EVENT_SCOPE(USoundWave::GetCompressedData);
 	FName PlatformSpecificFormat = GetPlatformSpecificFormat(Format, CompressionOverrides);
 
 	bool bContainedValidData = CompressedFormatData.Contains(PlatformSpecificFormat);
diff --git a/Engine/Source/Runtime/Engine/Private/WorldPartition/WorldPartition.cpp b/Engine/Source/Runtime/Engine/Private/WorldPartition/WorldPartition.cpp
--- a/Engine/Source/Runtime/Engine/Private/WorldPartition/WorldPartition.cpp
+++ b/Engine/Source/Runtime/Engine/Private/WorldPartition/WorldPartition.cpp
@@ -523,6 +523,8 @@ void UWorldPartition::LoadEditorCells(const FBox& Box)
 	TArray<UWorldPartitionEditorCell*> CellsToLoad;
 	if (EditorHash->GetIntersectingCells(Box, CellsToLoad))
 	{
+		TRACE_CPUPROFILER_EVENT_SCOPE(UWorldPartition::LoadEditorCells);
+
 		FWorldPartionCellUpdateContext CellUpdateContext(this);
 
 		int32 NumActorsToLoad = Algo::TransformAccumulate(CellsToLoad, [](UWorldPartitionEditorCell* Cell) { return Cell->Actors.Num() - Cell->LoadedActors.Num();}, 0);
@@ -617,6 +619,9 @@ void UWorldPartition::UnloadEditorCells(const FBox& Box)
 // Loads actors in Editor cell
 void UWorldPartition::UpdateLoadingEditorCell(UWorldPartitionEditorCell* Cell, bool bShouldBeLoaded)
 {
+	TRACE_CPUPROFILER_EVENT_SCOPE(UWorldPartition::UpdateLoadingEditorCell);
+	TRACE_CPUPROFILER_EVENT_SCOPE_TEXT(*FString::Printf(TEXT("UWorldPartition::UpdateLoadingEditorCell(%s)"), *Cell->GetFullName()));
+	
 	FWorldPartionCellUpdateContext CellUpdateContext(this);
 
 	UE_LOG(LogWorldPartition, Verbose, TEXT("UWorldPartition::UpdateLoadingEditorCell 0x%08llx [%s]"), Cell, bShouldBeLoaded ? TEXT("Load") : TEXT("Unload"));
diff --git a/Engine/Source/Runtime/RenderCore/Private/RenderingThread.cpp b/Engine/Source/Runtime/RenderCore/Private/RenderingThread.cpp
--- a/Engine/Source/Runtime/RenderCore/Private/RenderingThread.cpp
+++ b/Engine/Source/Runtime/RenderCore/Private/RenderingThread.cpp
@@ -1123,6 +1123,7 @@ static FAutoConsoleVariableRef CVarTimeoutForBlockOnRenderFence(
  */
 static void GameThreadWaitForTask(const FGraphEventRef& Task, ENamedThreads::Type TriggerThreadIndex = ENamedThreads::ActualRenderingThread, bool bEmptyGameThreadTasks = false)
 {
+	TRACE_CPUPROFILER_EVENT_SCOPE(GameThreadWaitForTask);
 	SCOPE_TIME_GUARD(TEXT("GameThreadWaitForTask"));
 
 	check(IsInGameThread());
@@ -1261,6 +1262,8 @@ void FlushRenderingCommands(bool bFlushDeferredDeletes)
 	{
 		return;
 	}
+
+	TRACE_CPUPROFILER_EVENT_SCOPE(FlushRenderingCommands);
 	FSuspendRenderingTickables SuspendRenderingTickables;
 
 	// Need to flush GT because render commands from threads other than GT are sent to
diff --git a/Engine/Source/Runtime/Renderer/Private/RendererScene.cpp b/Engine/Source/Runtime/Renderer/Private/RendererScene.cpp
--- a/Engine/Source/Runtime/Renderer/Private/RendererScene.cpp
+++ b/Engine/Source/Runtime/Renderer/Private/RendererScene.cpp
@@ -3348,6 +3348,8 @@ bool FScene::RequiresHitProxies() const
 
 void FScene::Release()
 {
+	TRACE_CPUPROFILER_EVENT_SCOPE(FScene::Release);
+
 #if !(UE_BUILD_SHIPPING || UE_BUILD_TEST)
 	// Verify that no components reference this scene being destroyed
 	static bool bTriggeredOnce = false;
@@ -3769,6 +3771,7 @@ struct FPrimitiveArraySortKey
 
 void FScene::UpdateAllPrimitiveSceneInfos(FRHICommandListImmediate& RHICmdList, bool bAsyncCreateLPIs)
 {
+	TRACE_CPUPROFILER_EVENT_SCOPE(Scene::UpdateAllPrimitiveSceneInfos);
 	SCOPED_NAMED_EVENT(FScene_UpdateAllPrimitiveSceneInfos, FColor::Orange);
 	SCOPE_CYCLE_COUNTER(STAT_UpdateScenePrimitiveRenderThreadTime);
 
@@ -4579,6 +4582,8 @@ void FRendererModule::UpdateStaticDrawLists()
 
 void UpdateStaticMeshesForMaterials(const TArray<const FMaterial*>& MaterialResourcesToUpdate)
 {
+	TRACE_CPUPROFILER_EVENT_SCOPE(UpdateStaticMeshesForMaterials);
+
 	TArray<UMaterialInterface*> UsedMaterials;
 	TSet<UMaterialInterface*> UsedMaterialsDependencies;
 	TMap<FScene*, TArray<FPrimitiveSceneInfo*>> UsedPrimitives;
