Copying //UE5/Dev-Cooker@14539516 to Main (//UE5/Main)

diff --git a/Engine/Source/Runtime/CoreUObject/Private/Tests/ObjectPtrTest.cpp b/Engine/Source/Runtime/CoreUObject/Private/Tests/ObjectPtrTest.cpp
new file mode 100644
--- /dev/null
+++ b/Engine/Source/Runtime/CoreUObject/Private/Tests/ObjectPtrTest.cpp
@@ -0,0 +1,242 @@
+// Copyright Epic Games, Inc. All Rights Reserved.
+
+#include "UObject/ObjectPtr.h"
+
+#include "ObjectRefTrackingTestBase.h"
+
+#include "Concepts/EqualityComparable.h"
+#include "Serialization/ArchiveCountMem.h"
+#include "Templates/Models.h"
+#include "UObject/Interface.h"
+#include "UObject/SoftObjectPath.h"
+
+#include <type_traits>
+
+namespace ObjectPtrTest
+{
+using FMutableObjectPtr = TObjectPtr<UObject>;
+using FMutableInterfacePtr = TObjectPtr<UInterface>;
+using FMutablePackagePtr = TObjectPtr<UPackage>;
+using FConstObjectPtr = TObjectPtr<const UObject>;
+using FConstInterfacePtr = TObjectPtr<const UInterface>;
+using FConstPackagePtr = TObjectPtr<const UPackage>;
+
+class UForwardDeclaredObjDerived;
+class FForwardDeclaredNotObjDerived;
+
+static_assert(sizeof(FObjectPtr) == sizeof(FObjectHandle), "FObjectPtr type must always compile to something equivalent to an FObjectHandle size.");
+static_assert(sizeof(FObjectPtr) == sizeof(void*), "FObjectPtr type must always compile to something equivalent to a pointer size.");
+static_assert(sizeof(TObjectPtr<UObject>) == sizeof(void*), "TObjectPtr<UObject> type must always compile to something equivalent to a pointer size.");
+
+// Ensure that a TObjectPtr is trivially copyable, (copy/move) constructible, (copy/move) assignable, and destructible
+static_assert(std::is_trivially_copyable<FMutableObjectPtr>::value, "TObjectPtr must be trivially copyable");
+static_assert(std::is_trivially_copy_constructible<FMutableObjectPtr>::value, "TObjectPtr must be trivially copy constructible");
+static_assert(std::is_trivially_move_constructible<FMutableObjectPtr>::value, "TObjectPtr must be trivially move constructible");
+static_assert(std::is_trivially_copy_assignable<FMutableObjectPtr>::value, "TObjectPtr must be trivially copy assignable");
+static_assert(std::is_trivially_move_assignable<FMutableObjectPtr>::value, "TObjectPtr must be trivially move assignable");
+static_assert(std::is_trivially_destructible<FMutableObjectPtr>::value, "TObjectPtr must be trivially destructible");
+static_assert(std::is_trivially_default_constructible<FMutableObjectPtr>::value, "TObjectPtr must be trivially default constructible");
+
+// Ensure that raw pointers can be used to construct wrapped object pointers and that const-ness isn't stripped when constructing or converting with raw pointers
+static_assert(std::is_constructible<FMutableObjectPtr, UObject*>::value, "TObjectPtr<UObject> must be constructible from a raw UObject*");
+static_assert(!std::is_constructible<FMutableObjectPtr, const UObject*>::value, "TObjectPtr<UObject> must not be constructible from a const raw UObject*");
+static_assert(std::is_convertible<FMutableObjectPtr, UObject*>::value, "TObjectPtr<UObject> must be convertible to a raw UObject*");
+static_assert(std::is_convertible<FMutableObjectPtr, const UObject*>::value, "TObjectPtr<UObject> must be convertible to a const raw UObject*");
+
+static_assert(std::is_constructible<FConstObjectPtr, UObject*>::value, "TObjectPtr<const UObject> must be constructible from a raw UObject*");
+static_assert(std::is_constructible<FConstObjectPtr, const UObject*>::value, "TObjectPtr<const UObject> must be constructible from a const raw UObject*");
+static_assert(!std::is_convertible<FConstObjectPtr, UObject*>::value, "TObjectPtr<const UObject> must not be convertible to a raw UObject*");
+static_assert(std::is_convertible<FConstObjectPtr, const UObject*>::value, "TObjectPtr<const UObject> must be convertible to a const raw UObject*");
+
+// Ensure that a TObjectPtr<const UObject> is constructible and assignable from a TObjectPtr<UObject> but not vice versa
+static_assert(std::is_constructible<FConstObjectPtr, const FMutableObjectPtr&>::value, "Missing constructor (TObjectPtr<const UObject> from TObjectPtr<UObject>)");
+static_assert(!std::is_constructible<FMutableObjectPtr, const FConstObjectPtr&>::value, "Invalid constructor (TObjectPtr<UObject> from TObjectPtr<const UObject>)");
+static_assert(std::is_assignable<FConstObjectPtr, const FMutableObjectPtr&>::value, "Missing assignment (TObjectPtr<const UObject> from TObjectPtr<UObject>)");
+static_assert(!std::is_assignable<FMutableObjectPtr, const FConstObjectPtr&>::value, "Invalid assignment (TObjectPtr<UObject> from TObjectPtr<const UObject>)");
+
+static_assert(std::is_constructible<FConstObjectPtr, const FConstObjectPtr&>::value, "Missing constructor (TObjectPtr<const UObject> from TObjectPtr<const UObject>)");
+static_assert(std::is_assignable<FConstObjectPtr, const FConstObjectPtr&>::value, "Missing assignment (TObjectPtr<const UObject> from TObjectPtr<const UObject>)");
+
+// Ensure that a TObjectPtr<UObject> is constructible and assignable from a TObjectPtr<UInterface> but not vice versa
+static_assert(std::is_constructible<FMutableObjectPtr, const FMutableInterfacePtr&>::value, "Missing constructor (TObjectPtr<UObject> from TObjectPtr<UInterface>)");
+static_assert(!std::is_constructible<FMutableInterfacePtr, const FMutableObjectPtr&>::value, "Invalid constructor (TObjectPtr<UInterface> from TObjectPtr<UObject>)");
+static_assert(std::is_constructible<FConstObjectPtr, const FConstInterfacePtr&>::value, "Missing constructor (TObjectPtr<const UObject> from TObjectPtr<const UInterface>)");
+static_assert(std::is_constructible<FConstObjectPtr, const FMutableInterfacePtr&>::value, "Missing constructor (TObjectPtr<const UObject> from TObjectPtr<UInterface>)");
+static_assert(!std::is_constructible<FConstInterfacePtr, const FConstObjectPtr&>::value, "Invalid constructor (TObjectPtr<const UInterface> from TObjectPtr<const UObject>)");
+static_assert(!std::is_constructible<FConstInterfacePtr, const FMutableObjectPtr&>::value, "Invalid constructor (TObjectPtr<const UInterface> from TObjectPtr<UObject>)");
+
+static_assert(std::is_assignable<FMutableObjectPtr, const FMutableInterfacePtr&>::value, "Missing assignment (TObjectPtr<UObject> from TObjectPtr<UInterface>)");
+static_assert(std::is_assignable<FConstObjectPtr, const FMutableInterfacePtr&>::value, "Missing assignment (TObjectPtr<const UObject> from TObjectPtr<UInterface>)");
+static_assert(std::is_assignable<FConstObjectPtr, const FConstInterfacePtr&>::value, "Missing assignment (TObjectPtr<const UObject> from TObjectPtr<const UInterface>)");
+static_assert(!std::is_assignable<FMutableInterfacePtr, const FMutableObjectPtr&>::value, "Invalid assignment (TObjectPtr<UInterface> from TObjectPtr<UObject>)");
+static_assert(!std::is_assignable<FConstInterfacePtr, const FMutableObjectPtr&>::value, "Invalid assignment (TObjectPtr<const UInterface> from TObjectPtr<UObject>)");
+static_assert(!std::is_assignable<FConstInterfacePtr, const FConstObjectPtr&>::value, "Invalid assignment (TObjectPtr<const UInterface> from TObjectPtr<const UObject>)");
+
+// Ensure that TObjectPtr<[const] UObject> is comparable with another TObjectPtr<[const] UObject> regardless of constness
+static_assert(TModels<CEqualityComparableWith, FConstObjectPtr, FConstObjectPtr>::Value, "Must be able to compare equality and inequality bidirectionally between TObjectPtr<const UObject> and TObjectPtr<const UObject>");
+static_assert(TModels<CEqualityComparableWith, FMutableObjectPtr, FConstObjectPtr>::Value, "Must be able to compare equality and inequality bidirectionally between TObjectPtr<UObject> and TObjectPtr<const UObject>");
+
+// Ensure that TObjectPtr<[const] UObject> is comparable with another TObjectPtr<[const] UInterface> regardless of constness
+static_assert(TModels<CEqualityComparableWith, FConstObjectPtr, FConstInterfacePtr>::Value, "Must be able to compare equality and inequality bidirectionally between TObjectPtr<const UObject> and TObjectPtr<const UInterface>");
+static_assert(TModels<CEqualityComparableWith, FMutableObjectPtr, FConstInterfacePtr>::Value, "Must be able to compare equality and inequality bidirectionally between TObjectPtr<UObject> and TObjectPtr<const UInterface>");
+static_assert(TModels<CEqualityComparableWith, FConstObjectPtr, FMutableInterfacePtr>::Value, "Must be able to compare equality and inequality bidirectionally between TObjectPtr<const UObject> and TObjectPtr<UInterface>");
+static_assert(TModels<CEqualityComparableWith, FMutableObjectPtr, FMutableInterfacePtr>::Value, "Must be able to compare equality and inequality bidirectionally between TObjectPtr<UObject> and TObjectPtr<UInterface>");
+
+// Ensure that TObjectPtr<[const] UPackage> is not comparable with a TObjectPtr<[const] UInterface> regardless of constness
+// TODO: This only ensures that at least one of the A==B,B==A,A!=B,B!=A operations fail, not that they all fail.
+#if !PLATFORM_WINDOWS || !defined(_MSC_EXTENSIONS) // MSVC static analyzer is run in non-conformance mode, and that causes these checks to fail.
+static_assert(!TModels<CEqualityComparableWith, FConstPackagePtr, FConstInterfacePtr>::Value, "Must not be able to compare equality and inequality bidirectionally between TObjectPtr<const UPackage> and TObjectPtr<const UInterface>");
+static_assert(!TModels<CEqualityComparableWith, FMutablePackagePtr, FConstInterfacePtr>::Value, "Must not be able to compare equality and inequality bidirectionally between TObjectPtr<UPackage> and TObjectPtr<const UInterface>");
+static_assert(!TModels<CEqualityComparableWith, FConstPackagePtr, FMutableInterfacePtr>::Value, "Must not be able to compare equality and inequality bidirectionally between TObjectPtr<const UPackage> and TObjectPtr<UInterface>");
+static_assert(!TModels<CEqualityComparableWith, FMutablePackagePtr, FMutableInterfacePtr>::Value, "Must not be able to compare equality and inequality bidirectionally between TObjectPtr<UPackage> and TObjectPtr<UInterface>");
+#endif // _MSC_EXTENSIONS
+
+// Ensure that TObjectPtr<[const] UObject> is comparable with a raw pointer of the same referenced type regardless of constness
+static_assert(TModels<CEqualityComparableWith, FConstObjectPtr, const UObject*>::Value, "Must be able to compare equality and inequality bidirectionally between TObjectPtr<const UObject> and const UObject*");
+static_assert(TModels<CEqualityComparableWith, FMutableObjectPtr, const UObject*>::Value, "Must be able to compare equality and inequality bidirectionally between TObjectPtr<UObject> and const UObject*");
+static_assert(TModels<CEqualityComparableWith, FConstObjectPtr, UObject*>::Value, "Must be able to compare equality and inequality bidirectionally between TObjectPtr<const UObject> and UObject*");
+static_assert(TModels<CEqualityComparableWith, FMutableObjectPtr, UObject*>::Value, "Must be able to compare equality and inequality bidirectionally between TObjectPtr<UObject> and UObject*");
+
+// Ensure that TObjectPtr<[const] UObject> is comparable with a UInterface raw pointer regardless of constness
+static_assert(TModels<CEqualityComparableWith, FConstObjectPtr, const UInterface*>::Value, "Must be able to compare equality and inequality bidirectionally between TObjectPtr<const UObject> and const UInterface*");
+static_assert(TModels<CEqualityComparableWith, FMutableObjectPtr, const UInterface*>::Value, "Must be able to compare equality and inequality bidirectionally between TObjectPtr<UObject> and const UInterface*");
+static_assert(TModels<CEqualityComparableWith, FConstObjectPtr, UInterface*>::Value, "Must be able to compare equality and inequality bidirectionally between TObjectPtr<const UObject> and UInterface*");
+static_assert(TModels<CEqualityComparableWith, FMutableObjectPtr, UInterface*>::Value, "Must be able to compare equality and inequality bidirectionally between TObjectPtr<UObject> and UInterface*");
+
+// Ensure that TObjectPtr<[const] UInterface> is comparable with a UObject raw pointer regardless of constness
+static_assert(TModels<CEqualityComparableWith, FConstInterfacePtr, const UObject*>::Value, "Must be able to compare equality and inequality bidirectionally between TObjectPtr<const UInterface> and const UObject*");
+static_assert(TModels<CEqualityComparableWith, FMutableInterfacePtr, const UObject*>::Value, "Must be able to compare equality and inequality bidirectionally between TObjectPtr<UInterface> and const UObject*");
+static_assert(TModels<CEqualityComparableWith, FConstInterfacePtr, UObject*>::Value, "Must be able to compare equality and inequality bidirectionally between TObjectPtr<const UInterface> and UObject*");
+static_assert(TModels<CEqualityComparableWith, FMutableInterfacePtr, UObject*>::Value, "Must be able to compare equality and inequality bidirectionally between TObjectPtr<UInterface> and UObject*");
+
+// Ensure that TObjectPtr<[const] UInterface> is not comparable with a UPackage raw pointer regardless of constness
+// TODO: This only ensures that at least one of the A==B,B==A,A!=B,B!=A operations fail, not that they all fail.
+static_assert(!TModels<CEqualityComparableWith, FConstInterfacePtr, const UPackage*>::Value, "Must not be able to compare equality and inequality bidirectionally between TObjectPtr<const UInterface> and const UPackage*");
+static_assert(!TModels<CEqualityComparableWith, FMutableInterfacePtr, const UPackage*>::Value, "Must not be able to compare equality and inequality bidirectionally between TObjectPtr<UInterface> and const UPackage*");
+static_assert(!TModels<CEqualityComparableWith, FConstInterfacePtr, UPackage*>::Value, "Must not be able to compare equality and inequality bidirectionally between TObjectPtr<const UInterface> and UPackage*");
+static_assert(!TModels<CEqualityComparableWith, FMutableInterfacePtr, UPackage*>::Value, "Must not be able to compare equality and inequality bidirectionally between TObjectPtr<UInterface> and UPackage*");
+
+// Ensure that TObjectPtr<[const] UInterface> is not comparable with a char raw pointer regardless of constness
+// TODO: This only ensures that at least one of the A==B,B==A,A!=B,B!=A operations fail, not that they all fail.
+static_assert(!TModels<CEqualityComparableWith, FConstObjectPtr, const char*>::Value, "Must not be able to compare equality and inequality bidirectionally between TObjectPtr<const UObject> and const UObject*");
+static_assert(!TModels<CEqualityComparableWith, FMutableObjectPtr, const char*>::Value, "Must not be able to compare equality and inequality bidirectionally between TObjectPtr<UObject> and const UObject*");
+static_assert(!TModels<CEqualityComparableWith, FConstObjectPtr, char*>::Value, "Must not be able to compare equality and inequality bidirectionally between TObjectPtr<const UObject> and UObject*");
+static_assert(!TModels<CEqualityComparableWith, FMutableObjectPtr, char*>::Value, "Must not be able to compare equality and inequality bidirectionally between TObjectPtr<UObject> and UObject*");
+
+// Ensure that TObjectPtr<[const] UObject> is comparable with nullptr regardless of constness
+static_assert(TModels<CEqualityComparableWith, FConstObjectPtr, TYPE_OF_NULLPTR>::Value, "Must be able to compare equality and inequality bidirectionally between TObjectPtr<const UObject> and nullptr");
+static_assert(TModels<CEqualityComparableWith, FMutableObjectPtr, TYPE_OF_NULLPTR>::Value, "Must be able to compare equality and inequality bidirectionally between TObjectPtr<UObject> and nullptr");
+
+#if !PLATFORM_WINDOWS || !defined(_MSC_EXTENSIONS) // MSVC static analyzer is run in non-conformance mode, and that causes these checks to fail.
+static_assert(!TModels<CEqualityComparableWith, FConstObjectPtr, long>::Value, "Should not be able to compare equality and inequality bidirectionally between TObjectPtr<const UObject> and long");
+static_assert(!TModels<CEqualityComparableWith, FMutableObjectPtr, long>::Value, "Should not be able to compare equality and inequality bidirectionally between TObjectPtr<UObject> and long");
+#endif // _MSC_EXTENSIONS
+
+// Ensure that the use of incomplete types doesn't provide a means to bypass type safety on TObjectPtr
+// NOTE: This is disabled because we're permitting this operation with a deprecation warning.
+//static_assert(!std::is_assignable<TObjectPtr<UForwardDeclaredObjDerived>, UForwardDeclaredObjDerived*>::value, "Should not be able to assign raw pointer of incomplete type that descends from UObject to exactly this type of TObjectPtr");
+//static_assert(!std::is_assignable<TObjectPtr<FForwardDeclaredNotObjDerived>, FForwardDeclaredNotObjDerived*>::value, "Should not be able to assign raw pointer of incomplete type that does not descend from UObject to exactly this type of TObjectPtr");
+
+#if WITH_DEV_AUTOMATION_TESTS
+
+class FObjectPtrTestBase : public FObjectRefTrackingTestBase
+{
+public:
+	FObjectPtrTestBase(const FString& InName, const bool bInComplexTask)
+	: FObjectRefTrackingTestBase(InName, bInComplexTask)
+	{
+	}
+
+protected:
+};
+
+#define TEST_NAME_ROOT TEXT("System.CoreUObject.ObjectPtr")
+constexpr const uint32 ObjectPtrTestFlags = EAutomationTestFlags::ApplicationContextMask | EAutomationTestFlags::EngineFilter;
+
+IMPLEMENT_CUSTOM_SIMPLE_AUTOMATION_TEST(FObjectPtrTestNullBehavior, FObjectPtrTestBase, TEST_NAME_ROOT TEXT(".NullBehavior"), ObjectPtrTestFlags)
+bool FObjectPtrTestNullBehavior::RunTest(const FString& Parameters)
+{
+	TObjectPtr<UObject> NullObjectPtr(nullptr);
+	TestTrue(TEXT("Nullptr should equal a null object pointer"), nullptr == NullObjectPtr);
+	TestTrue(TEXT("A null object pointer should equal nullptr"), NullObjectPtr == nullptr);
+	TestFalse(TEXT("A null object pointer should evaluate to false"), !!NullObjectPtr);
+	TestTrue(TEXT("Negation of a null object pointer should evaluate to true"), !NullObjectPtr);
+
+	return true;
+}
+
+IMPLEMENT_CUSTOM_SIMPLE_AUTOMATION_TEST(FObjectPtrTestDefaultSerialize, FObjectPtrTestBase, TEST_NAME_ROOT TEXT(".DefaultSerialize"), ObjectPtrTestFlags)
+bool FObjectPtrTestDefaultSerialize::RunTest(const FString& Parameters)
+{
+	FSnapshotObjectRefMetrics ObjectRefMetrics(*this);
+	FObjectPtr DefaultTexturePtr(FObjectRef {FName("/Engine/EngineResources/DefaultTexture"), NAME_None, NAME_None, FObjectPathId("DefaultTexture")});
+
+	ObjectRefMetrics.TestNumResolves(TEXT("Unexpected resolve count after initializing an FObjectPtr"), UE_WITH_OBJECT_HANDLE_LATE_RESOLVE ? 0 : 1);
+	ObjectRefMetrics.TestNumFailedResolves(TEXT("Unexpected resolve failure after initializing an FObjectPtr"), 0);
+	ObjectRefMetrics.TestNumReads(TEXT("NumReads should not change when initializing an FObjectPtr"), 0);
+
+	FArchiveCountMem Writer(nullptr);
+	Writer << DefaultTexturePtr;
+
+	ObjectRefMetrics.TestNumResolves(TEXT("Serializing an FObjectPtr should force it to resolve"), 1);
+	ObjectRefMetrics.TestNumFailedResolves(TEXT("Unexpected resolve failure after serializing an FObjectPtr"), 0);
+	ObjectRefMetrics.TestNumReads(TEXT("NumReads should increase after serializing an FObjectPtr"), 1);
+
+	Writer << DefaultTexturePtr;
+
+	ObjectRefMetrics.TestNumResolves(TEXT("Serializing an FObjectPtr twice should only require it to resolve once"), 1);
+	ObjectRefMetrics.TestNumFailedResolves(TEXT("Unexpected resolve failure after serializing an FObjectPtr"), 0);
+	ObjectRefMetrics.TestNumReads(TEXT("NumReads should increase after serializing an FObjectPtr"), 2);
+
+	return true;
+}
+
+IMPLEMENT_CUSTOM_SIMPLE_AUTOMATION_TEST(FObjectPtrTestSoftObjectPath, FObjectPtrTestBase, TEST_NAME_ROOT TEXT(".SoftObjectPath"), ObjectPtrTestFlags)
+bool FObjectPtrTestSoftObjectPath::RunTest(const FString& Parameters)
+{
+	FSnapshotObjectRefMetrics ObjectRefMetrics(*this);
+	FObjectPtr DefaultTexturePtr(FObjectRef {FName("/Engine/EngineResources/DefaultTexture"), NAME_None, NAME_None, FObjectPathId("DefaultTexture")});
+
+	ObjectRefMetrics.TestNumResolves(TEXT("Unexpected resolve count after initializing an FObjectPtr"), UE_WITH_OBJECT_HANDLE_LATE_RESOLVE ? 0 : 1);
+	ObjectRefMetrics.TestNumFailedResolves(TEXT("Unexpected resolve failure after initializing an FObjectPtr"), 0);
+	ObjectRefMetrics.TestNumReads(TEXT("NumReads should not change when initializing an FObjectPtr"), 0);
+
+	FSoftObjectPath DefaultTexturePath(DefaultTexturePtr);
+
+	ObjectRefMetrics.TestNumResolves(TEXT("Unexpected resolve count after initializing an FSoftObjectPath from an FObjectPtr"), UE_WITH_OBJECT_HANDLE_LATE_RESOLVE ? 0 : 1);
+
+	TestEqual(TEXT("Soft object path constructed from an FObjectPtr does not have the expected path value"), *DefaultTexturePath.ToString(), TEXT("/Engine/EngineResources/DefaultTexture.DefaultTexture"));
+
+	return true;
+}
+
+IMPLEMENT_CUSTOM_SIMPLE_AUTOMATION_TEST(FObjectPtrTestForwardDeclared, FObjectPtrTestBase, TEST_NAME_ROOT TEXT(".ForwardDeclared"), ObjectPtrTestFlags)
+bool FObjectPtrTestForwardDeclared::RunTest(const FString& Parameters)
+{
+	UForwardDeclaredObjDerived* PtrFwd = nullptr;
+	TObjectPtr<UForwardDeclaredObjDerived> ObjPtrFwd(MakeObjectPtrUnsafe<UForwardDeclaredObjDerived>(reinterpret_cast<UObject*>(PtrFwd)));
+	TestTrue(TEXT("Null forward declared pointer used to construct a TObjectPtr should result in a null TObjectPtr"), ObjPtrFwd.IsNull());
+	return true;
+}
+
+// @TODO: OBJPTR: We should have a test that ensures that we can (de)serialize an FObjectPtr to FLinkerSave/FLinkerLoad and that upon load the object
+//			pointer is not resolved if we are in a configuration that supports lazy load.  This is proving difficult due to the restrictions around how
+//			FLinkerSave/FLinkerLoad is used.
+// IMPLEMENT_CUSTOM_SIMPLE_AUTOMATION_TEST(FObjectPtrTestLinkerSerializeBehavior, FObjectPtrTestBase, TEST_NAME_ROOT TEXT(".LinkerSerialize"), ObjectPtrTestFlags)
+// bool FObjectPtrTestLinkerSerializeBehavior::RunTest(const FString& Parameters)
+// {
+// 	FSnapshotObjectRefMetrics ObjectRefMetrics(*this);
+// 	FObjectPtr DefaultTexturePtr(FObjectRef {FName("/Engine/EngineResources/DefaultTexture"), NAME_None, NAME_None, FObjectPathId("DefaultTexture")});
+
+// 	TUniquePtr<FLinkerSave> Linker = TUniquePtr<FLinkerSave>(new FLinkerSave(nullptr /*InOuter*/, false /*bForceByteSwapping*/, true /*bSaveUnversioned*/));
+// 	return true;
+// }
+
+#undef TEST_NAME_ROOT
+
+#endif // WITH_DEV_AUTOMATION_TESTS
+
+class UForwardDeclaredObjDerived: public UObject {};
+class FForwardDeclaredNotObjDerived {};
+
+}
diff --git a/Engine/Source/Runtime/CoreUObject/Private/Tests/ObjectRefTrackingTestBase.cpp b/Engine/Source/Runtime/CoreUObject/Private/Tests/ObjectRefTrackingTestBase.cpp
new file mode 100644
--- /dev/null
+++ b/Engine/Source/Runtime/CoreUObject/Private/Tests/ObjectRefTrackingTestBase.cpp
@@ -0,0 +1,15 @@
+// Copyright Epic Games, Inc. All Rights Reserved.
+
+#include "ObjectRefTrackingTestBase.h"
+
+#if WITH_DEV_AUTOMATION_TESTS
+
+#if UE_WITH_OBJECT_HANDLE_TRACKING
+ObjectHandleReferenceResolvedFunction* FObjectRefTrackingTestBase::PrevResolvedFunc = nullptr;
+ObjectHandleReadFunction* FObjectRefTrackingTestBase::PrevReadFunc = nullptr;
+#endif
+thread_local uint32 FObjectRefTrackingTestBase::NumResolves = 0;
+thread_local uint32 FObjectRefTrackingTestBase::NumFailedResolves = 0;
+thread_local uint32 FObjectRefTrackingTestBase::NumReads = 0;
+
+#endif // WITH_DEV_AUTOMATION_TESTS
diff --git a/Engine/Source/Runtime/CoreUObject/Private/Tests/ObjectRefTrackingTestBase.h b/Engine/Source/Runtime/CoreUObject/Private/Tests/ObjectRefTrackingTestBase.h
new file mode 100644
--- /dev/null
+++ b/Engine/Source/Runtime/CoreUObject/Private/Tests/ObjectRefTrackingTestBase.h
@@ -0,0 +1,121 @@
+// Copyright Epic Games, Inc. All Rights Reserved.
+
+#pragma once
+
+#include "CoreMinimal.h"
+#include "Misc/AutomationTest.h"
+#include "UObject/ObjectHandle.h"
+
+#if WITH_DEV_AUTOMATION_TESTS
+
+class FObjectRefTrackingTestBase : public FAutomationTestBase
+{
+public:
+	FObjectRefTrackingTestBase(const FString& InName, const bool bInComplexTask)
+	: FAutomationTestBase(InName, bInComplexTask)
+	{
+		ConditionalInstallCallbacks();
+	}
+
+	uint32 GetNumResolves() const { return NumResolves; }
+	uint32 GetNumFailedResolves() const { return NumFailedResolves; }
+	uint32 GetNumReads() const { return NumReads; }
+
+	struct FSnapshotObjectRefMetrics
+	{
+	public:
+		FSnapshotObjectRefMetrics(FObjectRefTrackingTestBase& InTest)
+		: Test(InTest)
+		, OriginalNumResolves(Test.GetNumResolves())
+		, OriginalNumFailedResolves(Test.GetNumFailedResolves())
+		, OriginalNumReads(Test.GetNumReads())
+		{}
+
+		bool TestNumResolves(const TCHAR* What, uint32 ExpectedDelta)
+		{
+#if UE_WITH_OBJECT_HANDLE_TRACKING
+			return Test.TestEqual(What, OriginalNumResolves + ExpectedDelta, Test.GetNumResolves());
+#else
+			return true;
+#endif
+		}
+
+		bool TestNumFailedResolves(const TCHAR* What, uint32 ExpectedDelta)
+		{
+#if UE_WITH_OBJECT_HANDLE_TRACKING
+			return Test.TestEqual(What, OriginalNumFailedResolves + ExpectedDelta, Test.GetNumFailedResolves());
+#else
+			return true;
+#endif
+		}
+
+		bool TestNumReads(const TCHAR* What, uint32 ExpectedDelta)
+		{
+#if UE_WITH_OBJECT_HANDLE_TRACKING
+			return Test.TestEqual(What, OriginalNumReads + ExpectedDelta, Test.GetNumReads());
+#else
+			return true;
+#endif
+		}
+	private:
+		FObjectRefTrackingTestBase& Test;
+		uint32 OriginalNumResolves;
+		uint32 OriginalNumFailedResolves;
+		uint32 OriginalNumReads;
+	};
+
+private:
+#if UE_WITH_OBJECT_HANDLE_TRACKING
+	static void OnRefResolved(const FObjectRef& ObjectRef, UPackage* Pkg, UObject* Obj)
+	{
+		NumResolves++;
+		if (!IsObjectRefNull(ObjectRef) && !Obj)
+		{
+			NumFailedResolves++;
+		}
+
+		if (PrevResolvedFunc)
+		{
+			PrevResolvedFunc(ObjectRef, Pkg, Obj);
+		}
+	}
+	static void OnRefRead(UObject* Obj)
+	{
+		NumReads++;
+		if (PrevReadFunc)
+		{
+			PrevReadFunc(Obj);
+		}
+	}
+#endif
+	
+	static void ConditionalInstallCallbacks()
+	{
+		static bool bCallbacksInstalled = false;
+		static FCriticalSection CallbackInstallationLock;
+
+		if (bCallbacksInstalled)
+		{
+			return;
+		}
+
+		FScopeLock ScopeLock(&CallbackInstallationLock);
+		if (!bCallbacksInstalled)
+		{
+#if UE_WITH_OBJECT_HANDLE_TRACKING
+			PrevResolvedFunc = SetObjectHandleReferenceResolvedCallback(OnRefResolved);
+			PrevReadFunc = SetObjectHandleReadCallback(OnRefRead);
+#endif
+			bCallbacksInstalled = true;
+		}
+	}
+
+#if UE_WITH_OBJECT_HANDLE_TRACKING
+	static ObjectHandleReferenceResolvedFunction* PrevResolvedFunc;
+	static ObjectHandleReadFunction* PrevReadFunc;
+#endif
+	static thread_local uint32 NumResolves;
+	static thread_local uint32 NumFailedResolves;
+	static thread_local uint32 NumReads;
+};
+#endif // WITH_DEV_AUTOMATION_TESTS
diff --git a/Engine/Source/Runtime/CoreUObject/Private/Tests/Virtualization/VirtualizedBulkDataTests.cpp b/Engine/Source/Runtime/CoreUObject/Private/Tests/Virtualization/VirtualizedBulkDataTests.cpp
new file mode 100644
--- /dev/null
+++ b/Engine/Source/Runtime/CoreUObject/Private/Tests/Virtualization/VirtualizedBulkDataTests.cpp
@@ -0,0 +1,219 @@
+// Copyright Epic Games, Inc. All Rights Reserved.
+
+#include "Virtualization/VirtualizedBulkData.h"
+#include "Virtualization/VirtualizationUtilities.h"
+#include "Memory/SharedBuffer.h"
+#include "Templates/UniquePtr.h"
+#include "Algo/AllOf.h"
+
+#include "Misc/AutomationTest.h"
+
+#if WITH_DEV_AUTOMATION_TESTS && WITH_EDITORONLY_DATA
+
+constexpr const uint32 TestFlags = EAutomationTestFlags::ApplicationContextMask | EAutomationTestFlags::EngineFilter;
+
+/**
+ * This test creates a very basic VirtualizedBulkData object with in memory payload and validates that we are able to retrieve the 
+ * payload via both TFuture and callback methods. It then creates copies of the object and makes sure that we can get the payload from
+ * the copies, even when the original source object has been reset.
+ */
+IMPLEMENT_SIMPLE_AUTOMATION_TEST(FVirtualizationWrapperTestBasic, TEXT("System.Core.VirtualizedBulkData.Basic"), TestFlags)
+bool FVirtualizationWrapperTestBasic::RunTest(const FString& Parameters)
+{
+	const int64 BufferSize = 1024;
+	TUniquePtr<uint8[]> SourceBuffer = MakeUnique<uint8[]>(BufferSize);
+	
+	// Write some values to SourceBuffer
+	for (int64 Index = 0; Index < BufferSize; ++Index)
+	{
+		SourceBuffer[Index] = (uint8)(FMath::Rand() % 255);
+	}
+
+	auto ValidateBulkData = [this, &SourceBuffer, BufferSize](const FVirtualizedUntypedBulkData& BulkDataToValidate, const TCHAR* Label)
+		{
+			FSharedBufferConstPtr RetrievedBuffer = BulkDataToValidate.GetData().Get();
+			TestEqual(FString::Printf(TEXT("%s buffer length"),Label), (int64)RetrievedBuffer->Size(), BufferSize);
+			TestTrue(FString::Printf(TEXT("SourceBuffer values == %s values"),Label), FMemory::Memcmp(SourceBuffer.Get(), RetrievedBuffer->GetData(), BufferSize) == 0);
+		};
+
+	// Create a basic bulkdata (but retain ownership of the buffer!)
+	FByteVirtualizedBulkData BulkData;
+	BulkData.UpdatePayload(MakeSharedBuffer(FSharedBuffer::Wrap, SourceBuffer.Get(), BufferSize));
+
+	// Test access via callback
+	BulkData.GetData([this,BufferSize,&SourceBuffer](FSharedBufferConstPtr Payload)
+		{
+			TestEqual(TEXT("Callback buffer length"), (int64)Payload->Size(), BufferSize);
+			TestTrue(TEXT("SourceBuffer values == Callback values"), FMemory::Memcmp(SourceBuffer.Get(), Payload->GetData(), BufferSize) == 0);
+		});
+
+	// Test accessing the data from the bulkdata object
+	ValidateBulkData(BulkData, TEXT("Retrieved"));
+
+	// Create a new bulkdata object via the copy constructor
+	FByteVirtualizedBulkData BulkDataCopy(BulkData);
+
+	// Create a new bulkdata object and copy by assignment (note we assign some junk data that will get overwritten)
+	FByteVirtualizedBulkData BulkDataAssignment;
+	BulkDataAssignment.UpdatePayload(MakeSharedBuffer(128));
+	BulkDataAssignment = BulkData;
+
+	// Test both bulkdata objects
+	ValidateBulkData(BulkDataCopy, TEXT("Copy Constructor"));
+	ValidateBulkData(BulkDataAssignment, TEXT("Copy Assignment"));
+
+	// Should not affect BulkDataAssignment/BulkDataCopy 
+	BulkData.Reset();
+
+	// Test both bulkdata objects again now that we reset the data
+	ValidateBulkData(BulkDataCopy, TEXT("Copy Constructor (after data reset)"));
+	ValidateBulkData(BulkDataAssignment, TEXT("Copy Assignment (after data reset)"));
+
+	return true;
+}
+
+/**
+ * This test will validate how VirtualizedBulkData behaves when it has no associated payload and make sure 
+ * that our assumptions are correct.
+ */
+IMPLEMENT_SIMPLE_AUTOMATION_TEST(FVirtualizationWrapperTestEmpty, TEXT("System.Core.VirtualizedBulkData.Empty"), TestFlags)
+bool FVirtualizationWrapperTestEmpty::RunTest(const FString& Parameters)
+{
+	FByteVirtualizedBulkData BulkData;
+
+	// Validate the general accessors
+	TestEqual(TEXT("Return value of ::GetBulkDataSize()"), BulkData.GetBulkDataSize(), (int64)0);
+	TestTrue(TEXT("Payload key is invalid"), !BulkData.GetKey().IsValid());
+	TestFalse(TEXT("Return value of ::IsDataLoaded()"), BulkData.IsDataLoaded());
+
+	// Validate the payload accessors
+	FSharedBufferConstPtr Payload = BulkData.GetData().Get();
+	TestFalse(TEXT("The payload from the GetData TFuture is invalid"), Payload.IsValid());
+
+	BulkData.GetData([this](FSharedBufferConstPtr Payload) 
+		{
+			TestFalse(TEXT("The payload via the GetData callback is invalid"), Payload.IsValid());
+		});
+
+	return true;
+}
+
+/**
+ * Test the various methods for updating the payload that a FVirtualizedUntypedBulkData owns
+ */
+IMPLEMENT_SIMPLE_AUTOMATION_TEST(FVirtualizationWrapperTestUpdatePayload, TEXT("System.Core.VirtualizedBulkData.UpdatePayload"), TestFlags)
+bool FVirtualizationWrapperTestUpdatePayload::RunTest(const FString& Parameters)
+{
+	// Create a memory buffer of all zeros
+	const int64 BufferSize = 1024;
+	TUniquePtr<uint8[]> OriginalData = MakeUnique<uint8[]>(BufferSize);
+	FMemory::Memzero(OriginalData.Get(), BufferSize);
+
+	// Pass the buffer to to bulkdata but retain ownership
+	FByteVirtualizedBulkData BulkData;
+	BulkData.UpdatePayload(MakeSharedBuffer(FSharedBuffer::Wrap, OriginalData.Get(), BufferSize));
+
+	// Access the payload, edit it and push it back into the bulkdata object
+	{
+		// The payload should be the same size and same contents as the original buffer but a different
+		// memory address since we retained ownership in the TUniquePtr, so the bulkdata object should 
+		// have created it's own copy.
+		FSharedBufferConstPtr Payload = BulkData.GetData().Get();
+		TestEqual(TEXT("Payload length"), (int64)Payload->Size(), BufferSize);
+		TestNotEqual(TEXT("OriginalData and the payload should have different memory addresses"), (uint8*)OriginalData.Get(), (uint8*)Payload->GetData());
+		TestTrue(TEXT("Orginal buffer == Payload data"), FMemory::Memcmp(OriginalData.Get(), Payload->GetData(), Payload->Size()) == 0);
+
+		// Make a copy of the payload that we can edit
+		FSharedBufferRef EditablePayload = MakeSharedBuffer(FSharedBuffer::Clone, Payload->GetView());
+
+		const uint8 NewValue = 255;
+		FMemory::Memset(EditablePayload->GetData(), NewValue, EditablePayload->Size());
+
+		// Update the bulkdata object with the new edited payload
+		BulkData.UpdatePayload(EditablePayload);
+
+		Payload = BulkData.GetData().Get();
+		TestEqual(TEXT("Updated payload length"), (int64)Payload->Size(), BufferSize);
+		TestEqual(TEXT("Payload and EditablePayload should have the same memory addresses"), (uint8*)Payload->GetData(), (uint8*)EditablePayload->GetData());
+
+		const bool bAllElementsCorrect = Algo::AllOf(TArrayView<uint8>((uint8*)Payload->GetData(), Payload->Size()), [NewValue](uint8 Val)
+			{ 
+				return Val == NewValue; 
+			});
+
+		TestTrue(TEXT("All payload elements correctly updated"), bAllElementsCorrect);
+	}
+
+	{
+		// Store the original data pointer address so we can test against it later, we should not actually use
+		// this pointer though as once we pass it to the bulkdata object we cannot be sure what happens to it.
+		uint8* OriginalDataPtr = OriginalData.Get();
+
+		// Update the bulkdata object with the original data but this time we give ownership of the buffer to
+		// the bulkdata object.
+		BulkData.UpdatePayload(MakeSharedBuffer(FSharedBuffer::AssumeOwnership, OriginalData.Release(), BufferSize));
+
+		FSharedBufferConstPtr Payload = BulkData.GetData().Get();
+		TestEqual(TEXT("Updated payload length"), (int64)Payload->Size(), BufferSize);
+		TestEqual(TEXT("Payload and OriginalDataPtr should have the same memory addresses"), (uint8*)Payload->GetData(), OriginalDataPtr);
+
+		// The original data was all zeros, so we can test for that to make sure that the contents are correct.
+		const bool bAllElementsCorrect = Algo::AllOf(TArrayView<uint8>((uint8*)Payload->GetData(), Payload->Size()), [](uint8 Val)
+			{
+				return Val == 0;
+			});
+
+		TestTrue(TEXT("All payload elements correctly updated"), bAllElementsCorrect);
+	}
+
+	return true;
+}
+
+/**
+ * This test will create a buffer, then serialize it to a VirtualizedBulkData object via FVirtualizedBulkDataWriter.
+ * Then we will serialize the VirtualizedBulkData object back to a second buffer and compare the results.
+ * If the reader and writers are working then the ReplicatedBuffer should be the same as the original SourceBuffer.
+ */
+IMPLEMENT_SIMPLE_AUTOMATION_TEST(FVirtualizationWrapperTestReaderWriter, TEXT("System.Core.VirtualizedBulkData.Reader/Writer"), TestFlags)
+bool FVirtualizationWrapperTestReaderWriter::RunTest(const FString& Parameters)
+{
+	const int64 BufferSize = 1024;
+
+	TUniquePtr<uint8[]> SourceBuffer = MakeUnique<uint8[]>(BufferSize);
+	TUniquePtr<uint8[]> ReplicatedBuffer = MakeUnique<uint8[]>(BufferSize);
+
+	// Write some values to SourceBuffer
+	for (int64 Index = 0; Index < BufferSize; ++Index)
+	{
+		SourceBuffer[Index] = (uint8)(FMath::Rand() % 255);
+	}
+
+	FByteVirtualizedBulkData BulkData;
+
+	// Serialize the SourceBuffer to BulkData 
+	{
+		FVirtualizedBulkDataWriter WriterAr(BulkData);
+		WriterAr.Serialize(SourceBuffer.Get(), BufferSize);
+	}
+
+	// Serialize BulkData back to ReplicatedBuffer
+	{
+		FVirtualizedBulkDataReader ReaderAr(BulkData);
+		ReaderAr.Serialize(ReplicatedBuffer.Get(), BufferSize);
+	}
+
+	// Now test that the buffer was restored to the original values
+	const bool bMemCmpResult = FMemory::Memcmp(SourceBuffer.Get(), ReplicatedBuffer.Get(), BufferSize) == 0;
+	TestTrue(TEXT("SourceBuffer values == ReplicatedBuffer values"), bMemCmpResult);
+
+	// Test writing nothing to an empty bulkdata object and then reading that bulkdata object
+	// to make sure that we deal with null buffers properly.
+	{
+		FByteVirtualizedBulkData EmptyBulkData;
+		FVirtualizedBulkDataWriter WriterAr(EmptyBulkData);
+		FVirtualizedBulkDataReader ReaderAr(EmptyBulkData);
+	}
+	return true;
+}
+
+#endif //WITH_DEV_AUTOMATION_TESTS && WITH_EDITORONLY_DATA
diff --git a/Engine/Source/Runtime/CoreUObject/Private/UObject/GarbageCollection.cpp b/Engine/Source/Runtime/CoreUObject/Private/UObject/GarbageCollection.cpp
--- a/Engine/Source/Runtime/CoreUObject/Private/UObject/GarbageCollection.cpp
+++ b/Engine/Source/Runtime/CoreUObject/Private/UObject/GarbageCollection.cpp
@@ -14,6 +14,8 @@
 #include "UObject/UObjectAllocator.h"
 #include "UObject/UObjectBase.h"
 #include "UObject/Object.h"
+#include "UObject/ObjectHandle.h"
+#include "UObject/ObjectPtr.h"
 #include "UObject/Class.h"
 #include "UObject/UObjectIterator.h"
 #include "UObject/UnrealType.h"
@@ -791,19 +793,19 @@ public:
 	*/
 	FORCEINLINE void HandleObjectReference(TArray<UObject*>& ObjectsToSerialize, const UObject * const ReferencingObject, UObject*& Object, const bool bAllowReferenceElimination)
 	{
+		if (Object == nullptr)
+		{
+			return;
+		}
+
 		// Disregard NULL objects and perform very fast check to see whether object is part of permanent
 		// object pool and should therefore be disregarded. The check doesn't touch the object and is
 		// cache friendly as it's just a pointer compare against to globals.
-		const bool IsInPermanentPool = GUObjectAllocator.ResidesInPermanentPool(Object);
-
-#if PERF_DETAILED_PER_CLASS_GC_STATS
-		if (IsInPermanentPool)
+		if (GUObjectAllocator.ResidesInPermanentPool(Object))
 		{
+#if PERF_DETAILED_PER_CLASS_GC_STATS
 			GCurrentObjectDisregardedObjectRefs++;
-		}
 #endif
-		if (Object == nullptr || IsInPermanentPool)
-		{
 			return;
 		}
 
@@ -925,7 +927,7 @@ public:
 	FORCEINLINE void HandleTokenStreamObjectReference(TArray<UObject*>& ObjectsToSerialize, UObject* ReferencingObject, UObject*& Object, const int32 TokenIndex, bool bAllowReferenceElimination)
 	{
 #if ENABLE_GC_OBJECT_CHECKS
-		if (Object)
+		if (Object && IsObjectHandleResolved(*reinterpret_cast<FObjectHandle*>(&Object)))
 		{
 			if (
 #if DO_POINTER_CHECKS_ON_GC
@@ -2417,6 +2419,16 @@ void FMulticastDelegateProperty::EmitReferenceInfo(UClass& OwnerClass, int32 Bas
 	FGCReferenceFixedArrayTokenHelper FixedArrayHelper(OwnerClass, BaseOffset + GetOffset_ForGC(), ArrayDim, this->ElementSize, *this);
 	OwnerClass.EmitObjectReference(BaseOffset + GetOffset_ForGC(), GetFName(), GCRT_MulticastDelegate);
 }
+/**
+ * Emits tokens used by realtime garbage collection code to passed in OwnerClass' ReferenceTokenStream. The offset emitted is relative
+ * to the passed in BaseOffset which is used by e.g. arrays of structs.
+ */
+void FObjectPtrProperty::EmitReferenceInfo(UClass& OwnerClass, int32 BaseOffset, TArray<const FStructProperty*>& EncounteredStructProps)
+{
+	FGCReferenceFixedArrayTokenHelper FixedArrayHelper(OwnerClass, BaseOffset + GetOffset_ForGC(), ArrayDim, sizeof(FObjectPtr), *this);
+	OwnerClass.EmitObjectReference(BaseOffset + GetOffset_ForGC(), GetFName(), GCRT_Object);
+}
+
 /**
  * Emits tokens used by realtime garbage collection code to passed in OwnerClass' ReferenceTokenStream. The offset emitted is relative
  * to the passed in BaseOffset which is used by e.g. arrays of structs.
@@ -2437,7 +2449,7 @@ void FArrayProperty::EmitReferenceInfo(UClass& OwnerClass, int32 BaseOffset, TAr
 			const uint32 SkipIndex = OwnerClass.ReferenceTokenStream.EmitReturn();
 			OwnerClass.ReferenceTokenStream.UpdateSkipIndexPlaceholder(SkipIndexIndex, SkipIndex);
 		}
-		else if( Inner->IsA(FObjectProperty::StaticClass()) )
+		else if( Inner->IsA(FObjectProperty::StaticClass()) || Inner->IsA(FObjectPtrProperty::StaticClass()) )
 		{
 			OwnerClass.EmitObjectReference(BaseOffset + GetOffset_ForGC(), GetFName(), bUsesFreezableAllocator ? GCRT_ArrayObjectFreezable : GCRT_ArrayObject);
 		}
diff --git a/Engine/Source/Runtime/CoreUObject/Private/UObject/Obj.cpp b/Engine/Source/Runtime/CoreUObject/Private/UObject/Obj.cpp
--- a/Engine/Source/Runtime/CoreUObject/Private/UObject/Obj.cpp
+++ b/Engine/Source/Runtime/CoreUObject/Private/UObject/Obj.cpp
@@ -54,6 +54,8 @@
 #include "UObject/CoreRedirects.h"
 #include "HAL/LowLevelMemTracker.h"
 #include "HAL/LowLevelMemStats.h"
+#include "Misc/PackageAccessTracking.h"
+#include "Misc/PackageAccessTrackingOps.h"
 
 DEFINE_LOG_CATEGORY(LogObj);
 
@@ -1053,6 +1055,7 @@ void UObject::ConditionalPostLoad()
 	{
 
 		check(IsInGameThread() || HasAnyFlags(RF_ClassDefaultObject|RF_ArchetypeObject) || IsPostLoadThreadSafe() || IsA(UClass::StaticClass()))
+		UE_TRACK_REFERENCING_PACKAGE_SCOPED(GetOutermost(), PackageAccessTrackingOps::NAME_PostLoad);
 
 #if !(UE_BUILD_SHIPPING || UE_BUILD_TEST)
 		FUObjectThreadContext& ThreadContext = FUObjectThreadContext::Get();
diff --git a/Engine/Source/Runtime/CoreUObject/Private/UObject/ObjectHandle.cpp b/Engine/Source/Runtime/CoreUObject/Private/UObject/ObjectHandle.cpp
new file mode 100644
--- /dev/null
+++ b/Engine/Source/Runtime/CoreUObject/Private/UObject/ObjectHandle.cpp
@@ -0,0 +1,421 @@
+// Copyright Epic Games, Inc. All Rights Reserved.
+
+#include "UObject/ObjectHandle.h"
+#include "HAL/CriticalSection.h"
+#include "HAL/PlatformAtomics.h"
+#include "Misc/PackageName.h"
+#include "Misc/ScopeRWLock.h"
+#include "UObject/Class.h"
+#include "UObject/Linker.h"
+#include "UObject/LinkerLoad.h"
+#include "UObject/Package.h"
+
+bool operator==(const FObjectHandleDataClassDescriptor& Lhs, const FObjectHandleDataClassDescriptor& Rhs)
+{
+	return (Lhs.PackageName == Rhs.PackageName) && (Lhs.ClassName == Rhs.ClassName);
+}
+
+namespace ObjectHandle_Private
+{
+	class FPackageId
+	{
+		static constexpr uint32 InvalidId = ~uint32(0u);
+		uint32 Id = InvalidId;
+
+		inline explicit FPackageId(int32 InId) : Id(InId) {}
+
+	public:
+		FPackageId() = default;
+
+		inline static FPackageId FromIndex(uint32 Index)
+		{
+			return FPackageId(Index);
+		}
+
+		inline bool IsValid() const
+		{
+			return Id != InvalidId;
+		}
+
+		inline uint32 ToIndex() const
+		{
+			check(Id != InvalidId);
+			return Id;
+		}
+
+		inline bool operator==(FPackageId Other) const
+		{
+			return Id == Other.Id;
+		}
+	};
+	enum class EObjectId: uint32
+	{
+		Invalid = 0,
+	};
+
+	union FObjectId
+	{
+	private:
+		uint32 RawData;
+	public:
+		struct
+		{
+			uint32 DataClassDescriptorId : 8;
+			uint32 ObjectPathId : 24;
+		} Components;
+
+		FObjectId() = default;
+		FObjectId(EObjectId Id) : RawData(static_cast<uint32>(Id)) {}
+
+		bool operator==(EObjectId Id) { return RawData == static_cast<uint32>(Id); }
+		bool operator!=(EObjectId Id) { return RawData != static_cast<uint32>(Id); }
+	};
+
+	static_assert(sizeof(FObjectId) == sizeof(uint32), "FObjectId type must always compile to something equivalent to a uint32 size.");
+
+	struct FObjectHandlePackageData
+	{
+		FMinimalName PackageName;
+		TArray<FObjectPathId> ObjectPaths;
+		TArray<FObjectHandleDataClassDescriptor> DataClassDescriptors;
+		TMap<FObjectPathId, FObjectId> PathToObjectId;
+		FRWLock Lock;
+	};
+
+	static_assert(offsetof(FObjectHandlePackageData, PackageName) == offsetof(FObjectHandlePackageDebugData, PackageName), "FObjectHandlePackageData and FObjectHandlePackageDebugData must match in position of PackageNameField.");
+	static_assert(offsetof(FObjectHandlePackageData, ObjectPaths) == offsetof(FObjectHandlePackageDebugData, ObjectPaths), "FObjectHandlePackageData and FObjectHandlePackageDebugData must match in position of ObjectPaths.");
+	static_assert(offsetof(FObjectHandlePackageData, DataClassDescriptors) == offsetof(FObjectHandlePackageDebugData, DataClassDescriptors), "FObjectHandlePackageData and FObjectHandlePackageDebugData must match in position of DataClassDescriptors.");
+	static_assert(sizeof(FObjectHandlePackageData) == sizeof(FObjectHandlePackageDebugData), "FObjectHandlePackageData and FObjectHandlePackageDebugData must match in size.");
+
+	struct FObjectHandleIndex
+	{
+		FRWLock Lock; // @TODO: OBJPTR: Want to change this to a striped lock per object bucket to allow more concurrency when adding and looking up objects in a package
+		TMap<uint32, FPackageId> NameHashToPackageId;
+		TArray<FObjectHandlePackageData> PackageData;
+	} GObjectHandleIndex;
+
+	static inline FPackedObjectRef Pack(FPackageId PackageId, FObjectId ObjectId)
+	{
+#if UE_WITH_OBJECT_HANDLE_LATE_RESOLVE || UE_WITH_OBJECT_HANDLE_TRACKING
+		checkf(PackageId.ToIndex() <= 0x7FFFFFFF, TEXT("Package count exceeded the space permitted within packed object references.  This implies over 2 billion packages are in use."));
+		return {static_cast<UPTRINT>(PackageId.ToIndex()) << PackageIdShift |
+				static_cast<UPTRINT>(ObjectId.Components.DataClassDescriptorId) << DataClassDescriptorIdShift |
+				static_cast<UPTRINT>(ObjectId.Components.ObjectPathId) << ObjectPathIdShift |
+				1};
+#else
+		unimplemented();
+		return {0};
+#endif
+	}
+
+	static inline void Unpack(FPackedObjectRef PackedObjectRef, FPackageId& OutPackageId, FObjectId& OutObjectId)
+	{
+#if UE_WITH_OBJECT_HANDLE_LATE_RESOLVE || UE_WITH_OBJECT_HANDLE_TRACKING
+		checkf((PackedObjectRef.EncodedRef & 1) == 1, TEXT("Packed object reference is malformed."));
+		OutObjectId.Components.ObjectPathId = (static_cast<uint32>((PackedObjectRef.EncodedRef >> ObjectPathIdShift) & ObjectPathIdMask));
+		OutObjectId.Components.DataClassDescriptorId = (static_cast<uint32>((PackedObjectRef.EncodedRef >> DataClassDescriptorIdShift) & DataClassDescriptorIdMask));
+		OutPackageId = FPackageId::FromIndex(static_cast<uint32>((PackedObjectRef.EncodedRef >> PackageIdShift) & PackageIdMask));
+#else
+		unimplemented();
+		OutPackageId = FPackageId();
+		OutObjectId = FObjectId(EObjectId::Invalid);
+#endif
+	}
+
+	static void MakeReferenceIds(FName PackageName, FName ClassPackageName, FName ClassName, FObjectPathId ObjectPath, FPackageId& OutPackageId, FObjectId& OutObjectId)
+	{
+		uint32 NameHash = GetTypeHash(PackageName);
+
+		//Biases for read-only locking by default at the expense of having to do a second map search if a write is necessary
+		//Doing one search with either a read or locked write seems impossible with the TMap interface at the moment.
+		FRWScopeLock GlobalLockScope(GObjectHandleIndex.Lock, SLT_ReadOnly);
+		FPackageId* FoundPackageId = GObjectHandleIndex.NameHashToPackageId.Find(NameHash);
+		FObjectHandlePackageData* PackageData = nullptr;
+		if (!FoundPackageId)
+		{
+			GlobalLockScope.ReleaseReadOnlyLockAndAcquireWriteLock_USE_WITH_CAUTION();
+			//Has to be FindOrAdd, as the NameHashToPackageId may have changed between relinquishing the read lock
+			//and acquiring the write lock.
+			FPackageId NextId = FPackageId::FromIndex(GObjectHandleIndex.PackageData.Num());
+			FoundPackageId = &GObjectHandleIndex.NameHashToPackageId.FindOrAdd(NameHash, NextId);
+			if (*FoundPackageId == NextId)
+			{
+				PackageData = &GObjectHandleIndex.PackageData.AddDefaulted_GetRef();
+				PackageData->PackageName = NameToMinimalName(PackageName);
+				GCoreObjectHandlePackageDebug = reinterpret_cast<FObjectHandlePackageDebugData*>(GObjectHandleIndex.PackageData.GetData());
+			}
+			else
+			{
+				PackageData = &GObjectHandleIndex.PackageData[FoundPackageId->ToIndex()];
+			}
+			//Can't reasonably switch back to a read lock here because the downgrade would have a window where 
+			//the global map could be modified and invalidate the pointer we're holding to an element in it.
+		}
+		else
+		{
+			PackageData = &GObjectHandleIndex.PackageData[FoundPackageId->ToIndex()];
+		}
+		checkf(NameToMinimalName(PackageName) == PackageData->PackageName,
+			TEXT("Package name collision detected in ObjectHandle system.  '%s' and '%s' both produce a hash of %d"), *PackageName.ToString(), *MinimalNameToName(PackageData->PackageName).ToString(), NameHash);
+
+		OutPackageId = *FoundPackageId;
+
+		FRWScopeLock LocalLockScope(PackageData->Lock, SLT_ReadOnly);
+		FObjectId* FoundId = PackageData->PathToObjectId.Find(ObjectPath);
+
+		if (FoundId)
+		{
+			check(*FoundId != EObjectId::Invalid);
+			OutObjectId = *FoundId;
+			return;
+		}
+
+		LocalLockScope.ReleaseReadOnlyLockAndAcquireWriteLock_USE_WITH_CAUTION();
+		//The PathToObjectId could have been modified when the read lock was released and the write
+		//lock was acquired, so we must check and see if the ObjectPath was added in that window.
+		FoundId = &PackageData->PathToObjectId.FindOrAdd(ObjectPath, EObjectId::Invalid);
+		if (*FoundId != EObjectId::Invalid)
+		{
+			OutObjectId = *FoundId;
+			return;
+		}
+		uint32 PathIndex = PackageData->ObjectPaths.Emplace(ObjectPath);
+		checkf(((PathIndex + 1) & ~ObjectHandle_Private::ObjectPathIdMask) == 0, TEXT("Path id overflowed space in ObjectHandle"));
+		FoundId->Components.ObjectPathId = PathIndex + 1;
+
+		if (!ClassName.IsNone() && !ClassPackageName.IsNone())
+		{
+			// @TODO: OBJPTR: This could be inefficient if there are a high number of references to blueprint data instances
+			//		or references to unique blueprints in a single package.  Evaluate whether that's likely to be
+			//		the case in practice.
+			FObjectHandleDataClassDescriptor DataClassDesc{NameToMinimalName(ClassPackageName), NameToMinimalName(ClassName)};
+			uint32 DataClassDescriptorIndex = PackageData->DataClassDescriptors.AddUnique(DataClassDesc);
+			checkf(((DataClassDescriptorIndex + 1) & ~ObjectHandle_Private::DataClassDescriptorIdMask) == 0, TEXT("Data class descriptor id overflowed space in ObjectHandle"));
+			FoundId->Components.DataClassDescriptorId = DataClassDescriptorIndex + 1;
+		}
+		OutObjectId = *FoundId;
+	}
+
+	static inline FPackedObjectRef MakePackedObjectRef(FName PackageName, FName ClassPackageName, FName ClassName, FObjectPathId ObjectPath)
+	{
+		FPackageId PackageId;
+		FObjectId ObjectId;
+		MakeReferenceIds(PackageName, ClassPackageName, ClassName, ObjectPath, PackageId, ObjectId);
+		return Pack(PackageId, ObjectId);
+	}
+
+	static void GetObjectDataFromId(FPackageId PackageId, FObjectId ObjectId, FMinimalName& OutPackageName, FObjectPathId& OutPathId, FMinimalName& OutClassPackageName, FMinimalName& OutClassName)
+	{
+		if (ObjectId == EObjectId::Invalid)
+		{
+			return;
+		}
+
+		FRWScopeLock GlobalLockScope(GObjectHandleIndex.Lock, SLT_ReadOnly);
+		FObjectHandlePackageData& FoundPackageData = GObjectHandleIndex.PackageData[PackageId.ToIndex()];
+
+		FRWScopeLock LocalLockScope(FoundPackageData.Lock, SLT_ReadOnly);
+		OutPackageName = FoundPackageData.PackageName;
+		check(static_cast<uint32>(FoundPackageData.ObjectPaths.Num()) >= ObjectId.Components.ObjectPathId);
+		OutPathId = FoundPackageData.ObjectPaths[ObjectId.Components.ObjectPathId - 1];
+
+		if (ObjectId.Components.DataClassDescriptorId > 0)
+		{
+			check(static_cast<uint32>(FoundPackageData.DataClassDescriptors.Num()) >= ObjectId.Components.DataClassDescriptorId);
+			FObjectHandleDataClassDescriptor& Desc = FoundPackageData.DataClassDescriptors[ObjectId.Components.DataClassDescriptorId - 1];
+			OutClassPackageName = Desc.PackageName;
+			OutClassName = Desc.ClassName;
+		}
+	}
+
+	static inline FObjectRef MakeObjectRef(FPackedObjectRef PackedObjectRef)
+	{
+		FObjectId ObjectId;
+		FPackageId PackageId;
+		Unpack(PackedObjectRef, PackageId, ObjectId);
+
+		FMinimalName PackageName;
+		FObjectPathId PathId;
+		FMinimalName ClassPackageName;
+		FMinimalName ClassName;
+		GetObjectDataFromId(PackageId, ObjectId, PackageName, PathId, ClassPackageName, ClassName);
+		return FObjectRef{MinimalNameToName(PackageName), MinimalNameToName(ClassPackageName), MinimalNameToName(ClassName), PathId};
+	}
+};
+
+static inline FName GetNameOrNone(UObject* Object)
+{
+	return Object ? Object->GetFName() : NAME_None;
+}
+
+FObjectRef MakeObjectRef(const UObject* Object)
+{
+	if (!Object)
+	{
+		return {NAME_None, NAME_None, NAME_None, FObjectPathId()};
+	}
+	
+	UObject* ClassGeneratedBy = Object->GetClass()->ClassGeneratedBy;
+	UPackage* ClassGeneratedByPackage = ClassGeneratedBy ? ClassGeneratedBy->GetOutermost() : nullptr;
+	return FObjectRef {GetNameOrNone(Object->GetOutermost()), GetNameOrNone(ClassGeneratedByPackage), GetNameOrNone(ClassGeneratedBy), FObjectPathId(Object)};
+}
+
+FObjectRef MakeObjectRef(FPackedObjectRef PackedObjectRef)
+{
+	if (IsPackedObjectRefNull(PackedObjectRef))
+	{
+		return {NAME_None, NAME_None, NAME_None, FObjectPathId()};
+	}
+
+	return ObjectHandle_Private::MakeObjectRef(PackedObjectRef);
+}
+
+FPackedObjectRef MakePackedObjectRef(const UObject* Object)
+{
+	if (!Object)
+	{
+		return {0};
+	}
+
+	FName PackageName = GetNameOrNone(Object->GetOutermost());
+	UObject* ClassGeneratedBy = Object->GetClass()->ClassGeneratedBy;
+	UPackage* ClassGeneratedByPackage = ClassGeneratedBy ? ClassGeneratedBy->GetOutermost() : nullptr;
+	return ObjectHandle_Private::MakePackedObjectRef(PackageName, GetNameOrNone(ClassGeneratedByPackage), GetNameOrNone(ClassGeneratedBy), FObjectPathId(Object));
+}
+
+FPackedObjectRef MakePackedObjectRef(const FObjectRef& ObjectRef)
+{
+	if (IsObjectRefNull(ObjectRef))
+	{
+		return {0};
+	}
+
+	return ObjectHandle_Private::MakePackedObjectRef(ObjectRef.PackageName, ObjectRef.ClassPackageName, ObjectRef.ClassName, ObjectRef.ObjectPath);
+}
+
+static inline UPackage* FindOrLoadPackage(FName PackageName, int32 LoadFlags)
+{
+	// @TODO: OBJPTR: Want to replicate the functional path of an import here.  See things like FindImportFast in BlueprintSupport.cpp
+	// 		 for additional behavior that we're not handling here yet.
+	FName* ScriptPackageName = FPackageName::FindScriptPackageName(PackageName);
+	UPackage* TargetPackage = (UPackage*)StaticFindObjectFast(UPackage::StaticClass(), nullptr, PackageName);
+	if (!ScriptPackageName && !TargetPackage)
+	{
+		// @TODO: OBJPTR: This seems like it will do a search on disk every time instead of leveraging known missing packages
+		// @TODO: OBJPTR: Instancing context may be important to consider when loading the package.
+		LoadFlags |= LOAD_NoWarn | LOAD_NoVerify; //This does nothing? | LOAD_DisableDependencyPreloading;
+		TargetPackage = LoadPackage(nullptr, *PackageName.ToString(), LoadFlags);
+	}
+	return TargetPackage;
+}
+
+UClass* ResolveObjectRefClass(const FObjectRef& ObjectRef, uint32 LoadFlags /*= LOAD_None*/)
+{
+	UClass* ClassObject = nullptr;
+	UPackage* ClassPackage = nullptr;
+	if (!ObjectRef.ClassPackageName.IsNone())
+	{
+		ClassPackage = FindOrLoadPackage(ObjectRef.ClassPackageName, LoadFlags);
+
+		if (!ObjectRef.ClassName.IsNone())
+		{
+			ClassObject = (UClass*)StaticFindObjectFast(UClass::StaticClass(), ClassPackage, ObjectRef.ClassName);
+			if (ClassObject)
+			{
+				if (ClassObject->HasAnyFlags(RF_NeedLoad) && ClassPackage->LinkerLoad)
+				{
+					ClassPackage->LinkerLoad->Preload(ClassObject);
+				}
+				ClassObject->GetDefaultObject(); // build the CDO if it isn't already built
+			}
+		}
+	}
+
+	ObjectHandle_Private::OnClassReferenceResolved(ObjectRef, ClassPackage, ClassObject);
+	return ClassObject;
+}
+
+UObject* ResolveObjectRef(const FObjectRef& ObjectRef, uint32 LoadFlags /*= LOAD_None*/)
+{
+	//NOTE: Written with the assumption that there's no partial loading of packages.
+	//		A package when loaded is fully loaded, so there's no possibility of needing
+	//		to do additional loading work up the outer chain for a target object.
+
+	if (IsObjectRefNull(ObjectRef))
+	{
+		ObjectHandle_Private::OnReferenceResolved(ObjectRef, nullptr, nullptr);
+		return nullptr;
+	}
+
+#if USE_CIRCULAR_DEPENDENCY_LOAD_DEFERRING
+	// @TODO: OBJPTR: Do we need something equivalent for handling circular dependencies during lazy load?
+	//DeferPotentialCircularImport(Index);
+#endif // USE_CIRCULAR_DEPENDENCY_LOAD_DEFERRING
+
+	ResolveObjectRefClass(ObjectRef, LoadFlags);
+
+	UPackage* TargetPackage = FindOrLoadPackage(ObjectRef.PackageName, LoadFlags);
+
+	if (!TargetPackage)
+	{
+		ObjectHandle_Private::OnReferenceResolved(ObjectRef, nullptr, nullptr);
+		return nullptr;
+	}
+
+	FObjectPathId::ResolvedNameContainerType ResolvedNames;
+	ObjectRef.ObjectPath.Resolve(ResolvedNames);
+
+	// @TODO: OBJPTR: This seems like a silly way to find objects in a hierarchy, investigate better options.
+	//		 Normal import resolving walks up the outer hierarchy and each outer can have a resolved
+	//		 result cached avoiding any further walking up the hierarchy. Is that something that
+	//		 could/should be emulated here?
+	UObject* CurrentObject = TargetPackage;
+	for (int32 ObjectPathIndex = 0; ObjectPathIndex < ResolvedNames.Num(); ++ObjectPathIndex)
+	{
+		CurrentObject = StaticFindObjectFast(nullptr, CurrentObject, ResolvedNames[ObjectPathIndex]);
+		if (!CurrentObject)
+		{
+			ObjectHandle_Private::OnReferenceResolved(ObjectRef, TargetPackage, nullptr);
+			return nullptr;
+		}
+	}
+
+	if (CurrentObject->HasAnyFlags(RF_NeedLoad) && TargetPackage->LinkerLoad)
+	{
+		TargetPackage->LinkerLoad->Preload(CurrentObject);
+	}
+	ObjectHandle_Private::OnReferenceResolved(ObjectRef, TargetPackage, CurrentObject);
+
+	return CurrentObject;
+}
+
+UClass* ResolvePackedObjectRefClass(FPackedObjectRef PackedObjectRef, uint32 LoadFlags /*= LOAD_None*/)
+{
+	return ResolveObjectRefClass(MakeObjectRef(PackedObjectRef), LoadFlags);
+}
+
+UObject* ResolvePackedObjectRef(FPackedObjectRef PackedObjectRef, uint32 LoadFlags /*= LOAD_None*/)
+{
+	return ResolveObjectRef(MakeObjectRef(PackedObjectRef), LoadFlags);
+}
+
+#if UE_WITH_OBJECT_HANDLE_TRACKING
+COREUOBJECT_API ObjectHandleReadFunction* ObjectHandle_Private::ObjectHandleReadCallback = nullptr;
+COREUOBJECT_API ObjectHandleClassResolvedFunction* ObjectHandle_Private::ObjectHandleClassResolvedCallback = nullptr;
+COREUOBJECT_API ObjectHandleReferenceResolvedFunction* ObjectHandle_Private::ObjectHandleReferenceResolvedCallback = nullptr;
+
+ObjectHandleReadFunction* SetObjectHandleReadCallback(ObjectHandleReadFunction* Function)
+{
+	return (ObjectHandleReadFunction*)FPlatformAtomics::InterlockedExchangePtr((void**)&ObjectHandle_Private::ObjectHandleReadCallback, (void*)Function);
+}
+
+ObjectHandleClassResolvedFunction* SetObjectHandleClassResolvedCallback(ObjectHandleClassResolvedFunction* Function)
+{
+	return (ObjectHandleClassResolvedFunction*)FPlatformAtomics::InterlockedExchangePtr((void**)&ObjectHandle_Private::ObjectHandleClassResolvedCallback, (void*)Function);
+}
+
+ObjectHandleReferenceResolvedFunction* SetObjectHandleReferenceResolvedCallback(ObjectHandleReferenceResolvedFunction* Function)
+{
+	return (ObjectHandleReferenceResolvedFunction*)FPlatformAtomics::InterlockedExchangePtr((void**)&ObjectHandle_Private::ObjectHandleReferenceResolvedCallback, (void*)Function);
+}
+#endif // UE_WITH_OBJECT_HANDLE_TRACKING
diff --git a/Engine/Source/Runtime/CoreUObject/Private/UObject/ObjectPathId.cpp b/Engine/Source/Runtime/CoreUObject/Private/UObject/ObjectPathId.cpp
new file mode 100644
--- /dev/null
+++ b/Engine/Source/Runtime/CoreUObject/Private/UObject/ObjectPathId.cpp
@@ -0,0 +1,256 @@
+// Copyright Epic Games, Inc. All Rights Reserved.
+
+#include "UObject/ObjectPathId.h"
+#include "Algo/FindLast.h"
+#include "Containers/Array.h"
+#include "Containers/Map.h"
+#include "Containers/StringView.h"
+#include "HAL/CriticalSection.h"
+#include "UObject/Linker.h"
+#include "UObject/Object.h"
+#include "UObject/Package.h"
+
+FCriticalSection GComplexPathLock;  // Could be changed to an RWLock later if needed
+
+// @TODO: OBJPTR: Consider if it is possible to have this be case-preserving while still having equality checks between two paths of differing case be equal.
+// @TODO: OBJPTR: Evaluate if the inline array for the object paths needs to be changed to something more lightweight.
+// Currently each unique object path takes up 48 bytes of memory:
+//  *8 bytes per entry in the ComplexPathHashToId map
+//  *40 byte per entry in the ComplexPaths array
+//My expectation is that a 3 name object path is generous in almost every case.  48 bytes
+//per complex path may be too expensive depending on how frequently we encounter complex paths.
+//If so, we can consider a packed pool to store the array elements in or other options for
+//representing shared path elements like just registering the paths as FNames and not having
+//our own storage at all.
+TMultiMap<uint32, uint32> GComplexPathHashToId;
+template class CORE_API TArray<FMinimalName, TInlineAllocator<3>>;
+TArray<TArray<FMinimalName, TInlineAllocator<3>>> GComplexPaths;
+/*
+template <int a, int b>
+struct assert_equality
+{
+	static_assert(a == b, "not equal");
+};
+
+//assert_equality<sizeof(TMultiMap<uint32, uint32>::ElementType), 0> A;
+assert_equality<sizeof(TArray<FMinimalName, TInlineAllocator<3>>), 0> B;
+*/
+
+static inline uint64 NameToUInt64(const FName& Name)
+{
+	return (static_cast<uint64>(Name.GetComparisonIndex().ToUnstableInt()) << 32) |
+			static_cast<uint64>(Name.GetNumber());
+}
+
+template <typename NameProducerType>
+void StoreObjectPathId(NameProducerType& NameProducer, uint64 SimplePathFlag, uint64& OutObjectPathId)
+{
+	FName Name = NameProducer();
+
+	if (Name == NAME_None)
+	{
+		return;
+	}
+
+	FName OuterName = NameProducer();
+	if (OuterName == NAME_None)
+	{
+		int32 Number = Name.GetNumber();
+		// Check if the number is in range by left shift + right shift
+		if (static_cast<int32>((static_cast<uint32>(Number) << 1) >> 1) == Number)
+		{
+			//Simple path scenario
+			OutObjectPathId = static_cast<uint64>(Name.GetComparisonIndex().ToUnstableInt()) << 32 |
+								static_cast<uint64>(static_cast<uint32>(static_cast<uint32>(Number) << 1)) |
+								SimplePathFlag;
+			return;
+		}
+	}
+
+	//Complex path scenario
+	TArray<FMinimalName, TInlineAllocator<3>> MinimalNames;
+	uint32 Key = GetTypeHash(NameToUInt64(Name));
+	MinimalNames.Emplace(Name.GetComparisonIndex(), Name.GetNumber());
+	while (OuterName != NAME_None)
+	{
+		Name = OuterName;
+		OuterName = NameProducer();
+		MinimalNames.Emplace(Name.GetComparisonIndex(), Name.GetNumber());
+		Key = HashCombine(Key, GetTypeHash(NameToUInt64(Name)));
+	}
+
+	FScopeLock ComplexPathScopeLock(&GComplexPathLock);
+	for (typename TMultiMap<uint32, uint32>::TConstKeyIterator It(GComplexPathHashToId, Key); It; ++It)
+	{
+		uint32 PotentialPathId = It.Value();
+		if (GComplexPaths[PotentialPathId-1] == MinimalNames)
+		{
+			OutObjectPathId = static_cast<uint32>(PotentialPathId << 1);
+			return;
+		}
+	}
+
+	uint32 NewId = GComplexPaths.Add(MinimalNames) + 1;
+	GComplexPathHashToId.Add(Key, NewId);
+	OutObjectPathId = static_cast<uint32>(NewId << 1);
+	GCoreComplexObjectPathDebug = GComplexPaths.GetData();
+}
+
+FObjectPathId::FObjectPathId(const UObject* Object)
+{
+	struct FLoadedObjectNamePathProducer
+	{
+		FLoadedObjectNamePathProducer(const UObject* InObject)
+		: CurrentObject(InObject)
+		{}
+
+		FName operator()(void)
+		{
+			if ((CurrentObject == nullptr) || (CurrentObject->GetClass() == UPackage::StaticClass()))
+			{
+				return NAME_None;
+			}
+			FName RetVal = CurrentObject->GetFName();
+			CurrentObject = CurrentObject->GetOuter();
+			return RetVal;
+		}
+	private:
+		const UObject* CurrentObject;
+	} NamePathProducer(Object);
+
+	StoreObjectPathId(NamePathProducer, static_cast<uint64>(EPathId::FlagSimple), PathId);
+}
+
+FObjectPathId::FObjectPathId(const FObjectImport& Import, const FLinkerTables& LinkerTables)
+{
+	MakeImportPathIdAndPackageName(Import, LinkerTables, *this);
+}
+
+FName FObjectPathId::MakeImportPathIdAndPackageName(const FObjectImport& Import, const FLinkerTables& LinkerTables, FObjectPathId& OutPathId)
+{
+	// @TODO: OBJPTR: Need to handle redirects.  FCoreRedirectObjectName could be used, but it doesn't fit
+	//		conveniently with the FName walk approach that is currently here. 
+	struct FImportObjectNamePathProducer
+	{
+		FImportObjectNamePathProducer(const FObjectImport& InImport, const FLinkerTables& InLinkerTables)
+		: CurrentImport(&InImport)
+		, LinkerTables(InLinkerTables)
+		{}
+
+		FName operator()(void)
+		{
+			if ((CurrentImport == nullptr) || CurrentImport->OuterIndex.IsNull())
+			{
+				return NAME_None;
+			}
+			FName RetVal = CurrentImport->ObjectName;
+			CurrentImport = &LinkerTables.Imp(CurrentImport->OuterIndex);
+			return RetVal;
+		}
+
+		FName GetPackageName()
+		{
+			if (CurrentImport && CurrentImport->OuterIndex.IsNull())
+			{
+				return CurrentImport->ObjectName;
+			}
+			return NAME_None;
+		}
+	private:
+		const FObjectImport* CurrentImport;
+		const FLinkerTables& LinkerTables;
+	} NamePathProducer(Import, LinkerTables);
+
+	StoreObjectPathId(NamePathProducer, static_cast<uint64>(EPathId::FlagSimple), OutPathId.PathId);
+	return NamePathProducer.GetPackageName();
+}
+
+#if !(UE_BUILD_SHIPPING || UE_BUILD_TEST)
+template <typename T>
+struct TStringViewNamePathProducer
+{
+	TStringViewNamePathProducer(T StringPath)
+	: CurrentStringView(StringPath)
+	{}
+
+	FName operator()(void)
+	{
+		if (CurrentStringView.IsEmpty())
+		{
+			return NAME_None;
+		}
+
+		typename T::SizeType FoundIndex = INDEX_NONE;
+		if (!FindLastSeparator(FoundIndex))
+		{
+			FName RetVal(CurrentStringView);
+			CurrentStringView.Reset();
+			return RetVal;
+		}
+
+		FName RetVal(CurrentStringView.RightChop(FoundIndex + 1));
+		CurrentStringView.RemoveSuffix(CurrentStringView.Len() - FoundIndex);
+		return RetVal;
+	}
+private:
+	static inline bool IsPathIdSeparator(typename T::ElementType Char) { return Char == '.' || Char == ':'; }
+
+	bool FindLastSeparator(typename T::SizeType& OutIndex)
+	{
+		if (const typename T::ElementType* Separator = Algo::FindLastByPredicate(CurrentStringView, IsPathIdSeparator))
+		{
+			OutIndex = static_cast<typename T::SizeType>(Separator - CurrentStringView.GetData());
+			return true;
+		}
+		OutIndex = INDEX_NONE;
+		return false;
+	}
+
+	T CurrentStringView;
+};
+
+FObjectPathId::FObjectPathId(FWideStringView StringPath)
+{
+	// @TODO: OBJPTR: Need to handle redirects.  FCoreRedirectObjectName could be used, but it doesn't fit
+	//		conveniently with the FName walk approach that is currently here.
+	TStringViewNamePathProducer<FWideStringView> NamePathProducer(StringPath);
+
+	StoreObjectPathId(NamePathProducer, static_cast<uint64>(EPathId::FlagSimple), PathId);
+}
+
+FObjectPathId::FObjectPathId(FAnsiStringView StringPath)
+{
+	// @TODO: OBJPTR: Need to handle redirects.  FCoreRedirectObjectName could be used, but it doesn't fit
+	//		conveniently with the FName walk approach that is currently here.
+	TStringViewNamePathProducer<FAnsiStringView> NamePathProducer(StringPath);
+
+	StoreObjectPathId(NamePathProducer, static_cast<uint64>(EPathId::FlagSimple), PathId);
+}
+#endif
+
+void FObjectPathId::Resolve(ResolvedNameContainerType& OutContainer) const
+{
+	if (IsNone())
+	{
+		return;
+	}
+
+	if (static_cast<uint64>(PathId) & static_cast<uint64>(EPathId::FlagSimple))
+	{
+		uint32 Number = static_cast<uint32>(PathId & 0x00000000FFFFFFE) >> 1;
+		uint32 Index = static_cast<uint32>(PathId >> 32);
+		FNameEntryId EntryId = FNameEntryId::FromUnstableInt(Index);
+		OutContainer.Emplace(EntryId, EntryId, Number);
+		return; 
+	}
+
+	FScopeLock ComplexPathScopeLock(&GComplexPathLock);
+	//Append reversed path
+	TArray<FMinimalName, TInlineAllocator<3>>& FoundContainer = GComplexPaths[(PathId >> 1) - 1];
+	OutContainer.Reserve(OutContainer.Num() + FoundContainer.Num());
+	for (int32 Index = FoundContainer.Num() - 1; Index >= 0; --Index)
+	{
+		FMinimalName& FoundMinimalName = FoundContainer[Index];
+		OutContainer.Emplace(FoundMinimalName.Index, FoundMinimalName.Index, FoundMinimalName.Number);
+	}
+}
diff --git a/Engine/Source/Runtime/CoreUObject/Private/UObject/PropertyObjectPtr.cpp b/Engine/Source/Runtime/CoreUObject/Private/UObject/PropertyObjectPtr.cpp
new file mode 100644
--- /dev/null
+++ b/Engine/Source/Runtime/CoreUObject/Private/UObject/PropertyObjectPtr.cpp
@@ -0,0 +1,120 @@
+// Copyright Epic Games, Inc. All Rights Reserved.
+
+#include "CoreMinimal.h"
+#include "UObject/ObjectMacros.h"
+#include "UObject/PropertyHelper.h"
+#include "UObject/UnrealType.h"
+
+/*-----------------------------------------------------------------------------
+	FObjectPtrProperty.
+-----------------------------------------------------------------------------*/
+IMPLEMENT_FIELD(FObjectPtrProperty)
+
+FString FObjectPtrProperty::GetCPPType(FString* ExtendedTypeText/*=NULL*/, uint32 CPPExportFlags/*=0*/) const
+{
+	return FString::Printf(TEXT("TObjectPtr<%s%s>"), PropertyClass->GetPrefixCPP(), *PropertyClass->GetName());
+}
+FString FObjectPtrProperty::GetCPPMacroType(FString& ExtendedTypeText) const
+{
+	ExtendedTypeText = FString::Printf(TEXT("TObjectPtr<%s%s>"), PropertyClass->GetPrefixCPP(), *PropertyClass->GetName());
+	return TEXT("OBJECTPTR");
+}
+
+FName FObjectPtrProperty::GetID() const
+{
+	return NAME_ObjectPtrProperty;
+}
+
+void FObjectPtrProperty::SerializeItem(FStructuredArchive::FSlot Slot, void* Value, void const* Defaults) const
+{
+	FArchive& UnderlyingArchive = Slot.GetUnderlyingArchive();
+
+	FObjectPtr* ObjectPtr = GetPropertyValuePtr(Value);
+
+	Slot << *ObjectPtr;
+}
+
+EConvertFromTypeResult FObjectPtrProperty::ConvertFromType(const FPropertyTag& Tag, FStructuredArchive::FSlot Slot, uint8* Data, UStruct* DefaultsStruct)
+{
+	// Convert raw object reference to TObjectPtr
+	if (Tag.Type == NAME_ObjectProperty)
+	{
+		uint8* DestAddress = ContainerPtrToValuePtr<uint8>(Data, Tag.ArrayIndex);
+		FArchive& UnderlyingArchive = Slot.GetUnderlyingArchive();
+		FObjectPtr* ObjectPtr = GetPropertyValuePtr(DestAddress);
+		Slot << *ObjectPtr;
+
+		return UnderlyingArchive.IsCriticalError() ? EConvertFromTypeResult::CannotConvert : EConvertFromTypeResult::Converted;
+	}
+
+	return EConvertFromTypeResult::UseSerializeItem;
+}
+
+bool FObjectPtrProperty::SameType(const FProperty* Other) const
+{
+	// @TODO: OBJPTR: Should this be done through a new, separate API on FProperty (eg: ImplicitConv)
+	return Super::SameType(Other) ||
+		(Other && Other->IsA<FObjectProperty>() &&
+			(PropertyClass == ((FObjectPropertyBase*)Other)->PropertyClass) );
+}
+
+bool FObjectPtrProperty::Identical(const void* A, const void* B, uint32 PortFlags) const
+{
+	FObjectPtr ObjectA = A ? *((FObjectPtr*)A) : FObjectPtr();
+	FObjectPtr ObjectB = B ? *((FObjectPtr*)B) : FObjectPtr();
+
+	if (ObjectA.IsNull() || ObjectB.IsNull())
+	{
+		return ObjectA.IsNull() == ObjectB.IsNull();
+	}
+
+	// Compare actual pointers. We don't do this during PIE because we want to be sure to serialize everything. An example is the LevelScriptActor being serialized against its CDO,
+	// which contains actor references. We want to serialize those references so they are fixed up.
+	const bool bDuplicatingForPIE = (PortFlags&PPF_DuplicateForPIE) != 0;
+	bool bResult = !bDuplicatingForPIE ? (ObjectA == ObjectB) : false;
+	// always serialize the cross level references, because they could be NULL
+	// @todo: okay, this is pretty hacky overall - we should have a PortFlag or something
+	// that is set during SavePackage. Other times, we don't want to immediately return false
+	// (instead of just this ExportDefProps case)
+	// instance testing
+	if (!bResult && ObjectA->GetClass() == ObjectB->GetClass())
+	{
+		bool bPerformDeepComparison = (PortFlags&PPF_DeepComparison) != 0;
+		if ((PortFlags&PPF_DeepCompareInstances) && !bPerformDeepComparison)
+		{
+			bPerformDeepComparison = ObjectA->IsTemplate() != ObjectB->IsTemplate();
+		}
+
+		if (!bResult && bPerformDeepComparison)
+		{
+			// In order for deep comparison to be match they both need to have the same name and that name needs to be included in the instancing table for the class
+			if (ObjectA->GetFName() == ObjectB->GetFName() && ObjectA->GetClass()->GetDefaultSubobjectByName(ObjectA->GetFName()))
+			{
+				checkSlow(ObjectA->IsDefaultSubobject() && ObjectB->IsDefaultSubobject() && ObjectA->GetClass()->GetDefaultSubobjectByName(ObjectA->GetFName()) == ObjectB->GetClass()->GetDefaultSubobjectByName(ObjectB->GetFName())); // equivalent
+				bResult = AreInstancedObjectsIdentical(ObjectA.Get(), ObjectB.Get(), PortFlags);
+			}
+		}
+	}
+	return bResult;
+}
+
+UObject* FObjectPtrProperty::GetObjectPropertyValue(const void* PropertyValueAddress) const
+{
+	return GetPropertyValue(PropertyValueAddress).Get();
+}
+
+void FObjectPtrProperty::SetObjectPropertyValue(void* PropertyValueAddress, UObject* Value) const
+{
+	SetPropertyValue(PropertyValueAddress, TCppType(Value));
+}
+
+bool FObjectPtrProperty::AllowCrossLevel() const
+{
+	return true;
+}
+
+uint32 FObjectPtrProperty::GetValueTypeHashInternal(const void* Src) const
+{
+	return GetTypeHash(GetPropertyValue(Src));
+}
+
diff --git a/Engine/Source/Runtime/CoreUObject/Private/UObject/SavePackage.cpp b/Engine/Source/Runtime/CoreUObject/Private/UObject/SavePackage.cpp
--- a/Engine/Source/Runtime/CoreUObject/Private/UObject/SavePackage.cpp
+++ b/Engine/Source/Runtime/CoreUObject/Private/UObject/SavePackage.cpp
@@ -66,6 +66,8 @@
 #include "UObject/AsyncWorkSequence.h"
 #include "Serialization/BulkDataManifest.h"
 #include "Misc/ScopeExit.h"
+#include "Misc/PackageAccessTracking.h"
+#include "Misc/PackageAccessTrackingOps.h"
 
 #if ENABLE_COOK_STATS
 #include "ProfilingDebugging/ScopedTimers.h"
@@ -86,6 +88,18 @@ static FCriticalSection InitializeCoreClassesCritSec;
 // bring the UObectGlobal declaration visible to non editor
 bool IsEditorOnlyObject(const UObject* InObject, bool bCheckRecursive, bool bCheckMarks);
 
+static void AppendAdditionalData(FLinkerSave& Linker)
+{
+	FArchive& Ar = Linker;
+
+	for (FLinkerSave::AdditionalDataCallback& Callback : Linker.AdditionalDataToAppend)
+	{
+		Callback(Ar);
+	}
+
+	Linker.AdditionalDataToAppend.Empty();
+}
+
 static bool EndSavingIfCancelled() 
 {
 	if ( GWarn->ReceivedUserCancel() )
@@ -2151,6 +2165,7 @@ FSavePackageResultStruct UPackage::Save(UPackage* InOuter, UObject* Base, EObjec
 		return UPackage::Save2(InOuter, Base, Filename, SaveArgs);
 	}
 
+	UE_TRACK_REFERENCING_PACKAGE_SCOPED(InOuter, PackageAccessTrackingOps::NAME_Save);
 	COOK_STAT(FScopedDurationTimer FuncSaveTimer(FSavePackageStats::SavePackageTimeSec));
 	COOK_STAT(FSavePackageStats::NumPackagesSaved++);
 	SCOPED_SAVETIMER(UPackage_Save);
@@ -4268,6 +4283,8 @@ FSavePackageResultStruct UPackage::Save(UPackage* InOuter, UObject* Base, EObjec
 				SavePackageUtilities::SaveBulkData(Linker.Get(), InOuter, Filename, TargetPlatform, SavePackageContext, bTextFormat, bDiffing,
 							 bComputeHash, AsyncWriteAndHashSequence, TotalPackageSizeUncompressed );
 
+				AppendAdditionalData(*Linker);
+
 #if WITH_EDITOR
 				if (bIsCooking && AdditionalFilesFromExports.Num() > 0)
 				{
diff --git a/Engine/Source/Runtime/CoreUObject/Private/UObject/SoftObjectPath.cpp b/Engine/Source/Runtime/CoreUObject/Private/UObject/SoftObjectPath.cpp
--- a/Engine/Source/Runtime/CoreUObject/Private/UObject/SoftObjectPath.cpp
+++ b/Engine/Source/Runtime/CoreUObject/Private/UObject/SoftObjectPath.cpp
@@ -19,6 +19,42 @@ FSoftObjectPath::FSoftObjectPath(const UObject* InObject)
 	}
 }
 
+FSoftObjectPath::FSoftObjectPath(FObjectHandle InObjectHandle)
+{
+	if (!IsObjectHandleNull(InObjectHandle))
+	{
+		//Try to compose the path without forcing the referenced data to be loaded
+		FObjectRef Ref = MakeObjectRef(InObjectHandle);
+		FObjectPathId::ResolvedNameContainerType ResolvedNames;
+		Ref.ObjectPath.Resolve(ResolvedNames);
+		TStringBuilder<128> CompletePath;
+		CompletePath.Append(Ref.PackageName.ToString());
+		CompletePath.Append(TEXT('.'));
+		for (int32 ResolvedNameIndex = 0; ResolvedNameIndex < ResolvedNames.Num(); ++ResolvedNameIndex)
+		{
+			switch (ResolvedNameIndex)
+			{
+				case 0:
+				{
+					break;
+				}
+				case 1:
+				{
+					CompletePath << SUBOBJECT_DELIMITER_CHAR;
+					break;
+				}
+				default:
+				{
+					CompletePath << '.';
+					break;
+				}
+			}
+			CompletePath << ResolvedNames[ResolvedNameIndex];
+		}
+		SetPath(CompletePath.ToString());
+	}
+}
+
 FString FSoftObjectPath::ToString() const
 {
 	// Most of the time there is no sub path so we can do a single string allocation
diff --git a/Engine/Source/Runtime/CoreUObject/Private/UObject/UObjectGlobals.cpp b/Engine/Source/Runtime/CoreUObject/Private/UObject/UObjectGlobals.cpp
--- a/Engine/Source/Runtime/CoreUObject/Private/UObject/UObjectGlobals.cpp
+++ b/Engine/Source/Runtime/CoreUObject/Private/UObject/UObjectGlobals.cpp
@@ -57,6 +57,8 @@
 #include "Misc/CoreDelegates.h"
 #include "ProfilingDebugging/CsvProfiler.h"
 #include "ProfilingDebugging/LoadTimeTracker.h"
+#include "Misc/PackageAccessTracking.h"
+#include "Misc/PackageAccessTracking.h"
 
 DEFINE_LOG_CATEGORY(LogUObjectGlobals);
 
@@ -1119,6 +1121,7 @@ UPackage* LoadPackageInternal(UPackage* InOuter, const TCHAR* InLongPackageNameO
 		}
 
 		FName PackageFName(*InPackageName);
+		UE_TRACK_REFERENCING_PACKAGE_SCOPED(PackageFName, PackageAccessTrackingOps::NAME_Load);
 
 		{
 			if (FCoreDelegates::OnSyncLoadPackage.IsBound())
@@ -1249,6 +1252,7 @@ UPackage* LoadPackageInternal(UPackage* InOuter, const TCHAR* InLongPackageNameO
 
 		Result = Linker->LinkerRoot;
 		checkf(Result, TEXT("LinkerRoot is null"));
+		UE_TRACK_REFERENCING_PACKAGE_SCOPED(Result, PackageAccessTrackingOps::NAME_Load);
 
 		auto EndLoadAndCopyLocalizationGatherFlag = [&]
 		{
@@ -1579,6 +1583,7 @@ void EndLoad(FUObjectSerializeContext* LoadContext)
 			// Finish loading everything.
 			{
 				SCOPED_LOADTIMER(PreLoadAndSerialize);
+				UE_TRACK_REFERENCING_PACKAGE_DELAYED_SCOPED(AccessRefScope, PackageAccessTrackingOps::NAME_PreLoad);
 				for (int32 i = 0; i < ObjLoaded.Num(); i++)
 				{
 					// Preload.
@@ -1586,6 +1591,7 @@ void EndLoad(FUObjectSerializeContext* LoadContext)
 					if (Obj->HasAnyFlags(RF_NeedLoad))
 					{
 						check(Obj->GetLinker());
+						UE_TRACK_REFERENCING_PACKAGE_DELAYED(AccessRefScope, Obj->GetOutermost());
 						Obj->GetLinker()->Preload(Obj);
 					}
 				}
@@ -3485,6 +3491,18 @@ protected:
 		return *this;
 	}
 
+	virtual FArchive& operator<<(FObjectPtr& Object) override
+	{
+		if (Object && IsObjectHandleResolved(Object.GetHandle()))
+		{
+			FReferenceCollector& CurrentCollector = GetCollector();
+			FProperty* OldCollectorSerializedProperty = CurrentCollector.GetSerializedProperty();
+			CurrentCollector.SetSerializedProperty(GetSerializedProperty());
+			CurrentCollector.AddReferencedObject(Object.ToTObjectPtr(), GetSerializingObject(), GetSerializedProperty());
+			CurrentCollector.SetSerializedProperty(OldCollectorSerializedProperty);
+		}
+		return *this;
+	}
 
 };
 
@@ -3598,6 +3616,18 @@ FReferenceCollector::~FReferenceCollector()
 	delete PersistentFrameReferenceCollectorArchive;
 }
 
+void FReferenceCollector::HandleObjectReferences(FObjectPtr* InObjects, const int32 ObjectNum, const UObject* InReferencingObject, const FProperty* InReferencingProperty)
+{
+	for (int32 ObjectIndex = 0; ObjectIndex < ObjectNum; ++ObjectIndex)
+	{
+		FObjectPtr& Object = InObjects[ObjectIndex];
+		if (IsObjectHandleResolved(Object.GetHandle()))
+		{
+			HandleObjectReference(*reinterpret_cast<UObject**>(&Object), InReferencingObject, InReferencingProperty);
+		}
+	}
+}
+
 void FReferenceCollector::CreateVerySlowReferenceCollectorArchive()
 {
 	check(DefaultReferenceCollectorArchive == nullptr);
@@ -4286,13 +4316,26 @@ namespace UE4CodeGen_Private
 
 			case EPropertyGenFlags::Object:
 			{
-				const FObjectPropertyParams* Prop = (const FObjectPropertyParams*)PropBase;
-				NewProp = new FObjectProperty(Outer, UTF8_TO_TCHAR(Prop->NameUTF8), Prop->ObjectFlags, Prop->Offset, Prop->PropertyFlags, Prop->ClassFunc ? Prop->ClassFunc() : nullptr);
+				if (EnumHasAllFlags(PropBase->Flags, EPropertyGenFlags::ObjectPtr))
+				{
+					const FObjectPtrPropertyParams* Prop = (const FObjectPtrPropertyParams*)PropBase;
+					NewProp = new FObjectPtrProperty(Outer, UTF8_TO_TCHAR(Prop->NameUTF8), Prop->ObjectFlags, Prop->Offset, Prop->PropertyFlags, Prop->ClassFunc ? Prop->ClassFunc() : nullptr);
 
 #if WITH_METADATA
-				MetaDataArray = Prop->MetaDataArray;
-				NumMetaData   = Prop->NumMetaData;
+					MetaDataArray = Prop->MetaDataArray;
+					NumMetaData   = Prop->NumMetaData;
 #endif
+				}
+				else
+				{
+					const FObjectPropertyParams* Prop = (const FObjectPropertyParams*)PropBase;
+					NewProp = new FObjectProperty(Outer, UTF8_TO_TCHAR(Prop->NameUTF8), Prop->ObjectFlags, Prop->Offset, Prop->PropertyFlags, Prop->ClassFunc ? Prop->ClassFunc() : nullptr);
+
+#if WITH_METADATA
+					MetaDataArray = Prop->MetaDataArray;
+					NumMetaData   = Prop->NumMetaData;
+#endif
+				}
 			}
 			break;
 
diff --git a/Engine/Source/Runtime/CoreUObject/Private/Virtualization/VirtualizationManager.cpp b/Engine/Source/Runtime/CoreUObject/Private/Virtualization/VirtualizationManager.cpp
new file mode 100644
--- /dev/null
+++ b/Engine/Source/Runtime/CoreUObject/Private/Virtualization/VirtualizationManager.cpp
@@ -0,0 +1,121 @@
+// Copyright Epic Games, Inc. All Rights Reserved.
+
+#include "Virtualization/VirtualizationManager.h"
+
+#include "HAL/FileManager.h"
+#include "Misc/ConfigCacheIni.h"
+#include "Misc/Paths.h"
+
+DEFINE_LOG_CATEGORY(LogVirtualization);
+
+FVirtualizationManager& FVirtualizationManager::Get()
+{
+	// TODO: Do we really need to make this a singleton? Easier for prototyping.
+	static FVirtualizationManager Singleton;
+	return Singleton;
+}
+
+FVirtualizationManager::FVirtualizationManager()
+	: bEnablePayloadPushing(true)
+	, MinPayloadLength(0)
+{
+	LoadSettingsFromConfigFiles();
+}
+
+bool FVirtualizationManager::PushData(const FSharedBufferConstPtr& Payload, const FGuid& Guid)
+{
+	checkf(Guid.IsValid(), TEXT("VirtualizationManager::PushData called with an invalid guid (%s)."), *Guid.ToString());
+
+	// Early out if the pushing of payloads is disabled
+	if (bEnablePayloadPushing == false)
+	{
+		return false;
+	}
+
+	// Early out if we have no payload
+	if (!Payload.IsValid() || Payload->Size() == 0)
+	{
+		return false;
+	}
+
+	// Early out if the payload length is below our minimum required length
+	if ((int64)Payload->Size() < MinPayloadLength)
+	{
+		return false;
+	}
+
+	// TODO: At this point we should be able to consider which back end actually gets used, the following code is 
+	//		 more of a reference implementation!
+	{
+		const FString FilePath = CreateFilePath(Guid);
+
+		TUniquePtr<FArchive> FileAr(IFileManager::Get().CreateFileWriter(*FilePath));
+		check(FileAr);
+
+		uint64 PayloadLength = Payload->Size();
+		// Cast away the const because FArchive requires a non-const pointer even if we are only reading from it (saving)
+		void* PayloadPtr = const_cast<void*>(Payload->GetData());
+
+		*FileAr << PayloadLength;
+		FileAr->Serialize(PayloadPtr, PayloadLength);
+	}
+
+	return true;
+}
+
+FSharedBufferConstPtr FVirtualizationManager::PullData(const FGuid& Guid)
+{
+	checkf(Guid.IsValid(), TEXT("VirtualizationManager::PullData called with an invalid guid (%s)."), *Guid.ToString());
+
+	// TODO: At this point we should be able to consider which back end actually gets used, the following code is 
+	//		 more of a reference implementation!
+	
+	const FString FilePath = CreateFilePath(Guid);
+
+	TUniquePtr<FArchive> FileAr(IFileManager::Get().CreateFileReader(*FilePath));
+	checkf(FileAr != nullptr, TEXT("Unable to open %s, virtualized data is inaccessible!"), *FilePath);
+
+	int64 DataLength = 0;
+	*FileAr << DataLength;
+
+	uint8* DataPtr = (uint8*)FMemory::Malloc(DataLength, DEFAULT_ALIGNMENT);
+	FileAr->Serialize(DataPtr, DataLength);
+
+	return MakeSharedBuffer(FSharedBuffer::AssumeOwnership, DataPtr, DataLength);
+}
+
+void FVirtualizationManager::LoadSettingsFromConfigFiles()
+{
+	FConfigFile PlatformEngineIni;
+	if (FConfigCacheIni::LoadLocalIniFile(PlatformEngineIni, TEXT("Engine"), true))
+	{
+		bool bEnablePayloadPushingFromIni = false;
+		if (PlatformEngineIni.GetBool(TEXT("Core.ContentVirtualization"), TEXT("EnablePushToBackend"), bEnablePayloadPushingFromIni))
+		{
+			bEnablePayloadPushing = bEnablePayloadPushingFromIni;
+		}
+		else
+		{
+			UE_LOG(LogVirtualization, Error, TEXT("Failed to load [Core.ContentVirtualization].EnablePushToBackend from ini file!"));
+		}
+
+		int64 MinPayloadLengthFromIni = 0;
+		if (PlatformEngineIni.GetInt64(TEXT("Core.ContentVirtualization"), TEXT("MinPayloadLength"), MinPayloadLengthFromIni))
+		{
+			MinPayloadLength = MinPayloadLengthFromIni;
+		}
+		else
+		{
+			UE_LOG(LogVirtualization, Error, TEXT("Failed to load [Core.ContentVirtualization].MinPayloadLength from ini file!"));;
+		}
+	}
+	else
+	{
+		UE_LOG(LogVirtualization, Error, TEXT("Failed to load 'Core.ContentVirtualization' ini file settings!"));
+	}
+}
+
+FString FVirtualizationManager::CreateFilePath(const FGuid& Guid)
+{
+	return FString::Printf(TEXT("%sVirtualizedPayloads/%s.payload"), *FPaths::ProjectSavedDir(), *Guid.ToString());
+}
diff --git a/Engine/Source/Runtime/CoreUObject/Private/Virtualization/VirtualizedBulkData.cpp b/Engine/Source/Runtime/CoreUObject/Private/Virtualization/VirtualizedBulkData.cpp
new file mode 100644
--- /dev/null
+++ b/Engine/Source/Runtime/CoreUObject/Private/Virtualization/VirtualizedBulkData.cpp
@@ -0,0 +1,414 @@
+// Copyright Epic Games, Inc. All Rights Reserved.
+
+#include "Virtualization/VirtualizedBulkData.h"
+
+#include "HAL/FileManager.h"
+#include "Misc/SecureHash.h"
+#include "Misc/ScopeLock.h"
+#include "Serialization/BulkData.h"
+#include "UObject/LinkerLoad.h"
+#include "UObject/LinkerSave.h"
+#include "UObject/Object.h"
+
+//#if WITH_EDITORONLY_DATA
+
+FVirtualizedUntypedBulkData::~FVirtualizedUntypedBulkData()
+{
+	Reset();
+}
+
+void FVirtualizedUntypedBulkData::CreateFromBulkData(FUntypedBulkData& InBulkData, const FGuid& InGuid)
+{
+	TRACE_CPUPROFILER_EVENT_SCOPE(FVirtualizedUntypedBulkData::CreateFromBulkData);
+
+	Reset();
+
+	Filename = InBulkData.GetFilename();
+
+	OffsetInFile = InBulkData.GetBulkDataOffsetInFile();
+	PayloadLength = InBulkData.GetBulkDataSize();
+	bIsDataStoredAsCompressed = InBulkData.IsStoredCompressedOnDisk();
+
+	Key = InGuid;
+}
+
+void FVirtualizedUntypedBulkData::Serialize(FArchive& Ar, UObject* Owner)
+{
+	TRACE_CPUPROFILER_EVENT_SCOPE(FVirtualizedUntypedBulkData::Serialize);
+
+	if (Ar.IsTransacting())
+	{
+		// Do not process the transaction if the owner is mid loading (see FUntypedBulkData::Serialize)
+		bool bNeedsTransaction = Ar.IsSaving() && (!Owner || !Owner->HasAnyFlags(RF_NeedLoad));
+
+		Ar << bNeedsTransaction;
+
+		if (bNeedsTransaction)
+		{
+			Ar << Key;
+			Ar << PayloadLength;
+			Ar << bIsVirtualized;		
+			Ar << bShouldCompressData;
+			Ar << bIsDataStoredAsCompressed;
+			Ar << Filename;
+			Ar << OffsetInFile;
+
+			bool bKeepBuffer = !Ar.IsLoading();
+			if (!bIsVirtualized && Filename.IsEmpty())
+			{
+				// Data is held in memory so we need to serialize it to the archive
+				bKeepBuffer = SerializeData(Ar, Payload);
+			}
+			// TODO:	Do we want to force data to be pulled/loaded and serialized to the archive?
+			//			Assuming the virtualized back end does not discard the payload or the file with 
+			//			the payload remains intact then it should not be needed?
+			
+			if (!bKeepBuffer)
+			{
+				Payload.Reset();
+			}
+		}
+	}
+	else if (Ar.IsPersistent() && !Ar.IsObjectReferenceCollector() && !Ar.ShouldSkipBulkData())
+	{
+		Ar << Key;
+		Ar << PayloadLength;
+
+		if (Ar.IsSaving())
+		{
+			FLinkerSave* LinkerSave = Cast<FLinkerSave>(Ar.GetLinker());
+
+			const bool bCanAttemptVirtualization = LinkerSave != nullptr;			
+
+			if (bCanAttemptVirtualization)
+			{
+				PushData();
+			}
+
+			Ar << bIsVirtualized;
+					
+			if (!bIsVirtualized)
+			{
+				// Need to load the payload so that we can write it out
+				FSharedBufferConstPtr PayloadToSerialize = GetDataInternal();
+
+				// Need to update the flag before it is written out
+				bIsDataStoredAsCompressed = !bIsVirtualized && bShouldCompressData;
+
+				Ar << bIsDataStoredAsCompressed;
+
+				// Write out a dummy value that we will write over once the payload is serialized
+				int64 PlaceholderValue = INDEX_NONE;
+
+				const int64 OffsetPos = Ar.Tell();
+				Ar << PlaceholderValue; // OffsetInFile
+
+				// The lambda is mutable so that PayloadToSerialize is not const (due to FArchive api not
+				// accepting const values)
+				auto SerializePayload = [this, OffsetPos, PayloadToSerialize](FArchive& Ar) mutable
+				{
+					int64 PayloadOffset = Ar.Tell();
+
+					SerializeData(Ar, PayloadToSerialize);
+
+					// Record the current archive offset (probably EOF but we cannot be sure)
+					const int64 ReturnPos = Ar.Tell();
+
+					// Update the offset/size entries that we set up during ::Serialize
+					Ar.Seek(OffsetPos);
+					Ar << PayloadOffset;
+
+					// Restore the archive's offset
+					Ar.Seek(ReturnPos);
+				};
+
+				// If we have a valid linker then we will defer serialization of the payload so that it will
+				// be placed at the end of the output file so we don't have to seek past the payload on load.
+				// If we do not have a linker then we might as well just serialize right away.
+				if(LinkerSave != nullptr)
+				{
+					 LinkerSave->AdditionalDataToAppend.Add(MoveTemp(SerializePayload));
+				}
+				else
+				{
+					SerializePayload(Ar);
+				}
+			}
+		}
+		else if (Ar.IsLoading())
+		{
+			Ar << bIsVirtualized;
+			
+			if (bIsVirtualized)
+			{
+				// We aren't going to use these members so reset them
+				bIsDataStoredAsCompressed = false;
+				OffsetInFile = INDEX_NONE;
+
+				Filename.Empty();
+			}
+			else
+			{
+				Ar << bIsDataStoredAsCompressed;
+
+				// If we can lazy load then find the filename, otherwise we will want
+				// to serialize immediately.
+				if (Ar.IsAllowingLazyLoading())
+				{
+					Filename = GetFilenameFromOwner(Owner);
+				}
+				else
+				{
+					Filename.Empty();
+				}
+				
+				OffsetInFile = INDEX_NONE;
+				Ar << OffsetInFile;
+
+				if (Filename.IsEmpty())
+				{
+					// If we have no filename then we need to load the data immediately
+					// as we will not be able to load it on demand.
+					SerializeData(Ar, Payload);
+				}		
+			}	
+		}
+	}
+}
+
+void FVirtualizedUntypedBulkData::CalculateKey()
+{
+	if (Payload.IsValid())
+	{
+		uint32 Hash[5] = {};
+		FSHA1::HashBuffer(Payload->GetData(), GetBulkDataSize(), (uint8*)Hash);
+
+		Key = FGuid(Hash[0] ^ Hash[4], Hash[1], Hash[2], Hash[3]);
+	}
+	else
+	{
+		Key.Invalidate();
+	}
+}
+
+FSharedBufferConstPtr FVirtualizedUntypedBulkData::LoadFromDisk() const
+{
+	TRACE_CPUPROFILER_EVENT_SCOPE(FVirtualizedUntypedBulkData::LoadFromDisk);
+
+	FSharedBufferConstPtr PayloadFromDisk;
+
+	// Open a reader to the file
+	TUniquePtr<FArchive> Ar(IFileManager::Get().CreateFileReader(*Filename, FILEREAD_Silent));
+	if (Ar.IsValid())
+	{
+		// Move the correct location of the data in the file
+		Ar->Seek(OffsetInFile);
+
+		// Now we can actually serialize it
+		SerializeData(*Ar, PayloadFromDisk);
+	}
+	else
+	{
+		UE_LOG(LogVirtualization, Error, TEXT("Could not open (%s) to read FVirtualizedUntypedBulkData"), *Filename);
+	}	
+
+	return PayloadFromDisk;
+}
+
+bool FVirtualizedUntypedBulkData::SerializeData(FArchive& Ar, FSharedBufferConstPtr& InPayload) const
+{
+	TRACE_CPUPROFILER_EVENT_SCOPE(FVirtualizedUntypedBulkData::SerializeData);
+
+	const int64 PayloadSize = Ar.IsSaving() ? InPayload->Size() : GetBulkDataSize();
+
+	if (PayloadSize > 0)
+	{
+		// Our serialization methods require a non-const pointer even when we are only reading from it
+		// as the methods themselves are bi-directional. This means if we are saving we need to cast 
+		// away the const qualifier.
+		void* DataPtr = Ar.IsLoading() ?
+						FMemory::Malloc(PayloadSize, DEFAULT_ALIGNMENT) :
+						const_cast<void*>(InPayload->GetData());
+
+		checkf(DataPtr != nullptr, TEXT("Attempting to serialize to/from an invalid shared buffer!"));
+
+		if (bIsDataStoredAsCompressed)
+		{
+			Ar.SerializeCompressed(DataPtr, PayloadSize, NAME_Zlib, COMPRESS_NoFlags, false);
+		}
+		else
+		{
+			Ar.Serialize(DataPtr, PayloadSize);
+		}
+
+		if (Ar.IsLoading())
+		{
+			InPayload = MakeSharedBuffer(FSharedBuffer::AssumeOwnership, DataPtr, PayloadSize);
+		}
+
+		return true;
+	}
+	else
+	{
+		// If we are loading and we failed to serialize anything then we should reset the SharedBuffer.
+		if (Ar.IsLoading())
+		{
+			InPayload.Reset();
+		}
+
+		return false;
+	}
+}
+
+void FVirtualizedUntypedBulkData::PushData()
+{
+	checkf(bIsVirtualized == false || Payload.IsValid() == false, TEXT("Cannot have a valid payload in memory if the payload is virtualized!")); // Sanity check
+
+	// We only need to push if the payload is not currently virtualized (either we have an updated
+	// payload in memory or the payload is currently non-virtualized and stored on disk)
+	if (!bIsVirtualized)
+	{ 
+		TRACE_CPUPROFILER_EVENT_SCOPE(FVirtualizedUntypedBulkData::PushData);
+
+		// We should only need to load from disk at this point if we are going from
+		// a non-virtualized payload to a virtualized one. If the bulkdata is merely being
+		// edited then we should have the payload in memory already and are just accessing a
+		// reference to it.
+		FSharedBufferConstPtr PayloadToPush = GetDataInternal();
+
+		if (FVirtualizationManager::Get().PushData(PayloadToPush, Key))
+		{
+			bIsVirtualized = true;
+		}
+	}	
+}
+
+FSharedBufferConstPtr FVirtualizedUntypedBulkData::PullData() const
+{
+	TRACE_CPUPROFILER_EVENT_SCOPE(FVirtualizedUntypedBulkData::PullData);
+
+	FSharedBufferConstPtr PulledPayload = FVirtualizationManager::Get().PullData(Key);
+
+	if (PulledPayload.IsValid())
+	{
+		checkf(	PayloadLength == PulledPayload->Size(),	
+				TEXT("Mismatch between serialized length (%d) and virtualized data length (%d)"), 
+				PayloadLength,
+				PulledPayload->Size());
+	}
+	else
+	{
+		UE_LOG(LogVirtualization, Warning, TEXT("Failed to pull virtual data with guid (%s)"), *Key.ToString());
+	}
+
+	return PulledPayload;
+}
+
+FString FVirtualizedUntypedBulkData::GetFilenameFromOwner(UObject* Owner) const
+{
+	if (Owner != nullptr)
+	{
+		UPackage* Package = Owner->GetOutermost();
+		checkf(Package != nullptr, TEXT("Owner was not a valid UPackage!"));
+
+		FLinker* Linker = FLinkerLoad::FindExistingLinkerForPackage(Package);
+		checkf(Linker != nullptr, TEXT("UPackage did not have a valid FLinkerLoad!"));
+
+		return Linker->Filename;
+	}
+
+	else
+	{
+		return FString();
+	}
+}
+
+bool FVirtualizedUntypedBulkData::CanUnloadData() const
+{
+	// We cannot unload the data if are unable to reload it from a file
+	return bIsVirtualized || Filename.IsEmpty() == false;
+}
+
+void FVirtualizedUntypedBulkData::Reset()
+{
+	Filename.Empty();
+	OffsetInFile = INDEX_NONE;
+	Key.Invalidate();
+	Payload.Reset();
+	PayloadLength = 0;
+	bIsVirtualized = false;
+	bIsDataStoredAsCompressed = false;
+}
+
+void FVirtualizedUntypedBulkData::UnloadData()
+{
+	Payload.Reset();
+}
+
+FSharedBufferConstPtr FVirtualizedUntypedBulkData::GetDataInternal() const
+{
+	TRACE_CPUPROFILER_EVENT_SCOPE(FVirtualizedUntypedBulkData::GetDataInternal);
+
+	// Early out there isn't any data to actually load
+	if (GetBulkDataSize() == 0)
+	{
+		return FSharedBufferConstPtr();
+	}
+
+	// Check if we already have the data in memory
+	if (Payload.IsValid())
+	{
+		return Payload;
+	}
+
+	if (bIsVirtualized)
+	{
+		FSharedBufferConstPtr Buffer = PullData();
+		check(!Payload.IsValid()); //Make sure that we did not assign the buffer internally
+		return Buffer;
+	}
+	else
+	{
+		FSharedBufferConstPtr Buffer = LoadFromDisk();
+		check(!Payload.IsValid()); //Make sure that we did not assign the buffer internally
+		return Buffer;
+	}
+}
+
+TFuture<FSharedBufferConstPtr> FVirtualizedUntypedBulkData::GetData() const
+{
+	TPromise<FSharedBufferConstPtr> Promise;
+
+	FSharedBufferConstPtr SharedBuffer = GetDataInternal();
+
+	// TODO: Not actually async yet!
+	Promise.SetValue(SharedBuffer);
+
+	return Promise.GetFuture();
+}
+
+void FVirtualizedUntypedBulkData::GetData(OnDataReadyCallback&& Callback) const
+{
+	Callback(GetDataInternal());
+}
+
+void FVirtualizedUntypedBulkData::UpdatePayload(const FSharedBufferConstRef& InPayload)
+{
+	TRACE_CPUPROFILER_EVENT_SCOPE(FVirtualizedUntypedBulkData::UpdatePayload);
+
+	UnloadData();
+
+	// Make sure that we own the memory in the shared buffer
+	Payload = MakeSharedBufferOwned(InPayload);
+	PayloadLength = (int64)InPayload->Size();
+
+	bIsVirtualized = false;
+	bIsDataStoredAsCompressed = false;
+
+	Filename.Empty();
+	OffsetInFile = INDEX_NONE;
+
+	CalculateKey();
+}
+
+//#endif //WITH_EDITORONLY_DATA
diff --git a/Engine/Source/Runtime/CoreUObject/Public/Misc/PackageAccessTracking.h b/Engine/Source/Runtime/CoreUObject/Public/Misc/PackageAccessTracking.h
new file mode 100644
--- /dev/null
+++ b/Engine/Source/Runtime/CoreUObject/Public/Misc/PackageAccessTracking.h
@@ -0,0 +1,56 @@
+// Copyright Epic Games, Inc. All Rights Reserved.
+
+#pragma once
+
+#include "CoreTypes.h"
+
+#define UE_WITH_PACKAGE_ACCESS_TRACKING 0 // TODO: Should be UE_WITH_OBJECT_HANDLE_TRACKING
+
+#if UE_WITH_PACKAGE_ACCESS_TRACKING
+
+#include "Misc/Optional.h"
+#include "Misc/PackageName.h"
+#include "UObject/NameTypes.h"
+#include "UObject/ObjectHandle.h"
+#include "UObject/Package.h"
+
+#define UE_TRACK_REFERENCING_PACKAGE_SCOPED(Package, OpName) PackageAccessTracking_Private::FPackageAccessRefScope ANONYMOUS_VARIABLE(PackageAccessTracker_)(Package, OpName);
+#define UE_TRACK_REFERENCING_PACKAGE_DELAYED_SCOPED(TrackerName, OpName) TOptional<PackageAccessTracking_Private::FPackageAccessRefScope> TrackerName; FName TrackerName##_OpName(OpName);
+#define UE_TRACK_REFERENCING_PACKAGE_DELAYED(TrackerName, Package) if (TrackerName) TrackerName->SetPackageName(Package->GetFName()); else TrackerName.Emplace(Package->GetFName(), TrackerName##_OpName);
+#else
+#define UE_TRACK_REFERENCING_PACKAGE_SCOPED(Package, OpName)
+#define UE_TRACK_REFERENCING_PACKAGE_DELAYED_SCOPED(TrackerName, OpName)
+#define UE_TRACK_REFERENCING_PACKAGE_DELAYED(TrackerName, Package)
+#endif //UE_WITH_PACKAGE_ACCESS_TRACKING
+
+///////////////////////////////////////////////////////////////////////////////////////////////////////////////////////
+
+#if UE_WITH_PACKAGE_ACCESS_TRACKING
+namespace PackageAccessTracking_Private
+{
+	class FPackageAccessRefScope
+	{
+	public:
+		COREUOBJECT_API FPackageAccessRefScope(FName InPackageName, FName InOpName);
+		COREUOBJECT_API FPackageAccessRefScope(const UPackage* InPackage, FName InOpName);
+
+		COREUOBJECT_API ~FPackageAccessRefScope();
+
+		FORCEINLINE FName GetPackageName() const { return PackageName; }
+		FORCEINLINE void SetPackageName(FName InPackageName)
+		{
+			checkf(FPackageName::IsValidLongPackageName(InPackageName.ToString(), true), TEXT("Invalid package name: %s"), *InPackageName.ToString());
+			PackageName = InPackageName;
+		}
+		FORCEINLINE FName GetOpName() const { return OpName; }
+		FORCEINLINE FPackageAccessRefScope* GetOuter() const { return Outer; }
+
+		static COREUOBJECT_API FPackageAccessRefScope* GetCurrentThreadScope();
+	private:
+		FName PackageName;
+		FName OpName;
+		FPackageAccessRefScope* Outer = nullptr;
+		static thread_local FPackageAccessRefScope* CurrentThreadScope;
+	};
+};
+#endif // UE_WITH_PACKAGE_ACCESS_TRACKING
