Copying //UE5/Dev-Cooker@14539516 to Main (//UE5/Main)

diff --git a/Engine/Source/Runtime/CoreUObject/Public/Serialization/ArchiveUObject.h b/Engine/Source/Runtime/CoreUObject/Public/Serialization/ArchiveUObject.h
--- a/Engine/Source/Runtime/CoreUObject/Public/Serialization/ArchiveUObject.h
+++ b/Engine/Source/Runtime/CoreUObject/Public/Serialization/ArchiveUObject.h
@@ -5,6 +5,7 @@
 #include "CoreMinimal.h"
 
 struct FLazyObjectPtr;
+struct FObjectPtr;
 struct FSoftObjectPtr;
 struct FSoftObjectPath;
 struct FWeakObjectPtr;
@@ -17,6 +18,7 @@ class COREUOBJECT_API FArchiveUObject : public FArchive
 public:
 
 	static FArchive& SerializeLazyObjectPtr(FArchive& Ar, FLazyObjectPtr& Value);
+	static FArchive& SerializeObjectPtr(FArchive& Ar, FObjectPtr& Value);
 	static FArchive& SerializeSoftObjectPtr(FArchive& Ar, FSoftObjectPtr& Value);
 	static FArchive& SerializeSoftObjectPath(FArchive& Ar, FSoftObjectPath& Value);
 	static FArchive& SerializeWeakObjectPtr(FArchive& Ar, FWeakObjectPtr& Value);
@@ -25,6 +27,7 @@ public:
 
 	//~ Begin FArchive Interface
 	virtual FArchive& operator<<(FLazyObjectPtr& Value) override { return SerializeLazyObjectPtr(*this, Value); }
+	virtual FArchive& operator<<(FObjectPtr& Value) override { return SerializeObjectPtr(*this, Value); }
 	virtual FArchive& operator<<(FSoftObjectPtr& Value) override { return SerializeSoftObjectPtr(*this, Value); }
 	virtual FArchive& operator<<(FSoftObjectPath& Value) override { return SerializeSoftObjectPath(*this, Value); }
 	virtual FArchive& operator<<(FWeakObjectPtr& Value) override { return SerializeWeakObjectPtr(*this, Value); }
diff --git a/Engine/Source/Runtime/CoreUObject/Public/Serialization/Formatters/JsonArchiveInputFormatter.h b/Engine/Source/Runtime/CoreUObject/Public/Serialization/Formatters/JsonArchiveInputFormatter.h
--- a/Engine/Source/Runtime/CoreUObject/Public/Serialization/Formatters/JsonArchiveInputFormatter.h
+++ b/Engine/Source/Runtime/CoreUObject/Public/Serialization/Formatters/JsonArchiveInputFormatter.h
@@ -75,6 +75,7 @@ public:
 	virtual void Serialize(FSoftObjectPtr& Value) override;
 	virtual void Serialize(FSoftObjectPath& Value) override;
 	virtual void Serialize(FLazyObjectPtr& Value) override;
+	virtual void Serialize(FObjectPtr& Value) override;
 	virtual void Serialize(TArray<uint8>& Value) override;
 	virtual void Serialize(void* Data, uint64 DataSize) override;
 
diff --git a/Engine/Source/Runtime/CoreUObject/Public/Serialization/Formatters/JsonArchiveOutputFormatter.h b/Engine/Source/Runtime/CoreUObject/Public/Serialization/Formatters/JsonArchiveOutputFormatter.h
--- a/Engine/Source/Runtime/CoreUObject/Public/Serialization/Formatters/JsonArchiveOutputFormatter.h
+++ b/Engine/Source/Runtime/CoreUObject/Public/Serialization/Formatters/JsonArchiveOutputFormatter.h
@@ -73,6 +73,7 @@ public:
 	virtual void Serialize(FSoftObjectPtr& Value) override;
 	virtual void Serialize(FSoftObjectPath& Value) override;
 	virtual void Serialize(FLazyObjectPtr& Value) override;
+	virtual void Serialize(FObjectPtr& Value) override;
 	virtual void Serialize(TArray<uint8>& Value) override;
 	virtual void Serialize(void* Data, uint64 DataSize) override;
 
diff --git a/Engine/Source/Runtime/CoreUObject/Public/Serialization/ObjectReader.h b/Engine/Source/Runtime/CoreUObject/Public/Serialization/ObjectReader.h
--- a/Engine/Source/Runtime/CoreUObject/Public/Serialization/ObjectReader.h
+++ b/Engine/Source/Runtime/CoreUObject/Public/Serialization/ObjectReader.h
@@ -63,6 +63,7 @@ public:
 
 	COREUOBJECT_API virtual FArchive& operator<<(FName& N) override;
 	COREUOBJECT_API virtual FArchive& operator<<(UObject*& Res) override;
+	COREUOBJECT_API virtual FArchive& operator<<(FObjectPtr& Value) override;
 	COREUOBJECT_API virtual FArchive& operator<<(FLazyObjectPtr& LazyObjectPtr) override;
 	COREUOBJECT_API virtual FArchive& operator<<(FSoftObjectPtr& Value) override;
 	COREUOBJECT_API virtual FArchive& operator<<(FSoftObjectPath& Value) override;
diff --git a/Engine/Source/Runtime/CoreUObject/Public/Serialization/ObjectWriter.h b/Engine/Source/Runtime/CoreUObject/Public/Serialization/ObjectWriter.h
--- a/Engine/Source/Runtime/CoreUObject/Public/Serialization/ObjectWriter.h
+++ b/Engine/Source/Runtime/CoreUObject/Public/Serialization/ObjectWriter.h
@@ -41,6 +41,7 @@ public:
 	//~ Begin FArchive Interface
 	COREUOBJECT_API virtual FArchive& operator<<(FName& N) override;
 	COREUOBJECT_API virtual FArchive& operator<<(UObject*& Res) override;
+	COREUOBJECT_API virtual FArchive& operator<<(FObjectPtr& Value) override;
 	COREUOBJECT_API virtual FArchive& operator<<(FLazyObjectPtr& Value) override;
 	COREUOBJECT_API virtual FArchive& operator<<(FSoftObjectPtr& Value) override;
 	COREUOBJECT_API virtual FArchive& operator<<(FSoftObjectPath& Value) override;
diff --git a/Engine/Source/Runtime/CoreUObject/Public/Templates/Casts.h b/Engine/Source/Runtime/CoreUObject/Public/Templates/Casts.h
--- a/Engine/Source/Runtime/CoreUObject/Public/Templates/Casts.h
+++ b/Engine/Source/Runtime/CoreUObject/Public/Templates/Casts.h
@@ -310,6 +310,16 @@ template< class T, class U > FORCEINLINE T* Cast       ( const TWeakObjectPtr<U>
 template< class T, class U > FORCEINLINE T* ExactCast  ( const TWeakObjectPtr<U>& Src                                                                   ) { return ExactCast  <T>(Src.Get()); }
 template< class T, class U > FORCEINLINE T* CastChecked( const TWeakObjectPtr<U>& Src, ECastCheckedType::Type CheckType = ECastCheckedType::NullChecked ) { return CastChecked<T>(Src.Get(), CheckType); }
 
+// object ptr versions
+template <class T, class U> FORCEINLINE TObjectPtr<T> Cast       ( const TObjectPtr<U>& Src                                                                   ) { return TObjectPtr<T>(Cast       <T>(Src.Get())); }
+template <class T, class U> FORCEINLINE TObjectPtr<T> ExactCast  ( const TObjectPtr<U>& Src                                                                   ) { return TObjectPtr<T>(ExactCast  <T>(Src.Get())); }
+template <class T, class U> FORCEINLINE TObjectPtr<T> CastChecked( const TObjectPtr<U>& Src, ECastCheckedType::Type CheckType = ECastCheckedType::NullChecked ) { return TObjectPtr<T>(CastChecked<T>(Src.Get(), CheckType)); }
+
+// const object ptr versions
+template <class T, class U> FORCEINLINE TObjectPtr<const T> Cast       ( const TObjectPtr<const U>& Src                                                                   ) { return TObjectPtr<const T>(Cast       <T>(Src.Get())); }
+template <class T, class U> FORCEINLINE TObjectPtr<const T> ExactCast  ( const TObjectPtr<const U>& Src                                                                   ) { return TObjectPtr<const T>(ExactCast  <T>(Src.Get())); }
+template <class T, class U> FORCEINLINE TObjectPtr<const T> CastChecked( const TObjectPtr<const U>& Src, ECastCheckedType::Type CheckType = ECastCheckedType::NullChecked ) { return TObjectPtr<const T>(CastChecked<T>(Src.Get(), CheckType)); }
+
 // TSubclassOf versions
 template< class T, class U > FORCEINLINE T* Cast       ( const TSubclassOf<U>& Src                                                                   ) { return Cast       <T>(*Src); }
 template< class T, class U > FORCEINLINE T* CastChecked( const TSubclassOf<U>& Src, ECastCheckedType::Type CheckType = ECastCheckedType::NullChecked ) { return CastChecked<T>(*Src, CheckType); }
@@ -350,6 +360,7 @@ DECLARE_CAST_BY_FLAG(UClass)							\
 DECLARE_CAST_BY_FLAG(FProperty)							\
 DECLARE_CAST_BY_FLAG(FObjectPropertyBase)				\
 DECLARE_CAST_BY_FLAG(FObjectProperty)					\
+DECLARE_CAST_BY_FLAG(FObjectPtrProperty)				\
 DECLARE_CAST_BY_FLAG(FWeakObjectProperty)				\
 DECLARE_CAST_BY_FLAG(FLazyObjectProperty)				\
 DECLARE_CAST_BY_FLAG(FSoftObjectProperty)				\
diff --git a/Engine/Source/Runtime/CoreUObject/Public/Templates/SubclassOf.h b/Engine/Source/Runtime/CoreUObject/Public/Templates/SubclassOf.h
--- a/Engine/Source/Runtime/CoreUObject/Public/Templates/SubclassOf.h
+++ b/Engine/Source/Runtime/CoreUObject/Public/Templates/SubclassOf.h
@@ -7,6 +7,14 @@
 #include "UObject/Field.h"
 #include "Templates/ChooseClass.h"
 
+#include <type_traits>
+
+template <typename T>
+struct TIsTSubclassOf
+{
+	enum { Value = false };
+};
+
 /**
  * Template to allow TClassType's to be passed around with type safety 
  */
@@ -36,6 +44,19 @@ public:
 	{
 	}
 
+	/** Constructor that takes a UClass and does a runtime check to make sure this is a compatible class */
+	template <
+		typename U,
+		std::enable_if_t<
+			!TIsTSubclassOf<std::decay_t<U>>::Value,
+			decltype(ImplicitConv<TClassType*>(std::declval<U>()))
+		>* = nullptr
+	>
+	FORCEINLINE TSubclassOf(U&& From)
+		: Class(From)
+	{
+	}
+
 	/** Copy Constructor, will only compile if types are compatible */
 	template <class TClassA, class = decltype(ImplicitConv<TClass*>((TClassA*)nullptr))>
 	FORCEINLINE TSubclassOf(const TSubclassOf<TClassA>& From) :
@@ -58,6 +79,20 @@ public:
 		return *this;
 	}
 	
+	/** Assignment operator from UClass, the type is checked on get not on set */
+	template <
+		typename U,
+		std::enable_if_t<
+			!TIsTSubclassOf<std::decay_t<U>>::Value,
+			decltype(ImplicitConv<TClassType*>(std::declval<U>()))
+		>* = nullptr
+	>
+	FORCEINLINE TSubclassOf& operator=(U&& From)
+	{
+		Class = From;
+		return *this;
+	}
+	
 	/** Dereference back into a UClass, does runtime type checking */
 	FORCEINLINE TClassType* operator*() const
 	{
@@ -125,3 +160,9 @@ public:
 private:
 	TClassType* Class;
 };
+
+template <typename T>
+struct TIsTSubclassOf<TSubclassOf<T>>
+{
+	enum { Value = true };
+};
diff --git a/Engine/Source/Runtime/CoreUObject/Public/UObject/Class.h b/Engine/Source/Runtime/CoreUObject/Public/UObject/Class.h
--- a/Engine/Source/Runtime/CoreUObject/Public/UObject/Class.h
+++ b/Engine/Source/Runtime/CoreUObject/Public/UObject/Class.h
@@ -14,6 +14,8 @@
 #include "Misc/FallbackStruct.h"
 #include "Misc/Guid.h"
 #include "Misc/Optional.h"
+#include "Misc/PackageAccessTracking.h"
+#include "Misc/PackageAccessTrackingOps.h"
 #include "Misc/ScopeRWLock.h"
 #include "Templates/IsAbstract.h"
 #include "Templates/IsEnum.h"
@@ -2865,6 +2867,7 @@ public:
 	{
 		if (ClassDefaultObject == nullptr && bCreateIfNeeded)
 		{
+			UE_TRACK_REFERENCING_PACKAGE_SCOPED((GetOutermost()), PackageAccessTrackingOps::NAME_CreateDefaultObject);
 			const_cast<UClass*>(this)->CreateDefaultObject();
 		}
 
diff --git a/Engine/Source/Runtime/CoreUObject/Public/UObject/FastReferenceCollector.h b/Engine/Source/Runtime/CoreUObject/Public/UObject/FastReferenceCollector.h
--- a/Engine/Source/Runtime/CoreUObject/Public/UObject/FastReferenceCollector.h
+++ b/Engine/Source/Runtime/CoreUObject/Public/UObject/FastReferenceCollector.h
@@ -563,6 +563,13 @@ public:
 	}
 
 private:
+	FORCEINLINE void ConditionalHandleTokenStreamObjectReference(TArray<UObject*>& ObjectsToSerialize, UObject* ReferencingObject, UObject*& Object, const int32 TokenIndex, bool bAllowReferenceElimination)
+	{
+		if (IsObjectHandleResolved(*reinterpret_cast<FObjectHandle*>(&Object)))
+		{
+			ReferenceProcessor.HandleTokenStreamObjectReference(ObjectsToSerialize, ReferencingObject, Object, TokenIndex, bAllowReferenceElimination);
+		}
+	}
 
 	FORCEINLINE bool MoveToNextContainerElementAndCheckIfValid(FStackEntry* StackEntry) const
 	{
@@ -748,7 +755,7 @@ private:
 						UObject**	ObjectPtr = (UObject**)(StackEntryData + ReferenceInfo.Offset);
 						UObject*&	Object = *ObjectPtr;
 						TokenReturnCount = ReferenceInfo.ReturnCount;
-						ReferenceProcessor.HandleTokenStreamObjectReference(NewObjectsToSerialize, CurrentObject, Object, ReferenceTokenStreamIndex, true);
+						ConditionalHandleTokenStreamObjectReference(NewObjectsToSerialize, CurrentObject, Object, ReferenceTokenStreamIndex, true);
 					}
 					break;
 					case GCRT_ArrayObject:
@@ -758,7 +765,7 @@ private:
 						TokenReturnCount = ReferenceInfo.ReturnCount;
 						for (int32 ObjectIndex = 0, ObjectNum = ObjectArray.Num(); ObjectIndex < ObjectNum; ++ObjectIndex)
 						{
-							ReferenceProcessor.HandleTokenStreamObjectReference(NewObjectsToSerialize, CurrentObject, ObjectArray[ObjectIndex], ReferenceTokenStreamIndex, true);
+							ConditionalHandleTokenStreamObjectReference(NewObjectsToSerialize, CurrentObject, ObjectArray[ObjectIndex], ReferenceTokenStreamIndex, true);
 						}
 					}
 					break;
@@ -769,7 +776,7 @@ private:
 						TokenReturnCount = ReferenceInfo.ReturnCount;
 						for (int32 ObjectIndex = 0, ObjectNum = ObjectArray.Num(); ObjectIndex < ObjectNum; ++ObjectIndex)
 						{
-							ReferenceProcessor.HandleTokenStreamObjectReference(NewObjectsToSerialize, CurrentObject, ObjectArray[ObjectIndex], ReferenceTokenStreamIndex, true);
+							ConditionalHandleTokenStreamObjectReference(NewObjectsToSerialize, CurrentObject, ObjectArray[ObjectIndex], ReferenceTokenStreamIndex, true);
 						}
 					}
 					break;
@@ -835,7 +842,7 @@ private:
 						UObject**	ObjectPtr = (UObject**)(StackEntryData + ReferenceInfo.Offset);
 						UObject*&	Object = *ObjectPtr;
 						TokenReturnCount = ReferenceInfo.ReturnCount;
-						ReferenceProcessor.HandleTokenStreamObjectReference(NewObjectsToSerialize, CurrentObject, Object, ReferenceTokenStreamIndex, false);
+						ConditionalHandleTokenStreamObjectReference(NewObjectsToSerialize, CurrentObject, Object, ReferenceTokenStreamIndex, false);
 					}
 					break;
 					case GCRT_ExternalPackage:
@@ -978,7 +985,7 @@ private:
 						{
 							UObject* OwnerObject = static_cast<UObject*>(FieldOwnerItem->Object);
 							UObject* PreviousOwner = OwnerObject;
-							ReferenceProcessor.HandleTokenStreamObjectReference(NewObjectsToSerialize, CurrentObject, OwnerObject, ReferenceTokenStreamIndex, true);
+							ConditionalHandleTokenStreamObjectReference(NewObjectsToSerialize, CurrentObject, OwnerObject, ReferenceTokenStreamIndex, true);
 							// Handle reference elimination (PendingKill owner)
 							if (PreviousOwner && !OwnerObject)
 							{
@@ -999,7 +1006,7 @@ private:
 							{
 								UObject* OwnerObject = static_cast<UObject*>(FieldOwnerItem->Object);
 								UObject* PreviousOwner = OwnerObject;
-								ReferenceProcessor.HandleTokenStreamObjectReference(NewObjectsToSerialize, CurrentObject, OwnerObject, ReferenceTokenStreamIndex, true);
+								ConditionalHandleTokenStreamObjectReference(NewObjectsToSerialize, CurrentObject, OwnerObject, ReferenceTokenStreamIndex, true);
 								// Handle reference elimination (PendingKill owner)
 								if (PreviousOwner && !OwnerObject)
 								{
@@ -1022,7 +1029,7 @@ private:
 							// We're dealing with an object reference (this code should be identical to GCRT_PersistentObject)
 							UObject**	ObjectPtr = (UObject**)(StackEntryData + ReferenceInfo.Offset);
 							UObject*&	Object = *ObjectPtr;
-							ReferenceProcessor.HandleTokenStreamObjectReference(NewObjectsToSerialize, CurrentObject, Object, ReferenceTokenStreamIndex, false);
+							ConditionalHandleTokenStreamObjectReference(NewObjectsToSerialize, CurrentObject, Object, ReferenceTokenStreamIndex, false);
 						}
 					}
 					break;
@@ -1034,7 +1041,7 @@ private:
 							// We're dealing with an object reference (this code should be identical to GCRT_Object and GCRT_Class)
 							UObject**	ObjectPtr = (UObject**)(StackEntryData + ReferenceInfo.Offset);
 							UObject*&	Object = *ObjectPtr;
-							ReferenceProcessor.HandleTokenStreamObjectReference(NewObjectsToSerialize, CurrentObject, Object, ReferenceTokenStreamIndex, true);
+							ConditionalHandleTokenStreamObjectReference(NewObjectsToSerialize, CurrentObject, Object, ReferenceTokenStreamIndex, true);
 						}
 					}
 					break;
diff --git a/Engine/Source/Runtime/CoreUObject/Public/UObject/Field.h b/Engine/Source/Runtime/CoreUObject/Public/UObject/Field.h
--- a/Engine/Source/Runtime/CoreUObject/Public/UObject/Field.h
+++ b/Engine/Source/Runtime/CoreUObject/Public/UObject/Field.h
@@ -23,6 +23,8 @@ Field.h: Declares FField property system fundamentals
 #include "Misc/CoreMiscDefines.h"
 #include "HAL/ThreadSafeCounter.h"
 
+#include <type_traits>
+
 class FProperty;
 class FField;
 class FFieldVariant;
@@ -222,10 +224,15 @@ public:
 	{
 		Container.Field = const_cast<FField*>(InField);
 	}
-	FFieldVariant(const UObject* InObject)
+
+	template <
+		typename T,
+		decltype(ImplicitConv<const UObject*>(std::declval<T>()))* = nullptr
+	>
+	FFieldVariant(T&& InObject)
 		: bIsUObject(true)
 	{
-		Container.Object = const_cast<UObject*>(InObject);
+		Container.Object = const_cast<UObject*>(ImplicitConv<const UObject*>(InObject));
 	}
 
 	FFieldVariant(TYPE_OF_NULLPTR)
diff --git a/Engine/Source/Runtime/CoreUObject/Public/UObject/LinkerLoad.h b/Engine/Source/Runtime/CoreUObject/Public/UObject/LinkerLoad.h
--- a/Engine/Source/Runtime/CoreUObject/Public/UObject/LinkerLoad.h
+++ b/Engine/Source/Runtime/CoreUObject/Public/UObject/LinkerLoad.h
@@ -535,6 +535,15 @@ public:
 	 * @return	ObjectName for the FObjectResource at ResourceIndex, or NAME_None if not found
 	 */
 	FName ResolveResourceName( FPackageIndex ResourceIndex );
+
+	/**
+	 * Returns the Object associated with the resource indicated.
+	 *
+	 * @param	ResourceIndex	location of the object resource
+	 *
+	 * @return	The UObject at ResourceIndex, or nullptr if not found
+	 */
+	UObject* ResolveResource(FPackageIndex ResourceIndex);
 	
 	int32 FindExportIndex( FName ClassName, FName ClassPackage, FName ObjectName, FPackageIndex ExportOuterIndex );
 	
@@ -842,6 +851,8 @@ private:
 		return Ar;
 	}
 
+	virtual FArchive& operator<<(FObjectPtr& ObjectPtr) override;
+
 	FORCEINLINE virtual FArchive& operator<<(FSoftObjectPtr& Value) override
 	{
 		FArchive& Ar = *this;
diff --git a/Engine/Source/Runtime/CoreUObject/Public/UObject/LinkerSave.h b/Engine/Source/Runtime/CoreUObject/Public/UObject/LinkerSave.h
--- a/Engine/Source/Runtime/CoreUObject/Public/UObject/LinkerSave.h
+++ b/Engine/Source/Runtime/CoreUObject/Public/UObject/LinkerSave.h
@@ -64,6 +64,15 @@ public:
 	};
 	TArray<FBulkDataStorageInfo> BulkDataToAppend;
 
+	// TODO: Look into removing BulkDataToAppend and use AdditionalDataToAppend instead.
+
+	using AdditionalDataCallback = TUniqueFunction<void(FArchive& Ar)>;
+	/** 
+	 * Array of callbacks that will be invoked when it is possible to serialize out data 
+	 * to the end of the output file.
+	 */
+	TArray<AdditionalDataCallback> AdditionalDataToAppend;
+
 	/** A mapping of package name to generated script SHA keys */
 	COREUOBJECT_API static TMap<FString, TArray<uint8> > PackagesToScriptSHAMap;
 
diff --git a/Engine/Source/Runtime/CoreUObject/Public/UObject/ObjectHandle.h b/Engine/Source/Runtime/CoreUObject/Public/UObject/ObjectHandle.h
new file mode 100644
--- /dev/null
+++ b/Engine/Source/Runtime/CoreUObject/Public/UObject/ObjectHandle.h
@@ -0,0 +1,366 @@
+// Copyright Epic Games, Inc. All Rights Reserved.
+
+#pragma once
+
+#include "Containers/ScriptArray.h"
+#include "HAL/CriticalSection.h"
+#include "HAL/Platform.h"
+#include "UObject/NameTypes.h"
+#include "UObject/Object.h"
+#include "UObject/ObjectPathId.h"
+
+#define UE_WITH_OBJECT_HANDLE_LATE_RESOLVE WITH_EDITORONLY_DATA
+#define UE_WITH_OBJECT_HANDLE_TRACKING WITH_EDITORONLY_DATA
+
+/**
+ * FObjectRef represents a heavyweight reference that contains the specific pieces of information needed to reference an object
+ * (or null) that may or may not be loaded yet.  The expectation is that given the imports of a package we have loaded, we should
+ * be able to create an FObjectRef to objects in a package we haven't yet loaded.  For this reason, FObjectRef has to be derivable
+ * from the serialized (not transient) contents of an FObjectImport.
+ */
+struct FObjectRef
+{
+	FName PackageName;
+	FName ClassPackageName;
+	FName ClassName;
+	FObjectPathId ObjectPath;
+
+	bool operator==(const FObjectRef&& Other) const
+	{
+		return (PackageName == Other.PackageName) && (ObjectPath == Other.ObjectPath) && (ClassPackageName == Other.ClassPackageName) && (ClassName == Other.ClassName);
+	}
+};
+
+inline bool IsObjectRefNull(const FObjectRef& ObjectRef) { return ObjectRef.PackageName.IsNone() && ObjectRef.ObjectPath.IsNone(); }
+
+COREUOBJECT_API FObjectRef MakeObjectRef(const UObject* Object);
+COREUOBJECT_API FObjectRef MakeObjectRef(struct FPackedObjectRef ObjectRef);
+COREUOBJECT_API UObject* ResolveObjectRef(const FObjectRef& ObjectRef, uint32 LoadFlags = LOAD_None);
+COREUOBJECT_API UClass* ResolveObjectRefClass(const FObjectRef& ObjectRef, uint32 LoadFlags = LOAD_None);
+
+
+/**
+ * FPackedObjectRef represents a lightweight reference that can fit in the space of a pointer and be able to refer to an object
+ * (or null) that may or may not be loaded without pointing to its location in memory (even if it is currently loaded).
+ */
+struct FPackedObjectRef
+{
+	// Must be 0 for a reference to null.
+	// The least significant bit must always be 1 in a non-null reference.
+	UPTRINT EncodedRef;
+};
+
+inline bool IsPackedObjectRefNull(FPackedObjectRef ObjectRef) { return !ObjectRef.EncodedRef; }
+
+COREUOBJECT_API FPackedObjectRef MakePackedObjectRef(const UObject* Object);
+COREUOBJECT_API FPackedObjectRef MakePackedObjectRef(const FObjectRef& ObjectRef);
+COREUOBJECT_API UObject* ResolvePackedObjectRef(FPackedObjectRef ObjectRef, uint32 LoadFlags = LOAD_None);
+COREUOBJECT_API UClass* ResolvePackedObjectRefClass(FPackedObjectRef ObjectRef, uint32 LoadFlags = LOAD_None);
+
+inline bool operator==(FPackedObjectRef LHS, FPackedObjectRef RHS) { return LHS.EncodedRef == RHS.EncodedRef; }
+inline bool operator!=(FPackedObjectRef LHS, FPackedObjectRef RHS) { return LHS.EncodedRef != RHS.EncodedRef; }
+inline uint32 GetTypeHash(FPackedObjectRef ObjectRef) { return GetTypeHash(ObjectRef.EncodedRef); }
+
+/**
+ * FObjectHandle is either a packed object ref or the resolved pointer to an object.  Depending on configuration
+ * when you create a handle, it may immediately be resolved to a pointer.
+ */
+#if UE_WITH_OBJECT_HANDLE_LATE_RESOLVE
+
+struct FObjectHandle
+{
+	UPTRINT PointerOrRef;
+};
+inline bool operator==(FObjectHandle LHS, FObjectHandle RHS);
+inline bool operator!=(FObjectHandle LHS, FObjectHandle RHS);
+inline uint32 GetTypeHash(FObjectHandle Handle);
+
+#else
+
+using FObjectHandle = UObject*;
+//NOTE: operator==, operator!=, GetTypeHash fall back to the default on UObject* or void* through coercion.
+
+#endif
+
+inline FObjectHandle MakeObjectHandle(FPackedObjectRef ObjectRef);
+inline FObjectHandle MakeObjectHandle(const FObjectRef& ObjectRef);
+inline FObjectHandle MakeObjectHandle(UObject* Object);
+
+inline bool IsObjectHandleNull(FObjectHandle Handle);
+inline bool IsObjectHandleResolved(FObjectHandle Handle);
+
+/** Read the handle as a pointer without checking if it is resolved. Invalid to call for unresolved handles. */
+inline UObject* ReadObjectHandlePointerNoCheck(FObjectHandle Handle);
+/** Read the handle as a packed object ref without checking if it is unresolved. Invalid to call for resolved handles. */
+inline FPackedObjectRef ReadObjectHandlePackedObjectRefNoCheck(FObjectHandle Handle);
+
+inline UObject* ResolveObjectHandle(FObjectHandle& Handle);
+inline UClass* ResolveObjectHandleClass(FObjectHandle Handle);
+
+/** Read the handle as a pointer if resolved, and otherwise return null. */
+inline UObject* ReadObjectHandlePointer(FObjectHandle Handle);
+/** Read the handle as a packed object ref if unresolved, and otherwise return the null packed object ref. */
+inline FPackedObjectRef ReadObjectHandlePackedObjectRef(FObjectHandle Handle);
+
+#if UE_WITH_OBJECT_HANDLE_LATE_RESOLVE
+inline FObjectRef MakeObjectRef(FObjectHandle Handle);
+inline FPackedObjectRef MakePackedObjectRef(FObjectHandle Handle);
+#endif
+
+
+/**
+ * FObjectHandles can optionally support tracking.  Because of the low level nature of object handles, anything that
+ * registers itself for these callbacks should ensure that it is:
+ * 1) error free (ie: should not cause exceptions even in unusual circumstances)
+ * 2) fault tolerant (ie: could be called at a time when an exception has happened)
+ * 3) thread-safe (ie: could be called from any thread)
+ * 4) high performance (ie: will be called many times)
+ */
+#if UE_WITH_OBJECT_HANDLE_TRACKING
+ /**
+  * Callback notifying when an object value is read from a handle.  Fired regardless of whether the handle
+  * was resolved as part of the read operation or not and whether the object being read is null or not.
+  *
+  * @param ReadObject	The object that was read from a handle.
+  */
+using ObjectHandleReadFunction = void(UObject* ReadObject);
+
+/**
+ * Callback notifying when a class is resolved from an object handle or object reference.
+ * Classes are resolved either independently for a given handle/reference or as part of each object resolve.
+ *
+ * @param SourceRef		The object reference (either user provided or internally computed from a handle) from which the class was resolved.
+ * @param ClassPackage	The package containing the resolved class.
+ * @param Class			The resolved class.
+ */
+using ObjectHandleClassResolvedFunction = void(const FObjectRef& SourceRef, UPackage* ClassPackage, UClass* Class);
+
+/**
+ * Callback notifying when an object is resolved from an object handle or object reference.
+ *
+ * @param SourceRef		The object reference (either user provided or internally computed from a handle) from which the object was resolved.
+ * @param ObjectPackage	The package containing the resolved object.
+ * @param Object		The resolved object.
+ */
+using ObjectHandleReferenceResolvedFunction = void(const FObjectRef& SourceRef, UPackage* ObjectPackage, UObject* Object);
+
+/**
+ * Installs a new callback for notifications that an object value has been read from a handle.
+ *
+ * @param Function		The new handle read callback to install.
+ * @return				The previous handle read callback (or nullptr).  The caller is expected to store this and call in their own handle read callback.
+ */
+COREUOBJECT_API /*UE_NODISCARD*/ ObjectHandleReadFunction* SetObjectHandleReadCallback(ObjectHandleReadFunction* Function);
+
+/**
+ * Installs a new callback for notifications that a class has been resolved from an object handle or object reference.
+ *
+ * @param Function		The new class resolved callback to install.
+ * @return				The previous object resolved callback (or nullptr).  The caller is expected to store this and call in their own class resolved callback.
+ */
+COREUOBJECT_API /*UE_NODISCARD*/ ObjectHandleClassResolvedFunction* SetObjectHandleClassResolvedCallback(ObjectHandleClassResolvedFunction* Function);
+
+/**
+ * Installs a new callback for notifications that an object has been resolved from an object handle or object reference.
+ *
+ * @param Function		The new object resolved callback to install.
+ * @return				The previous object resolved callback (or nullptr).  The caller is expected to store this and call in their own object resolved callback.
+ */
+COREUOBJECT_API /*UE_NODISCARD*/ ObjectHandleReferenceResolvedFunction* SetObjectHandleReferenceResolvedCallback(ObjectHandleReferenceResolvedFunction* Function);
+#endif
+
+
+
+///////////////////////////////////////////////////////////////////////////////////////////////////////////////////////
+
+struct FObjectHandlePackageDebugData
+{
+	FMinimalName PackageName;
+	FScriptArray ObjectPaths;
+	FScriptArray DataClassDescriptors;
+	uint8 _Padding[sizeof(FRWLock) + sizeof(FScriptMap)];
+};
+
+struct FObjectHandleDataClassDescriptor
+{
+	FMinimalName PackageName;
+	FMinimalName ClassName;
+};
+
+namespace ObjectHandle_Private
+{
+	constexpr uint32 ObjectPathIdShift = 1;
+	constexpr uint32 ObjectPathIdMask = 0x00FF'FFFF;
+
+	constexpr uint32 DataClassDescriptorIdShift = 25;
+	constexpr uint32 DataClassDescriptorIdMask = 0x0000'00FF;
+
+	constexpr uint32 PackageIdShift = 33;
+	constexpr uint32 PackageIdMask = 0x7FFF'FFFF;
+
+#if UE_WITH_OBJECT_HANDLE_TRACKING
+	extern COREUOBJECT_API ObjectHandleReadFunction* ObjectHandleReadCallback;
+	extern COREUOBJECT_API ObjectHandleClassResolvedFunction* ObjectHandleClassResolvedCallback;
+	extern COREUOBJECT_API ObjectHandleReferenceResolvedFunction* ObjectHandleReferenceResolvedCallback;
+
+	// @TODO: OBJPTR: Re-assess callsites for these On* functions to see if we have the coverage we want for each event type.
+	inline void OnHandleRead(UObject* Object)
+	{
+		if (ObjectHandleReadFunction* Callback = ObjectHandleReadCallback)
+		{
+			Callback(Object);
+		}
+	}
+	
+	inline void OnClassReferenceResolved(const FObjectRef& ObjectRef, UPackage* Package, UClass* Class)
+	{
+		if (ObjectHandleClassResolvedFunction* Callback = ObjectHandleClassResolvedCallback)
+		{
+			Callback(ObjectRef, Package, Class);
+		}
+	}
+
+	inline void OnReferenceResolved(const FObjectRef& ObjectRef, UPackage* Package, UObject* Object)
+	{
+		if (ObjectHandleReferenceResolvedFunction* Callback = ObjectHandleReferenceResolvedCallback)
+		{
+			Callback(ObjectRef, Package, Object);
+		}
+	}
+#else
+	inline void OnHandleRead(UObject* Object) {}
+	inline void OnClassReferenceResolved(const FObjectRef& ObjectRef, UPackage* Package, UClass* Class) {}
+	inline void OnReferenceResolved(const FObjectRef& ObjectRef, UPackage* Package, UObject* Object) {}
+#endif
+}
+
+#if UE_WITH_OBJECT_HANDLE_LATE_RESOLVE
+
+inline bool IsObjectHandleNull(FObjectHandle Handle) { return !Handle.PointerOrRef; }
+inline bool IsObjectHandleResolved(FObjectHandle Handle) { return !(Handle.PointerOrRef & 1); }
+
+inline UObject* ReadObjectHandlePointerNoCheck(FObjectHandle Handle)
+{
+	return reinterpret_cast<UObject*>(Handle.PointerOrRef);
+}
+
+inline FPackedObjectRef ReadObjectHandlePackedObjectRefNoCheck(FObjectHandle Handle)
+{
+	return {Handle.PointerOrRef};
+}
+
+inline FObjectHandle MakeObjectHandle(FPackedObjectRef ObjectRef) { return {ObjectRef.EncodedRef}; }
+inline FObjectHandle MakeObjectHandle(const FObjectRef& ObjectRef) { return MakeObjectHandle(MakePackedObjectRef(ObjectRef)); }
+inline FObjectHandle MakeObjectHandle(UObject* Object) { return {UPTRINT(Object)}; }
+
+inline bool operator==(FObjectHandle LHS, FObjectHandle RHS)
+{
+	if (IsObjectHandleResolved(LHS) == IsObjectHandleResolved(RHS))
+	{
+		return LHS.PointerOrRef == RHS.PointerOrRef;
+	}
+	else
+	{
+		return MakePackedObjectRef(LHS) == MakePackedObjectRef(RHS);
+	}
+}
+inline bool operator!=(FObjectHandle LHS, FObjectHandle RHS)
+{
+	return !(LHS == RHS);
+}
+
+inline uint32 GetTypeHash(FObjectHandle Handle)
+{
+	// Hash happens on the packed object ref to avoid having to resolve for the sake of hashing
+	return IsObjectHandleResolved(Handle)?
+		GetTypeHash(MakePackedObjectRef(ReadObjectHandlePointerNoCheck(Handle)).EncodedRef):
+		GetTypeHash(Handle.PointerOrRef);
+}
+
+#else
+
+inline bool IsObjectHandleNull(FObjectHandle Handle) { return !Handle; }
+inline bool IsObjectHandleResolved(FObjectHandle Handle) { return true; }
+
+inline UObject* ReadObjectHandlePointerNoCheck(FObjectHandle Handle) { return Handle; }
+inline FPackedObjectRef ReadObjectHandlePackedObjectRefNoCheck(FObjectHandle Handle) { return FPackedObjectRef(); }
+
+inline FObjectHandle MakeObjectHandle(FPackedObjectRef ObjectRef) { return ResolvePackedObjectRef(ObjectRef); }
+inline FObjectHandle MakeObjectHandle(const FObjectRef& ObjectRef) { return ResolveObjectRef(ObjectRef); }
+inline FObjectHandle MakeObjectHandle(UObject* Object) { return Object; }
+
+#endif
+
+inline UObject* ResolveObjectHandle(FObjectHandle& Handle)
+{
+	FObjectHandle LocalHandle = Handle;
+	if (IsObjectHandleResolved(LocalHandle))
+	{
+		UObject* ResolvedObject = ReadObjectHandlePointerNoCheck(LocalHandle);
+		ObjectHandle_Private::OnHandleRead(ResolvedObject);
+		return ResolvedObject;
+	}
+	else
+	{
+		LocalHandle = MakeObjectHandle(ResolvePackedObjectRef(ReadObjectHandlePackedObjectRefNoCheck(LocalHandle)));
+		UObject* ResolvedObject = ReadObjectHandlePointerNoCheck(LocalHandle);
+		Handle = LocalHandle;
+		ObjectHandle_Private::OnHandleRead(ResolvedObject);
+		return ResolvedObject;
+	}
+}
+
+inline UClass* ResolveObjectHandleClass(FObjectHandle Handle)
+{
+	if (IsObjectHandleResolved(Handle))
+	{
+		return ReadObjectHandlePointerNoCheck(Handle)->GetClass();
+	}
+	else
+	{
+		// @TODO: OBJPTR: This should be cached somewhere instead of resolving on every call
+		return ResolvePackedObjectRefClass(ReadObjectHandlePackedObjectRefNoCheck(Handle));
+	}
+}
+
+inline UObject* ReadObjectHandlePointer(FObjectHandle Handle)
+{
+	if (IsObjectHandleResolved(Handle))
+	{
+		UObject* ResolvedObject = ReadObjectHandlePointerNoCheck(Handle);
+		ObjectHandle_Private::OnHandleRead(ResolvedObject);
+		return ResolvedObject;
+	}
+	return nullptr;
+}
+
+inline FPackedObjectRef ReadObjectHandlePackedObjectRef(FObjectHandle Handle)
+{
+	return !IsObjectHandleResolved(Handle) ? ReadObjectHandlePackedObjectRefNoCheck(Handle) : FPackedObjectRef();
+}
+
+#if UE_WITH_OBJECT_HANDLE_LATE_RESOLVE
+inline FObjectRef MakeObjectRef(FObjectHandle Handle)
+{
+	if (IsObjectHandleResolved(Handle))
+	{
+		return MakeObjectRef(ReadObjectHandlePointerNoCheck(Handle));
+	}
+	else
+	{
+		return MakeObjectRef(ReadObjectHandlePackedObjectRefNoCheck(Handle));
+	}
+}
+
+inline FPackedObjectRef MakePackedObjectRef(FObjectHandle Handle)
+{
+	if (IsObjectHandleResolved(Handle))
+	{
+		return MakePackedObjectRef(ReadObjectHandlePointerNoCheck(Handle));
+	}
+	else
+	{
+		return ReadObjectHandlePackedObjectRefNoCheck(Handle);
+	}
+}
+#endif
diff --git a/Engine/Source/Runtime/CoreUObject/Public/UObject/ObjectKey.h b/Engine/Source/Runtime/CoreUObject/Public/UObject/ObjectKey.h
--- a/Engine/Source/Runtime/CoreUObject/Public/UObject/ObjectKey.h
+++ b/Engine/Source/Runtime/CoreUObject/Public/UObject/ObjectKey.h
@@ -5,6 +5,8 @@
 #include "CoreMinimal.h"
 #include "UObject/WeakObjectPtr.h"
 
+#include <type_traits>
+
 /** FObjectKey is an immutable, copyable key which can be used to uniquely identify an object for the lifetime of the application */
 struct FObjectKey
 {
@@ -28,6 +30,14 @@ public:
 			ObjectSerialNumber = Weak.ObjectSerialNumber;
 		}
 	}
+	template <
+		typename U,
+		decltype(ImplicitConv<const UObject*>(std::declval<U>()))* = nullptr
+	>
+	FORCEINLINE FObjectKey(U Object)
+		: FObjectKey(ImplicitConv<const UObject*>(Object))
+	{
+	}
 
 	/** Compare this key with another */
 	FORCEINLINE bool operator==(const FObjectKey& Other) const
diff --git a/Engine/Source/Runtime/CoreUObject/Public/UObject/ObjectMacros.h b/Engine/Source/Runtime/CoreUObject/Public/UObject/ObjectMacros.h
--- a/Engine/Source/Runtime/CoreUObject/Public/UObject/ObjectMacros.h
+++ b/Engine/Source/Runtime/CoreUObject/Public/UObject/ObjectMacros.h
@@ -347,6 +347,7 @@ enum EClassCastFlags : uint64
 	CASTCLASS_FMulticastInlineDelegateProperty	= 0x0004000000000000,
 	CASTCLASS_FMulticastSparseDelegateProperty	= 0x0008000000000000,
 	CASTCLASS_FFieldPathProperty			= 0x0010000000000000,
+	CASTCLASS_FObjectPtrProperty			= 0x0020000000000000,
 };
 
 #define CASTCLASS_AllFlags ((EClassCastFlags)0xFFFFFFFFFFFFFFFF)
diff --git a/Engine/Source/Runtime/CoreUObject/Public/UObject/ObjectPtr.h b/Engine/Source/Runtime/CoreUObject/Public/UObject/ObjectPtr.h
new file mode 100644
--- /dev/null
+++ b/Engine/Source/Runtime/CoreUObject/Public/UObject/ObjectPtr.h
@@ -0,0 +1,499 @@
+// Copyright Epic Games, Inc. All Rights Reserved.
+
+#pragma once
+
+#include "HAL/Platform.h"
+#include "UObject/Class.h"
+#include "UObject/Object.h"
+#include "UObject/ObjectHandle.h"
+
+#include <type_traits>
+
+#define UE_WITH_OBJECT_PTR_DEPRECATIONS 0
+#if UE_WITH_OBJECT_PTR_DEPRECATIONS
+	#define UE_OBJPTR_DEPRECATED(Version, Message) UE_DEPRECATED(Version, Message)
+#else
+	#define UE_OBJPTR_DEPRECATED(Version, Message) 
+#endif
+
+template <typename T>
+struct TObjectPtr;
+
+struct FObjectPtr
+{
+public:
+	FObjectPtr() = default;
+
+	FORCEINLINE FObjectPtr(TYPE_OF_NULLPTR)
+		: Handle(MakeObjectHandle(nullptr))
+	{
+	}
+
+	explicit FORCEINLINE FObjectPtr(UObject* Object)
+		: Handle(MakeObjectHandle(Object))
+	{
+	}
+
+	UE_OBJPTR_DEPRECATED(5.0, "Construction with incomplete type pointer is deprecated.  Please update this code to use MakeObjectPtrUnsafe.")
+	explicit FORCEINLINE FObjectPtr(void* IncompleteObject)
+		: Handle(MakeObjectHandle(reinterpret_cast<UObject*>(IncompleteObject)))
+	{
+	}
+
+	explicit FORCEINLINE FObjectPtr(const FObjectRef& ObjectRef)
+		: Handle(MakeObjectHandle(ObjectRef))
+	{
+	}
+
+	explicit FORCEINLINE FObjectPtr(const FPackedObjectRef& PackedObjectRef)
+		: Handle(MakeObjectHandle(PackedObjectRef))
+	{
+	}
+
+	FORCEINLINE UObject* Get() const
+	{
+		return ResolveObjectHandle(Handle);
+	}
+
+	FORCEINLINE UClass* GetClass() const
+	{
+		return ResolveObjectHandleClass(Handle);
+	}
+
+	// @TODO: OBJPTR: Garbage collector interaction (for Rob's requested feature to make incremental
+	//		 reachability analysis possible) could be hard coded here and not conditionally
+	//		 compiled out.
+	FObjectPtr(FObjectPtr&&) = default;
+	FObjectPtr(const FObjectPtr&) = default;
+	FObjectPtr& operator=(FObjectPtr&&) = default;
+	FObjectPtr& operator=(const FObjectPtr&) = default;
+
+	FObjectPtr& operator=(UObject* Other)
+	{
+		// @TODO: OBJPTR: Garbage collector interaction (for Rob's requested feature to make incremental
+		//		 reachability analysis possible) could be hard coded here and not conditionally
+		//		 compiled out.
+		Handle = MakeObjectHandle(Other);
+		return *this;
+	}
+
+	UE_OBJPTR_DEPRECATED(5.0, "Assignment with incomplete type pointer is deprecated.  Please update this code to use MakeObjectPtrUnsafe.")
+	FObjectPtr& operator=(void* IncompleteOther)
+	{
+		// @TODO: OBJPTR: Garbage collector interaction (for Rob's requested feature to make incremental
+		//		 reachability analysis possible) could be hard coded here and not conditionally
+		//		 compiled out.
+		Handle = MakeObjectHandle(reinterpret_cast<UObject*>(IncompleteOther));
+		return *this;
+	}
+
+	FObjectPtr& operator=(TYPE_OF_NULLPTR)
+	{
+		// @TODO: OBJPTR: Garbage collector interaction (for Rob's requested feature to make incremental
+		//		 reachability analysis possible) could be hard coded here and not conditionally
+		//		 compiled out.
+		Handle = MakeObjectHandle(nullptr);
+		return *this;
+	}
+
+	FORCEINLINE bool operator==(FObjectPtr Other) const { return (Handle == Other.Handle); }
+	FORCEINLINE bool operator!=(FObjectPtr Other) const { return (Handle != Other.Handle); }
+
+	// @TODO: OBJPTR: ToTObjectPtr will be removed in the future when a proper casting layer is added
+	UE_OBJPTR_DEPRECATED(5.0, "Use of ToTObjectPtr is unsafe and is deprecated.")
+	FORCEINLINE TObjectPtr<UObject>& ToTObjectPtr();
+	UE_OBJPTR_DEPRECATED(5.0, "Use of ToTObjectPtr is unsafe and is deprecated.")
+	FORCEINLINE const TObjectPtr<UObject>& ToTObjectPtr() const;
+
+	// @TODO: OBJPTR: The presence of this conversion operator triggers a bug in MSVC where implicit conversion of the
+	//			derived TObjectPtr type to UObject* results in a compile error saying that the conversion  operator is
+	//			in a private base class.  That shouldn't matter as the derived TObjectPtr type has its own templated
+	//			conversion operator.  This issue does not come up in clang compiled builds.
+	//FORCEINLINE operator UObject*() const { return Get();  }
+	FORCEINLINE UObject* operator->() const { return Get(); }
+	FORCEINLINE UObject& operator*() const { return *Get(); }
+
+	FORCEINLINE bool IsNull() const { return IsObjectHandleNull(Handle); }
+	explicit FORCEINLINE operator bool() const { return !IsNull(); }
+	FORCEINLINE FObjectHandle GetHandle() const { return Handle; }
+	FORCEINLINE FObjectHandle& GetHandleRef() const { return Handle; }
+
+	FORCEINLINE bool IsA(const UClass* SomeBase) const
+	{
+		checkfSlow(SomeBase, TEXT("IsA(NULL) cannot yield meaningful results"));
+
+		const UClass* ThisClass = GetClass();
+		checkfSlow(ThisClass, TEXT("Object pointer class did not resolve to a non-null value"));
+
+		return ThisClass->IsChildOf(SomeBase);
+	}
+
+	template <typename T>
+	FORCEINLINE bool IsA() const
+	{
+		return IsA(T::StaticClass());
+	}
+
+private:
+	friend FORCEINLINE uint32 GetTypeHash(FObjectPtr Object)
+	{
+		return GetTypeHash(Object.Handle);
+	}
+
+	mutable FObjectHandle Handle;
+};
+
+template <typename T>
+struct TPrivateObjectPtr;
+
+template <typename T>
+struct TIsTObjectPtr
+{
+	enum { Value = false };
+};
+
+template <typename T>
+struct TObjectPtr : private FObjectPtr
+{
+public:
+
+	TObjectPtr() = default;
+	TObjectPtr(TObjectPtr<T>&&) = default;
+	TObjectPtr(const TObjectPtr<T>&) = default;
+
+	FORCEINLINE TObjectPtr(TYPE_OF_NULLPTR)
+		: FObjectPtr(nullptr)
+	{
+	}
+
+	template <
+		typename U,
+		decltype(ImplicitConv<T*>(std::declval<U*>()))* = nullptr
+	>
+	FORCEINLINE TObjectPtr(const TObjectPtr<U>& Other)
+		: FObjectPtr(static_cast<const FObjectPtr&>(Other))
+	{
+	}
+
+	template <
+		typename U,
+		std::enable_if_t<
+			!TIsTObjectPtr<std::decay_t<U>>::Value,
+			decltype(ImplicitConv<T*>(std::declval<U>()))
+		>* = nullptr
+	>
+	explicit FORCEINLINE TObjectPtr(U&& Object)
+		: FObjectPtr(const_cast<std::remove_const_t<T>*>(ImplicitConv<T*>(Object)))
+	{
+		// It would be desirable for this constructor to be implicit to allow:
+		// 1) Uniform initialization syntax of a TObjectPtr field from a raw object pointer value
+		// 2) Allowing the implicit construction of a TObjectPtr<Derived> from a raw object pointer to Base (or something implicitly convertible to that)
+		// However making this implicit introduces a compile-time ambiguity in ternary operators of the form "condition ? TObjectPtr : raw pointer"
+	}
+
+	explicit FORCEINLINE TObjectPtr(TPrivateObjectPtr<T>&& PrivatePtr)
+		: FObjectPtr(const_cast<UObject*>(PrivatePtr.Pointer))
+	{
+	}
+
+	TObjectPtr<T>& operator=(TObjectPtr<T>&&) = default;
+	TObjectPtr<T>& operator=(const TObjectPtr<T>&) = default;
+
+	FORCEINLINE TObjectPtr<T>& operator=(TYPE_OF_NULLPTR)
+	{
+		FObjectPtr::operator=(nullptr);
+		return *this;
+	}
+
+	template <
+		typename U,
+		decltype(ImplicitConv<T*>(std::declval<U*>()))* = nullptr
+	>
+	FORCEINLINE TObjectPtr<T>& operator=(const TObjectPtr<U>& Other)
+	{
+		FObjectPtr::operator=(static_cast<const FObjectPtr&>(Other));
+		return *this;
+	}
+
+	template <
+		typename U,
+		std::enable_if_t<
+			!TIsTObjectPtr<std::decay_t<U>>::Value,
+			decltype(ImplicitConv<T*>(std::declval<U>()))
+		>* = nullptr
+	>
+	FORCEINLINE TObjectPtr<T>& operator=(U&& Object)
+	{
+		FObjectPtr::operator=(const_cast<std::remove_const_t<T>*>(ImplicitConv<T*>(Object)));
+		return *this;
+	}
+
+	FORCEINLINE TObjectPtr<T>& operator=(TPrivateObjectPtr<T>&& PrivatePtr)
+	{
+		FObjectPtr::operator=(const_cast<UObject*>(PrivatePtr.Pointer));
+		return *this;
+	}
+
+	template <
+		typename U,
+		typename Base = std::decay_t<decltype(false ? std::declval<std::decay_t<T*>>() : std::declval<std::decay_t<U*>>())>
+	>
+	FORCEINLINE bool operator==(const TObjectPtr<U>& Other) const
+	{
+		return FObjectPtr::operator==(static_cast<const FObjectPtr&>(Other));
+	}
+
+	template <
+		typename U,
+		typename Base = std::decay_t<decltype(false ? std::declval<std::decay_t<T*>>() : std::declval<std::decay_t<U*>>())>
+	>
+	FORCEINLINE bool operator!=(const TObjectPtr<U>& Other) const
+	{
+		return FObjectPtr::operator!=(static_cast<const FObjectPtr&>(Other));
+	}
+
+	// @TODO: OBJPTR: There is a risk that the FObjectPtr is storing a reference to the wrong type.  This could
+	//			happen if data was serialized at a time when a pointer field was declared to be of type A, but then the declaration
+	//			changed and the pointer field is now of type B.  Upon deserialization of pre-existing data, we'll be holding
+	//			a reference to the wrong type of object which we'll just send back static_casted as the wrong type.  Doing
+	//			a check or checkSlow here could catch this, but it would be better if the check could happen elsewhere that
+	//			isn't called as frequently.
+	FORCEINLINE T* Get() const { return (T*)(FObjectPtr::Get()); }
+	using FObjectPtr::GetClass;
+	
+	FORCEINLINE operator T* () const { return Get(); }
+	template <typename U>
+	UE_OBJPTR_DEPRECATED(5.0, "Explicit cast to other raw pointer types is deprecated.  Please use the Cast API or get the raw pointer with ToRawPtr and cast that instead.")
+	explicit FORCEINLINE operator U* () const { return (U*)Get(); }
+	explicit FORCEINLINE operator UPTRINT() const { return (UPTRINT)Get(); }
+	FORCEINLINE T* operator->() const { return Get(); }
+	FORCEINLINE T& operator*() const { return *Get(); }
+
+	UE_OBJPTR_DEPRECATED(5.0, "Conversion to a mutable pointer is deprecated.  Please pass a TObjectPtr<T>& instead so that assignment can be tracked accurately.")
+	FORCEINLINE operator T*& () { return GetInternalRef(); }
+
+	using FObjectPtr::IsNull;
+	using FObjectPtr::operator bool;
+	using FObjectPtr::GetHandle;
+
+	using FObjectPtr::IsA;
+
+	friend FORCEINLINE uint32 GetTypeHash(const TObjectPtr<T>& ObjectPtr)
+	{
+		return GetTypeHash(static_cast<const FObjectPtr&>(ObjectPtr));
+	}
+
+	friend FORCEINLINE FArchive& operator<<(FArchive& Ar, TObjectPtr<T>& ObjectPtr)
+	{
+		Ar << static_cast<FObjectPtr&>(ObjectPtr);
+		return Ar;
+	}
+
+	friend struct FObjectPtr;
+	template <typename U> friend struct TObjectPtr;
+private:
+	// @TODO: OBJPTR: There is a risk of a gap in access tracking here.  The caller may get a mutable pointer, write to it, then
+	//			read from it.  That last read would happen without an access being recorded.  Not sure if there is a good way
+	//			to handle this case without forcing the calling code to be modified.
+	FORCEINLINE T*& GetInternalRef()
+	{
+		FObjectPtr::Get();
+		return (T*&)FObjectPtr::GetHandleRef();
+	}
+};
+
+template <typename T>
+struct TIsTObjectPtr<TObjectPtr<T>>
+{
+	enum { Value = true };
+};
+
+template <typename T>
+struct TRemoveObjectPointer
+{
+	typedef T Type;
+};
+template <typename T>
+struct TRemoveObjectPointer<TObjectPtr<T>>
+{
+	typedef T Type;
+};
+
+namespace ObjectPtr_Private
+{
+	/** Coerce to pointer through implicit conversion to const T* (overload through less specific "const T*" parameter to avoid ambiguity with other coercion options that may also exist. */
+	template <typename T>
+	FORCEINLINE const T* CoerceToPointer(const T* Other)
+	{
+		return Other;
+	}
+
+	/** Coerce to pointer through implicit conversion to CommonPointerType where CommonPointerType is deduced, and must be a C++ pointer, not a wrapper type. */
+	template <
+		typename T,
+		typename U,
+		typename CommonPointerType =  decltype(std::declval<bool>() ? std::declval<const T*>() : std::declval<U>()),
+		std::enable_if_t<
+			std::is_pointer<CommonPointerType>::value
+		>* = nullptr
+	>
+	FORCEINLINE auto CoerceToPointer(U&& Other) -> CommonPointerType
+	{
+		return Other;
+	}
+
+	/** Coerce to pointer through the use of a ".Get()" member, which is the convention within Unreal smart pointer types. */
+	template <
+		typename T,
+		typename U,
+		std::enable_if_t<
+			!TIsTObjectPtr<std::decay_t<U>>::Value,
+			decltype(std::declval<U>().Get())
+		>* = nullptr
+	>
+	FORCEINLINE auto CoerceToPointer(U&& Other) -> decltype(std::declval<U>().Get())
+	{
+		return Other.Get();
+	}
+
+	/** Perform shallow equality check between a TObjectPtr and another (non TObjectPtr) type that we can coerce to a pointer. */
+	template <
+		typename T,
+		typename U,
+		std::enable_if_t<
+			!TIsTObjectPtr<std::decay_t<U>>::Value,
+			decltype(CoerceToPointer<T>(std::declval<U>()) == std::declval<const T*>())
+		>* = nullptr
+	>
+	FORCEINLINE bool IsObjectPtrEqual(const TObjectPtr<T>& Ptr, U&& Other)
+	{
+		auto OtherPtr = ObjectPtr_Private::CoerceToPointer<T>(Other);
+		return (Ptr.IsNull() == (OtherPtr == nullptr)) && (Ptr.Get() == OtherPtr);
+	}
+}
+
+// Equality/Inequality comparisons against another type that can be implicitly converted to the pointer type kept in a TObjectPtr
+template <
+	typename T,
+	typename U,
+	decltype(ObjectPtr_Private::IsObjectPtrEqual(std::declval<const TObjectPtr<T>&>(), std::declval<U&&>()))* = nullptr
+>
+FORCEINLINE bool operator==(const TObjectPtr<T>& Ptr, U&& Other)
+{
+	return ObjectPtr_Private::IsObjectPtrEqual(Ptr, Other);
+}
+template <
+	typename T,
+	typename U,
+	decltype(ObjectPtr_Private::IsObjectPtrEqual(std::declval<const TObjectPtr<T>&>(), std::declval<U&&>()))* = nullptr
+>
+FORCEINLINE bool operator==(U&& Other, const TObjectPtr<T>& Ptr)
+{
+	return ObjectPtr_Private::IsObjectPtrEqual(Ptr, Other);
+}
+template <
+	typename T,
+	typename U,
+	decltype(ObjectPtr_Private::IsObjectPtrEqual(std::declval<const TObjectPtr<T>&>(), std::declval<U&&>()))* = nullptr
+>
+FORCEINLINE bool operator!=(const TObjectPtr<T>& Ptr, U&& Other)
+{
+	return !ObjectPtr_Private::IsObjectPtrEqual(Ptr, Other);
+}
+template <
+	typename T,
+	typename U,
+	decltype(ObjectPtr_Private::IsObjectPtrEqual(std::declval<const TObjectPtr<T>&>(), std::declval<U&&>()))* = nullptr
+>
+FORCEINLINE bool operator!=(U&& Other, const TObjectPtr<T>& Ptr)
+{
+	return !ObjectPtr_Private::IsObjectPtrEqual(Ptr, Other);
+}
+
+template <typename T>
+TPrivateObjectPtr<T> MakeObjectPtrUnsafe(const UObject* Obj);
+
+template <typename T>
+struct TPrivateObjectPtr
+{
+public:
+	TPrivateObjectPtr(const TPrivateObjectPtr<T>& Other) = default;
+
+private:
+	/** Only for use by MakeObjectPtrUnsafe */
+	explicit TPrivateObjectPtr(const UObject* InPointer)
+		: Pointer(InPointer)
+	{
+	}
+
+	const UObject* Pointer;
+	friend struct TObjectPtr<T>;
+	friend TPrivateObjectPtr MakeObjectPtrUnsafe<T>(const UObject* Obj);
+};
+
+/** Used to allow the caller to provide a pointer to an incomplete type of T that has explicitly cast to a UObject. */
+template <typename T>
+TPrivateObjectPtr<T> MakeObjectPtrUnsafe(const UObject* Obj)
+{
+	return TPrivateObjectPtr<T>{Obj};
+}
+
+template <typename T>
+FORCEINLINE T* ToRawPtr(const TObjectPtr<T>& Ptr)
+{
+	// NOTE: This is specifically not getting a reference to the internal pointer.
+	return Ptr.Get();
+}
+
+template <typename T>
+FORCEINLINE T* ToRawPtr(T* Ptr)
+{
+	return Ptr;
+}
+
+template <typename T, SIZE_T Size>
+FORCEINLINE T** ToRawPtrArrayUnsafe(TObjectPtr<T>(&ArrayOfPtr)[Size])
+{
+	for (TObjectPtr<T>& Item : ArrayOfPtr)
+	{
+		// NOTE: Relying on the fact that the TObjectPtr will cache the resolved pointer in place after calling Get.
+		Item.Get();
+	}
+	return reinterpret_cast<T**>(ArrayOfPtr);
+}
+
+template <typename T>
+FORCEINLINE T** ToRawPtrArrayUnsafe(T** ArrayOfPtr)
+{
+	return ArrayOfPtr;
+}
+
+
+
+template <typename T>
+FORCEINLINE TWeakObjectPtr<T> MakeWeakObjectPtr(TObjectPtr<T> Ptr)
+{
+	return TWeakObjectPtr<T>(Ptr);
+}
+
+FORCEINLINE TObjectPtr<UObject>& FObjectPtr::ToTObjectPtr()
+{
+	return *static_cast<TObjectPtr<UObject>*>(this);
+}
+
+FORCEINLINE const TObjectPtr<UObject>& FObjectPtr::ToTObjectPtr() const
+{
+	return *static_cast<const TObjectPtr<UObject>*>(this);
+}
+
+/** Exchange variants between TObjectPtr<T> and raw pointer to T */
+template <typename T>
+inline void Exchange(TObjectPtr<T>& A, T*& B)
+{
+	Swap((T*&)A, B);
+}
+template <typename T>
+inline void Exchange(T*& A, TObjectPtr<T>& B)
+{
+	Swap(A, (T*&)B);
+}
diff --git a/Engine/Source/Runtime/CoreUObject/Public/UObject/SoftObjectPath.h b/Engine/Source/Runtime/CoreUObject/Public/UObject/SoftObjectPath.h
--- a/Engine/Source/Runtime/CoreUObject/Public/UObject/SoftObjectPath.h
+++ b/Engine/Source/Runtime/CoreUObject/Public/UObject/SoftObjectPath.h
@@ -7,6 +7,7 @@
 #include "HAL/ThreadSafeCounter.h"
 #include "HAL/ThreadSingleton.h"
 #include "UObject/Class.h"
+#include "UObject/ObjectPtr.h"
 
 /**
  * A struct that contains a string reference to an object, either a top level asset or a subobject.
@@ -50,6 +51,20 @@ struct COREUOBJECT_API FSoftObjectPath
 	/** Construct from an existing object in memory */
 	FSoftObjectPath(const UObject* InObject);
 
+	/** Construct from a FObjectHandle/FObjectPtr/TObjectPtr which may or may not have the referenced object in memory */
+	explicit FSoftObjectPath(FObjectHandle InObjectHandle);
+
+	FSoftObjectPath(FObjectPtr InObject)
+		: FSoftObjectPath(InObject.GetHandle())
+	{
+	}
+
+	template <typename T>
+	FSoftObjectPath(const TObjectPtr<T>& InObject)
+		: FSoftObjectPath(InObject.GetHandle())
+	{
+	}
+
 	~FSoftObjectPath() {}
 	
 	FSoftObjectPath& operator=(const FSoftObjectPath& Path)	= default;
diff --git a/Engine/Source/Runtime/CoreUObject/Public/UObject/Stack.h b/Engine/Source/Runtime/CoreUObject/Public/UObject/Stack.h
--- a/Engine/Source/Runtime/CoreUObject/Public/UObject/Stack.h
+++ b/Engine/Source/Runtime/CoreUObject/Public/UObject/Stack.h
@@ -50,6 +50,7 @@ enum EPropertyType
 	CPT_MulticastDelegate,
 	CPT_WeakObjectReference,
 	CPT_LazyObjectReference,
+	CPT_ObjectPtrReference,
 	CPT_SoftObjectReference,
 	CPT_Double,
 	CPT_Map,
diff --git a/Engine/Source/Runtime/CoreUObject/Public/UObject/UObjectGlobals.h b/Engine/Source/Runtime/CoreUObject/Public/UObject/UObjectGlobals.h
--- a/Engine/Source/Runtime/CoreUObject/Public/UObject/UObjectGlobals.h
+++ b/Engine/Source/Runtime/CoreUObject/Public/UObject/UObjectGlobals.h
@@ -20,7 +20,10 @@
 
 struct FCustomPropertyListNode;
 struct FObjectInstancingGraph;
+struct FObjectPtr;
 struct FStaticConstructObjectParameters;
+template <typename T>
+struct TObjectPtr;
 
 COREUOBJECT_API DECLARE_LOG_CATEGORY_EXTERN(LogUObjectGlobals, Log, All);
 
@@ -1341,6 +1344,19 @@ T* DuplicateObject(T const* SourceObject,UObject* Outer, const FName Name = NAME
 	}
 	return nullptr;
 }
+template <typename T>
+T* DuplicateObject(const TObjectPtr<T>& SourceObject,UObject* Outer, const FName Name = NAME_None)
+{
+	if (SourceObject != nullptr)
+	{
+		if (Outer == nullptr || Outer == INVALID_OBJECT)
+		{
+			Outer = (UObject*)GetTransientPackage();
+		}
+		return (T*)StaticDuplicateObject(SourceObject,Outer,Name);
+	}
+	return nullptr;
+}
 
 /**
  * Determines whether the specified object should load values using PerObjectConfig rules
@@ -1901,6 +1917,150 @@ public:
 		}
 	}
 
+
+
+	/**
+	 * Adds object reference.
+	 *
+	 * @param Object Referenced object.
+	 * @param ReferencingObject Referencing object (if available).
+	 * @param ReferencingProperty Referencing property (if available).
+	 */
+	template<class UObjectType>
+	void AddReferencedObject(TObjectPtr<UObjectType>& Object, const UObject* ReferencingObject = nullptr, const FProperty* ReferencingProperty = nullptr)
+	{
+		if (IsObjectHandleResolved(Object.GetHandle()))
+		{
+			HandleObjectReference(*reinterpret_cast<UObject**>(&Object), ReferencingObject, ReferencingProperty);
+		}
+	}
+
+	/**
+	 * Adds const object reference, this reference can still be nulled out if forcefully collected.
+	 *
+	 * @param Object Referenced object.
+	 * @param ReferencingObject Referencing object (if available).
+	 * @param ReferencingProperty Referencing property (if available).
+	 */
+	template<class UObjectType>
+	void AddReferencedObject(TObjectPtr<const UObjectType>& Object, const UObject* ReferencingObject = nullptr, const FProperty* ReferencingProperty = nullptr)
+	{
+		static_assert(sizeof(UObjectType) > 0, "AddReferencedObject: Element must be a pointer to a fully-defined type");
+		static_assert(TPointerIsConvertibleFromTo<UObjectType, const UObjectBase>::Value, "AddReferencedObject: Element must be a pointer to a type derived from UObject");
+		if (IsObjectHandleResolved(Object.GetHandle()))
+		{
+			HandleObjectReference(*reinterpret_cast<UObject**>(&Object), ReferencingObject, ReferencingProperty);
+		}
+	}
+
+	/**
+	* Adds references to an array of objects.
+	*
+	* @param ObjectArray Referenced objects array.
+	* @param ReferencingObject Referencing object (if available).
+	* @param ReferencingProperty Referencing property (if available).
+	*/
+	template<class UObjectType>
+	void AddReferencedObjects(TArray<TObjectPtr<UObjectType>>& ObjectArray, const UObject* ReferencingObject = nullptr, const FProperty* ReferencingProperty = nullptr)
+	{
+		static_assert(sizeof(UObjectType) > 0, "AddReferencedObjects: Elements must be pointers to a fully-defined type");
+		static_assert(TPointerIsConvertibleFromTo<UObjectType, const UObjectBase>::Value, "AddReferencedObjects: Elements must be pointers to a type derived from UObject");
+		// Cannot use a reinterpret_cast due to MSVC (and not Clang) emitting a warning:
+		// C4946: reinterpret_cast used between related classes: ...
+		HandleObjectReferences((FObjectPtr*)(ObjectArray.GetData()), ObjectArray.Num(), ReferencingObject, ReferencingProperty);
+	}
+
+	/**
+	* Adds references to an array of const objects, these objects can still be nulled out if forcefully collected.
+	*
+	* @param ObjectArray Referenced objects array.
+	* @param ReferencingObject Referencing object (if available).
+	* @param ReferencingProperty Referencing property (if available).
+	*/
+	template<class UObjectType>
+	void AddReferencedObjects(TArray<TObjectPtr<const UObjectType>>& ObjectArray, const UObject* ReferencingObject = nullptr, const FProperty* ReferencingProperty = nullptr)
+	{
+		static_assert(sizeof(UObjectType) > 0, "AddReferencedObjects: Elements must be pointers to a fully-defined type");
+		static_assert(TPointerIsConvertibleFromTo<UObjectType, const UObjectBase>::Value, "AddReferencedObjects: Elements must be pointers to a type derived from UObject");
+		// Cannot use a reinterpret_cast due to MSVC (and not Clang) emitting a warning:
+		// C4946: reinterpret_cast used between related classes: ...
+		HandleObjectReferences((FObjectPtr*)(ObjectArray.GetData()), ObjectArray.Num(), ReferencingObject, ReferencingProperty);
+	}
+
+	/**
+	* Adds references to a set of objects.
+	*
+	* @param ObjectSet Referenced objects set.
+	* @param ReferencingObject Referencing object (if available).
+	* @param ReferencingProperty Referencing property (if available).
+	*/
+	template<class UObjectType>
+	void AddReferencedObjects(TSet<TObjectPtr<UObjectType>>& ObjectSet, const UObject* ReferencingObject = nullptr, const FProperty* ReferencingProperty = nullptr)
+	{
+		static_assert(sizeof(UObjectType) > 0, "AddReferencedObjects: Elements must be pointers to a fully-defined type");
+		static_assert(TPointerIsConvertibleFromTo<UObjectType, const UObjectBase>::Value, "AddReferencedObjects: Elements must be pointers to a type derived from UObject");
+		for (auto& Object : ObjectSet)
+		{
+			if (IsObjectHandleResolved(Object.GetHandle()))
+			{
+				HandleObjectReference(*reinterpret_cast<UObject**>(&Object), ReferencingObject, ReferencingProperty);
+			}
+		}
+	}
+
+	/**
+	 * Adds references to a map of objects.
+	 *
+	 * @param ObjectArray Referenced objects map.
+	 * @param ReferencingObject Referencing object (if available).
+	 * @param ReferencingProperty Referencing property (if available).
+	 */
+	template <typename KeyType, typename ValueType, typename Allocator, typename KeyFuncs>
+	void AddReferencedObjects(TMapBase<TObjectPtr<KeyType>, ValueType, Allocator, KeyFuncs>& Map, const UObject* ReferencingObject = nullptr, const FProperty* ReferencingProperty = nullptr)
+	{
+		static_assert(sizeof(KeyType) > 0, "AddReferencedObjects: Keys must be pointers to a fully-defined type");
+		static_assert(TPointerIsConvertibleFromTo<KeyType, const UObjectBase>::Value, "AddReferencedObjects: Keys must be pointers to a type derived from UObject");
+		for (auto& It : Map)
+		{
+			if (IsObjectHandleResolved(It.Key.GetHandle()))
+			{
+				HandleObjectReference(*reinterpret_cast<UObject**>(&It.Key), ReferencingObject, ReferencingProperty);
+			}
+		}
+	}
+	template <typename KeyType, typename ValueType, typename Allocator, typename KeyFuncs>
+	void AddReferencedObjects(TMapBase<KeyType, TObjectPtr<ValueType>, Allocator, KeyFuncs>& Map, const UObject* ReferencingObject = nullptr, const FProperty* ReferencingProperty = nullptr)
+	{
+		static_assert(sizeof(ValueType) > 0, "AddReferencedObjects: Values must be pointers to a fully-defined type");
+		static_assert(TPointerIsConvertibleFromTo<ValueType, const UObjectBase>::Value, "AddReferencedObjects: Values must be pointers to a type derived from UObject");
+		for (auto& It : Map)
+		{
+			if (IsObjectHandleResolved(It.Value.GetHandle()))
+			{
+				HandleObjectReference(*reinterpret_cast<UObject**>(&It.Value), ReferencingObject, ReferencingProperty);
+			}
+		}
+	}
+	template <typename KeyType, typename ValueType, typename Allocator, typename KeyFuncs>
+	void AddReferencedObjects(TMapBase<TObjectPtr<KeyType>, TObjectPtr<ValueType>, Allocator, KeyFuncs>& Map, const UObject* ReferencingObject = nullptr, const FProperty* ReferencingProperty = nullptr)
+	{
+		static_assert(sizeof(KeyType) > 0, "AddReferencedObjects: Keys must be pointers to a fully-defined type");
+		static_assert(sizeof(ValueType) > 0, "AddReferencedObjects: Values must be pointers to a fully-defined type");
+		static_assert(TPointerIsConvertibleFromTo<KeyType, const UObjectBase>::Value, "AddReferencedObjects: Keys must be pointers to a type derived from UObject");
+		static_assert(TPointerIsConvertibleFromTo<ValueType, const UObjectBase>::Value, "AddReferencedObjects: Values must be pointers to a type derived from UObject");
+		for (auto& It : Map)
+		{
+			if (IsObjectHandleResolved(It.Key.GetHandle()))
+			{
+				HandleObjectReference(*reinterpret_cast<UObject**>(&It.Key), ReferencingObject, ReferencingProperty);
+			}
+			if (IsObjectHandleResolved(It.Value.GetHandle()))
+			{
+				HandleObjectReference(*reinterpret_cast<UObject**>(&It.Value), ReferencingObject, ReferencingProperty);
+			}
+		}
+	}
+
 	/**
 	 * If true archetype references should not be added to this collector.
 	 */
@@ -1965,7 +2125,7 @@ protected:
 
 	/**
 	* Handle multiple object references. Called by AddReferencedObjects.
-	* DEFAULT IMPLEMENTAION IS SLOW as it calls HandleObjectReference multiple times. In order to optimize it, provide your own implementation.
+	* DEFAULT IMPLEMENTATION IS SLOW as it calls HandleObjectReference multiple times. In order to optimize it, provide your own implementation.
 	*
 	* @param Object Referenced object.
 	* @param ReferencingObject Referencing object (if available).
@@ -1980,6 +2140,16 @@ protected:
 		}
 	}
 
+	/**
+	* Handle multiple object references. Called by AddReferencedObjects.
+	* DEFAULT IMPLEMENTATION IS SLOW as it calls HandleObjectReference multiple times. In order to optimize it, provide your own implementation.
+	*
+	* @param Object Referenced object.
+	* @param ReferencingObject Referencing object (if available).
+	* @param ReferencingProperty Referencing property (if available).
+	*/
+	virtual void HandleObjectReferences(FObjectPtr* InObjects, const int32 ObjectNum, const UObject* InReferencingObject, const FProperty* InReferencingProperty);
+
 private:
 
 	/** Creates the roxy archive that uses serialization to add objects to this collector */
@@ -2346,7 +2516,8 @@ namespace UE4CodeGen_Private
 		FieldPath         = 0x1F,
 
 		// Property-specific flags
-		NativeBool        = 0x20
+		NativeBool        = 0x20,
+		ObjectPtr         = 0x20,
 	};
 
 	ENUM_CLASS_FLAGS(EPropertyGenFlags)
@@ -2640,6 +2811,7 @@ namespace UE4CodeGen_Private
 	typedef FGenericPropertyParams FTextPropertyParams;
 	typedef FObjectPropertyParams  FWeakObjectPropertyParams;
 	typedef FObjectPropertyParams  FLazyObjectPropertyParams;
+	typedef FObjectPropertyParams  FObjectPtrPropertyParams;
 	typedef FObjectPropertyParams  FSoftObjectPropertyParams;
 
 	struct FFunctionParams
diff --git a/Engine/Source/Runtime/CoreUObject/Public/UObject/UnrealType.h b/Engine/Source/Runtime/CoreUObject/Public/UObject/UnrealType.h
--- a/Engine/Source/Runtime/CoreUObject/Public/UObject/UnrealType.h
+++ b/Engine/Source/Runtime/CoreUObject/Public/UObject/UnrealType.h
@@ -24,6 +24,7 @@
 #include "UObject/LazyObjectPtr.h"
 #include "UObject/Object.h"
 #include "UObject/ObjectMacros.h"
+#include "UObject/ObjectPtr.h"
 #include "UObject/PropertyPortFlags.h"
 #include "UObject/PropertyTag.h"
 #include "UObject/ScriptInterface.h"
@@ -2282,6 +2283,46 @@ public:
 	// End of FObjectPropertyBase interface
 };
 
+//
+// Describes a reference variable to another object which may be nil.
+//
+class COREUOBJECT_API FObjectPtrProperty : public TFObjectPropertyBase<FObjectPtr>
+{
+	DECLARE_FIELD(FObjectPtrProperty, TFObjectPropertyBase<FObjectPtr>, CASTCLASS_FObjectPtrProperty)
+
+	using Super::Super;
+
+#if WITH_EDITORONLY_DATA
+	explicit FObjectPtrProperty(UField* InField)
+		: TFObjectPropertyBase(InField)
+	{
+	}
+#endif // WITH_EDITORONLY_DATA
+
+	// UHT interface
+	virtual FString GetCPPMacroType( FString& ExtendedTypeText ) const  override;
+	virtual FString GetCPPType( FString* ExtendedTypeText, uint32 CPPExportFlags ) const override;
+	// End of UHT interface
+
+	// FProperty interface
+	virtual FName GetID() const override;
+	virtual bool SameType(const FProperty* Other) const override;
+	virtual bool Identical(const void* A, const void* B, uint32 PortFlags) const override;
+	virtual void SerializeItem(FStructuredArchive::FSlot Slot, void* Value, void const* Defaults) const override;
+	virtual void EmitReferenceInfo(UClass& OwnerClass, int32 BaseOffset, TArray<const FStructProperty*>& EncounteredStructProps) override;
+	virtual EConvertFromTypeResult ConvertFromType(const FPropertyTag& Tag, FStructuredArchive::FSlot Slot, uint8* Data, UStruct* DefaultsStruct) override;
+	// End of FProperty interface
+
+	// FObjectProperty interface
+	virtual UObject* GetObjectPropertyValue(const void* PropertyValueAddress) const override;
+	virtual void SetObjectPropertyValue(void* PropertyValueAddress, UObject* Value) const override;
+	virtual bool AllowCrossLevel() const override;
+private:
+	virtual uint32 GetValueTypeHashInternal(const void* Src) const override;
+public:
+	// End of FObjectProperty interface
+};
+
 //
 // Describes a reference variable to another object which may be nil, and may turn nil at any point
 //
diff --git a/Engine/Source/Runtime/CoreUObject/Public/UObject/WeakObjectPtr.h b/Engine/Source/Runtime/CoreUObject/Public/UObject/WeakObjectPtr.h
--- a/Engine/Source/Runtime/CoreUObject/Public/UObject/WeakObjectPtr.h
+++ b/Engine/Source/Runtime/CoreUObject/Public/UObject/WeakObjectPtr.h
@@ -10,6 +10,8 @@
 #include "UObject/UObjectArray.h"
 #include "UObject/FastReferenceCollectorOptions.h"
 
+#include <type_traits>
+
 /**
  * FWeakObjectPtr is a weak pointer to a UObject. 
  * It can return nullptr later if the object is garbage collected.
@@ -31,13 +33,27 @@ public:
 	{
 		Reset();
 	}
+
+	/**
+	 * Construct from nullptr or something that can be implicitly converted to nullptr (eg: NULL)
+	 * @param Object object to create a weak pointer to
+	 */
+	FORCEINLINE FWeakObjectPtr(TYPE_OF_NULLPTR)
+	{
+		(*this) = nullptr;
+	}
+
 	/**  
-	 * Construct from an object pointer
+	 * Construct from an object pointer or something that can be implicitly converted to an object pointer
 	 * @param Object object to create a weak pointer to
 	 */
-	FORCEINLINE FWeakObjectPtr(const class UObject* Object)
+	template <
+		typename U,
+		decltype(ImplicitConv<const UObject*>(std::declval<U>()))* = nullptr
+	>
+	FORCEINLINE FWeakObjectPtr(U&& Object)
 	{
-		(*this)=Object;
+		(*this) = ImplicitConv<const UObject*>(Object);
 	}
 
 	/**  
diff --git a/Engine/Source/Runtime/CoreUObject/Public/Virtualization/VirtualizationManager.h b/Engine/Source/Runtime/CoreUObject/Public/Virtualization/VirtualizationManager.h
new file mode 100644
--- /dev/null
+++ b/Engine/Source/Runtime/CoreUObject/Public/Virtualization/VirtualizationManager.h
@@ -0,0 +1,53 @@
+// Copyright Epic Games, Inc. All Rights Reserved.
+
+#pragma once
+
+#include "Logging/LogMacros.h"
+#include "Misc/Guid.h"
+#include "Memory/SharedBuffer.h"
+
+// For UE_USE_VIRTUALBULKDATA
+#include "UObject/UE5CookerObjectVersion.h"
+
+COREUOBJECT_API DECLARE_LOG_CATEGORY_EXTERN(LogVirtualization, Log, All);
+
+// TODO: Do we want to keep this header public for UE5 release? If the only interaction should be
+//		 via FVirtualizedUntypedBulkData or other such classes we might not want to expose this 
+//		 at all.
+
+// TODO: Remove when UE_USE_VIRTUALBULKDATA is removed  
+/** Some methods can be turned to const when we are using VirtualBulkData, this macro will help with the upgrade path */
+#if UE_USE_VIRTUALBULKDATA
+	#define UE_VBD_CONST const
+#else
+	#define UE_VBD_CONST
+#endif //UE_USE_VIRTUALBULKDATA
+
+/** This is used as a wrapper around the various potential back end implementations. 
+	The calling code shouldn't need to care about which back ends are actually in use. */
+class COREUOBJECT_API FVirtualizationManager
+{
+public:
+	/** Singleton access */
+	static FVirtualizationManager& Get();
+
+	FVirtualizationManager();
+	~FVirtualizationManager() = default;
+
+	/** Push a payload to the backends*/
+	bool PushData(const FSharedBufferConstPtr& Payload, const FGuid& Guid);
+
+	/** Pull a payload from the backends*/
+	FSharedBufferConstPtr PullData(const FGuid& Guid);
+
+private:
+
+	void LoadSettingsFromConfigFiles();
+	FString CreateFilePath(const FGuid& Guid);
+
+	/** Are payloads allowed to be virtualized. Defaults to true. */
+	bool bEnablePayloadPushing;
+
+	/** The minimum length for a payload to be considered for virtualization. Defaults to 0 bytes. */
+	int64 MinPayloadLength;
+};
diff --git a/Engine/Source/Runtime/CoreUObject/Public/Virtualization/VirtualizationUtilities.h b/Engine/Source/Runtime/CoreUObject/Public/Virtualization/VirtualizationUtilities.h
new file mode 100644
--- /dev/null
+++ b/Engine/Source/Runtime/CoreUObject/Public/Virtualization/VirtualizationUtilities.h
@@ -0,0 +1,191 @@
+// Copyright Epic Games, Inc. All Rights Reserved.
+
+#pragma once
+
+#include "Containers/UnrealString.h"
+#include "Virtualization/VirtualizedBulkData.h"
+#include "Serialization/BufferReader.h"
+
+#if WITH_EDITORONLY_DATA
+
+class FVirtualizedBulkDataWriter : public FArchive
+{
+public:
+	FVirtualizedBulkDataWriter(FVirtualizedUntypedBulkData& InBulkData, bool bIsPersistent = false)
+		: BulkData(InBulkData)
+	{
+		SetIsSaving(true);
+		SetIsPersistent(bIsPersistent);
+
+		FSharedBufferConstPtr Payload = InBulkData.GetData().Get();
+
+		if (Payload.IsValid())
+		{
+			const int64 CurrentDataLength = Payload->Size();
+
+			// Clone the payload so that we have a local copy that we can
+			// append additional data to.
+			Buffer = FMemory::Malloc(CurrentDataLength, DEFAULT_ALIGNMENT);
+			FMemory::Memcpy(Buffer, Payload->GetData(), CurrentDataLength);
+
+			BufferLength = CurrentDataLength;
+
+			// Start at the end of the existing data
+			CurPos = CurrentDataLength;
+			DataLength = CurrentDataLength;
+		}
+		else
+		{
+			Buffer = nullptr;
+			BufferLength = 0;
+
+			CurPos = 0;
+			DataLength = 0;
+		}
+	}
+
+	~FVirtualizedBulkDataWriter()
+	{
+		// Remove the slack from the allocated bulk data
+		Buffer = FMemory::Realloc(Buffer, DataLength, DEFAULT_ALIGNMENT);
+		BulkData.UpdatePayload(MakeSharedBuffer(FSharedBuffer::AssumeOwnership, Buffer, DataLength));
+	}
+
+	virtual void Serialize(void* Data, int64 Num)
+	{
+		// Determine if we need to reallocate the buffer to fit the next item
+		const int64 NewPos = CurPos + Num;
+		checkf(NewPos >= CurPos, TEXT("Serialization has somehow gone backwards"));
+
+		if (NewPos > BufferLength)
+		{
+			// If so, resize to the new size + 3/8 additional slack
+			const int64 NewLength = NewPos + 3 * NewPos / 8 + 16;
+			Buffer = FMemory::Realloc(Buffer, NewLength, DEFAULT_ALIGNMENT);
+			BufferLength = NewLength;
+		}
+
+		FMemory::Memcpy(static_cast<unsigned char*>(Buffer) + CurPos, Data, Num);
+
+		CurPos += Num;
+		DataLength = FMath::Max(DataLength, CurPos);
+	}
+
+	using FArchive::operator<<; // For visibility of the overloads we don't override
+
+	virtual FArchive& operator<<(class FName& Name) override
+	{
+		// FNames are serialized as strings in BulkData
+		FString StringName = Name.ToString();
+		*this << StringName;
+		return *this;
+	}
+
+	virtual int64 Tell() { return CurPos; }
+	virtual int64 TotalSize() { return DataLength; }
+
+	virtual void Seek(int64 InPos)
+	{
+		check(InPos >= 0);
+		check(InPos <= DataLength);
+		CurPos = InPos;
+	}
+
+	virtual bool AtEnd()
+	{
+		return CurPos >= DataLength;
+	}
+
+	virtual FString GetArchiveName() const
+	{
+		return TEXT("FVirtualizedBulkDataWriter");
+	}
+
+protected:
+	/** The target bulkdata object */
+	FVirtualizedUntypedBulkData& BulkData;
+
+	/** Pointer to the data buffer */
+	void* Buffer;
+	/** Length of the data buffer (stored as bytes) */
+	int64 BufferLength;
+
+	/** Current position in the buffer for serialization operations */
+	int64 CurPos;
+
+	/** The length of valid data in the data buffer (stored as bytes) */
+	int64 DataLength;
+};
+
+
+namespace UE4VirtualizedBulkData_Private
+{
+	/** Wraps access to the FVirtualizedUntypedBulkData data for FVirtualizedBulkDataReader
+	so that it can be done before the FBufferReaderBase constructor is called */
+	class DataAccessWrapper
+	{
+	protected:
+		DataAccessWrapper(FVirtualizedUntypedBulkData& InBulkData)
+			: Payload(InBulkData.GetData().Get())
+
+		{
+
+		}
+
+		virtual ~DataAccessWrapper() = default;
+
+		void* GetData() const
+		{
+			if(Payload.IsValid())
+			{
+				// It's okay to remove the const qualifier here as it will only be passed
+				// on to FBufferReaderBase, which will not change the data at all, but takes a non-const
+				// pointer so that it can free the memory if requested, which we don't.
+				return const_cast<void*>(Payload->GetData());
+			}
+			else
+			{
+				return nullptr;
+			}
+		}
+
+		int64 GetDataLength() const
+		{
+			return Payload.IsValid() ? Payload->Size() : 0;
+		}
+
+	private:
+		FSharedBufferConstPtr Payload;
+	};
+}
+
+class FVirtualizedBulkDataReader :	public UE4VirtualizedBulkData_Private::DataAccessWrapper,
+	public FBufferReaderBase
+{
+public:
+	FVirtualizedBulkDataReader(FVirtualizedUntypedBulkData& InBulkData, bool bIsPersistent = false)
+		: DataAccessWrapper(InBulkData)
+		, FBufferReaderBase(GetData(), GetDataLength(), false, bIsPersistent)
+	{
+	}
+
+	virtual ~FVirtualizedBulkDataReader() = default;
+
+	using FArchive::operator<<; // For visibility of the overloads we don't override
+
+	virtual FArchive& operator<<(class FName& Name) override
+	{
+		// FNames are serialized as strings in BulkData
+		FString StringName;
+		*this << StringName;
+		Name = FName(*StringName);
+		return *this;
+	}
+
+	virtual FString GetArchiveName() const
+	{
+		return TEXT("FVirtualizedBulkDataReader");
+	}	
+};
+
+#endif //WITH_EDITORONLY_DATA
diff --git a/Engine/Source/Runtime/CoreUObject/Public/Virtualization/VirtualizedBulkData.h b/Engine/Source/Runtime/CoreUObject/Public/Virtualization/VirtualizedBulkData.h
new file mode 100644
--- /dev/null
+++ b/Engine/Source/Runtime/CoreUObject/Public/Virtualization/VirtualizedBulkData.h
@@ -0,0 +1,173 @@
+// Copyright Epic Games, Inc. All Rights Reserved.
+
+#pragma once
+
+#include "Async/Future.h"
+#include "HAL/Platform.h"
+#include "Memory/SharedBuffer.h"
+#include "Misc/Guid.h"
+#include "Virtualization/VirtualizationManager.h"
+
+class FArchive;
+class UObject;
+
+//TODO: At some point it might be a good idea to uncomment this to make sure that FVirtualizedUntypedBulkData is
+//		never used at runtime (requires a little too much reworking of some assets for now though)
+//#if WITH_EDITORONLY_DATA
+
+/**
+ * The goal of this class is to provide an editor time version of BulkData that will work with the content
+ * virtualization system.
+ *
+ * Assuming that the DDC is hot, the virtualized payloads are accessed relatively infrequently, usually when the package
+ * is being edited in the editor in some manner. So the payload access is designed around this. If the data is frequently 
+ * accessed when running the editor then the user would not gain from having it virtualized as they would end up pulling
+ * it immediately anyway.
+ *
+ * The biggest difference with normal bulkdata is that the access times might be significantly longer if the
+ * payload is not readily available which is why the only way to access the payload is by a TFuture or a callback
+ * so that the caller is forced to consider how to handle the potential stall and hopefully organize their code 
+ * in such a way that the time lag is not noticeable to the user.
+ *
+ * The second biggest difference is that the caller will own the payload memory once it is returned to them, unlike
+ * the old bulkdata class which would retain ownership. Which forces the calling code to be in control of when the 
+ * memory is actually allocated and for how long. With the old bulkdata class a single access would leave that memory
+ * allocated in a non-obvious way and would lead to memory bloats within the editor.
+ *
+ * The method ::GetGuid can be used to access a unique identifier for the payload, currently it is based on the 
+ * payload itself, so that two objects with the same payload would both have the same Guid. The intent is that we 
+ * would be able to share local copies of the payload between branches to reduce the cost of having multiple branches 
+ * with similar data on the same machine.
+ * 
+ */
+
+/** The base class with no type */
+class COREUOBJECT_API FVirtualizedUntypedBulkData
+{
+public:
+	using OnDataReadyCallback = TUniqueFunction<void(FSharedBufferConstPtr)>;
+
+	FVirtualizedUntypedBulkData() = default;
+	~FVirtualizedUntypedBulkData();
+
+	/** 
+	* Convenience method to make it easier to convert from BulkData to FVirtualizedBulkData and sets the Guid 
+	*
+	* @param BulkData	The bulkdata object to create from.
+	* @param Guid		A guid associated with the bulkdata object which will be used to identify the payload.
+	*					This MUST remain the same between sessions so that the payloads key remains consistent!
+	*/
+	void CreateFromBulkData(FUntypedBulkData& BulkData, const FGuid& Guid);
+
+	void Serialize(FArchive& Ar, UObject* Owner);
+
+	/** Reset to a truly empty state */
+	void Reset();
+
+	// TODO: Probably can just remove this as there probably isn't a good use case for unloading updated payloads as there is no
+	// way for us to restore it. In that case ::Reset might as well be used.
+	/** Unloads the data (if possible) but leaves it in a state where the data can be reloaded */
+	void UnloadData();
+
+	/** Returns the a key that identifies the payload */
+	const FGuid& GetKey() const { return Key; }
+
+	/** Returns the size of the payload in bytes */
+	int64 GetBulkDataSize() const { return PayloadLength; }
+
+	// TODO: This (IsDataLoaded) is only needed for TextureDerivedData.cpp (which assumes that the data 
+	// needs to be loaded in order to be able to run on a background thread. Since FVirtualizedUntypedBulkData 
+	// will aim to be thread safe we could change TextureDerivedData and remove this.
+	/** Returns if the data is being held in memory (true) or will be loaded from disk (false) */
+	bool IsDataLoaded() const { return Payload.IsValid(); }
+
+	/** 
+	 * Enable/Disable if the payload should be compressed if it ends up being serialized to disk on the local
+	 * machine (if virtualization is disabled or if the payload does not meet the virtualization requirements
+	 * of the project for example)
+	 */
+	void EnableDataCompression(bool bFlag) { bShouldCompressData = bFlag;  }
+
+	/**
+	 * Returns an immutable FSharedBuffer reference to the payload data. 
+	 *
+	 * To convert to a mutable reference call'FSharedBufferRef Mutable = MakeSharedBuffer(FSharedBuffer::Clone, Immutable->GetView());'
+	 * which will return a newly allocated clone of the payload which the caller will have full ownership over.
+	 */
+	TFuture<FSharedBufferConstPtr> GetData() const;
+
+	/** Returns an immutable FSharedBuffer reference to the payload data via the OnDataReadyCallback */
+	void GetData(OnDataReadyCallback&& Callback) const;
+		
+	/**
+	 * Allows the existing payload to be replaced with a new one.
+	 *
+	 * To pass in a raw pointer use 'MakeSharedBuffer(...)' to create a valid FSharedBufferConstRef
+	 * Use 'FSharedBuffer::Wrap' if you want to retain ownership on the data being passed in, and use 
+	 * 'FSharedBuffer::AssumeOwnership' if you are okay with the bulkdata object taking over ownership of it.
+	 * The bulkdata object must own it's internal buffer, so if you pass in a non-owned FSharedBuffer (ie 
+	 * by using 'FSharedBuffer::Wrap') then a clone of the data will be created internally and assigned to
+	 * the bulkdata object.
+	 */
+	void UpdatePayload(const FSharedBufferConstRef& InPayload);
+
+private:
+
+	FSharedBufferConstPtr GetDataInternal() const;
+
+	void CalculateKey();
+
+	FSharedBufferConstPtr LoadFromDisk() const;
+	bool SerializeData(FArchive& Ar, FSharedBufferConstPtr& Payload) const;
+
+	void PushData();
+	FSharedBufferConstPtr PullData() const;
+
+	FString GetFilenameFromOwner(UObject* Owner) const;
+
+	bool CanUnloadData() const;
+
+	/** Unique identifier for the payload*/
+	FGuid Key;
+
+	/** Pointer to the payload if it is held in memory (it has been updated but not yet saved to disk for example) */
+	FSharedBufferConstPtr Payload;
+
+	/** Length of the payload in bytes */
+	int64 PayloadLength = 0;
+
+	/** Is the data actually virtualized or not? */
+	bool bIsVirtualized = false;
+	
+	//---- The remaining members are used when the payload is not virtualized.
+	
+	/** Should we compress the data next time (don't want to serialize this)
+	TODO: Not currently affected by ::Reset which is an odd design and needs fixing */
+	bool bShouldCompressData = false;
+
+	/** If the data has been saved to disk rather than being Virtualized, was it stored in a compressed format? */
+	bool bIsDataStoredAsCompressed = false ;
+
+	/** File containing the payload (this will be empty if the payload does not come from a file)*/
+	FString Filename;
+
+	/** Offset of the payload in the file that contains it (INDEX_NONE if the payload does not come from a file)*/
+	int64 OffsetInFile = INDEX_NONE;
+};
+
+//TODO: Probably remove this and change FVirtualizedUntypedBulkData to always be TVirtualizedBulkData<uint8>
+/** Type safe versions */
+template<typename DataType>
+class TVirtualizedBulkData final : public FVirtualizedUntypedBulkData
+{
+public:
+	TVirtualizedBulkData() = default;
+	~TVirtualizedBulkData() = default;
+};
+
+using FByteVirtualizedBulkData	= TVirtualizedBulkData<uint8>;
+using FWordVirtualizedBulkData	= TVirtualizedBulkData<uint16>;
+using FIntVirtualizedBulkData	= TVirtualizedBulkData<int32>;
+using FFloatVirtualizedBulkData	= TVirtualizedBulkData<float>;
+
+//#endif //WITH_EDITORONLY_DATA
diff --git a/Engine/Source/Runtime/Engine/Classes/Engine/World.h b/Engine/Source/Runtime/Engine/Classes/Engine/World.h
--- a/Engine/Source/Runtime/Engine/Classes/Engine/World.h
+++ b/Engine/Source/Runtime/Engine/Classes/Engine/World.h
@@ -3762,7 +3762,7 @@ public:
 	virtual void UpdateConstraintActors();
 
 	/** Gets all LightMaps and ShadowMaps associated with this world. Specify the level or leave null for persistent */
-	void GetLightMapsAndShadowMaps(ULevel* Level, TArray<UTexture2D*>& OutLightMapsAndShadowMaps);
+	void GetLightMapsAndShadowMaps(ULevel* Level, TArray<UTexture2D*>& OutLightMapsAndShadowMaps, bool bForceLazyLoad = true);
 
 public:
 	/** Rename this world such that it has the prefix on names for the given PIE Instance ID */
diff --git a/Engine/Source/Runtime/Engine/Classes/GameFramework/Actor.h b/Engine/Source/Runtime/Engine/Classes/GameFramework/Actor.h
--- a/Engine/Source/Runtime/Engine/Classes/GameFramework/Actor.h
+++ b/Engine/Source/Runtime/Engine/Classes/GameFramework/Actor.h
@@ -1886,31 +1886,31 @@ public:
 	/** Returns the transform of the RootComponent of this Actor*/ 
 	FORCEINLINE const FTransform& GetActorTransform() const
 	{
-		return TemplateGetActorTransform(RootComponent);
+		return TemplateGetActorTransform(ToRawPtr(RootComponent));
 	}
 
 	/** Returns the location of the RootComponent of this Actor*/ 
 	FORCEINLINE FVector GetActorLocation() const
 	{
-		return TemplateGetActorLocation(RootComponent);
+		return TemplateGetActorLocation(ToRawPtr(RootComponent));
 	}
 
 	/** Returns the rotation of the RootComponent of this Actor */
 	FORCEINLINE FRotator GetActorRotation() const
 	{
-		return TemplateGetActorRotation(RootComponent);
+		return TemplateGetActorRotation(ToRawPtr(RootComponent));
 	}
 
 	/** Returns the scale of the RootComponent of this Actor */
 	FORCEINLINE FVector GetActorScale() const
 	{
-		return TemplateGetActorScale(RootComponent);
+		return TemplateGetActorScale(ToRawPtr(RootComponent));
 	}
 
 	/** Returns the quaternion of the RootComponent of this Actor */
 	FORCEINLINE FQuat GetActorQuat() const
 	{
-		return TemplateGetActorQuat(RootComponent);
+		return TemplateGetActorQuat(ToRawPtr(RootComponent));
 	}
 
 #if WITH_EDITOR
@@ -3690,17 +3690,17 @@ FORCEINLINE_DEBUGGABLE USceneComponent* AActor::K2_GetRootComponent() const
 
 FORCEINLINE_DEBUGGABLE FVector AActor::GetActorForwardVector() const
 {
-	return TemplateGetActorForwardVector(RootComponent);
+	return TemplateGetActorForwardVector(ToRawPtr(RootComponent));
 }
 
 FORCEINLINE_DEBUGGABLE FVector AActor::GetActorUpVector() const
 {
-	return TemplateGetActorUpVector(RootComponent);
+	return TemplateGetActorUpVector(ToRawPtr(RootComponent));
 }
 
 FORCEINLINE_DEBUGGABLE FVector AActor::GetActorRightVector() const
 {
-	return TemplateGetActorRightVector(RootComponent);
+	return TemplateGetActorRightVector(ToRawPtr(RootComponent));
 }
 
 
diff --git a/Engine/Source/Runtime/Engine/Private/Animation/AnimStreamable.cpp b/Engine/Source/Runtime/Engine/Private/Animation/AnimStreamable.cpp
--- a/Engine/Source/Runtime/Engine/Private/Animation/AnimStreamable.cpp
+++ b/Engine/Source/Runtime/Engine/Private/Animation/AnimStreamable.cpp
@@ -362,7 +362,7 @@ void UAnimStreamable::PostLoad()
 	Super::PostLoad();
 
 #if WITH_EDITOR
-	if (UAnimSequence* NonConstSeq = const_cast<UAnimSequence*>(SourceSequence))
+	if (UAnimSequence* NonConstSeq = const_cast<UAnimSequence*>(ToRawPtr(SourceSequence)))
 	{
 		if (FLinkerLoad* Linker = NonConstSeq->GetLinker())
 		{
diff --git a/Engine/Source/Runtime/Engine/Private/Distributions.cpp b/Engine/Source/Runtime/Engine/Private/Distributions.cpp
--- a/Engine/Source/Runtime/Engine/Private/Distributions.cpp
+++ b/Engine/Source/Runtime/Engine/Private/Distributions.cpp
@@ -106,7 +106,7 @@ static void BuildZeroLookupTable( FDistributionLookupTable* OutTable, int32 Valu
  * @param Distribution - The distribution for which to build a lookup table.
  */
 template <typename DistributionType>
-void BuildLookupTable( FDistributionLookupTable* OutTable, const DistributionType* Distribution )
+void BuildLookupTable( FDistributionLookupTable* OutTable, const DistributionType& Distribution )
 {
 	check(IsInGameThread() || IsInAsyncLoadingThread());
 	check(Distribution);
diff --git a/Engine/Source/Runtime/Engine/Private/StaticMesh.cpp b/Engine/Source/Runtime/Engine/Private/StaticMesh.cpp
--- a/Engine/Source/Runtime/Engine/Private/StaticMesh.cpp
+++ b/Engine/Source/Runtime/Engine/Private/StaticMesh.cpp
@@ -4268,9 +4268,12 @@ void UStaticMesh::ClearMeshDescriptions()
 // If static mesh derived data needs to be rebuilt (new format, serialization
 // differences, etc.) replace the version GUID below with a new one.
 // In case of merge conflicts with DDC versions, you MUST generate a new GUID
-// and set this new GUID as the version.                                       
+// and set this new GUID as the version.
+#if UE_USE_VIRTUALBULKDATA
+	#define MESHDATAKEY_STATICMESH_DERIVEDDATA_VER TEXT("1F22FDD3B8C54996AAD54400037B85DF")
+#else
 #define MESHDATAKEY_STATICMESH_DERIVEDDATA_VER TEXT("BEDE1E3CC35D4D8FA944458825A18AFB")
-
+#endif
 
 static const FString& GetMeshDataKeyStaticMeshDerivedDataVersion()
 {
diff --git a/Engine/Source/Runtime/Engine/Private/UnrealEngine.cpp b/Engine/Source/Runtime/Engine/Private/UnrealEngine.cpp
--- a/Engine/Source/Runtime/Engine/Private/UnrealEngine.cpp
+++ b/Engine/Source/Runtime/Engine/Private/UnrealEngine.cpp
@@ -2482,6 +2482,20 @@ static void LoadEngineTexture(TextureType*& InOutTexture, const TCHAR* InName)
 	}
 }
 
+template <typename TextureType>
+static void LoadEngineTexture(TObjectPtr<TextureType>& InOutTexture, const TCHAR* InName)
+{
+	// Fully copied from specialization above to allow for mutable pointer access without relying on deprecated explicit cast to mutable
+	if (!InOutTexture)
+	{
+		InOutTexture = LoadObject<TextureType>(nullptr, InName, nullptr, LOAD_None, nullptr);
+	}
+	if (FPlatformProperties::RequiresCookedData() && InOutTexture)
+	{
+		InOutTexture->AddToRoot();
+	}
+}
+
 static void LoadCustomTimeStep(UEngine* Engine)
 {
 	if (Engine->CustomTimeStepClassName.IsValid())
diff --git a/Engine/Source/Runtime/Engine/Private/World.cpp b/Engine/Source/Runtime/Engine/Private/World.cpp
--- a/Engine/Source/Runtime/Engine/Private/World.cpp
+++ b/Engine/Source/Runtime/Engine/Private/World.cpp
@@ -7399,16 +7399,18 @@ void UWorld::CopyGameState(AGameModeBase* FromGameMode, AGameStateBase* FromGame
 	SetGameState(FromGameState);
 }
 
-void UWorld::GetLightMapsAndShadowMaps(ULevel* Level, TArray<UTexture2D*>& OutLightMapsAndShadowMaps)
+void UWorld::GetLightMapsAndShadowMaps(ULevel* Level, TArray<UTexture2D*>& OutLightMapsAndShadowMaps, bool bForceLazyLoad /*= true*/)
 {
 	class FFindLightmapsArchive : public FArchiveUObject
 	{
 		/** The array of textures discovered */
 		TArray<UTexture2D*>& TextureList;
+		bool bForceLazyLoad;
 
 	public:
-		FFindLightmapsArchive(UObject* InSearch, TArray<UTexture2D*>& OutTextureList)
+		FFindLightmapsArchive(UObject* InSearch, TArray<UTexture2D*>& OutTextureList, bool bInForceLazyLoad)
 			: TextureList(OutTextureList)
+			, bForceLazyLoad(bInForceLazyLoad)
 		{
 			ArIsObjectReferenceCollector = true;
 			ArIsModifyingWeakAndStrongReferences = true; // While we are not modifying them, we want to follow weak references as well
@@ -7443,6 +7445,33 @@ void UWorld::GetLightMapsAndShadowMaps(ULevel* Level, TArray<UTexture2D*>& OutLi
 
 			return *this;
 		}
+
+		FArchive& operator<<(FObjectPtr& Obj)
+		{
+			// @TODO: OBJPTR: Could some or all of this behavior be generalized for use in other reference collectors?
+			//			Could add another Ar* flag to control whether lazy loads get resolved.  Could add a derivative
+			//			of FArchiveUObject that filters references by type.
+
+			// Don't check null references or objects already visited. Also, skip UWorlds as they will pull in more levels than desired
+			if (!Obj.IsNull() && !Obj.IsA<UWorld>())
+			{
+				if (Obj.IsA<ULightMapTexture2D>() ||
+					Obj.IsA<UShadowMapTexture2D>() ||
+					Obj.IsA<ULightMapVirtualTexture2D>())
+				{
+					UTexture2D* Tex = Cast<UTexture2D>(Obj.Get());
+					if (ensure(Tex))
+					{
+						TextureList.Add(Tex);
+					}
+				}
+				else if (IsObjectHandleResolved(Obj.GetHandle()) || bForceLazyLoad)
+				{
+					return FArchiveUObject::operator<<(Obj);
+				}
+			}
+			return *this;
+		}
 	};
 
 	UObject* SearchObject = Level;
@@ -7451,7 +7480,7 @@ void UWorld::GetLightMapsAndShadowMaps(ULevel* Level, TArray<UTexture2D*>& OutLi
 		SearchObject = PersistentLevel;
 	}
 
-	FFindLightmapsArchive FindArchive(SearchObject, OutLightMapsAndShadowMaps);
+	FFindLightmapsArchive FindArchive(SearchObject, OutLightMapsAndShadowMaps, bForceLazyLoad);
 }
 
 void UWorld::CreateFXSystem()
diff --git a/Engine/Source/Runtime/Engine/Private/WorldSettings.cpp b/Engine/Source/Runtime/Engine/Private/WorldSettings.cpp
--- a/Engine/Source/Runtime/Engine/Private/WorldSettings.cpp
+++ b/Engine/Source/Runtime/Engine/Private/WorldSettings.cpp
@@ -359,7 +359,7 @@ PRAGMA_DISABLE_DEPRECATION_WARNINGS
 	{
 		if (Ar.CustomVer(FEnterpriseObjectVersion::GUID) < FEnterpriseObjectVersion::BookmarkExtensibilityUpgrade)
 		{
-			UBookmarkBase** LocalBookmarks = reinterpret_cast<UBookmarkBase**>(static_cast<UBookMark**>(BookMarks)); //-V777
+			UBookmarkBase** LocalBookmarks = reinterpret_cast<UBookmarkBase**>(static_cast<UBookMark**>(ToRawPtrArrayUnsafe(BookMarks))); //-V777
 			const int32 NumBookmarks = sizeof(BookMarks) / sizeof(UBookMark*);
 			BookmarkArray = TArray<UBookmarkBase*>(LocalBookmarks, NumBookmarks);
 			AdjustNumberOfBookmarks();
diff --git a/Engine/Source/Runtime/Experimental/GeometryCollectionEngine/Private/GeometryCollection/GeometryCollectionComponent.cpp b/Engine/Source/Runtime/Experimental/GeometryCollectionEngine/Private/GeometryCollection/GeometryCollectionComponent.cpp
--- a/Engine/Source/Runtime/Experimental/GeometryCollectionEngine/Private/GeometryCollection/GeometryCollectionComponent.cpp
+++ b/Engine/Source/Runtime/Experimental/GeometryCollectionEngine/Private/GeometryCollection/GeometryCollectionComponent.cpp
@@ -1738,7 +1738,7 @@ void UGeometryCollectionComponent::OnCreatePhysicsState()
 		if (RestCollection)
 		{
 			//hack: find a better place for this
-			UGeometryCollection* RestCollectionMutable = const_cast<UGeometryCollection*>(RestCollection);
+			UGeometryCollection* RestCollectionMutable = const_cast<UGeometryCollection*>(ToRawPtr(RestCollection));
 			RestCollectionMutable->EnsureDataIsCooked();
 		}
 #endif
@@ -2154,7 +2154,7 @@ UGeometryCollection* FGeometryCollectionEdit::GetRestCollection()
 {
 	if (Component)
 	{
-		return const_cast<UGeometryCollection*>(Component->RestCollection);	//const cast is ok here since we are explicitly in edit mode. Should all this editor code be in an editor module?
+		return const_cast<UGeometryCollection*>(ToRawPtr(Component->RestCollection));	//const cast is ok here since we are explicitly in edit mode. Should all this editor code be in an editor module?
 	}
 	return nullptr;
 }
diff --git a/Engine/Source/Runtime/Experimental/GeometryCollectionEngine/Private/GeometryCollection/GeometryCollectionDebugDrawActor.cpp b/Engine/Source/Runtime/Experimental/GeometryCollectionEngine/Private/GeometryCollection/GeometryCollectionDebugDrawActor.cpp
--- a/Engine/Source/Runtime/Experimental/GeometryCollectionEngine/Private/GeometryCollection/GeometryCollectionDebugDrawActor.cpp
+++ b/Engine/Source/Runtime/Experimental/GeometryCollectionEngine/Private/GeometryCollection/GeometryCollectionDebugDrawActor.cpp
@@ -497,8 +497,8 @@ void AGeometryCollectionDebugDrawActor::UpdatePropertyValue(T1& PropertyValue, c
 	}
 }
 
-template<>
-void AGeometryCollectionDebugDrawActor::UpdatePropertyValue<AChaosSolverActor*, FString>(AChaosSolverActor*& PropertyValue, const TAutoConsoleVariable<FString>& ConsoleVariable, bool& bHasChanged)
+template<typename T, decltype(ImplicitConv<AChaosSolverActor*>(std::declval<T>()))*>
+void AGeometryCollectionDebugDrawActor::UpdatePropertyValue(T& PropertyValue, const TAutoConsoleVariable<FString>& ConsoleVariable, bool& bHasChanged)
 {
 	AChaosSolverActor* NewValue = nullptr;
 
diff --git a/Engine/Source/Runtime/Experimental/GeometryCollectionEngine/Public/GeometryCollection/GeometryCollectionDebugDrawActor.h b/Engine/Source/Runtime/Experimental/GeometryCollectionEngine/Public/GeometryCollection/GeometryCollectionDebugDrawActor.h
--- a/Engine/Source/Runtime/Experimental/GeometryCollectionEngine/Public/GeometryCollection/GeometryCollectionDebugDrawActor.h
+++ b/Engine/Source/Runtime/Experimental/GeometryCollectionEngine/Public/GeometryCollection/GeometryCollectionDebugDrawActor.h
@@ -5,6 +5,8 @@
 #include "GameFramework/Actor.h"
 #include "Components/BillboardComponent.h"
 
+#include <type_traits>
+
 #include "GeometryCollectionDebugDrawActor.generated.h"
 
 class AChaosSolverActor;
@@ -504,6 +506,10 @@ private:
 	template<typename T1, typename T2>
 	void UpdatePropertyValue(T1& PropertyValue, const TAutoConsoleVariable<T2>& ConsoleVariable, bool& bHasChanged);
 
+	/** Property update function helper. */
+	template<typename T, decltype(ImplicitConv<AChaosSolverActor*>(std::declval<T>()))* = nullptr>
+	void UpdatePropertyValue(T& PropertyValue, const TAutoConsoleVariable<FString>& ConsoleVariable, bool& bHasChanged);
+
 	/** Callback on console variable. */
 	void OnCVarsChanged();
 
