Avoid warning emitted from IncludeTool by keeping FObjectPtr a 'using' in all configurations.

diff --git a/Engine/Source/Runtime/CoreUObject/Public/UObject/ObjectHandle.h b/Engine/Source/Runtime/CoreUObject/Public/UObject/ObjectHandle.h
--- a/Engine/Source/Runtime/CoreUObject/Public/UObject/ObjectHandle.h
+++ b/Engine/Source/Runtime/CoreUObject/Public/UObject/ObjectHandle.h
@@ -67,10 +67,12 @@ inline uint32 GetTypeHash(FPackedObjectRef ObjectRef) { return GetTypeHash(Objec
  */
 #if UE_WITH_OBJECT_HANDLE_LATE_RESOLVE
 
-struct FObjectHandle
+struct FObjectHandleInternal
 {
 	UPTRINT PointerOrRef;
 };
+using FObjectHandle = FObjectHandleInternal;
+
 inline bool operator==(FObjectHandle LHS, FObjectHandle RHS);
 inline bool operator!=(FObjectHandle LHS, FObjectHandle RHS);
 inline uint32 GetTypeHash(FObjectHandle Handle);
