Fixes for HL2 compile.
- ObjectPtrTest was just checking for windows.  HL2 falls into the same category
- D3D12RHI.cs now has AGS properly so that D3D12RootSignature includes process correctly

diff --git a/Engine/Source/Runtime/CoreUObject/Private/Tests/ObjectPtrTest.cpp b/Engine/Source/Runtime/CoreUObject/Private/Tests/ObjectPtrTest.cpp
--- a/Engine/Source/Runtime/CoreUObject/Private/Tests/ObjectPtrTest.cpp
+++ b/Engine/Source/Runtime/CoreUObject/Private/Tests/ObjectPtrTest.cpp
@@ -84,7 +84,7 @@ static_assert(TModels<CEqualityComparableWith, FMutableObjectPtr, FMutableInterf
 
 // Ensure that TObjectPtr<[const] UPackage> is not comparable with a TObjectPtr<[const] UInterface> regardless of constness
 // TODO: This only ensures that at least one of the A==B,B==A,A!=B,B!=A operations fail, not that they all fail.
-#if !PLATFORM_WINDOWS || !defined(_MSC_EXTENSIONS) // MSVC static analyzer is run in non-conformance mode, and that causes these checks to fail.
+#if !(PLATFORM_WINDOWS || PLATFORM_HOLOLENS) || !defined(_MSC_EXTENSIONS) // MSVC static analyzer is run in non-conformance mode, and that causes these checks to fail.
 static_assert(!TModels<CEqualityComparableWith, FConstPackagePtr, FConstInterfacePtr>::Value, "Must not be able to compare equality and inequality bidirectionally between TObjectPtr<const UPackage> and TObjectPtr<const UInterface>");
 static_assert(!TModels<CEqualityComparableWith, FMutablePackagePtr, FConstInterfacePtr>::Value, "Must not be able to compare equality and inequality bidirectionally between TObjectPtr<UPackage> and TObjectPtr<const UInterface>");
 static_assert(!TModels<CEqualityComparableWith, FConstPackagePtr, FMutableInterfacePtr>::Value, "Must not be able to compare equality and inequality bidirectionally between TObjectPtr<const UPackage> and TObjectPtr<UInterface>");
@@ -127,7 +127,7 @@ static_assert(!TModels<CEqualityComparableWith, FMutableObjectPtr, char*>::Value
 static_assert(TModels<CEqualityComparableWith, FConstObjectPtr, TYPE_OF_NULLPTR>::Value, "Must be able to compare equality and inequality bidirectionally between TObjectPtr<const UObject> and nullptr");
 static_assert(TModels<CEqualityComparableWith, FMutableObjectPtr, TYPE_OF_NULLPTR>::Value, "Must be able to compare equality and inequality bidirectionally between TObjectPtr<UObject> and nullptr");
 
-#if !PLATFORM_WINDOWS || !defined(_MSC_EXTENSIONS) // MSVC static analyzer is run in non-conformance mode, and that causes these checks to fail.
+#if !(PLATFORM_WINDOWS || PLATFORM_HOLOLENS) || !defined(_MSC_EXTENSIONS) // MSVC static analyzer is run in non-conformance mode, and that causes these checks to fail.
 static_assert(!TModels<CEqualityComparableWith, FConstObjectPtr, long>::Value, "Should not be able to compare equality and inequality bidirectionally between TObjectPtr<const UObject> and long");
 static_assert(!TModels<CEqualityComparableWith, FMutableObjectPtr, long>::Value, "Should not be able to compare equality and inequality bidirectionally between TObjectPtr<UObject> and long");
 #endif // _MSC_EXTENSIONS
diff --git a/Engine/Source/Runtime/D3D12RHI/D3D12RHI.Build.cs b/Engine/Source/Runtime/D3D12RHI/D3D12RHI.Build.cs
--- a/Engine/Source/Runtime/D3D12RHI/D3D12RHI.Build.cs
+++ b/Engine/Source/Runtime/D3D12RHI/D3D12RHI.Build.cs
@@ -42,10 +42,10 @@ public class D3D12RHI : ModuleRules
 		{
 			AddEngineThirdPartyPrivateStaticDependencies(Target, "DX12");
 			AddEngineThirdPartyPrivateStaticDependencies(Target, "DX11");
-            if (Target.Platform != UnrealTargetPlatform.HoloLens)
+			AddEngineThirdPartyPrivateStaticDependencies(Target, "AMD_AGS");
+			if (Target.Platform != UnrealTargetPlatform.HoloLens)
             {
                 AddEngineThirdPartyPrivateStaticDependencies(Target, "NVAPI");
-				AddEngineThirdPartyPrivateStaticDependencies(Target, "AMD_AGS");
             	AddEngineThirdPartyPrivateStaticDependencies(Target, "NVAftermath");
             	AddEngineThirdPartyPrivateStaticDependencies(Target, "IntelMetricsDiscovery");
             }
