Improve output of UObject referenced by FGCObject

diff --git a/Engine/Source/Runtime/CoreUObject/Private/UObject/Obj.cpp b/Engine/Source/Runtime/CoreUObject/Private/UObject/Obj.cpp
--- a/Engine/Source/Runtime/CoreUObject/Private/UObject/Obj.cpp
+++ b/Engine/Source/Runtime/CoreUObject/Private/UObject/Obj.cpp
@@ -44,7 +44,7 @@
 #include "Misc/RedirectCollector.h"
 #include "UObject/GCScopeLock.h"
 #include "ProfilingDebugging/LoadTimeTracker.h"
-
+#include "UObject/GCObject.h"
 #include "Serialization/ArchiveUObjectFromStructuredArchive.h"
 #include "Serialization/ArchiveDescribeReference.h"
 #include "UObject/FindStronglyConnected.h"
@@ -2905,6 +2905,7 @@ void UObject::OutputReferencers( FOutputDevice& Ar, FReferencerInformationList*
 					ObjectReachability += TEXT(" (standalone)");
 				}
 
+				UGCObjectReferencer* GCObjectReferencer = Cast<UGCObjectReferencer>(RefInfo.Referencer);
 				Ar.Logf( TEXT("   %s (%i)\r\n"), *ObjectReachability, RefInfo.TotalReferences );
 				for ( int32 i = 0; i < RefInfo.TotalReferences; i++ )
 				{
@@ -2915,7 +2916,15 @@ void UObject::OutputReferencers( FOutputDevice& Ar, FReferencerInformationList*
 					}
 					else
 					{
-						Ar.Logf(TEXT("      %i) [[native reference]]\r\n"), i);
+						FString ReferencerName;
+						if (GCObjectReferencer && GCObjectReferencer->GetReferencerName(this, ReferencerName))
+						{
+							Ar.Logf(TEXT("      %i) [[FGCObject %s]]\r\n"), i, *ReferencerName);
+						}
+						else
+						{
+							Ar.Logf(TEXT("      %i) [[native reference]]\r\n"), i);
+						}
 					}
 				}
 			}
