Enable package access tracking for use in cook-time dependency tracking.

diff --git a/Engine/Source/Runtime/CoreUObject/Public/Misc/PackageAccessTracking.h b/Engine/Source/Runtime/CoreUObject/Public/Misc/PackageAccessTracking.h
--- a/Engine/Source/Runtime/CoreUObject/Public/Misc/PackageAccessTracking.h
+++ b/Engine/Source/Runtime/CoreUObject/Public/Misc/PackageAccessTracking.h
@@ -3,17 +3,15 @@
 #pragma once
 
 #include "CoreTypes.h"
-
-#define UE_WITH_PACKAGE_ACCESS_TRACKING 0 // TODO: Should be UE_WITH_OBJECT_HANDLE_TRACKING
-
-#if UE_WITH_PACKAGE_ACCESS_TRACKING
-
 #include "Misc/Optional.h"
 #include "Misc/PackageName.h"
 #include "UObject/NameTypes.h"
 #include "UObject/ObjectHandle.h"
 #include "UObject/Package.h"
 
+#define UE_WITH_PACKAGE_ACCESS_TRACKING UE_WITH_OBJECT_HANDLE_TRACKING
+
+#if UE_WITH_PACKAGE_ACCESS_TRACKING
 #define UE_TRACK_REFERENCING_PACKAGE_SCOPED(Package, OpName) PackageAccessTracking_Private::FPackageAccessRefScope ANONYMOUS_VARIABLE(PackageAccessTracker_)(Package, OpName);
 #define UE_TRACK_REFERENCING_PACKAGE_DELAYED_SCOPED(TrackerName, OpName) TOptional<PackageAccessTracking_Private::FPackageAccessRefScope> TrackerName; FName TrackerName##_OpName(OpName);
 #define UE_TRACK_REFERENCING_PACKAGE_DELAYED(TrackerName, Package) if (TrackerName) TrackerName->SetPackageName(Package->GetFName()); else TrackerName.Emplace(Package->GetFName(), TrackerName##_OpName);
