Small Misc changes for Save2:
* Modified EnableNewSave cvar to use multiple values: 0: disabled, 1: enabled for cook, 2: enabled for uncooked, 3: enabled everywhere
* Added proper Prestream Packages support to Save2 since it isn't something that will be deprecated in the end
* Fixed new package not being saved properly when passing in a Base Object because they aren't properly marked as fully loaded in Save1
* Added a way for SaveConcurrent to use serial save with Save1 for now when not being passed in the SAVE_Concurrent flag
* Switched Save2 traces to use the SaveTimeChannel like Save1

diff --git a/Engine/Source/Editor/UnrealEd/Private/Commandlets/SavePackageUtilitiesCommandlet.cpp b/Engine/Source/Editor/UnrealEd/Private/Commandlets/SavePackageUtilitiesCommandlet.cpp
--- a/Engine/Source/Editor/UnrealEd/Private/Commandlets/SavePackageUtilitiesCommandlet.cpp
+++ b/Engine/Source/Editor/UnrealEd/Private/Commandlets/SavePackageUtilitiesCommandlet.cpp
@@ -25,26 +25,11 @@ int32 USavePackageUtilitiesCommandlet::Main(const FString& Params)
 {
 	InitParameters(Params);
 
-	auto GetPackageAsset = [](UPackage* InPackage) -> UObject*
-	{
-		UObject* Asset = nullptr;
-		ForEachObjectWithPackage(InPackage, [&Asset](UObject* Object)
-			{
-				if (Object->IsAsset())
-				{
-					Asset = Object;
-					return false;
-				}
-				return true;
-			}, /*bIncludeNestedObjects*/ false);
-		return Asset;
-	};
-
 	for (const FString& PackageName : PackageNames)
 	{
 		// Load Package
 		UPackage* Package = LoadPackage(nullptr, *PackageName, LOAD_None);
-		UObject* Asset = GetPackageAsset(Package);
+		UObject* Asset = Package->FindAssetInPackage();
 		FString Filename = FPaths::CreateTempFilename(*FPaths::ProjectSavedDir());
 
 		FSavePackageArgs SaveArgs;
@@ -65,12 +50,11 @@ int32 USavePackageUtilitiesCommandlet::Main(const FString& Params)
 		// New Save Package
 		FSavePackageResultStruct NewResult;
 		{
-			EnableNewSave->Set(1);
+			EnableNewSave->Set(3); // Enable new save cooked and uncooked data
 			NewResult = GEditor->Save(Package, Asset, SaveArgs.TopLevelFlags, *Filename,
 				GError, nullptr, false, true, SaveArgs.SaveFlags, SaveArgs.TargetPlatform,
 				FDateTime::MinValue(), false, /*DiffMap*/ nullptr,
 				nullptr);
-
 		}
 
 		// Old Save Package
diff --git a/Engine/Source/Editor/UnrealEd/Private/CookOnTheFlyServer.cpp b/Engine/Source/Editor/UnrealEd/Private/CookOnTheFlyServer.cpp
--- a/Engine/Source/Editor/UnrealEd/Private/CookOnTheFlyServer.cpp
+++ b/Engine/Source/Editor/UnrealEd/Private/CookOnTheFlyServer.cpp
@@ -3700,14 +3700,15 @@ void UCookOnTheFlyServer::SaveCookedPackage(UE::Cook::FPackageData& PackageData,
 					if (DiffModeHelper.IsRunningCookLinkerDiff())
 					{
 						static IConsoleVariable* EnableNewSave = IConsoleManager::Get().FindConsoleVariable(TEXT("SavePackage.EnableNewSave"));
-						bool bPreviousCvarValue = EnableNewSave->GetBool();
+						int32 PreviousCvarValue = EnableNewSave->GetInt();
 
 						// If the linker diff is comparing the two save algo, switch the cvar to the other algo before saving again,
 						// Otherwise the linker diff mode is tracking if the save is consistent across multiple save
 						bool bAlgoLinkerDiff = DiffModeHelper.GetLinkerDiffMode() == FDiffModeCookServerUtils::LDM_Algo;
 						if (bAlgoLinkerDiff)
 						{
-							EnableNewSave->Set(!bPreviousCvarValue);
+							// see CVarEnablePackageNewSave definition in SavePackageUtilities.cpp for value meaning
+							EnableNewSave->Set(PreviousCvarValue & 1 ? 0 : 1);
 						}
 
 						FSavePackageResultStruct NewResult = GEditor->Save(Package, World, FlagsToCook, *PlatFilename,
@@ -3717,7 +3718,7 @@ void UCookOnTheFlyServer::SaveCookedPackage(UE::Cook::FPackageData& PackageData,
 
 						if (bAlgoLinkerDiff)
 						{
-							EnableNewSave->Set(bPreviousCvarValue);
+							EnableNewSave->Set(PreviousCvarValue);
 						}
 
 						if (Result.LinkerSave && NewResult.LinkerSave)
diff --git a/Engine/Source/Runtime/CoreUObject/Private/UObject/SavePackage.cpp b/Engine/Source/Runtime/CoreUObject/Private/UObject/SavePackage.cpp
--- a/Engine/Source/Runtime/CoreUObject/Private/UObject/SavePackage.cpp
+++ b/Engine/Source/Runtime/CoreUObject/Private/UObject/SavePackage.cpp
@@ -73,13 +73,6 @@
 #include "ProfilingDebugging/ScopedTimers.h"
 #endif
 
-#if UE_TRACE_ENABLED && !UE_BUILD_SHIPPING
-UE_TRACE_CHANNEL(SaveTimeChannel)
-#define SCOPED_SAVETIMER(TimerName) TRACE_CPUPROFILER_EVENT_SCOPE_ON_CHANNEL(TimerName, SaveTimeChannel)
-#else
-#define SCOPED_SAVETIMER(TimerName)
-#endif
-
 static FCriticalSection InitializeCoreClassesCritSec;
 
 #define VALIDATE_INITIALIZECORECLASSES 0
@@ -2150,16 +2143,12 @@ public:
 
 #endif
 
-// Switch to SavePackage to Save2
-static TAutoConsoleVariable<int32> CVarEnablePackageNewSave(TEXT("SavePackage.EnableNewSave"), 0, TEXT("Enable new package save mechanism over the old one."));
-
 FSavePackageResultStruct UPackage::Save(UPackage* InOuter, UObject* Base, EObjectFlags TopLevelFlags, const TCHAR* Filename,
 	FOutputDevice* Error, FLinkerNull* ConformNO, bool bForceByteSwapping, bool bWarnOfLongFilename, uint32 SaveFlags,
 	const class ITargetPlatform* TargetPlatform, const FDateTime& FinalTimeStamp, bool bSlowTask, FArchiveDiffMap* InOutDiffMap,
 	FSavePackageContext* SavePackageContext)
 {
-
-	if (CVarEnablePackageNewSave.GetValueOnAnyThread() > 0)
+	if (SavePackageUtilities::IsNewSaveEnabled(TargetPlatform != nullptr))
 	{
 		FSavePackageArgs SaveArgs = { (ITargetPlatform*)TargetPlatform, TopLevelFlags, SaveFlags, bForceByteSwapping, bWarnOfLongFilename, bSlowTask, FinalTimeStamp, Error, InOutDiffMap, SavePackageContext };
 		return UPackage::Save2(InOuter, Base, Filename, SaveArgs);
@@ -2279,7 +2268,8 @@ FSavePackageResultStruct UPackage::Save(UPackage* InOuter, UObject* Base, EObjec
 		TAsyncWorkSequence<FMD5> AsyncWriteAndHashSequence;
 
 		// Make sure package is fully loaded before saving. 
-		if (!Base && !InOuter->IsFullyLoaded())
+		// IsFullyLoaded has important byproducts for new packages, so make sure we evaluate it rather than short-circuiting it
+		if (!InOuter->IsFullyLoaded() && !Base)
 		{
 			if (!(SaveFlags & SAVE_NoError))
 			{
diff --git a/Engine/Source/Runtime/CoreUObject/Private/UObject/SavePackage/PackageHarvester.cpp b/Engine/Source/Runtime/CoreUObject/Private/UObject/SavePackage/PackageHarvester.cpp
--- a/Engine/Source/Runtime/CoreUObject/Private/UObject/SavePackage/PackageHarvester.cpp
+++ b/Engine/Source/Runtime/CoreUObject/Private/UObject/SavePackage/PackageHarvester.cpp
@@ -230,7 +230,6 @@ void FPackageHarvester::ProcessExport(UObject* InObject)
 			}
 		}
 
-		//@todo FH: Is this even useful anymore!
 		if (SaveContext.IsProcessingPrestreamingRequests())
 		{
 			Deps.Reset();
@@ -298,13 +297,11 @@ void FPackageHarvester::TryHarvestImport(UObject* InObject)
 	if (!bExcludePackageFromCook && !bExcluded && !SaveContext.IsUnsaveable(InObject))
 	{
 		bool bIsNative = IsObjNative(InObject);
-		SaveContext.AddImport(InObject);
+		bool bIsEditorOnly = false;
 #if WITH_EDITORONLY_DATA
-		if (!bIsEditorOnlyExportOnStack && !IsEditorOnlyPropertyOnTheStack())
+		bIsEditorOnly = bIsEditorOnlyExportOnStack || IsEditorOnlyPropertyOnTheStack();
 #endif
-		{
-			SaveContext.ImportsUsedInGame.Add(InObject);
-		}
+		SaveContext.AddImport(InObject, bIsEditorOnly);
 
 		UObject* ObjOuter = InObject->GetOuter();
 		UClass* ObjClass = InObject->GetClass();
diff --git a/Engine/Source/Runtime/CoreUObject/Private/UObject/SavePackage/SaveContext.h b/Engine/Source/Runtime/CoreUObject/Private/UObject/SavePackage/SaveContext.h
--- a/Engine/Source/Runtime/CoreUObject/Private/UObject/SavePackage/SaveContext.h
+++ b/Engine/Source/Runtime/CoreUObject/Private/UObject/SavePackage/SaveContext.h
@@ -344,9 +344,13 @@ public:
 		bGenerateFileStub = true;
 	}
 
-	void AddImport(UObject* InObject)
+	void AddImport(UObject* InObject, bool bIsEditorOnlyImport = false)
 	{
 		Imports.Add(InObject);
+		if (!bIsEditorOnlyImport)
+		{
+			ImportsUsedInGame.Add(InObject);
+		}
 	}
 
 	void AddExport(UObject* InObj, bool bNotAlwaysLoadedForEditorGame)
@@ -448,6 +452,16 @@ public:
 		return PrestreamPackages;
 	}
 
+	TSet<UPackage*>& GetPrestreamPackages()
+	{
+		return PrestreamPackages;
+	}
+
+	bool IsPrestreamPackage(UPackage* InPackage) const
+	{
+		return PrestreamPackages.Contains(InPackage);
+	}
+
 	void AddPrestreamPackages(UPackage* InPackage)
 	{
 		PrestreamPackages.Add(InPackage);
@@ -575,7 +589,7 @@ private:
 	TMap<UObject*, TSet<UObject*>> ExportObjectDependencies;
 	// Map of objects to their native dependencies
 	TMap<UObject*, TSet<UObject*>> ExportNativeObjectDependencies;
-	// Set of harvested prestream packages, should be deprecated, not implementated in Save2 yet
+	// Set of harvested prestream packages
 	TSet<UPackage*> PrestreamPackages;
 	// Harvested custom versions
 	FCustomVersionContainer CustomVersions;
diff --git a/Engine/Source/Runtime/CoreUObject/Private/UObject/SavePackage/SavePackageUtilities.cpp b/Engine/Source/Runtime/CoreUObject/Private/UObject/SavePackage/SavePackageUtilities.cpp
--- a/Engine/Source/Runtime/CoreUObject/Private/UObject/SavePackage/SavePackageUtilities.cpp
+++ b/Engine/Source/Runtime/CoreUObject/Private/UObject/SavePackage/SavePackageUtilities.cpp
@@ -30,6 +30,7 @@
 #include "UObject/UObjectThreadContext.h"
 
 DEFINE_LOG_CATEGORY(LogSavePackage);
+UE_TRACE_CHANNEL_DEFINE(SaveTimeChannel);
 
 #if ENABLE_COOK_STATS
 #include "ProfilingDebugging/ScopedTimers.h"
@@ -172,6 +173,15 @@ const FName NAME_World("World");
 const FName NAME_Level("Level");
 const FName NAME_PrestreamPackage("PrestreamPackage");
 
+// Switch to SavePackage to Save2, 0:disabled, 1: enabled for cook, 2: enabled for uncooked, 3: enabled everywhere
+static TAutoConsoleVariable<int32> CVarEnablePackageNewSave(TEXT("SavePackage.EnableNewSave"), 0, TEXT("Enable new package save mechanism over the old one."));
+
+bool IsNewSaveEnabled(bool bForCooking)
+{
+	int32 EnableNewSave = CVarEnablePackageNewSave.GetValueOnAnyThread();
+	return (bForCooking && (EnableNewSave & 1) != 0) || (!bForCooking && (EnableNewSave & 2) != 0);
+}
+
 void GetBlueprintNativeCodeGenReplacement(UObject* InObj, UClass*& ObjClass, UObject*& ObjOuter, FName& ObjName, const ITargetPlatform* TargetPlatform)
 {
 #if WITH_EDITOR
diff --git a/Engine/Source/Runtime/CoreUObject/Private/UObject/SavePackage/SavePackageUtilities.h b/Engine/Source/Runtime/CoreUObject/Private/UObject/SavePackage/SavePackageUtilities.h
--- a/Engine/Source/Runtime/CoreUObject/Private/UObject/SavePackage/SavePackageUtilities.h
+++ b/Engine/Source/Runtime/CoreUObject/Private/UObject/SavePackage/SavePackageUtilities.h
@@ -22,6 +22,14 @@ template<typename StateType> class TAsyncWorkSequence;
 
 DECLARE_LOG_CATEGORY_EXTERN(LogSavePackage, Log, All);
 
+// Save Time trace
+#if UE_TRACE_ENABLED && !UE_BUILD_SHIPPING
+UE_TRACE_CHANNEL_EXTERN(SaveTimeChannel)
+#define SCOPED_SAVETIMER(TimerName) TRACE_CPUPROFILER_EVENT_SCOPE_ON_CHANNEL(TimerName, SaveTimeChannel)
+#else
+#define SCOPED_SAVETIMER(TimerName)
+#endif
+
 struct FLargeMemoryDelete
 {
 	void operator()(uint8* Ptr) const
@@ -288,6 +296,9 @@ namespace SavePackageUtilities
 	extern const FName NAME_Level;
 	extern const FName NAME_PrestreamPackage;
 
+	// return if the new save algorithm is enabled for cooked or uncooked
+	bool IsNewSaveEnabled(bool bForCooking = false);
+
 	void GetBlueprintNativeCodeGenReplacement(UObject* InObj, UClass*& ObjClass, UObject*& ObjOuter, FName& ObjName, const ITargetPlatform* TargetPlatform);
 
 	void IncrementOutstandingAsyncWrites();
diff --git a/Engine/Source/Runtime/CoreUObject/Private/UObject/SavePackage2.cpp b/Engine/Source/Runtime/CoreUObject/Private/UObject/SavePackage2.cpp
--- a/Engine/Source/Runtime/CoreUObject/Private/UObject/SavePackage2.cpp
+++ b/Engine/Source/Runtime/CoreUObject/Private/UObject/SavePackage2.cpp
@@ -74,7 +74,7 @@ ESavePackageResult ValidateBlueprintNativeCodeGenReplacement(FSaveContext& SaveC
 
 ESavePackageResult ValidatePackage(FSaveContext& SaveContext)
 {
-	TRACE_CPUPROFILER_EVENT_SCOPE(UPackage_ValidatePackage);
+	SCOPED_SAVETIMER(UPackage_ValidatePackage);
 
 	// Platform can't save the package
 	if (!FPlatformProperties::HasEditorOnlyData())
@@ -208,7 +208,7 @@ FORCEINLINE void EnsurePackageLocalization(UPackage* InPackage)
 
 ESavePackageResult RoutePresave(FSaveContext& SaveContext)
 {
-	TRACE_CPUPROFILER_EVENT_SCOPE(UPackage_RoutePresave);
+	SCOPED_SAVETIMER(UPackage_RoutePresave);
 
 	// Just route presave on all objects in the package while skipping unsaveable objects. 
 	// This should be more efficient then trying to restrict to just the actual export, 
@@ -250,7 +250,7 @@ ESavePackageResult RoutePresave(FSaveContext& SaveContext)
 
 ESavePackageResult HarvestPackage(FSaveContext& SaveContext)
 {
-	TRACE_CPUPROFILER_EVENT_SCOPE(UPackage_Save_HarvestPackage);
+	SCOPED_SAVETIMER(UPackage_Save_HarvestPackage);
 
 	FPackageHarvester Harvester(SaveContext);
 	EObjectFlags TopLevelFlags = SaveContext.GetTopLevelFlags();
@@ -285,22 +285,24 @@ ESavePackageResult HarvestPackage(FSaveContext& SaveContext)
 		}
 	}
 
-	// Harvest Prestream package class name if needed
-	if (SaveContext.GetPrestreamPackages().Num() > 0)
+	// Trim PrestreamPackage list 
+	TSet<UPackage*>& PrestreamPackages = SaveContext.GetPrestreamPackages();
+	TSet<UPackage*> KeptPrestreamPackages;
+	for (UPackage* Pkg : PrestreamPackages)
 	{
-		Harvester.HarvestName(SavePackageUtilities::NAME_PrestreamPackage);
-		
-		//@todo FH: is this really needed?
-		//TSet<UPackage*> KeptPrestreamPackages;
-		//for (UPackage* Pkg : PrestreamPackages)
-		//{
-		//	if (!Pkg->HasAnyMarks(OBJECTMARK_TagImp))
-		//	{
-		//		Pkg->Mark(OBJECTMARK_TagImp);
-		//		KeptPrestreamPackages.Add(Pkg);
-		//	}
-		//}
-		//Exchange(PrestreamPackages, KeptPrestreamPackages);
+		// If the prestream package hasn't been otherwise already marked as an import, keep it as such and mark it as an import
+		if (!SaveContext.IsImport(Pkg))
+		{
+			KeptPrestreamPackages.Add(Pkg);
+			SaveContext.AddImport(Pkg);
+		}
+	}
+	Exchange(PrestreamPackages, KeptPrestreamPackages);
+
+	// Harvest the PrestreamPackage class name if needed
+	if (PrestreamPackages.Num() > 0)
+	{
+		Harvester.HarvestName(SavePackageUtilities::NAME_PrestreamPackage);		
 	}
 
 	// if we have a WorldTileInfo, we need to harvest its dependencies as well, i.e. Custom Version
@@ -325,7 +327,7 @@ static FNameEntryId NAME_UniqueObjectNameForCookingComparisonIndex = FName("Uniq
 
 ESavePackageResult ValidateExports(FSaveContext& SaveContext)
 {
-	TRACE_CPUPROFILER_EVENT_SCOPE(UPackage_Save_ValidateExports);
+	SCOPED_SAVETIMER(UPackage_Save_ValidateExports);
 
 	// Check if we gathered any exports
 	if (SaveContext.GetExports().Num() == 0)
@@ -520,7 +522,7 @@ ESavePackageResult ValidateIllegalReferences(FSaveContext& SaveContext, TArray<U
 
 ESavePackageResult ValidateImports(FSaveContext& SaveContext)
 {
-	TRACE_CPUPROFILER_EVENT_SCOPE(UPackage_Save_ValidateImports);
+	SCOPED_SAVETIMER(UPackage_Save_ValidateImports);
 
 	TArray<UObject*> TopLevelObjects;
 	GetObjectsWithPackage(SaveContext.GetPackage(), TopLevelObjects, false);
@@ -572,6 +574,15 @@ ESavePackageResult ValidateImports(FSaveContext& SaveContext)
 		ensureAlways(SaveContext.NameExists(Import->GetClass()->GetFName().GetComparisonIndex()));
 		ensureAlways(SaveContext.NameExists(Import->GetClass()->GetOuter()->GetFName().GetComparisonIndex()));
 
+		// if the import is marked as a prestream package, we dont need to validate further
+		if (SaveContext.IsPrestreamPackage(ImportPackage))
+		{
+			ensureAlways(Import == ImportPackage);
+			// These are not errors
+			UE_LOG(LogSavePackage, Display, TEXT("Prestreaming package %s "), *ImportPackage->GetPathName()); //-V595
+			continue;
+		}
+
 		// if an import outer is an export and that import doesn't have a specific package set then, there's an error
 		const bool bWrongImport = Import->GetOuter() && Import->GetOuter()->IsInPackage(SaveContext.GetPackage()) && Import->GetExternalPackage() == nullptr;
 		if (bWrongImport)
@@ -591,14 +602,6 @@ ESavePackageResult ValidateImports(FSaveContext& SaveContext)
 		}
 		check(!bWrongImport || Import->HasAllFlags(RF_Transient) || Import->IsNative());
 
-		// @todo FH: validate prestream packages still needed..
-		if (SaveContext.GetPrestreamPackages().Contains(ImportPackage))
-		{
-			// These are not errors
-			UE_LOG(LogSavePackage, Display, TEXT("Prestreaming package %s "), *ImportPackage->GetPathName()); //-V595
-			continue;
-		}
-
 		// if this import shares a outer with top level object of this package then the reference is acceptable
 		if (!SaveContext.IsCooking() && (IsInAnyTopLevelObject(Import) || AnyTopLevelObjectIsIn(Import) || AnyTopLevelObjectHasSameOutermostObject(Import)))
 		{
@@ -645,7 +648,7 @@ ESavePackageResult CreateLinker(FSaveContext& SaveContext)
 	// The temp file will be saved in the game save folder to not have to deal with potentially too long paths.
 	// Since the temp filename may include a 32 character GUID as well, limit the user prefix to 32 characters.
 	{
-		TRACE_CPUPROFILER_EVENT_SCOPE(UPackage_Save_CreateLinkerSave);
+		SCOPED_SAVETIMER(UPackage_Save_CreateLinkerSave);
 
 		if (SaveContext.IsDiffing())
 		{
@@ -777,7 +780,7 @@ ESavePackageResult BuildLinker(FSaveContext& SaveContext)
 
 	// Build Name Map
 	{
-		TRACE_CPUPROFILER_EVENT_SCOPE(UPackage_Save_BuildNameMap);
+		SCOPED_SAVETIMER(UPackage_Save_BuildNameMap);
 		Linker->Summary.NameOffset = 0;
 		Linker->Summary.NameCount = 0;
 		Linker->NameMap.Append(SaveContext.GetReferencedNames().Array());
@@ -798,7 +801,7 @@ ESavePackageResult BuildLinker(FSaveContext& SaveContext)
 		Linker->Summary.GatherableTextDataCount = 0;
 		if (!SaveContext.IsFilterEditorOnly())
 		{
-			TRACE_CPUPROFILER_EVENT_SCOPE(UPackage_Save_BuildGatherableTextData);
+			SCOPED_SAVETIMER(UPackage_Save_BuildGatherableTextData);
 
 			// Gathers from the given package
 			SaveContext.GatherableTextResultFlags = EPropertyLocalizationGathererResultFlags::Empty;
@@ -809,7 +812,7 @@ ESavePackageResult BuildLinker(FSaveContext& SaveContext)
 	// Build ImportMap
 	TMap<UObject*, UObject*> ReplacedImportOuters;
 	{
-		TRACE_CPUPROFILER_EVENT_SCOPE(UPackage_Save_BuildImportMap);
+		SCOPED_SAVETIMER(UPackage_Save_BuildImportMap);
 		for (UObject* Import : SaveContext.GetImports())
 		{
 			UClass* ImportClass = Import->GetClass();
@@ -826,8 +829,8 @@ ESavePackageResult BuildLinker(FSaveContext& SaveContext)
 			}
 			FObjectImport& ObjectImport = Linker->ImportMap.Add_GetRef(FObjectImport(Import, ImportClass));
 
-			//@todo FH: validate this
-			if (SaveContext.GetPrestreamPackages().Contains((UPackage*)Import))
+			// If the package import is a prestream package, mark it as such by hacking its class name
+			if (SaveContext.IsPrestreamPackage(Cast<UPackage>(Import)))
 			{
 				ObjectImport.ClassName = SavePackageUtilities::NAME_PrestreamPackage;
 			}
@@ -843,7 +846,7 @@ ESavePackageResult BuildLinker(FSaveContext& SaveContext)
 
 	// Build ExportMap & Package Netplay data
 	{
-		TRACE_CPUPROFILER_EVENT_SCOPE(UPackage_Save_BuildExportMap);
+		SCOPED_SAVETIMER(UPackage_Save_BuildExportMap);
 		for (const FTaggedExport& TaggedExport : SaveContext.GetExports())
 		{
 			FObjectExport& Export = Linker->ExportMap.Add_GetRef(FObjectExport(TaggedExport.Obj, TaggedExport.bNotAlwaysLoadedForEditorGame));
@@ -862,7 +865,7 @@ ESavePackageResult BuildLinker(FSaveContext& SaveContext)
 		// for example where FAnimInstanceProxy PostLoad actually depends on  UAnimBlueprintGeneratedClass PostLoad to be properly initialized.
 		FObjectExportSortHelper ExportSortHelper;
 		{
-			TRACE_CPUPROFILER_EVENT_SCOPE(UPackage_Save_SortExports);
+			SCOPED_SAVETIMER(UPackage_Save_SortExports);
 			ExportSortHelper.SortExports(Linker, nullptr);
 		}
 		Linker->Summary.ExportCount = Linker->ExportMap.Num();
@@ -886,7 +889,7 @@ ESavePackageResult BuildLinker(FSaveContext& SaveContext)
 
 	// Build DependsMap
 	{
-		TRACE_CPUPROFILER_EVENT_SCOPE(UPackage_Save_BuildExportDependsMap);
+		SCOPED_SAVETIMER(UPackage_Save_BuildExportDependsMap);
 
 		Linker->DependsMap.AddZeroed(Linker->ExportMap.Num());
 		for (int32 ExpIndex = 0; ExpIndex < Linker->ExportMap.Num(); ++ExpIndex)
@@ -931,7 +934,7 @@ ESavePackageResult BuildLinker(FSaveContext& SaveContext)
 
 	// Map Export Indices
 	{
-		TRACE_CPUPROFILER_EVENT_SCOPE(UPackage_Save_MapExportIndices);
+		SCOPED_SAVETIMER(UPackage_Save_MapExportIndices);
 		for (FObjectExport& Export : Linker->ExportMap)
 		{
 			// Set class index.
@@ -1336,7 +1339,7 @@ ESavePackageResult WritePackageHeader(FStructuredArchive::FRecord& StructuredArc
 	// Write Name Map
 	Linker->Summary.NameOffset = SaveContext.OffsetAfterPackageFileSummary;
 	{
-		TRACE_CPUPROFILER_EVENT_SCOPE(UPackage_Save_BuildNameMap);
+		SCOPED_SAVETIMER(UPackage_Save_BuildNameMap);
 		Linker->Summary.NameCount = Linker->NameMap.Num();
 		for (const FNameEntryId NameEntryId : Linker->NameMap)
 		{
@@ -1346,13 +1349,13 @@ ESavePackageResult WritePackageHeader(FStructuredArchive::FRecord& StructuredArc
 
 	// Write GatherableText
 	{
-		TRACE_CPUPROFILER_EVENT_SCOPE(UPackage_Save_WriteGatherableTextData);
+		SCOPED_SAVETIMER(UPackage_Save_WriteGatherableTextData);
 		WriteGatherableText(StructuredArchiveRoot, SaveContext);
 	}
 
 	// Save Dummy Import Map, overwritten later.
 	{
-		TRACE_CPUPROFILER_EVENT_SCOPE(UPackage_Save_WriteDummyImportMap);
+		SCOPED_SAVETIMER(UPackage_Save_WriteDummyImportMap);
 		Linker->Summary.ImportOffset = Linker->Tell();
 		for (FObjectImport& Import : Linker->ImportMap)
 		{
@@ -1363,7 +1366,7 @@ ESavePackageResult WritePackageHeader(FStructuredArchive::FRecord& StructuredArc
 
 	// Save Dummy Export Map, overwritten later.
 	{
-		TRACE_CPUPROFILER_EVENT_SCOPE(UPackage_Save_WriteDummyExportMap);
+		SCOPED_SAVETIMER(UPackage_Save_WriteDummyExportMap);
 		Linker->Summary.ExportOffset = Linker->Tell();
 		for (FObjectExport& Export : Linker->ExportMap)
 		{
@@ -1374,7 +1377,7 @@ ESavePackageResult WritePackageHeader(FStructuredArchive::FRecord& StructuredArc
 
 	// Save Depend Map
 	{
-		TRACE_CPUPROFILER_EVENT_SCOPE(UPackage_Save_WriteDependsMap);
+		SCOPED_SAVETIMER(UPackage_Save_WriteDependsMap);
 
 		FStructuredArchive::FStream DependsStream = StructuredArchiveRoot.EnterStream(SA_FIELD_NAME(TEXT("DependsMap")));
 		Linker->Summary.DependsOffset = Linker->Tell();
@@ -1401,7 +1404,7 @@ ESavePackageResult WritePackageHeader(FStructuredArchive::FRecord& StructuredArc
 	// Write Soft Package references & Searchable Names
 	if (!SaveContext.IsFilterEditorOnly())
 	{
-		TRACE_CPUPROFILER_EVENT_SCOPE(UPackage_Save_SaveSoftPackagesAndSearchableNames);
+		SCOPED_SAVETIMER(UPackage_Save_SaveSoftPackagesAndSearchableNames);
 
 		// Save soft package references
 		Linker->Summary.SoftPackageReferencesOffset = Linker->Tell();
@@ -1427,23 +1430,23 @@ ESavePackageResult WritePackageHeader(FStructuredArchive::FRecord& StructuredArc
 
 	// Save thumbnails
 	{
-		TRACE_CPUPROFILER_EVENT_SCOPE(UPackage_Save_SaveThumbnails);
+		SCOPED_SAVETIMER(UPackage_Save_SaveThumbnails);
 		SavePackageUtilities::SaveThumbnails(SaveContext.GetPackage(), Linker, StructuredArchiveRoot.EnterField(SA_FIELD_NAME(TEXT("Thumbnails"))));
 	}
 	{
 		// Save asset registry data so the editor can search for information about assets in this package
-		TRACE_CPUPROFILER_EVENT_SCOPE(UPackage_Save_SaveAssetRegistryData);
+		SCOPED_SAVETIMER(UPackage_Save_SaveAssetRegistryData);
 		UE::AssetRegistry::WritePackageData(StructuredArchiveRoot, SaveContext.IsCooking(), SaveContext.GetPackage(), Linker, SaveContext.GetImportsUsedInGame(), SaveContext.GetSoftPackagesUsedInGame());
 	}
 	// Save level information used by World browser
 	{
-		TRACE_CPUPROFILER_EVENT_SCOPE(UPackage_Save_WorldLevelData);
+		SCOPED_SAVETIMER(UPackage_Save_WorldLevelData);
 		SavePackageUtilities::SaveWorldLevelInfo(SaveContext.GetPackage(), Linker, StructuredArchiveRoot);
 	}
 
 	// Write Preload Dependencies
 	{
-		TRACE_CPUPROFILER_EVENT_SCOPE(UPackage_Save_PreloadDependencies);
+		SCOPED_SAVETIMER(UPackage_Save_PreloadDependencies);
 		SavePreloadDependencies(StructuredArchiveRoot, SaveContext);
 	}
 	Linker->Summary.TotalHeaderSize = Linker->Tell();
@@ -1456,18 +1459,18 @@ ESavePackageResult WritePackageTextHeader(FStructuredArchive::FRecord& Structure
 
 	// Write GatherableText
 	{
-		TRACE_CPUPROFILER_EVENT_SCOPE(UPackage_Save_WriteGatherableTextData);
+		SCOPED_SAVETIMER(UPackage_Save_WriteGatherableTextData);
 		WriteGatherableText(StructuredArchiveRoot, SaveContext);
 	}
 
 	// Save thumbnails
 	{
-		TRACE_CPUPROFILER_EVENT_SCOPE(UPackage_Save_SaveThumbnails);
+		SCOPED_SAVETIMER(UPackage_Save_SaveThumbnails);
 		SavePackageUtilities::SaveThumbnails(SaveContext.GetPackage(), Linker, StructuredArchiveRoot.EnterField(SA_FIELD_NAME(TEXT("Thumbnails"))));
 	}
 	// Save level information used by World browser
 	{
-		TRACE_CPUPROFILER_EVENT_SCOPE(UPackage_Save_WorldLevelData);
+		SCOPED_SAVETIMER(UPackage_Save_WorldLevelData);
 		SavePackageUtilities::SaveWorldLevelInfo(SaveContext.GetPackage(), Linker, StructuredArchiveRoot);
 	}
 
@@ -1477,7 +1480,7 @@ ESavePackageResult WritePackageTextHeader(FStructuredArchive::FRecord& Structure
 ESavePackageResult WriteExports(FStructuredArchive::FRecord& StructuredArchiveRoot, FSaveContext& SaveContext)
 {
 	//COOK_STAT(FScopedDurationTimer SaveTimer(FSavePackageStats::SerializeExportsTimeSec));
-	TRACE_CPUPROFILER_EVENT_SCOPE(UPackage_Save_SaveExports);
+	SCOPED_SAVETIMER(UPackage_Save_SaveExports);
 	FLinkerSave* Linker = SaveContext.Linker.Get();
 	FScopedSlowTask SlowTask(Linker->ExportMap.Num(), FText(), SaveContext.IsUsingSlowTask());
 
@@ -1496,7 +1499,7 @@ ESavePackageResult WriteExports(FStructuredArchive::FRecord& StructuredArchiveRo
 		FObjectExport& Export = Linker->ExportMap[i];
 		if (Export.Object)
 		{
-			TRACE_CPUPROFILER_EVENT_SCOPE(UPackage_Save_SaveExport);
+			SCOPED_SAVETIMER(UPackage_Save_SaveExport);
 
 			// Save the object data.
 			Export.SerialOffset = Linker->Tell();
@@ -1599,7 +1602,7 @@ ESavePackageResult WriteAdditionalExportFiles(FSaveContext& SaveContext)
 
 ESavePackageResult UpdatePackageHeader(FStructuredArchive::FRecord& StructuredArchiveRoot, FSaveContext& SaveContext)
 {
-	TRACE_CPUPROFILER_EVENT_SCOPE(UPackage_Save_UpdatePackageHeader);
+	SCOPED_SAVETIMER(UPackage_Save_UpdatePackageHeader);
 
 	FLinkerSave* Linker = SaveContext.Linker.Get();
 #if WITH_EDITOR
@@ -1688,7 +1691,7 @@ ESavePackageResult UpdatePackageHeader(FStructuredArchive::FRecord& StructuredAr
 
 ESavePackageResult FinalizeFile(FStructuredArchive::FRecord& StructuredArchiveRoot, FSaveContext& SaveContext)
 {
-	TRACE_CPUPROFILER_EVENT_SCOPE(UPackage_Save_FinalizeFile);
+	SCOPED_SAVETIMER(UPackage_Save_FinalizeFile);
 
 	// In the concurrent case, it is called right after routing presave so it can be done in batch before going concurrent
 	if (!SaveContext.IsConcurrent())
@@ -2031,7 +2034,7 @@ ESavePackageResult InnerSave(FSaveContext& SaveContext)
 	// Mark Exports & Package RF_Loaded
 	SlowTask.EnterProgressFrame();
 	{
-		TRACE_CPUPROFILER_EVENT_SCOPE(UPackage_Save_MarkExportLoaded);
+		SCOPED_SAVETIMER(UPackage_Save_MarkExportLoaded);
 		FLinkerSave* Linker = SaveContext.Linker.Get();
 		// Mark exports and the package as RF_Loaded after they've been serialized
 		// This is to ensue that newly created packages are properly marked as loaded (since they now exist on disk and 
@@ -2075,7 +2078,7 @@ FSavePackageResultStruct UPackage::Save2(UPackage* InPackage, UObject* InAsset,
 {
 	//COOK_STAT(FScopedDurationTimer FuncSaveTimer(FSavePackageStats::SavePackageTimeSec));
 	//COOK_STAT(FSavePackageStats::NumPackagesSaved++);
-	TRACE_CPUPROFILER_EVENT_SCOPE(UPackage_Save2);
+	SCOPED_SAVETIMER(UPackage_Save2);
 	FSaveContext SaveContext(InPackage, InAsset, InFilename, SaveArgs);
 
 	// Create the slow task dialog if needed
@@ -2162,21 +2165,6 @@ FSavePackageResultStruct UPackage::Save2(UPackage* InPackage, UObject* InAsset,
 
 ESavePackageResult UPackage::SaveConcurrent(TArrayView<FPackageSaveInfo> InPackages, FSavePackageArgs& SaveArgs, TArray<FSavePackageResultStruct>& OutResults)
 {
-	auto GetPackageAsset = [](FPackageSaveInfo& PackageSaveInfo) -> UObject*
-	{
-		UObject* Asset = nullptr;
-		ForEachObjectWithPackage(PackageSaveInfo.Package, [&Asset](UObject* Object)
-			{
-				if (Object->IsAsset())
-				{
-					Asset = Object;
-					return false;
-				}
-				return true;
-			}, /*bIncludeNestedObjects*/ false);
-		return Asset;
-	};
-
 	const int32 TotalSaveSteps = 4;
 	FScopedSlowTask SlowTask(TotalSaveSteps, NSLOCTEXT("Core", "SavingFiles", "Saving files..."), SaveArgs.bSlowTask);
 	SlowTask.MakeDialog(!!(SaveArgs.SaveFlags & SAVE_FromAutosave));
@@ -2185,10 +2173,10 @@ ESavePackageResult UPackage::SaveConcurrent(TArrayView<FPackageSaveInfo> InPacka
 	SlowTask.EnterProgressFrame();
 	TArray<FSaveContext> PackageSaveContexts;
 	{
-		TRACE_CPUPROFILER_EVENT_SCOPE(UPackage_SaveConcurrent_PreSave);
+		SCOPED_SAVETIMER(UPackage_SaveConcurrent_PreSave);
 		for (FPackageSaveInfo& PackageSaveInfo : InPackages)
 		{
-			FSaveContext& SaveContext = PackageSaveContexts.Emplace_GetRef(PackageSaveInfo.Package, GetPackageAsset(PackageSaveInfo), *PackageSaveInfo.Filename, SaveArgs, nullptr);
+			FSaveContext& SaveContext = PackageSaveContexts.Emplace_GetRef(PackageSaveInfo.Package, PackageSaveInfo.Package->FindAssetInPackage(), *PackageSaveInfo.Filename, SaveArgs, nullptr);
 
 			// Validation
 			SaveContext.Result = ValidatePackage(SaveContext);
@@ -2204,7 +2192,7 @@ ESavePackageResult UPackage::SaveConcurrent(TArrayView<FPackageSaveInfo> InPacka
 			// PreSave Asset
 			if (SaveContext.GetAsset())
 			{
-				TRACE_CPUPROFILER_EVENT_SCOPE(UPackage_SaveConcurrent_PreSaveRoot);
+				SCOPED_SAVETIMER(UPackage_SaveConcurrent_PreSaveRoot);
 				SaveContext.SetPreSaveCleanup(SaveContext.GetAsset()->PreSaveRoot(SaveContext.GetFilename()));
 			}
 
@@ -2220,17 +2208,15 @@ ESavePackageResult UPackage::SaveConcurrent(TArrayView<FPackageSaveInfo> InPacka
 	SlowTask.EnterProgressFrame();
 	{
 		// Flush async loading and reset loaders
-		TRACE_CPUPROFILER_EVENT_SCOPE(UPackage_SaveConcurrent_ResetLoadersForSave);
+		SCOPED_SAVETIMER(UPackage_SaveConcurrent_ResetLoadersForSave);
 		ResetLoadersForSave(InPackages);
 	}
 
 	SlowTask.EnterProgressFrame();
 	{
-		TRACE_CPUPROFILER_EVENT_SCOPE(UPackage_SaveConcurrent);
-
+		SCOPED_SAVETIMER(UPackage_SaveConcurrent);
 		// Use concurrent new save only if new save is enabled, otherwise use old save
-		static const IConsoleVariable* EnableNewSave = IConsoleManager::Get().FindConsoleVariable(TEXT("SavePackage.EnableNewSave"));
-		if (EnableNewSave->GetInt())
+		if (SavePackageUtilities::IsNewSaveEnabled(SaveArgs.TargetPlatform != nullptr))
 		{
 			// Passing in false here so that GIsSavingPackage is set to true on top of locking the GC
 			FScopedSavingFlag IsSavingFlag(false);
@@ -2241,7 +2227,8 @@ ESavePackageResult UPackage::SaveConcurrent(TArrayView<FPackageSaveInfo> InPacka
 					InnerSave(PackageSaveContexts[PackageIdx]);
 				});
 		}
-		else
+		// save concurrently if the SAVE_Concurrent flag is present
+		else if (SaveArgs.SaveFlags & SAVE_Concurrent)
 		{
 			GIsSavingPackage = true;
 			ParallelFor(PackageSaveContexts.Num(), [&PackageSaveContexts](int32 PackageIdx)
@@ -2254,12 +2241,22 @@ ESavePackageResult UPackage::SaveConcurrent(TArrayView<FPackageSaveInfo> InPacka
 				});
 			GIsSavingPackage = false;
 		}
+		// otherwise save serial
+		else
+		{
+			for (FSaveContext& SaveContext : PackageSaveContexts)
+			{
+				Save(SaveContext.GetPackage(), SaveContext.GetAsset(), SaveArgs.TopLevelFlags, SaveContext.GetFilename(),
+					SaveArgs.Error, nullptr, SaveArgs.bForceByteSwapping, SaveArgs.bWarnOfLongFilename, SaveArgs.SaveFlags | SAVE_Concurrent, // still pass in the flag so that we don't redo things we have handled already
+					SaveArgs.TargetPlatform, SaveArgs.FinalTimeStamp, false, nullptr, nullptr);
+			}
+		}
 	}
 
 	// Run Post Concurrent Save
 	SlowTask.EnterProgressFrame();
 	{
-		TRACE_CPUPROFILER_EVENT_SCOPE(UPackage_SaveConcurrent_PostSave);
+		SCOPED_SAVETIMER(UPackage_SaveConcurrent_PostSave);
 		for (FSaveContext& SaveContext : PackageSaveContexts)
 		{
 			// PostSave Asset
