SavePackage private object errors: Identify private archetypes as the MostLikelyCulprit. Always print out the MostLikelyCulprit when running a commandlet rather than prompting.
#rnx

diff --git a/Engine/Source/Runtime/CoreUObject/Private/UObject/SavePackage.cpp b/Engine/Source/Runtime/CoreUObject/Private/UObject/SavePackage.cpp
--- a/Engine/Source/Runtime/CoreUObject/Private/UObject/SavePackage.cpp
+++ b/Engine/Source/Runtime/CoreUObject/Private/UObject/SavePackage.cpp
@@ -3044,7 +3044,6 @@ FSavePackageResultStruct UPackage::Save(UPackage* InOuter, UObject* Base, EObjec
 				if (IllegalObjectsInOtherMaps.Num() )
 				{
 					UObject* MostLikelyCulprit = nullptr;
-					const FProperty* PropertyRef = nullptr;
 
 					// construct a string containing up to the first 5 objects problem objects
 					FString ObjectNames;
@@ -3077,19 +3076,16 @@ FSavePackageResultStruct UPackage::Save(UPackage* InOuter, UObject* Base, EObjec
 					bool bFindCulprit = IsRunningCommandlet() || (FMessageDialog::Open(EAppMsgType::YesNo, Message) == EAppReturnType::Yes);
 					if (bFindCulprit)
 					{
-						SavePackageUtilities::FindMostLikelyCulprit(IllegalObjectsInOtherMaps, MostLikelyCulprit, PropertyRef);
-						if (MostLikelyCulprit != nullptr && PropertyRef != nullptr)
+						FString Referencer;
+						SavePackageUtilities::FindMostLikelyCulprit(IllegalObjectsInOtherMaps, MostLikelyCulprit, Referencer);
+						if (MostLikelyCulprit != nullptr)
 						{
-							CulpritString = FString::Printf(TEXT("%s (%s)"), *MostLikelyCulprit->GetFullName(), *PropertyRef->GetName());
-						}
-						else if (MostLikelyCulprit != nullptr)
-						{
-							CulpritString = FString::Printf(TEXT("%s (Unknown property)"), *MostLikelyCulprit->GetFullName());
+							CulpritString = FString::Printf(TEXT("%s (%s)"), *MostLikelyCulprit->GetFullName(), *Referencer);
 						}
 					}
 					else if (!ObjectNames.IsEmpty())
 					{
-							CulpritString = ObjectNames;
+						CulpritString = ObjectNames;
 					}
 
 					FString ErrorMessage = FString::Printf(TEXT("Can't save %s: Graph is linked to object %s in external map"), Filename, *CulpritString);
@@ -3108,7 +3104,6 @@ FSavePackageResultStruct UPackage::Save(UPackage* InOuter, UObject* Base, EObjec
 				if (PrivateObjects.Num())
 				{
 					UObject* MostLikelyCulprit = nullptr;
-					const FProperty* PropertyRef = nullptr;
 					
 					// construct a string containing up to the first 5 objects problem objects
 					FString ObjectNames;
@@ -3138,12 +3133,14 @@ FSavePackageResultStruct UPackage::Save(UPackage* InOuter, UObject* Base, EObjec
 					const FText Message = FText::Format( NSLOCTEXT("Core", "LinkedToPrivateObjectsInOtherPackage_FindCulpritQ", "Can't save {FileName}: Graph is linked to private object(s) in an external package.\nExternal Object(s):\n{ObjectNames}  \nTry to find the chain of references to that object (may take some time)?"), Args );
 
 					FString CulpritString = TEXT( "Unknown" );
-					if (FMessageDialog::Open(EAppMsgType::YesNo, Message) == EAppReturnType::Yes)
+					bool bFindCulprit = IsRunningCommandlet() || (FMessageDialog::Open(EAppMsgType::YesNo, Message) == EAppReturnType::Yes);
+					if (bFindCulprit)
 					{
-						SavePackageUtilities::FindMostLikelyCulprit(PrivateObjects, MostLikelyCulprit, PropertyRef);
+						FString Referencer;
+						SavePackageUtilities::FindMostLikelyCulprit(PrivateObjects, MostLikelyCulprit, Referencer);
 						CulpritString = FString::Printf(TEXT("%s (%s)"),
 							(MostLikelyCulprit != nullptr) ? *MostLikelyCulprit->GetFullName() : TEXT("(unknown culprit)"),
-							(PropertyRef != nullptr) ? *PropertyRef->GetName() : TEXT("unknown property ref"));
+							*Referencer);
 					}
 
 					if (!(SaveFlags & SAVE_NoError))
diff --git a/Engine/Source/Runtime/CoreUObject/Private/UObject/SavePackage/SavePackageUtilities.cpp b/Engine/Source/Runtime/CoreUObject/Private/UObject/SavePackage/SavePackageUtilities.cpp
--- a/Engine/Source/Runtime/CoreUObject/Private/UObject/SavePackage/SavePackageUtilities.cpp
+++ b/Engine/Source/Runtime/CoreUObject/Private/UObject/SavePackage/SavePackageUtilities.cpp
@@ -457,15 +457,43 @@ void ConditionallyExcludeObjectForTarget(UObject* Obj, EObjectMark ExcludedObjec
  * @param	BadObjects	array of objects that are considered "bad" (e.g. non- RF_Public, in different map package, ...)
  * @return	UObject that is considered the most likely culprit causing them to be referenced or NULL
  */
-void FindMostLikelyCulprit(TArray<UObject*> BadObjects, UObject*& MostLikelyCulprit, const FProperty*& PropertyRef)
+void FindMostLikelyCulprit(TArray<UObject*> BadObjects, UObject*& MostLikelyCulprit, FString& OutReferencer)
 {
-	MostLikelyCulprit = nullptr;
+	UObject* ArchetypeCulprit = nullptr;
+	UObject* ReferencedCulprit = nullptr;
+	const FProperty* ReferencedCulpritReferencer = nullptr;
 
 	// Iterate over all objects that are marked as unserializable/ bad and print out their referencers.
 	for (int32 BadObjIndex = 0; BadObjIndex < BadObjects.Num(); BadObjIndex++)
 	{
 		UObject* Obj = BadObjects[BadObjIndex];
 
+		// SavePackage adds references to the class archetype manually; if this type is a class archetype and it is private, mark it as an error
+		// for that reason rather than checking references. Class archetypes must be public since instances of their class in other packages can refer to them
+		if (Obj->HasAnyFlags(RF_ArchetypeObject | RF_DefaultSubObject | RF_ClassDefaultObject))
+		{
+			UE_LOG(LogSavePackage, Warning, TEXT("%s is a private Archetype object"), *Obj->GetFullName());
+			TArray<const TCHAR*> Flags;
+			auto AddFlagIfPresent = [&Flags, Obj](EObjectFlags InFlag, const TCHAR* Descriptor)
+			{
+				if (Obj->HasAnyFlags(InFlag))
+				{
+					Flags.Add(Descriptor);
+				}
+			};
+			AddFlagIfPresent(RF_ArchetypeObject, TEXT("RF_ArchetypeObject"));
+			AddFlagIfPresent(RF_ClassDefaultObject, TEXT("RF_ClassDefaultObject"));
+			AddFlagIfPresent(RF_DefaultSubObject, TEXT("RF_DefaultSubObject"));
+			UE_LOG(LogSavePackage, Warning, TEXT("\tThis object is an archetype (flags include %s) but is private. This is a code error from the generator of the object. All archetype objects must be public."),
+				*FString::Join(Flags, TEXT("|")));
+
+			if (ArchetypeCulprit == nullptr)
+			{
+				ArchetypeCulprit = Obj;
+			}
+			continue;
+		}
+
 		UE_LOG(LogSavePackage, Warning, TEXT("\r\nReferencers of %s:"), *Obj->GetFullName());
 
 		FReferencerInformationList Refs;
@@ -488,14 +516,34 @@ void FindMostLikelyCulprit(TArray<UObject*> BadObjects, UObject*& MostLikelyCulp
 					{
 						const FProperty* Prop = Refs.ExternalReferences[i].ReferencingProperties[j];
 						UE_LOG(LogSavePackage, Warning, TEXT("\t\t%i) %s"), j, *Prop->GetFullName());
-						PropertyRef = Prop;
+						ReferencedCulpritReferencer = Prop;
 					}
 
-					MostLikelyCulprit = Obj;
+					// Later ReferencedCulprits are higher priority than earlier culprits. TODO: Not sure if this is an intentional behavior or if they choice was arbitrary.
+					ReferencedCulprit = Obj;
 				}
 			}
 		}
 	}
+
+	if (ArchetypeCulprit)
+	{
+		// ArchetypeCulprits are the most likely to be the problem; they are definitely a problem
+		MostLikelyCulprit = ArchetypeCulprit;
+		OutReferencer = TEXT("Referenced because it is an archetype object");
+	}
+	else
+	{
+		MostLikelyCulprit = ReferencedCulprit; // Might be null, in which case we didn't find one
+		if (ReferencedCulpritReferencer)
+		{
+			OutReferencer = *ReferencedCulpritReferencer->GetName();
+		}
+		else
+		{
+			OutReferencer = TEXT("Unknown property");
+		}
+	}
 }
 
 void AddFileToHash(FString const& Filename, FMD5& Hash)
diff --git a/Engine/Source/Runtime/CoreUObject/Private/UObject/SavePackage/SavePackageUtilities.h b/Engine/Source/Runtime/CoreUObject/Private/UObject/SavePackage/SavePackageUtilities.h
--- a/Engine/Source/Runtime/CoreUObject/Private/UObject/SavePackage/SavePackageUtilities.h
+++ b/Engine/Source/Runtime/CoreUObject/Private/UObject/SavePackage/SavePackageUtilities.h
@@ -312,7 +312,7 @@ namespace SavePackageUtilities
 	bool HasUnsaveableOuter(UObject* InObj, UPackage* InSavingPackage);
 	void CheckObjectPriorToSave(FArchiveUObject& Ar, UObject* InObj, UPackage* InSavingPackage);
 	void ConditionallyExcludeObjectForTarget(UObject* Obj, EObjectMark ExcludedObjectMarks, const ITargetPlatform* TargetPlatform);
-	void FindMostLikelyCulprit(TArray<UObject*> BadObjects, UObject*& MostLikelyCulprit, const FProperty*& PropertyRef);
+	void FindMostLikelyCulprit(TArray<UObject*> BadObjects, UObject*& MostLikelyCulprit, FString& OutReferencer);
 	void AddFileToHash(FString const& Filename, FMD5& Hash);
 
 	void WriteToFile(const FString& Filename, const uint8* InDataPtr, int64 InDataSize);
diff --git a/Engine/Source/Runtime/CoreUObject/Private/UObject/SavePackage2.cpp b/Engine/Source/Runtime/CoreUObject/Private/UObject/SavePackage2.cpp
--- a/Engine/Source/Runtime/CoreUObject/Private/UObject/SavePackage2.cpp
+++ b/Engine/Source/Runtime/CoreUObject/Private/UObject/SavePackage2.cpp
@@ -414,7 +414,6 @@ ESavePackageResult ValidateIllegalReferences(FSaveContext& SaveContext, TArray<U
 	if (ObjectsInOtherMaps.Num() > 0)
 	{
 		UObject* MostLikelyCulprit = nullptr;
-		const FProperty* PropertyRef = nullptr;
 
 		// construct a string containing up to the first 5 objects problem objects
 		FString ObjectNames;
@@ -447,14 +446,11 @@ ESavePackageResult ValidateIllegalReferences(FSaveContext& SaveContext, TArray<U
 		bool bFindCulprit = IsRunningCommandlet() || (FMessageDialog::Open(EAppMsgType::YesNo, Message) == EAppReturnType::Yes);
 		if (bFindCulprit)
 		{
-			SavePackageUtilities::FindMostLikelyCulprit(ObjectsInOtherMaps, MostLikelyCulprit, PropertyRef);
-			if (MostLikelyCulprit != nullptr && PropertyRef != nullptr)
+			FString Referencer;
+			SavePackageUtilities::FindMostLikelyCulprit(ObjectsInOtherMaps, MostLikelyCulprit, Referencer);
+			if (MostLikelyCulprit != nullptr)
 			{
-				CulpritString = FString::Printf(TEXT("%s (%s)"), *MostLikelyCulprit->GetFullName(), *PropertyRef->GetName());
-			}
-			else if (MostLikelyCulprit != nullptr)
-			{
-				CulpritString = FString::Printf(TEXT("%s (Unknown property)"), *MostLikelyCulprit->GetFullName());
+				CulpritString = FString::Printf(TEXT("%s (%s)"), *MostLikelyCulprit->GetFullName(), *Referencer);
 			}
 		}
 
@@ -473,7 +469,6 @@ ESavePackageResult ValidateIllegalReferences(FSaveContext& SaveContext, TArray<U
 	if (PrivateObjects.Num() > 0)
 	{
 		UObject* MostLikelyCulprit = nullptr;
-		const FProperty* PropertyRef = nullptr;
 
 		// construct a string containing up to the first 5 objects problem objects
 		FString ObjectNames;
@@ -503,12 +498,14 @@ ESavePackageResult ValidateIllegalReferences(FSaveContext& SaveContext, TArray<U
 		const FText Message = FText::Format(NSLOCTEXT("Core", "LinkedToPrivateObjectsInOtherPackage_FindCulpritQ", "Can't save {FileName}: Graph is linked to private object(s) in an external package.\nExternal Object(s):\n{ObjectNames}  \nTry to find the chain of references to that object (may take some time)?"), Args);
 
 		FString CulpritString = TEXT("Unknown");
-		if (FMessageDialog::Open(EAppMsgType::YesNo, Message) == EAppReturnType::Yes)
+		bool bFindCulprit = IsRunningCommandlet() || (FMessageDialog::Open(EAppMsgType::YesNo, Message) == EAppReturnType::Yes);
+		if (bFindCulprit)
 		{
-			SavePackageUtilities::FindMostLikelyCulprit(PrivateObjects, MostLikelyCulprit, PropertyRef);
+			FString Referencer;
+			SavePackageUtilities::FindMostLikelyCulprit(PrivateObjects, MostLikelyCulprit, Referencer);
 			CulpritString = FString::Printf(TEXT("%s (%s)"),
 				(MostLikelyCulprit != nullptr) ? *MostLikelyCulprit->GetFullName() : TEXT("(unknown culprit)"),
-				(PropertyRef != nullptr) ? *PropertyRef->GetName() : TEXT("unknown property ref"));
+				*Referencer);
 		}
 
 		if (SaveContext.IsGenerateSaveError())
