Activate Save2 on cook by default
Activate Save2 everywhere on QAGame

diff --git a/Engine/Source/Runtime/CoreUObject/Private/UObject/SavePackage/SavePackageUtilities.cpp b/Engine/Source/Runtime/CoreUObject/Private/UObject/SavePackage/SavePackageUtilities.cpp
--- a/Engine/Source/Runtime/CoreUObject/Private/UObject/SavePackage/SavePackageUtilities.cpp
+++ b/Engine/Source/Runtime/CoreUObject/Private/UObject/SavePackage/SavePackageUtilities.cpp
@@ -174,7 +174,7 @@ const FName NAME_Level("Level");
 const FName NAME_PrestreamPackage("PrestreamPackage");
 
 // Switch to SavePackage to Save2, 0:disabled, 1: enabled for cook, 2: enabled for uncooked, 3: enabled everywhere
-static TAutoConsoleVariable<int32> CVarEnablePackageNewSave(TEXT("SavePackage.EnableNewSave"), 0, TEXT("Enable new package save mechanism over the old one."));
+static TAutoConsoleVariable<int32> CVarEnablePackageNewSave(TEXT("SavePackage.EnableNewSave"), 1, TEXT("Enable new package save mechanism over the old one."));
 
 bool IsNewSaveEnabled(bool bForCooking)
 {
