Make sure the package localization cache is up to date before going in the GIsSavingPackage block in Save2

diff --git a/Engine/Source/Runtime/CoreUObject/Private/UObject/SavePackage2.cpp b/Engine/Source/Runtime/CoreUObject/Private/UObject/SavePackage2.cpp
--- a/Engine/Source/Runtime/CoreUObject/Private/UObject/SavePackage2.cpp
+++ b/Engine/Source/Runtime/CoreUObject/Private/UObject/SavePackage2.cpp
@@ -8,6 +8,7 @@
 #include "Blueprint/BlueprintSupport.h"
 #include "HAL/FileManager.h"
 #include "Internationalization/TextPackageNamespaceUtil.h"
+#include "Internationalization/PackageLocalizationManager.h"
 #include "IO/IoDispatcher.h"
 #include "Misc/AssetRegistryInterface.h"
 #include "Misc/ConfigCacheIni.h"
@@ -204,6 +205,9 @@ FORCEINLINE void EnsurePackageLocalization(UPackage* InPackage)
 		TextNamespaceUtil::EnsurePackageNamespace(InPackage);
 	}
 #endif // USE_STABLE_LOCALIZATION_KEYS
+
+	// Also make sure the localization cache is up to date, since updating it during the GIsSavingPackage won't allow object resolving
+	FPackageLocalizationManager::Get().ConditionalUpdateCache();
 }
 
 ESavePackageResult RoutePresave(FSaveContext& SaveContext)
