Move setting GIsSavingPackage in Save2 to after reference gathering

diff --git a/Engine/Source/Runtime/CoreUObject/Private/UObject/SavePackage2.cpp b/Engine/Source/Runtime/CoreUObject/Private/UObject/SavePackage2.cpp
--- a/Engine/Source/Runtime/CoreUObject/Private/UObject/SavePackage2.cpp
+++ b/Engine/Source/Runtime/CoreUObject/Private/UObject/SavePackage2.cpp
@@ -1911,6 +1911,10 @@ ESavePackageResult InnerSave(FSaveContext& SaveContext)
 		return SaveContext.Result;
 	}
 
+	// @todo: Need to adjust GIsSavingPackage to properly prevent generating reference once, package harvesting is done
+	// GIsSavingPackage is too harsh however, since it should be scope only to the current package
+	FScopedSavingFlag IsSavingFlag(SaveContext.IsConcurrent());
+
 	// Validate Exports
 	SlowTask.EnterProgressFrame();
 	SaveContext.Result = ValidateExports(SaveContext);
@@ -2133,7 +2137,8 @@ FSavePackageResultStruct UPackage::Save2(UPackage* InPackage, UObject* InAsset,
 
 	SlowTask.EnterProgressFrame();
 	{
-		FScopedSavingFlag IsSavingFlag(SaveContext.IsConcurrent());
+		// @todo: Once GIsSavingPackage is reworked we should reinstore the saving flag here for the gc lock
+		//FScopedSavingFlag IsSavingFlag(SaveContext.IsConcurrent());
 		SaveContext.Result = InnerSave(SaveContext);
 		if (SaveContext.Result != ESavePackageResult::Success)
 		{
@@ -2219,8 +2224,9 @@ ESavePackageResult UPackage::SaveConcurrent(TArrayView<FPackageSaveInfo> InPacka
 		// Use concurrent new save only if new save is enabled, otherwise use old save
 		if (SavePackageUtilities::IsNewSaveEnabled(SaveArgs.TargetPlatform != nullptr))
 		{
+			// @todo: Once GIsSavingPackage is reworked we should reinstore the saving flag here for the gc lock
 			// Passing in false here so that GIsSavingPackage is set to true on top of locking the GC
-			FScopedSavingFlag IsSavingFlag(false);
+			//FScopedSavingFlag IsSavingFlag(false);
 
 			// Concurrent Part
 			ParallelFor(PackageSaveContexts.Num(), [&PackageSaveContexts](int32 PackageIdx)
