Added check for ArrayIndex >= 0  in Container*PtrToValuePtrInternal, because the pointer math can access invalid memory if negative.

This isn't related to any particular bug, but while debugging I found myself confirming that this was not the case - especially since INDEX_NONE may appear to be a valid value for ArrayIndex.

diff --git a/Engine/Source/Runtime/CoreUObject/Public/UObject/UnrealType.h b/Engine/Source/Runtime/CoreUObject/Public/UObject/UnrealType.h
--- a/Engine/Source/Runtime/CoreUObject/Public/UObject/UnrealType.h
+++ b/Engine/Source/Runtime/CoreUObject/Public/UObject/UnrealType.h
@@ -365,7 +365,7 @@ private:
 
 	FORCEINLINE void* ContainerVoidPtrToValuePtrInternal(void* ContainerPtr, int32 ArrayIndex) const
 	{
-		check(ArrayIndex < ArrayDim);
+		checkf((ArrayIndex >= 0) && (ArrayIndex < ArrayDim), TEXT("Array index out of bounds: %i from an array of size %i"), ArrayIndex, ArrayDim);
 		check(ContainerPtr);
 
 		if (0)
@@ -379,7 +379,7 @@ private:
 
 	FORCEINLINE void* ContainerUObjectPtrToValuePtrInternal(UObject* ContainerPtr, int32 ArrayIndex) const
 	{
-		check(ArrayIndex < ArrayDim);
+		checkf((ArrayIndex >= 0) && (ArrayIndex < ArrayDim), TEXT("Array index out of bounds: %i from an array of size %i"), ArrayIndex, ArrayDim);
 		check(ContainerPtr);
 
 		// in the future, these checks will be tested if the property is supposed be from a UClass
