Allow a Reset TStrongObjectPtr to be destroyed during static de-init on Mac

On Mac the Main thread isn't the Game thread, so this trips the assert in ~FInternalReferenceCollector(), even if the TStrongObjectPtr has been manually Reset prior.

diff --git a/Engine/Source/Runtime/CoreUObject/Public/UObject/StrongObjectPtr.h b/Engine/Source/Runtime/CoreUObject/Public/UObject/StrongObjectPtr.h
--- a/Engine/Source/Runtime/CoreUObject/Public/UObject/StrongObjectPtr.h
+++ b/Engine/Source/Runtime/CoreUObject/Public/UObject/StrongObjectPtr.h
@@ -139,13 +139,24 @@ public:
 
 	FORCEINLINE_DEBUGGABLE void Reset(ObjectType* InNewObject = nullptr)
 	{
-		if (ReferenceCollector)
+		if (InNewObject)
 		{
-			ReferenceCollector->Set(InNewObject);
+			if (ReferenceCollector)
+			{
+				// Update the referenced object
+				ReferenceCollector->Set(InNewObject);
+			}
+			else
+			{
+				// Lazily create the ReferenceCollector to allow TStrongObjectPtr to be used during static initialization
+				ReferenceCollector = MakeUnique<UE4StrongObjectPtr_Private::FInternalReferenceCollector>(InNewObject);
+			}
 		}
-		else if (InNewObject)
+		else
 		{
-			ReferenceCollector = MakeUnique<UE4StrongObjectPtr_Private::FInternalReferenceCollector>(InNewObject);
+			// Destroy the ReferenceCollector immediately to allow TStrongObjectPtr to be manually cleared prior to 
+			// static deinitialization, as not all platforms use the main thread as their game thread
+			ReferenceCollector.Reset();
 		}
 	}
 
