Fixed Save2 when using unversioned serialization where the archive were properly set up but the necessary associated package flag was missing
Changed Save2 to use the old import sorting algo to keep binary consistent between the new and old save algo

Encapsulated "UseUnversionedPropertySerialization" in LinkerSave to set both the archive mode and the necessary package flag
Moved FObjectImportSortHelper to SavePackageUtilities
Cache "s.ProcessPrestreamingRequests" for Save2

diff --git a/Engine/Source/Runtime/CoreUObject/Private/UObject/LinkerSave.cpp b/Engine/Source/Runtime/CoreUObject/Private/UObject/LinkerSave.cpp
--- a/Engine/Source/Runtime/CoreUObject/Private/UObject/LinkerSave.cpp
+++ b/Engine/Source/Runtime/CoreUObject/Private/UObject/LinkerSave.cpp
@@ -359,3 +359,28 @@ void FLinkerSave::UsingCustomVersion(const struct FGuid& Guid)
 		UE_LOG(LogLinker, Warning, TEXT("%s"), *CustomVersionWarning);
 	}
 }
+
+void FLinkerSave::SetUseUnversionedPropertySerialization(bool bInUseUnversioned)
+{
+	FArchiveUObject::SetUseUnversionedPropertySerialization(bInUseUnversioned);
+	if (Saver)
+	{
+		Saver->SetUseUnversionedPropertySerialization(bInUseUnversioned);
+	}
+	if (bInUseUnversioned)
+	{
+		Summary.PackageFlags |= PKG_UnversionedProperties;
+		if (LinkerRoot)
+		{
+			LinkerRoot->SetPackageFlags(PKG_UnversionedProperties);
+		}
+	}
+	else
+	{
+		Summary.PackageFlags &= ~PKG_UnversionedProperties;
+		if (LinkerRoot)
+		{
+			LinkerRoot->ClearPackageFlags(PKG_UnversionedProperties);
+		}
+	}
+}
diff --git a/Engine/Source/Runtime/CoreUObject/Private/UObject/SavePackage.cpp b/Engine/Source/Runtime/CoreUObject/Private/UObject/SavePackage.cpp
--- a/Engine/Source/Runtime/CoreUObject/Private/UObject/SavePackage.cpp
+++ b/Engine/Source/Runtime/CoreUObject/Private/UObject/SavePackage.cpp
@@ -736,131 +736,6 @@ void FPackageNameMapSaver::UpdateLinker(FLinkerSave& Linker, FLinkerLoad* Confor
 	}
 }
 
-/**
- * Helper structure to encapsulate sorting a linker's import table according to the of the import table of the package being conformed against.
- */
-struct FObjectImportSortHelper
-{
-private:
-	/**
-	 * Map of UObject => full name; optimization for sorting.
-	 */
-	TMap<UObject*,FString>			ObjectToFullNameMap;
-
-	/** the linker that we're sorting names for */
-	friend struct TDereferenceWrapper<FObjectImport,FObjectImportSortHelper>;
-
-	/** Comparison function used by Sort */
-	bool operator()( const FObjectImport& A, const FObjectImport& B ) const
-	{
-		int32 Result = 0;
-		if ( A.XObject == nullptr )
-		{
-			Result = 1;
-		}
-		else if ( B.XObject == nullptr )
-		{
-			Result = -1;
-		}
-		else
-		{
-			const FString* FullNameA = ObjectToFullNameMap.Find(A.XObject);
-			const FString* FullNameB = ObjectToFullNameMap.Find(B.XObject);
-			checkSlow(FullNameA);
-			checkSlow(FullNameB);
-
-			Result = FCString::Stricmp(**FullNameA, **FullNameB);
-		}
-
-		return Result < 0;
-	}
-
-public:
-
-	/**
-	 * Sorts imports according to the order in which they occur in the list of imports.  If a package is specified to be conformed against, ensures that the order
-	 * of the imports match the order in which the corresponding imports occur in the old package.
-	 *
-	 * @param	Linker				linker containing the imports that need to be sorted
-	 * @param	LinkerToConformTo	optional linker to conform against.
-	 */
-	void SortImports( FLinkerSave* Linker, FLinkerLoad* LinkerToConformTo=nullptr )
-	{
-		int32 SortStartPosition=0;
-		TArray<FObjectImport>& Imports = Linker->ImportMap;
-		if ( LinkerToConformTo )
-		{
-			// intended to be a copy
-			TArray<FObjectImport> Orig = Imports;
-			Imports.Empty(Imports.Num());
-
-			// this array tracks which imports from the new package exist in the old package
-			TArray<uint8> Used;
-			Used.AddZeroed(Orig.Num());
-
-			TMap<FString,int32> OriginalImportIndexes;
-			OriginalImportIndexes.Reserve(Orig.Num());
-			ObjectToFullNameMap.Reserve(Orig.Num());
-			for ( int32 i = 0; i < Orig.Num(); i++ )
-			{
-				FObjectImport& Import = Orig[i];
-				FString ImportFullName = Import.XObject->GetFullName();
-
-				OriginalImportIndexes.Add( *ImportFullName, i );
-				ObjectToFullNameMap.Add(Import.XObject, *ImportFullName);
-			}
-
-			for( int32 i=0; i<LinkerToConformTo->ImportMap.Num(); i++ )
-			{
-				// determine whether the new version of the package contains this export from the old package
-				int32* OriginalImportPosition = OriginalImportIndexes.Find( *LinkerToConformTo->GetImportFullName(i) );
-				if( OriginalImportPosition )
-				{
-					// this import exists in the new package as well,
-					// create a copy of the FObjectImport located at the original index and place it
-					// into the matching position in the new package's import map
-					FObjectImport* NewImport = new(Imports) FObjectImport( Orig[*OriginalImportPosition] );
-					check(NewImport->XObject == Orig[*OriginalImportPosition].XObject);
-					Used[ *OriginalImportPosition ] = 1;
-				}
-				else
-				{
-					// this import no longer exists in the new package
-					new(Imports)FObjectImport( nullptr );
-				}
-			}
-
-			SortStartPosition = LinkerToConformTo->ImportMap.Num();
-			for( int32 i=0; i<Used.Num(); i++ )
-			{
-				if( !Used[i] )
-				{
-					// the FObjectImport located at pos "i" in the original import table did not
-					// exist in the old package - add it to the end of the import table
-					new(Imports) FObjectImport( Orig[i] );
-				}
-			}
-		}
-		else
-		{
-			ObjectToFullNameMap.Reserve(Linker->ImportMap.Num());
-			for ( int32 ImportIndex = 0; ImportIndex < Linker->ImportMap.Num(); ImportIndex++ )
-			{
-				const FObjectImport& Import = Linker->ImportMap[ImportIndex];
-				if ( Import.XObject )
-				{
-					ObjectToFullNameMap.Add(Import.XObject, Import.XObject->GetFullName());
-				}
-			}
-		}
-
-		if ( SortStartPosition < Linker->ImportMap.Num() )
-		{
-			Sort( &Linker->ImportMap[SortStartPosition], Linker->ImportMap.Num() - SortStartPosition, *this );
-		}
-	}
-};
-
 class FExportReferenceSorter : public FArchiveUObject
 {
 	/**
diff --git a/Engine/Source/Runtime/CoreUObject/Private/UObject/SavePackage/SaveContext.h b/Engine/Source/Runtime/CoreUObject/Private/UObject/SavePackage/SaveContext.h
--- a/Engine/Source/Runtime/CoreUObject/Private/UObject/SavePackage/SaveContext.h
+++ b/Engine/Source/Runtime/CoreUObject/Private/UObject/SavePackage/SaveContext.h
@@ -94,7 +94,8 @@ public:
 
 		bCanUseUnversionedPropertySerialization = CanUseUnversionedPropertySerialization(SaveArgs.TargetPlatform);
 		bTextFormat = FString(Filename).EndsWith(FPackageName::GetTextAssetPackageExtension()) || FString(Filename).EndsWith(FPackageName::GetTextMapPackageExtension());
-		if (const IConsoleVariable* ProcessPrestreamingRequests = IConsoleManager::Get().FindConsoleVariable(TEXT("s.ProcessPrestreamingRequests")))
+		static const IConsoleVariable* ProcessPrestreamingRequests = IConsoleManager::Get().FindConsoleVariable(TEXT("s.ProcessPrestreamingRequests"));
+		if (ProcessPrestreamingRequests)
 		{
 			bIsProcessingPrestreamPackages = ProcessPrestreamingRequests->GetInt() > 0;
 		}
diff --git a/Engine/Source/Runtime/CoreUObject/Private/UObject/SavePackage/SavePackageUtilities.cpp b/Engine/Source/Runtime/CoreUObject/Private/UObject/SavePackage/SavePackageUtilities.cpp
--- a/Engine/Source/Runtime/CoreUObject/Private/UObject/SavePackage/SavePackageUtilities.cpp
+++ b/Engine/Source/Runtime/CoreUObject/Private/UObject/SavePackage/SavePackageUtilities.cpp
@@ -769,6 +769,106 @@ bool IsEditorOnlyObject(const UObject* InObject, bool bCheckRecursive, bool bChe
 	return false;
 }
 
+bool FObjectImportSortHelper::operator()(const FObjectImport& A, const FObjectImport& B) const
+{
+	int32 Result = 0;
+	if (A.XObject == nullptr)
+	{
+		Result = 1;
+	}
+	else if (B.XObject == nullptr)
+	{
+		Result = -1;
+	}
+	else
+	{
+		const FString* FullNameA = ObjectToFullNameMap.Find(A.XObject);
+		const FString* FullNameB = ObjectToFullNameMap.Find(B.XObject);
+		checkSlow(FullNameA);
+		checkSlow(FullNameB);
+
+		Result = FCString::Stricmp(**FullNameA, **FullNameB);
+	}
+
+	return Result < 0;
+}
+
+void FObjectImportSortHelper::SortImports(FLinkerSave* Linker, FLinkerLoad* LinkerToConformTo)
+{
+	int32 SortStartPosition = 0;
+	TArray<FObjectImport>& Imports = Linker->ImportMap;
+	if (LinkerToConformTo)
+	{
+		// intended to be a copy
+		TArray<FObjectImport> Orig = Imports;
+		Imports.Empty(Imports.Num());
+
+		// this array tracks which imports from the new package exist in the old package
+		TArray<uint8> Used;
+		Used.AddZeroed(Orig.Num());
+
+		TMap<FString, int32> OriginalImportIndexes;
+		OriginalImportIndexes.Reserve(Orig.Num());
+		ObjectToFullNameMap.Reserve(Orig.Num());
+		for (int32 i = 0; i < Orig.Num(); i++)
+		{
+			FObjectImport& Import = Orig[i];
+			FString ImportFullName = Import.XObject->GetFullName();
+
+			OriginalImportIndexes.Add(*ImportFullName, i);
+			ObjectToFullNameMap.Add(Import.XObject, *ImportFullName);
+		}
+
+		for (int32 i = 0; i < LinkerToConformTo->ImportMap.Num(); i++)
+		{
+			// determine whether the new version of the package contains this export from the old package
+			int32* OriginalImportPosition = OriginalImportIndexes.Find(*LinkerToConformTo->GetImportFullName(i));
+			if (OriginalImportPosition)
+			{
+				// this import exists in the new package as well,
+				// create a copy of the FObjectImport located at the original index and place it
+				// into the matching position in the new package's import map
+				FObjectImport* NewImport = new(Imports) FObjectImport(Orig[*OriginalImportPosition]);
+				check(NewImport->XObject == Orig[*OriginalImportPosition].XObject);
+				Used[*OriginalImportPosition] = 1;
+			}
+			else
+			{
+				// this import no longer exists in the new package
+				new(Imports)FObjectImport(nullptr);
+			}
+		}
+
+		SortStartPosition = LinkerToConformTo->ImportMap.Num();
+		for (int32 i = 0; i < Used.Num(); i++)
+		{
+			if (!Used[i])
+			{
+				// the FObjectImport located at pos "i" in the original import table did not
+				// exist in the old package - add it to the end of the import table
+				new(Imports) FObjectImport(Orig[i]);
+			}
+		}
+	}
+	else
+	{
+		ObjectToFullNameMap.Reserve(Linker->ImportMap.Num());
+		for (int32 ImportIndex = 0; ImportIndex < Linker->ImportMap.Num(); ImportIndex++)
+		{
+			const FObjectImport& Import = Linker->ImportMap[ImportIndex];
+			if (Import.XObject)
+			{
+				ObjectToFullNameMap.Add(Import.XObject, Import.XObject->GetFullName());
+			}
+		}
+	}
+
+	if (SortStartPosition < Linker->ImportMap.Num())
+	{
+		Sort(&Linker->ImportMap[SortStartPosition], Linker->ImportMap.Num() - SortStartPosition, *this);
+	}
+}
+
 bool FObjectExportSortHelper::operator()(const FObjectExport& A, const FObjectExport& B) const
 {
 	int32 Result = 0;
diff --git a/Engine/Source/Runtime/CoreUObject/Private/UObject/SavePackage/SavePackageUtilities.h b/Engine/Source/Runtime/CoreUObject/Private/UObject/SavePackage/SavePackageUtilities.h
--- a/Engine/Source/Runtime/CoreUObject/Private/UObject/SavePackage/SavePackageUtilities.h
+++ b/Engine/Source/Runtime/CoreUObject/Private/UObject/SavePackage/SavePackageUtilities.h
@@ -73,6 +73,37 @@ struct FCanSkipEditorReferencedPackagesWhenCooking
 	FORCEINLINE operator bool() const { return bCanSkipEditorReferencedPackagesWhenCooking; }
 };
 
+ /**
+  * Helper structure to encapsulate sorting a linker's import table alphabetically, taking into account conforming to other linkers.
+  * @note Save2 should not have to use this sorting long term
+  * @note Conforming is a deprecated feature however
+  */
+struct FObjectImportSortHelper
+{
+private:
+	/**
+	 * Map of UObject => full name; optimization for sorting.
+	 */
+	TMap<UObject*, FString>			ObjectToFullNameMap;
+
+	/** the linker that we're sorting names for */
+	friend struct TDereferenceWrapper<FObjectImport, FObjectImportSortHelper>;
+
+	/** Comparison function used by Sort */
+	bool operator()(const FObjectImport& A, const FObjectImport& B) const;
+
+public:
+
+	/**
+	 * Sorts imports according to the order in which they occur in the list of imports.  If a package is specified to be conformed against, ensures that the order
+	 * of the imports match the order in which the corresponding imports occur in the old package.
+	 *
+	 * @param	Linker				linker containing the imports that need to be sorted
+	 * @param	LinkerToConformTo	optional linker to conform against.
+	 */
+	void SortImports(FLinkerSave* Linker, FLinkerLoad* LinkerToConformTo = nullptr);
+};
+
 /**
  * Helper structure to encapsulate sorting a linker's export table alphabetically, taking into account conforming to other linkers.
  * @note Save2 should not have to use this sorting long term
diff --git a/Engine/Source/Runtime/CoreUObject/Private/UObject/SavePackage2.cpp b/Engine/Source/Runtime/CoreUObject/Private/UObject/SavePackage2.cpp
--- a/Engine/Source/Runtime/CoreUObject/Private/UObject/SavePackage2.cpp
+++ b/Engine/Source/Runtime/CoreUObject/Private/UObject/SavePackage2.cpp
@@ -756,8 +756,6 @@ ESavePackageResult BuildLinker(FSaveContext& SaveContext)
 
 		bool bUseUnversionedProperties = SaveContext.IsUsingUnversionedProperties();
 		SaveContext.Linker->SetUseUnversionedPropertySerialization(bUseUnversionedProperties);
-		SaveContext.Linker->Saver->SetUseUnversionedPropertySerialization(bUseUnversionedProperties);
-
 
 #if WITH_EDITOR
 		if (SaveContext.IsCooking())
@@ -841,7 +839,15 @@ ESavePackageResult BuildLinker(FSaveContext& SaveContext)
 				ObjectImport.ObjectName = ReplacedName;
 			}
 		}
-		Sort(&Linker->ImportMap[0], Linker->ImportMap.Num(), FObjectResourceSortHelper());
+		//Sort(&Linker->ImportMap[0], Linker->ImportMap.Num(), FObjectResourceSortHelper());
+
+		// @todo: To stay consistent with the old save and prevent binary diff between the algo, use the old import sort for now
+		// a future cvar could allow projects use the less expensive sort in their own time down the line
+		FObjectImportSortHelper ImportSortHelper;
+		{
+			SCOPED_SAVETIMER(UPackage_Save_SortImports);
+			ImportSortHelper.SortImports(Linker, nullptr);
+		}
 		Linker->Summary.ImportCount = Linker->ImportMap.Num();
 	}
 
@@ -862,8 +868,10 @@ ESavePackageResult BuildLinker(FSaveContext& SaveContext)
 		}
 		//Sort(&Linker->ExportMap[0], Linker->ExportMap.Num(), FObjectResourceSortHelper());
 
-		// @todo: To remove but for now object sort freaking matter in an incidental manner where it should be properly tracked with dependencies
-		// for example where FAnimInstanceProxy PostLoad actually depends on  UAnimBlueprintGeneratedClass PostLoad to be properly initialized.
+		// @todo: To stay consistent with the old save and prevent binary diff between the algo, use the old import sort for now
+		// a future cvar could allow projects use the less expensive sort in their own time down the line
+		// Also, currently the export sort order matters in an incidental manner where it should be properly tracked with dependencies instead.
+		// for example where FAnimInstanceProxy PostLoad actually depends on UAnimBlueprintGeneratedClass PostLoad to be properly initialized.
 		FObjectExportSortHelper ExportSortHelper;
 		{
 			SCOPED_SAVETIMER(UPackage_Save_SortExports);
diff --git a/Engine/Source/Runtime/CoreUObject/Public/UObject/LinkerSave.h b/Engine/Source/Runtime/CoreUObject/Public/UObject/LinkerSave.h
--- a/Engine/Source/Runtime/CoreUObject/Public/UObject/LinkerSave.h
+++ b/Engine/Source/Runtime/CoreUObject/Public/UObject/LinkerSave.h
@@ -101,6 +101,12 @@ public:
 	virtual void SetSerializeContext(FUObjectSerializeContext* InLoadContext) override;
 	FUObjectSerializeContext* GetSerializeContext() override;
 	virtual void UsingCustomVersion(const struct FGuid& Guid) override;
+	/**
+	 * Sets whether tagged property serialization should be replaced by faster unversioned serialization.
+	 * This sets it on itself, the summary, the actual Saver Archive if any and set the proper associated flag on the LinkerRoot
+	 */
+	virtual void SetUseUnversionedPropertySerialization(bool bInUseUnversioned) override;
+
 
 #if WITH_EDITOR
 	// proxy for debugdata
