CIS Fix: Change the ::Unlock method for BulkData2 to be const to match the original API.

- This stops code compiling fine in the editor then breaking when compiled for a different config.

#rnx

diff --git a/Engine/Source/Runtime/CoreUObject/Private/Serialization/BulkData2.cpp b/Engine/Source/Runtime/CoreUObject/Private/Serialization/BulkData2.cpp
--- a/Engine/Source/Runtime/CoreUObject/Private/Serialization/BulkData2.cpp
+++ b/Engine/Source/Runtime/CoreUObject/Private/Serialization/BulkData2.cpp
@@ -944,7 +944,7 @@ const void* FBulkDataBase::LockReadOnly() const
 	return GetDataBufferReadOnly();
 }
 
-void FBulkDataBase::Unlock()
+void FBulkDataBase::Unlock() const
 {
 	checkf(LockStatus != LOCKSTATUS_Unlocked, TEXT("Attempting to unlock a BulkData object that is not locked"));
 	LockStatus = LOCKSTATUS_Unlocked;
@@ -952,7 +952,9 @@ void FBulkDataBase::Unlock()
 	// Free pointer if we're guaranteed to only to access the data once.
 	if (IsSingleUse())
 	{
-		FreeData();
+		// Cast away const so that we can match the original bulkdata api
+		// which had ::Unlock as const too.
+		const_cast<FBulkDataBase*>(this)->FreeData();
 	}
 }
 
diff --git a/Engine/Source/Runtime/CoreUObject/Public/Serialization/BulkData2.h b/Engine/Source/Runtime/CoreUObject/Public/Serialization/BulkData2.h
--- a/Engine/Source/Runtime/CoreUObject/Public/Serialization/BulkData2.h
+++ b/Engine/Source/Runtime/CoreUObject/Public/Serialization/BulkData2.h
@@ -136,10 +136,9 @@ protected:
 	void Serialize(FArchive& Ar, UObject* Owner, int32 Index, bool bAttemptFileMapping, int32 ElementSize);
 
 public:
-	// Unimplemented:
 	void* Lock(uint32 LockFlags);
 	const void* LockReadOnly() const;
-	void Unlock();
+	void Unlock() const;
 	bool IsLocked() const;
 
 	void* Realloc(int64 SizeInBytes);
