Fixed bug where setting an external package to a package wasn't a noop as it should have resulting in GetObjectsInPackage to return the package itself after undo/redo operation.

diff --git a/Engine/Source/Runtime/CoreUObject/Private/UObject/UObjectBase.cpp b/Engine/Source/Runtime/CoreUObject/Private/UObject/UObjectBase.cpp
--- a/Engine/Source/Runtime/CoreUObject/Private/UObject/UObjectBase.cpp
+++ b/Engine/Source/Runtime/CoreUObject/Private/UObject/UObjectBase.cpp
@@ -246,6 +246,13 @@ UPackage* UObjectBase::GetExternalPackageInternal() const
 
 void UObjectBase::SetExternalPackage(UPackage* InPackage)
 {
+	// if we have no outer, consider this a package, packages have themselves as their external package and that shouldn't be added to the object hash
+	if (OuterPrivate == nullptr)
+	{
+		// Just validate that we tried to set ourselves or nothing as our external package which is a no-op. anything else is illegal for package
+		check(GetClass()->IsChildOf(UPackage::StaticClass()) && (InPackage == this || InPackage == nullptr));
+		return;
+	}
 	HashObjectExternalPackage(this, InPackage);
 	if (InPackage)
 	{
