PackageResourceManager and PackagePaths
Create a new interface IPackageResourceManager that handles all attempts to load packages from storage. The default implementation passes through to loading from IFileManager. Future implementations will load from other services such as the EditorDomain and TargetDomain.
Add FPackagePath class that holds a LocalFilePath or LongPackageName, or a relative path to a package under a mounted content root (which can be converted into either the LocalFilePath or LongPackageName).
Modify LinkerLoad, BulkData, and AsyncLoader to use FPackagePath and IResourceManager.
Modify UPackage and FLinkerLoad to have a FPackagePath instead of a Filename.
#rn Major Core

diff --git a/Engine/Plugins/Developer/Concert/ConcertSync/ConcertSyncClient/Source/ConcertSyncClient/Private/ConcertClientPackageManager.cpp b/Engine/Plugins/Developer/Concert/ConcertSync/ConcertSyncClient/Source/ConcertSyncClient/Private/ConcertClientPackageManager.cpp
--- a/Engine/Plugins/Developer/Concert/ConcertSync/ConcertSyncClient/Source/ConcertSyncClient/Private/ConcertClientPackageManager.cpp
+++ b/Engine/Plugins/Developer/Concert/ConcertSync/ConcertSyncClient/Source/ConcertSyncClient/Private/ConcertClientPackageManager.cpp
@@ -209,7 +209,7 @@ void FConcertClientPackageManager::QueueDirtyPackagesForReload()
 		if (WorldPackage && WorldPackage->IsDirty() &&
 			(WorldPackage->HasAnyPackageFlags(PKG_PlayInEditor | PKG_InMemoryOnly) ||
 			WorldPackage->HasAnyFlags(RF_Transient) ||
-			WorldPackage->FileName != WorldPackage->GetFName()))
+			WorldPackage->GetLoadedPath().GetPackageFName() != WorldPackage->GetFName()))
 		{
 			DirtyPkgs.Remove(WorldPackage);
 		}
diff --git a/Engine/Plugins/FX/Niagara/Source/NiagaraEditor/Private/Toolkits/NiagaraSystemToolkit.cpp b/Engine/Plugins/FX/Niagara/Source/NiagaraEditor/Private/Toolkits/NiagaraSystemToolkit.cpp
--- a/Engine/Plugins/FX/Niagara/Source/NiagaraEditor/Private/Toolkits/NiagaraSystemToolkit.cpp
+++ b/Engine/Plugins/FX/Niagara/Source/NiagaraEditor/Private/Toolkits/NiagaraSystemToolkit.cpp
@@ -226,7 +226,7 @@ void FNiagaraSystemToolkit::InitializeWithSystem(const EToolkitMode::Type Mode,
 	{
 		FString ExportText;
 		SystemViewModel->DumpToText(ExportText);
-		FString FilePath = System->GetOutermost()->FileName.ToString();
+		FString FilePath = System->GetOutermost()->GetLoadedPath().GetPackageName();
 		FString PathPart, FilenamePart, ExtensionPart;
 		FPaths::Split(FilePath, PathPart, FilenamePart, ExtensionPart);
 
@@ -303,7 +303,7 @@ void FNiagaraSystemToolkit::InitializeWithEmitter(const EToolkitMode::Type Mode,
 	{
 		FString ExportText;
 		SystemViewModel->DumpToText(ExportText);
-		FString FilePath = Emitter->GetOutermost()->FileName.ToString();
+		FString FilePath = Emitter->GetOutermost()->GetLoadedPath().GetPackageName();
 		FString PathPart, FilenamePart, ExtensionPart;
 		FPaths::Split(FilePath, PathPart, FilenamePart, ExtensionPart);
 
@@ -1592,11 +1592,11 @@ bool FNiagaraSystemToolkit::OnRequestClose()
 
 		if (SystemToolkitMode == ESystemToolkitMode::System)
 		{
-			FilePath = System->GetOutermost()->FileName.ToString();
+			FilePath = System->GetOutermost()->GetLoadedPath().GetPackageName();
 		}
 		else if (SystemToolkitMode == ESystemToolkitMode::Emitter)
 		{
-			FilePath = Emitter->GetOutermost()->FileName.ToString();
+			FilePath = Emitter->GetOutermost()->GetLoadedPath().GetPackageName();
 		}
 
 		FString PathPart, FilenamePart, ExtensionPart;
diff --git a/Engine/Plugins/Online/OnlineFramework/Source/Hotfix/Private/OnlineHotfixManager.cpp b/Engine/Plugins/Online/OnlineFramework/Source/Hotfix/Private/OnlineHotfixManager.cpp
--- a/Engine/Plugins/Online/OnlineFramework/Source/Hotfix/Private/OnlineHotfixManager.cpp
+++ b/Engine/Plugins/Online/OnlineFramework/Source/Hotfix/Private/OnlineHotfixManager.cpp
@@ -1132,7 +1132,7 @@ bool UOnlineHotfixManager::HotfixPakFile(const FCloudFileHeader& FileHeader)
 			UPackage* Package = *it;
 			if (Package && Package->ContainsMap())
 			{
-				const FString FileName = FPackageName::LongPackageNameToFilename(Package->FileName.ToString(), FPackageName::GetMapPackageExtension());
+				const FString FileName = FPackageName::LongPackageNameToFilename(Package->GetLoadedPath().GetPackageName(), FPackageName::GetMapPackageExtension());
 				if (PakFile->PakContains(FileName))
 				{
 					bHotfixNeedsMapReload = true;
diff --git a/Engine/Source/Developer/HierarchicalLODUtilities/Private/HierarchicalLODUtilities.cpp b/Engine/Source/Developer/HierarchicalLODUtilities/Private/HierarchicalLODUtilities.cpp
--- a/Engine/Source/Developer/HierarchicalLODUtilities/Private/HierarchicalLODUtilities.cpp
+++ b/Engine/Source/Developer/HierarchicalLODUtilities/Private/HierarchicalLODUtilities.cpp
@@ -180,10 +180,8 @@ UPackage* FHierarchicalLODUtilities::CreateOrRetrieveLevelHLODPackage(const ULev
 	HLODPackage->FullyLoad();
 	HLODPackage->SetPackageFlags(PKG_ContainsMapData);		// PKG_ContainsMapData required so FEditorFileUtils::GetDirtyContentPackages can treat this as a map package
 
-	// Target level filename
-	const FString HLODLevelFileName = FPackageName::LongPackageNameToFilename(HLODLevelPackageName);
-	// This is a hack to avoid save file dialog when we will be saving HLOD map package
-	HLODPackage->FileName = FName(*HLODLevelFileName);
+	// Target PackagePath; this is a hack to avoid save file dialog when we will be saving HLOD map package
+	HLODPackage->SetLoadedPath(FPackagePath::FromPackageNameChecked(HLODLevelPackageName));
 
 	return HLODPackage;
 }
@@ -224,10 +222,8 @@ UPackage* FHierarchicalLODUtilities::CreateOrRetrieveLevelHLODPackage(const ULev
 	HLODPackage->Modify();
 	HLODPackage->SetPackageFlags(PKG_ContainsMapData);		// PKG_ContainsMapData required so FEditorFileUtils::GetDirtyContentPackages can treat this as a map package
 
-	// Target level filename
-	const FString HLODLevelFileName = FPackageName::LongPackageNameToFilename(HLODLevelPackageName);
-	// This is a hack to avoid save file dialog when we will be saving HLOD map package
-	HLODPackage->FileName = FName(*HLODLevelFileName);
+	// Target PackagePath; this is a hack to avoid save file dialog when we will be saving HLOD map package
+	HLODPackage->SetLoadedPath(FPackagePath::FromPackageNameChecked(HLODLevelPackageName));
 
 	return HLODPackage;
 }
@@ -271,10 +267,8 @@ UPackage* CreateOrRetrieveImposterMeshPackage(const UMaterialInterface* InImpost
 	UPackage* MeshPackage = CreatePackage( *MeshPackageName);
 	MeshPackage->FullyLoad();
 
-	// Target filename
-	const FString MeshPackageFileName = FPackageName::LongPackageNameToFilename(MeshPackageName);
-	// This is a hack to avoid save file dialog when we will be saving imposter mesh package
-	MeshPackage->FileName = FName(*MeshPackageFileName);
+	// Target PackagePath; this is a hack to avoid save file dialog when we will be saving imposter mesh package
+	MeshPackage->SetLoadedPath(FPackagePath::FromPackageNameChecked(MeshPackageName));
 
 	return MeshPackage;
 }
diff --git a/Engine/Source/Editor/UnrealEd/Private/AutoReimport/ContentDirectoryMonitor.cpp b/Engine/Source/Editor/UnrealEd/Private/AutoReimport/ContentDirectoryMonitor.cpp
--- a/Engine/Source/Editor/UnrealEd/Private/AutoReimport/ContentDirectoryMonitor.cpp
+++ b/Engine/Source/Editor/UnrealEd/Private/AutoReimport/ContentDirectoryMonitor.cpp
@@ -434,39 +434,46 @@ void FContentDirectoryMonitor::ProcessModifications(const DirectoryWatcher::FTim
 					const FString PackagePath = UPackageTools::SanitizePackageName(MountedContentPath / FPaths::GetPath(Change.Filename.Get()));
 					const FString FullDestPath = PackagePath / NewAssetName;
 
-					if (ExistingPackage && ExistingPackage->FileName.ToString() == FullDestPath)
+					const FText SrcPathText = FText::FromString(Assets[0].PackageName.ToString());
+					const FText DstPathText = FText::FromString(FullDestPath);
+					FPackagePath DestPackagePath;
+					if (!FPackagePath::TryFromMountedName(FullDestPath, DestPackagePath))
 					{
-						// No need to process this asset - it's already been moved to the right location
-						Cache.CompleteTransaction(MoveTemp(Change));
-						continue;
-					}
-
-					const FText SrcPathText = FText::FromString(Assets[0].PackageName.ToString()),
-						DstPathText = FText::FromString(FullDestPath);
-
-					if (FPackageName::DoesPackageExist(*FullDestPath))
-					{
-						Context.AddMessage(EMessageSeverity::Warning, FText::Format(LOCTEXT("MoveWarning_ExistingAsset", "Can't move {0} to {1} - one already exists."), SrcPathText, DstPathText));
+						Context.AddMessage(EMessageSeverity::Warning, FText::Format(LOCTEXT("MoveWarning_ExistingAsset", "Can't move {0} to {1} - {1} is not in a mounted path."), SrcPathText, DstPathText));
 					}
 					else
 					{
-						TArray<FAssetRenameData> RenameData;
-						RenameData.Emplace(Asset, PackagePath, NewAssetName);
+						if (ExistingPackage && ExistingPackage->GetLoadedPath() == DestPackagePath)
+						{
+							// No need to process this asset - it's already been moved to the right location
+							Cache.CompleteTransaction(MoveTemp(Change));
+							continue;
+						}
 
-						Context.AddMessage(EMessageSeverity::Info, FText::Format(LOCTEXT("Success_MovedAsset", "Moving asset {0} to {1}."),	SrcPathText, DstPathText));
+						if (FPackageName::DoesPackageExist(DestPackagePath, &DestPackagePath))
+						{
+							Context.AddMessage(EMessageSeverity::Warning, FText::Format(LOCTEXT("MoveWarning_ExistingAsset", "Can't move {0} to {1} - one already exists."), SrcPathText, DstPathText));
+						}
+						else
+						{
+							TArray<FAssetRenameData> RenameData;
+							RenameData.Emplace(Asset, PackagePath, NewAssetName);
 
-						FModuleManager::LoadModuleChecked<FAssetToolsModule>("AssetTools").Get().RenameAssetsWithDialog(RenameData);
+							Context.AddMessage(EMessageSeverity::Info, FText::Format(LOCTEXT("Success_MovedAsset", "Moving asset {0} to {1}."), SrcPathText, DstPathText));
 
-						TArray<FString> Filenames;
-						Filenames.Add(FullFilename);
+							FModuleManager::LoadModuleChecked<FAssetToolsModule>("AssetTools").Get().RenameAssetsWithDialog(RenameData);
 
-						// Update the reimport file names
-						FReimportManager::Instance()->UpdateReimportPaths(Asset, Filenames);
-						Asset->MarkPackageDirty();
+							TArray<FString> Filenames;
+							Filenames.Add(FullFilename);
 
-						if (!bAssetWasDirty)
-						{
-							OutPackagesToSave.Add(Asset->GetOutermost());
+							// Update the reimport file names
+							FReimportManager::Instance()->UpdateReimportPaths(Asset, Filenames);
+							Asset->MarkPackageDirty();
+
+							if (!bAssetWasDirty)
+							{
+								OutPackagesToSave.Add(Asset->GetOutermost());
+							}
 						}
 					}
 				}
diff --git a/Engine/Source/Editor/UnrealEd/Private/Commandlets/ContentCommandlets.cpp b/Engine/Source/Editor/UnrealEd/Private/Commandlets/ContentCommandlets.cpp
--- a/Engine/Source/Editor/UnrealEd/Private/Commandlets/ContentCommandlets.cpp
+++ b/Engine/Source/Editor/UnrealEd/Private/Commandlets/ContentCommandlets.cpp
@@ -11,13 +11,14 @@
 #include "Misc/Paths.h"
 #include "Misc/ConfigCacheIni.h"
 #include "Modules/ModuleManager.h"
-#include "UObject/ObjectMacros.h"
-#include "UObject/Object.h"
 #include "UObject/Class.h"
+#include "UObject/MetaData.h"
+#include "UObject/Object.h"
+#include "UObject/ObjectMacros.h"
+#include "UObject/Package.h"
+#include "UObject/PackageResourceManager.h"
 #include "UObject/UObjectHash.h"
 #include "UObject/UObjectIterator.h"
-#include "UObject/Package.h"
-#include "UObject/MetaData.h"
 #include "Misc/PackageName.h"
 #include "Misc/EngineVersion.h"
 #include "Misc/RedirectCollector.h"
@@ -575,14 +576,10 @@ void UResavePackagesCommandlet::LoadAndSaveOnePackage(const FString& Filename)
 		return;
 	}
 
+	FPackagePath PackagePath = FPackagePath::FromLocalPath(Filename);
 	if (CollectionFilter.Num())
 	{
-		FString PackageNameToCreate;
-		if (!FPackageName::TryConvertFilenameToLongPackageName(Filename, PackageNameToCreate))
-		{
-			PackageNameToCreate = Filename;
-		}
-
+		FString PackageNameToCreate = PackagePath.GetPackageNameOrFallback();
 		FName PackageName(*PackageNameToCreate);
 		if (!CollectionFilter.Contains(PackageName))
 		{
@@ -592,17 +589,13 @@ void UResavePackagesCommandlet::LoadAndSaveOnePackage(const FString& Filename)
 
 	if (bOnlyMaterials)
 	{
-		FString PackageNameToCreate;
-		if (!FPackageName::TryConvertFilenameToLongPackageName(Filename, PackageNameToCreate))
-		{
-			PackageNameToCreate = Filename;
-		}
+		FString PackageNameToCreate = PackagePath.GetPackageNameOrFallback();
+		FName PackageFNameToCreate(*PackageNameToCreate);
 
 		FAssetRegistryModule& AssetRegistryModule = FModuleManager::LoadModuleChecked<FAssetRegistryModule>(TEXT("AssetRegistry"));
 		IAssetRegistry& AssetRegistry = AssetRegistryModule.Get();
-		FName PackageName(*PackageNameToCreate);
 		TArray<FAssetData> PackageAssetData;
-		if (!AssetRegistry.GetAssetsByPackageName(PackageName, PackageAssetData, true))
+		if (!AssetRegistry.GetAssetsByPackageName(PackageFNameToCreate, PackageAssetData, true))
 		{
 			return;
 		}
@@ -652,7 +645,7 @@ void UResavePackagesCommandlet::LoadAndSaveOnePackage(const FString& Filename)
 		// Get the package linker.
 		VerboseMessage(TEXT("Pre GetPackageLinker"));
 
-		FLinkerLoad* Linker = LoadPackageLinker(nullptr, *Filename, LOAD_NoVerify);
+		FLinkerLoad* Linker = LoadPackageLinker(nullptr, PackagePath, LOAD_NoVerify);
 	
 		// Bail early if we don't have a valid linker (package was out of date, etc)
 		if( !Linker )
@@ -675,7 +668,9 @@ void UResavePackagesCommandlet::LoadAndSaveOnePackage(const FString& Filename)
 
 			// Only rebuild static meshes on load for the to be saved package.
 			extern ENGINE_API FName GStaticMeshPackageNameToRebuild;
-			GStaticMeshPackageNameToRebuild = FName(*FPackageName::FilenameToLongPackageName(Filename));
+			FName PackageFName = PackagePath.GetPackageFName();
+			check(!PackageFName.IsNone());
+			GStaticMeshPackageNameToRebuild = PackageFName;
 
 			// Assert if package couldn't be opened so we have no chance of messing up saving later packages.
 			UPackage* Package = nullptr;
@@ -2407,7 +2402,7 @@ int32 UWrangleContentCommandlet::Main( const FString& Params )
 			}
 		}
 
-		// now add them to the list of all packges to load
+		// now add them to the list of all packages to load
 		for (int32 PackageIndex = 0; PackageIndex < PerMapPackagesToLoad.Num(); PackageIndex++)
 		{
 			PackagesToFullyLoad.Add(TEXT("Package"), *PerMapPackagesToLoad[PackageIndex]);
@@ -2437,102 +2432,105 @@ int32 UWrangleContentCommandlet::Main( const FString& Params )
 			{
 				// save a copy of the packagename (not a reference in case the PackgesToLoad array gets realloced)
 				FString PackageName = PackagesToLoad[PackageIndex];
-				FString PackageFilename;
-
-				if( FPackageName::DoesPackageExist( PackageName, NULL, &PackageFilename ) == true )
+				FPackagePath PackagePath;
+				if (!FPackagePath::TryFromMountedName(PackageName, PackagePath) ||
+					!FPackageName::DoesPackageExist(PackagePath, &PackagePath))
 				{
-					SET_WARN_COLOR(COLOR_WHITE);
-					UE_LOG(LogContentCommandlet, Warning, TEXT("Fully loading %s..."), *PackageFilename);
-					CLEAR_WARN_COLOR();
+					continue;
+				}
+
+				FString PackageFilename = PackagePath.GetLocalFullPath();
+				SET_WARN_COLOR(COLOR_WHITE);
+				UE_LOG(LogContentCommandlet, Warning, TEXT("Fully loading %s..."), *PackageFilename);
+				CLEAR_WARN_COLOR();
 
 	// @todo josh: track redirects in this package and then save the package instead of copy it if there were redirects
 	// or make sure that the following redirects marks the package dirty (which maybe it shouldn't do in the editor?)
 
-					// load the package fully
-					UPackage* Package = LoadPackage(nullptr, *PackageFilename, LOAD_None);
+				// load the package fully
+				UPackage* Package = LoadPackage(nullptr, PackagePath, LOAD_None);
+
+				FLinkerLoad* Linker = LoadPackageLinker(nullptr, PackagePath, LOAD_Quiet | LOAD_NoWarn | LOAD_NoVerify);
 
-					FLinkerLoad* Linker = LoadPackageLinker(nullptr, *PackageFilename, LOAD_Quiet | LOAD_NoWarn | LOAD_NoVerify);
+				// look for special package types
+				bool bIsMap = Linker->ContainsMap();
+				bool bIsScriptPackage = Linker->ContainsCode();
 
-					// look for special package types
-					bool bIsMap = Linker->ContainsMap();
-					bool bIsScriptPackage = Linker->ContainsCode();
+				// collect all public objects loaded
+				for(FObjectIterator ObjectIt; ObjectIt; ++ObjectIt)
+				{
+					UObject* Object = *ObjectIt;
 
-					// collect all public objects loaded
-					for(FObjectIterator ObjectIt; ObjectIt; ++ObjectIt)
+					// record all public referenced objects (skipping over top level packages)
+					if (/*Object->HasAnyFlags(RF_Public) &&*/ Object->GetOuter() != NULL)
 					{
-						UObject* Object = *ObjectIt;
+						// is this public object in a fully loaded package?
+						bool bIsObjectInFullyLoadedPackage = Object->IsIn(Package);
 
-						// record all public referenced objects (skipping over top level packages)
-						if (/*Object->HasAnyFlags(RF_Public) &&*/ Object->GetOuter() != NULL)
+						if (bIsMap && bIsObjectInFullyLoadedPackage && Object->HasAnyFlags(RF_Public))
 						{
-							// is this public object in a fully loaded package?
-							bool bIsObjectInFullyLoadedPackage = Object->IsIn(Package);
-
-							if (bIsMap && bIsObjectInFullyLoadedPackage && Object->HasAnyFlags(RF_Public))
-							{
-								UE_LOG(LogContentCommandlet, Warning, TEXT("Clearing public flag on map object %s"), *Object->GetFullName());
-								Object->ClearFlags(RF_Public);
-								// mark that we need to save the package since we modified it (instead of copying it)
-								Object->MarkPackageDirty();
-							}
-							else
-							{
-								// record that this object was referenced
-								ReferenceObject(Object, AllReferencedPublicObjects, bIsObjectInFullyLoadedPackage);
-							}
+							UE_LOG(LogContentCommandlet, Warning, TEXT("Clearing public flag on map object %s"), *Object->GetFullName());
+							Object->ClearFlags(RF_Public);
+							// mark that we need to save the package since we modified it (instead of copying it)
+							Object->MarkPackageDirty();
+						}
+						else
+						{
+							// record that this object was referenced
+							ReferenceObject(Object, AllReferencedPublicObjects, bIsObjectInFullyLoadedPackage);
 						}
 					}
+				}
 
-					// add any sublevels of this world to the list of levels to load
-					for (TObjectIterator<UWorld> WorldIt; WorldIt; ++WorldIt)
+				// add any sublevels of this world to the list of levels to load
+				for (TObjectIterator<UWorld> WorldIt; WorldIt; ++WorldIt)
+				{
+					UWorld*	World = *WorldIt;
+					// iterate over streaming level objects loading the levels.
+					for (ULevelStreaming* StreamingLevel : World->GetStreamingLevels())
 					{
-						UWorld*	World = *WorldIt;
-						// iterate over streaming level objects loading the levels.
-						for (ULevelStreaming* StreamingLevel : World->GetStreamingLevels())
+						if (StreamingLevel)
 						{
-							if (StreamingLevel)
+							FString SubLevelName = StreamingLevel->GetWorldAssetPackageName();
+							// add this sublevel's package to the list of packages to load if it's not already in the master list of packages
+							if (PackagesToFullyLoad.FindKey(SubLevelName) == NULL)
 							{
-								FString SubLevelName = StreamingLevel->GetWorldAssetPackageName();
-								// add this sublevel's package to the list of packages to load if it's not already in the master list of packages
-								if (PackagesToFullyLoad.FindKey(SubLevelName) == NULL)
-								{
-									PackagesToLoad.AddUnique(SubLevelName);
-								}
+								PackagesToLoad.AddUnique(SubLevelName);
 							}
 						}
 					}
+				}
 
-					// save/copy the package if desired, and only if it's not a script package (script code is
-					// not cutdown, so we always use original script code)
-					if (bShouldSavePackages && !bIsScriptPackage)
-					{
-						// make the name of the location to put the package
-						FString CutdownPackageName = MakeCutdownFilename(PackageFilename);
+				// save/copy the package if desired, and only if it's not a script package (script code is
+				// not cutdown, so we always use original script code)
+				if (bShouldSavePackages && !bIsScriptPackage)
+				{
+					// make the name of the location to put the package
+					FString CutdownPackageName = MakeCutdownFilename(PackageFilename);
 						
-						// if the package was modified by loading it, then we should save the package
-						if (Package->IsDirty())
+					// if the package was modified by loading it, then we should save the package
+					if (Package->IsDirty())
+					{
+						// save the fully load packages
+						UE_LOG(LogContentCommandlet, Warning, TEXT("Saving fully loaded package %s..."), *CutdownPackageName);
+						if (!SavePackageHelper(Package, CutdownPackageName))
 						{
-							// save the fully load packages
-							UE_LOG(LogContentCommandlet, Warning, TEXT("Saving fully loaded package %s..."), *CutdownPackageName);
-							if (!SavePackageHelper(Package, CutdownPackageName))
-							{
-								UE_LOG(LogContentCommandlet, Error, TEXT("Failed to save package %s..."), *CutdownPackageName);
-							}
+							UE_LOG(LogContentCommandlet, Error, TEXT("Failed to save package %s..."), *CutdownPackageName);
 						}
-						else
+					}
+					else
+					{
+						UE_LOG(LogContentCommandlet, Warning, TEXT("Copying fully loaded package %s..."), *CutdownPackageName);
+						// copy the unmodified file (faster than saving) (0 is success)
+						if (IFileManager::Get().Copy(*CutdownPackageName, *PackageFilename) != 0)
 						{
-							UE_LOG(LogContentCommandlet, Warning, TEXT("Copying fully loaded package %s..."), *CutdownPackageName);
-							// copy the unmodified file (faster than saving) (0 is success)
-							if (IFileManager::Get().Copy(*CutdownPackageName, *PackageFilename) != 0)
-							{
-								UE_LOG(LogContentCommandlet, Error, TEXT("Failed to copy package to %s..."), *CutdownPackageName);
-							}
+							UE_LOG(LogContentCommandlet, Error, TEXT("Failed to copy package to %s..."), *CutdownPackageName);
 						}
 					}
-
-					// close this package
-					CollectGarbage(RF_NoFlags);
 				}
+
+				// close this package
+				CollectGarbage(RF_NoFlags);
 			}
 		}
 
@@ -2562,8 +2560,13 @@ int32 UWrangleContentCommandlet::Main( const FString& Params )
 		// Iterate over all files doing stuff.
 		for (int32 PackageIndex = 0; PackageIndex < AllPackages.Num(); PackageIndex++)
 		{
-			FString PackageFilename(AllPackages[PackageIndex]);
-			FString PackageName = FPackageName::FilenameToLongPackageName(PackageFilename);
+			FPackagePath PackagePath;
+			if (!FPackagePath::TryFromMountedName(AllPackages[PackageIndex], PackagePath))
+			{
+				UE_LOG(LogContentCommandlet, Warning, TEXT("Skipping unmounted path %s..."), *AllPackages[PackageIndex]);
+				continue;
+			}
+			FString PackageName = PackagePath.GetPackageName();
 
 			// the list of objects in this package
 			FPackageObjects* PackageObjs = NULL;
@@ -2571,9 +2574,9 @@ int32 UWrangleContentCommandlet::Main( const FString& Params )
 			// this will be set to true if every object in the package is unnecessary
 			bool bAreAllObjectsUnnecessary = false;
 
-			if (FPaths::GetExtension(PackageFilename, true) == FPackageName::GetMapPackageExtension() )
+			if (PackagePath.GetHeaderExtension() == EPackageExtension::Map)
 			{
-				UE_LOG(LogContentCommandlet, Warning, TEXT("Skipping map %s..."), *PackageFilename);
+				UE_LOG(LogContentCommandlet, Warning, TEXT("Skipping map %s..."), *PackagePath.GetDebugName());
 				continue;
 			}
 			else
@@ -2585,25 +2588,25 @@ int32 UWrangleContentCommandlet::Main( const FString& Params )
 				// and mark the whole package as unreferenced
 				if (PackageObjs == NULL)
 				{
-					UE_LOG(LogContentCommandlet, Warning, TEXT("No objects in %s were referenced..."), *PackageFilename);
-					new(UnnecessaryPublicObjects) FUnreferencedObject(PackageName, 
-						TEXT("ENTIRE PACKAGE"), IFileManager::Get().FileSize(*PackageFilename));
+					UE_LOG(LogContentCommandlet, Warning, TEXT("No objects in %s were referenced..."), *PackagePath.GetDebugName());
+					new(UnnecessaryPublicObjects) FUnreferencedObject(PackageName,
+						TEXT("ENTIRE PACKAGE"), IPackageResourceManager::Get().FileSize(PackagePath));
 
-					// all objects in this package are unnecasstry
+					// all objects in this package are unnecessary
 					bAreAllObjectsUnnecessary = true;
 				}
 				else if (PackageObjs->bIsFullyLoadedPackage)
 				{
-					UE_LOG(LogContentCommandlet, Warning, TEXT("Skipping fully loaded package %s..."), *PackageFilename);
+					UE_LOG(LogContentCommandlet, Warning, TEXT("Skipping fully loaded package %s..."), *PackagePath.GetDebugName());
 					continue;
 				}
 				else
 				{
-					UE_LOG(LogContentCommandlet, Warning, TEXT("Scanning %s..."), *PackageFilename);
+					UE_LOG(LogContentCommandlet, Warning, TEXT("Scanning %s..."), *PackagePath.GetDebugName());
 				}
 			}
 
-			FLinkerLoad* Linker = LoadPackageLinker(nullptr, *PackageFilename, LOAD_Quiet | LOAD_NoWarn | LOAD_NoVerify);
+			FLinkerLoad* Linker = LoadPackageLinker(nullptr, PackagePath, LOAD_Quiet | LOAD_NoWarn | LOAD_NoVerify);
 
 			// go through the exports in the package, looking for public objects
 			for (int32 ExportIndex = 0; ExportIndex < Linker->ExportMap.Num(); ExportIndex++)
@@ -2614,7 +2617,7 @@ int32 UWrangleContentCommandlet::Main( const FString& Params )
 				// some packages may have brokenness in them so we want to just continue so we can wrangle
 				if( Export.ObjectName == NAME_None )
 				{
-					UE_LOG(LogContentCommandlet, Warning, TEXT( "    Export.ObjectName == NAME_None  for Package: %s " ), *PackageFilename );
+					UE_LOG(LogContentCommandlet, Warning, TEXT( "    Export.ObjectName == NAME_None  for Package: %s " ), *PackagePath.GetDebugName());
 					continue;
 				}
 
@@ -2637,11 +2640,11 @@ int32 UWrangleContentCommandlet::Main( const FString& Params )
 					}
 
 					// look for existing entry
-					FPackageObjects* ObjectsInPackage = UnnecessaryObjectsByPackage.Find(*PackageFilename);
+					FPackageObjects* ObjectsInPackage = UnnecessaryObjectsByPackage.Find(PackagePath.GetLocalBaseFilenameWithPath());
 					// if not found, make a new one
 					if (ObjectsInPackage == NULL)
 					{
-						ObjectsInPackage = &UnnecessaryObjectsByPackage.Add(*PackageFilename, FPackageObjects());
+						ObjectsInPackage = &UnnecessaryObjectsByPackage.Add(PackagePath.GetLocalBaseFilenameWithPath(), FPackageObjects());
 					}
 
 					// get object's class
@@ -2829,7 +2832,7 @@ int32 UWrangleContentCommandlet::Main( const FString& Params )
 			if (PackageIt.Value().ReferencedObjects.Num() == 0)
 			{
 				// just load it, and don't need a reference to it
-				FullyLoadedPackage = LoadPackage(NULL, *PackageIt.Key(), LOAD_None);
+				FullyLoadedPackage = LoadPackage(NULL, FPackagePath::FromLocalPath(PackageIt.Key()), LOAD_None);
 			}
 			else
 			{
@@ -2886,13 +2889,14 @@ int32 UWrangleContentCommandlet::Main( const FString& Params )
 			}
 
 			// find the one we moved this packages objects to
-			FString PackagePath = PackageIt.Key();
-			FString PackageName = FPackageName::FilenameToLongPackageName(PackagePath);
+			FPackagePath PackagePath = FPackagePath::FromLocalPath(PackageIt.Key());
+			FString PackageFilePath = PackagePath.GetLocalFullPath();
+			FString PackageName = PackagePath.GetPackageName();
 			UPackage* MovedPackage = FindPackage(NULL, *FString::Printf(TEXT("%s/NFS_%s"), *FPackageName::GetLongPackagePath(PackageName), *FPackageName::GetLongPackageAssetName(PackageName)));
 			check(MovedPackage);
 
 			// convert the new name to a a NFS directory directory
-			FString MovedFilename = MakeCutdownFilename(FString::Printf(TEXT("%s/NFS_%s"), *FPaths::GetPath(PackagePath), *FPaths::GetCleanFilename(PackagePath)), TEXT("NFSContent"));
+			FString MovedFilename = MakeCutdownFilename(FString::Printf(TEXT("%s/NFS_%s"), *FPaths::GetPath(PackageFilePath), *FPaths::GetCleanFilename(PackageFilePath)), TEXT("NFSContent"));
 			UE_LOG(LogContentCommandlet, Warning, TEXT("Saving package %s [%d/%d]"), *MovedFilename, PackageIndex, NumPackages);
 			// finally save it out
 			SavePackageHelper(MovedPackage, *MovedFilename);
diff --git a/Engine/Source/Editor/UnrealEd/Private/Commandlets/PackageUtilities.cpp b/Engine/Source/Editor/UnrealEd/Private/Commandlets/PackageUtilities.cpp
--- a/Engine/Source/Editor/UnrealEd/Private/Commandlets/PackageUtilities.cpp
+++ b/Engine/Source/Editor/UnrealEd/Private/Commandlets/PackageUtilities.cpp
@@ -766,8 +766,9 @@ int32 ULoadPackageCommandlet::Main( const FString& Params )
 
 		UE_LOG(LogPackageUtilities, Display, TEXT("Loading %s"), *Filename );
 
-		FString PackageName;
-		if (FPackageName::TryConvertFilenameToLongPackageName(Filename, PackageName))
+		FPackagePath PackagePath = FPackagePath::FromLocalPath(Filename);
+		FString PackageName = PackagePath.GetPackageName();
+		if (!PackageName.IsEmpty())
 		{
 			UPackage* Package = FindObject<UPackage>(nullptr, *PackageName, true);
 			if (Package != NULL && !bLoadAllPackages)
@@ -778,12 +779,12 @@ int32 ULoadPackageCommandlet::Main( const FString& Params )
 
 		if (bCheckForLegacyPackages)
 		{
-			FLinkerLoad* Linker = LoadPackageLinker(nullptr, *Filename, LOAD_NoVerify);
+			FLinkerLoad* Linker = LoadPackageLinker(nullptr, PackagePath, LOAD_NoVerify);
 			MinVersion = FMath::Min<int32>(MinVersion, Linker->Summary.GetFileVersionUE4());
 		}
 		else
 		{
-			UPackage* Package = LoadPackage(nullptr, *Filename, LOAD_None );
+			UPackage* Package = LoadPackage(nullptr, PackagePath, LOAD_None );
 			if(Package == nullptr)
 			{
 				UE_LOG(LogPackageUtilities, Error, TEXT("Error loading %s!"), *Filename );
@@ -906,7 +907,7 @@ FLinkerLoad* CreateLinkerForFilename(FUObjectSerializeContext* LoadContext, cons
 	{
 		Package = CreatePackage( *TempPackageName);
 	}
-	FLinkerLoad* Linker = FLinkerLoad::CreateLinker(LoadContext, Package, *InFilename, LOAD_NoVerify);
+	FLinkerLoad* Linker = FLinkerLoad::CreateLinker(LoadContext, Package, FPackagePath::FromLocalPath(InFilename), LOAD_NoVerify);
 	return Linker;
 }
 
@@ -942,7 +943,7 @@ void FPkgInfoReporter_Log::GeneratePackageReport( FLinkerLoad* InLinker /*=nullp
 	Out.Logf(ELogVerbosity::Display, TEXT("Package '%s' Summary"), *LinkerName.ToString() );
 	Out.Logf(ELogVerbosity::Display, TEXT("--------------------------------------------") );
 
-	Out.Logf(ELogVerbosity::Display, TEXT("\t         Filename: %s"), *Linker->Filename);
+	Out.Logf(ELogVerbosity::Display, TEXT("\t         Filename: %s"), *Linker->GetPackagePath().GetLocalFullPath());
 	Out.Logf(ELogVerbosity::Display, TEXT("\t     File Version: %i"), Linker->UE4Ver() );
 	Out.Logf(ELogVerbosity::Display, TEXT("\t   Engine Version: %s"), *Linker->Summary.SavedByEngineVersion.ToString());
 	Out.Logf(ELogVerbosity::Display, TEXT("\t   Compat Version: %s"), *Linker->Summary.CompatibleWithEngineVersion.ToString());
@@ -2406,4 +2407,4 @@ int32 UReplaceActorCommandlet::Main(const FString& Params)
 	}
 	GWorld = UWorld::CreateWorld(EWorldType::Editor, false );
 	return 0;
-}
\ No newline at end of file
+}
diff --git a/Engine/Source/Editor/UnrealEd/Private/Commandlets/WorldPartitionConvertCommandlet.cpp b/Engine/Source/Editor/UnrealEd/Private/Commandlets/WorldPartitionConvertCommandlet.cpp
--- a/Engine/Source/Editor/UnrealEd/Private/Commandlets/WorldPartitionConvertCommandlet.cpp
+++ b/Engine/Source/Editor/UnrealEd/Private/Commandlets/WorldPartitionConvertCommandlet.cpp
@@ -550,14 +550,14 @@ bool UWorldPartitionConvertCommandlet::RenameWorldPackageWithSuffix(UWorld* Worl
 
 	FString OldPackageName = Package->GetName();
 	FString NewPackageName = OldPackageName + ConversionSuffix;
-	FString NewPackageFilename = Package->FileName.ToString().Replace(*OldPackageName, *NewPackageName);
+	FString NewPackageResourceName = Package->GetLoadedPath().GetPackageName().Replace(*OldPackageName, *NewPackageName);
 	bRenamedSuccess = Package->Rename(*NewPackageName, nullptr, REN_NonTransactional | REN_DontCreateRedirectors | REN_ForceNoResetLoaders);
 	if (!bRenamedSuccess)
 	{
 		UE_LOG(LogWorldPartitionConvertCommandlet, Error, TEXT("Unable to rename package to %s"), *NewPackageName);
 		return false;
 	}
-	Package->FileName = *NewPackageFilename;
+	Package->SetLoadedPath(FPackagePath::FromPackageNameChecked(NewPackageResourceName));
 
 	return true;
 }
@@ -1066,12 +1066,12 @@ int32 UWorldPartitionConvertCommandlet::Main(const FString& Params)
 
 		if (LevelHasLevelScriptBlueprint(SubLevel))
 		{
-			MapsWithLevelScriptsBPs.Add(SubPackage->FileName.ToString());
+			MapsWithLevelScriptsBPs.Add(SubPackage->GetLoadedPath().GetPackageName());
 		}
 
 		if (LevelHasMapBuildData(SubLevel))
 		{
-			MapsWithMapBuildData.Add(SubPackage->FileName.ToString());
+			MapsWithMapBuildData.Add(SubPackage->GetLoadedPath().GetPackageName());
 		}
 
 		UE_LOG(LogWorldPartitionConvertCommandlet, Log, TEXT("Converting %s"), *SubWorld->GetName());
diff --git a/Engine/Source/Editor/UnrealEd/Private/CookOnTheFlyServer.cpp b/Engine/Source/Editor/UnrealEd/Private/CookOnTheFlyServer.cpp
--- a/Engine/Source/Editor/UnrealEd/Private/CookOnTheFlyServer.cpp
+++ b/Engine/Source/Editor/UnrealEd/Private/CookOnTheFlyServer.cpp
@@ -3641,9 +3641,10 @@ private:
 		{
 			bInclude = true;
 		}
-		else if (InPackage->LinkerLoad && InPackage->LinkerLoad->Filename.MatchesWildcard(InWildcard))
+		else if (InPackage->LinkerLoad)
 		{
-			bInclude = true;
+			FString Filename(InPackage->LinkerLoad->GetPackagePath().GetLocalFullPath());
+			bInclude = Filename.MatchesWildcard(InWildcard);
 		}
 		return bInclude;
 	}
@@ -5517,8 +5518,8 @@ void UCookOnTheFlyServer::GenerateLongPackageNames(TArray<FName>& FilesInPath)
 		else
 		{
 			FString LongPackageName;
-			FString FailureReason;
-			if (FPackageName::TryConvertFilenameToLongPackageName(FileInPath, LongPackageName, &FailureReason))
+			FPackageName::EErrorCode FailureReason;
+			if (FPackageName::TryConvertToMountedPath(FileInPath, nullptr /* LocalPath */, &LongPackageName, nullptr /* ObjectName */, nullptr /* SubObjectName */, nullptr /* Extension */, nullptr /* FlexNameType */, &FailureReason))
 			{
 				const FName LongPackageFName(*LongPackageName);
 				bool bIsAlreadyAdded;
@@ -5530,7 +5531,7 @@ void UCookOnTheFlyServer::GenerateLongPackageNames(TArray<FName>& FilesInPath)
 			}
 			else
 			{
-				LogCookerMessage(FString::Printf(TEXT("Unable to generate long package name for %s because %s"), *FileInPath, *FailureReason), EMessageSeverity::Warning);
+				LogCookerMessage(FString::Printf(TEXT("Unable to generate long package name, %s"), *FileInPath, *FPackageName::FormatErrorAsString(FileInPath, FailureReason)), EMessageSeverity::Warning);
 			}
 		}
 	}
diff --git a/Engine/Source/Editor/UnrealEd/Private/Cooker/CookPackageData.cpp b/Engine/Source/Editor/UnrealEd/Private/Cooker/CookPackageData.cpp
--- a/Engine/Source/Editor/UnrealEd/Private/Cooker/CookPackageData.cpp
+++ b/Engine/Source/Editor/UnrealEd/Private/Cooker/CookPackageData.cpp
@@ -26,7 +26,7 @@ namespace Cook
 
 
 	FPackageData::FPackageData(FPackageDatas& PackageDatas, const FName& InPackageName, const FName& InFileName)
-		: PackageName(InPackageName), FileName(InFileName), PackageDatas(PackageDatas)
+		: PackageName(InPackageName), FileName(InFileName), PackageDatas(PackageDatas), PreloadableFileFormat(EPackageFormat::Binary)
 		, bIsUrgent(0), bIsVisited(0), bIsPreloadAttempted(0), bIsPreloaded(0), bIsGeneratorPackage(0), bHasSaveCache(0), bHasBeginPrepareSaveFailed(0), bCookedPlatformDataStarted(0), bCookedPlatformDataCalled(0), bCookedPlatformDataComplete(0), bMonitorIsCooked(0)
 	{
 		SetState(EPackageState::Idle);
@@ -687,9 +687,22 @@ namespace Cook
 		{
 			TStringBuilder<NAME_SIZE> FileNameString;
 			GetFileName().ToString(FileNameString);
-			PreloadableFile = MakeShared<FPreloadableFile>(FileNameString.ToString());
+			PreloadableFile = MakeShared<FPreloadableArchive>(FileNameString.ToString());
 			PackageDatas.GetMonitor().OnPreloadAllocatedChanged(*this, true);
-			PreloadableFile->InitializeAsync(FPreloadableFile::Flags::PreloadHandle | FPreloadableFile::Flags::Prime);
+			PreloadableFile->InitializeAsync([this]()
+				{
+					TStringBuilder<NAME_SIZE> FileNameString;
+					// Note this is an unsynchronized read of this->GetFilename. It is allowed because GetFilename does not change until this is destructed, and the destructor of this waits for PreloadableFile to finish initialization
+					this->GetFileName().ToString(FileNameString);
+					FPackagePath PackagePath = FPackagePath::FromLocalPath(FileNameString);
+					FOpenPackageResult Result = IPackageResourceManager::Get().OpenReadPackage(PackagePath);
+					if (Result.Archive)
+					{
+						this->PreloadableFileFormat = Result.Format; // PreloadableFileFormat is an atomic, so this write is safe with other reads
+					}
+					return Result.Archive.Release();
+				},
+				FPreloadableFile::Flags::PreloadHandle | FPreloadableFile::Flags::Prime);
 		}
 		if (!PreloadableFile->IsInitialized())
 		{
@@ -708,14 +721,18 @@ namespace Cook
 		{
 			SetIsPreloadAttempted(true);
 			PreloadableFile.Reset();
+			PreloadableFileFormat = EPackageFormat::Binary;
 			return true;
 		}
 
-		if (!FPreloadableFile::TryRegister(PreloadableFile))
+		TStringBuilder<NAME_SIZE> FileNameString;
+		GetFileName().ToString(FileNameString);
+		if (!IPackageResourceManager::TryRegisterPreloadableArchive(FPackagePath::FromLocalPath(FileNameString), PreloadableFile, PreloadableFileFormat))
 		{
-			UE_LOG(LogCook, Warning, TEXT("Duplicate attempts to register %s for preload."), *GetFileName().ToString());
+			UE_LOG(LogCook, Warning, TEXT("Failed to register %s for preload."), *GetFileName().ToString());
 			SetIsPreloadAttempted(true);
 			PreloadableFile.Reset();
+			PreloadableFileFormat = EPackageFormat::Binary;
 			return true;
 		}
 
@@ -729,7 +746,9 @@ namespace Cook
 		if (GetIsPreloaded())
 		{
 			check(PreloadableFile);
-			if (FPreloadableFile::UnRegister(PreloadableFile))
+			TStringBuilder<NAME_SIZE> FileNameString;
+			GetFileName().ToString(FileNameString);
+			if (IPackageResourceManager::UnRegisterPreloadableArchive(FPackagePath::FromLocalPath(FileNameString)))
 			{
 				UE_LOG(LogCook, Display, TEXT("PreloadableFile was created for %s but never used. This is wasteful and bad for cook performance."), *PackageName.ToString());
 			}
@@ -738,12 +757,12 @@ namespace Cook
 		else
 		{
 			check(!PreloadableFile || !PreloadableFile->IsCacheAllocated());
-			check(!PreloadableFile || !FPreloadableFile::UnRegister(PreloadableFile));
 		}
 
 		if (PreloadableFile)
 		{
 			PreloadableFile.Reset();
+			PreloadableFileFormat = EPackageFormat::Binary;
 			PackageDatas.GetMonitor().OnPreloadAllocatedChanged(*this, false);
 		}
 		SetIsPreloaded(false);
diff --git a/Engine/Source/Editor/UnrealEd/Private/Cooker/CookPackageData.h b/Engine/Source/Editor/UnrealEd/Private/Cooker/CookPackageData.h
--- a/Engine/Source/Editor/UnrealEd/Private/Cooker/CookPackageData.h
+++ b/Engine/Source/Editor/UnrealEd/Private/Cooker/CookPackageData.h
@@ -2,6 +2,7 @@
 
 #pragma once
 
+#include <atomic>
 #include "CookTypes.h"
 #include "Containers/Array.h"
 #include "Containers/ArrayView.h"
@@ -15,6 +16,7 @@
 #include "Templates/SharedPointer.h"
 #include "UObject/GCObject.h"
 #include "UObject/NameTypes.h"
+#include "UObject/PackageResourceManager.h"
 #include "UObject/WeakObjectPtr.h"
 #include "UObject/WeakObjectPtrTemplates.h"
 
@@ -340,9 +342,10 @@ namespace Cook
 		FName CookedFileName; // Override filename to be used when saving cooked package (used by generated packages)
 		TWeakObjectPtr<UPackage> Package;
 		FPackageDatas& PackageDatas; // The one-per-CookOnTheFlyServer owner of this PackageData
-		TSharedPtr<FPreloadableFile> PreloadableFile;
+		TSharedPtr<FPreloadableArchive> PreloadableFile;
 		int32 NumPendingCookedPlatformData = 0;
 		int32 CookedPlatformDataNextIndex = -1;
+		std::atomic<EPackageFormat> PreloadableFileFormat;
 
 		uint32 State : int32(EPackageState::BitCount);
 		uint32 bIsUrgent : 1;
diff --git a/Engine/Source/Editor/UnrealEd/Private/EditorCommandLineUtils.cpp b/Engine/Source/Editor/UnrealEd/Private/EditorCommandLineUtils.cpp
--- a/Engine/Source/Editor/UnrealEd/Private/EditorCommandLineUtils.cpp
+++ b/Engine/Source/Editor/UnrealEd/Private/EditorCommandLineUtils.cpp
@@ -640,7 +640,7 @@ static void EditorCommandLineUtilsImpl::RunAssetMerge(FMergeAsset const& Base, F
 		{
 			if ((Resolution != EMergeResult::Cancelled) && !MergingPackage->IsDirty())
 			{
-				FString SrcFilePath = MergingPackage->FileName.ToString();
+				FString SrcFilePath = MergingPackage->GetLoadedPath().GetLocalFullPath();
 				IFileManager::Get().Copy(*DstFilePath, *SrcFilePath);
 			}
 		}
diff --git a/Engine/Source/Editor/UnrealEd/Private/EditorServer.cpp b/Engine/Source/Editor/UnrealEd/Private/EditorServer.cpp
--- a/Engine/Source/Editor/UnrealEd/Private/EditorServer.cpp
+++ b/Engine/Source/Editor/UnrealEd/Private/EditorServer.cpp
@@ -5484,7 +5484,7 @@ void ListMapPackageDependencies(const TCHAR* InStr)
 		if (!Linker)
 		{
 			// Create a new linker object which goes off and tries load the file.
-			Linker = GetPackageLinker(NULL, *(ProcessingPackage->GetName()), LOAD_None, NULL, NULL );
+			Linker = GetPackageLinker(nullptr, FPackagePath::FromPackageNameChecked(ProcessingPackage->GetName()), LOAD_None, nullptr, nullptr);
 		}
 		if (Linker)
 		{
@@ -5528,7 +5528,15 @@ void ListMapPackageDependencies(const TCHAR* InStr)
 				if (!Linker)
 				{
 					// Create a new linker object which goes off and tries load the file.
-					Linker = GetPackageLinker(NULL, *RefdPkgName,  LOAD_None, NULL, NULL );
+					FPackagePath PackagePath;
+					if (!FPackagePath::TryFromPackageName(RefdPkgName, PackagePath))
+					{
+						UE_LOG(LogEditorServer, Warning, TEXT("Invalid PackageName \"%s\""), *RefdPkgName);
+					}
+					else
+					{
+						Linker = GetPackageLinker(nullptr, PackagePath, LOAD_None, nullptr, nullptr);
+					}
 				}
 				if (Linker)
 				{
diff --git a/Engine/Source/Editor/UnrealEd/Private/FileHelpers.cpp b/Engine/Source/Editor/UnrealEd/Private/FileHelpers.cpp
--- a/Engine/Source/Editor/UnrealEd/Private/FileHelpers.cpp
+++ b/Engine/Source/Editor/UnrealEd/Private/FileHelpers.cpp
@@ -2505,35 +2505,32 @@ bool FEditorFileUtils::LoadMap(const FString& InFilename, bool LoadAsTemplate, b
 
 	const FScopedBusyCursor BusyCursor;
 
-	FString Filename( InFilename );
-
+	FString Filename;
 	FString LongMapPackageName;
-	if ( FPackageName::IsValidLongPackageName(InFilename) )
-	{
-		LongMapPackageName = InFilename;
-		FPackageName::TryConvertLongPackageNameToFilename(InFilename, Filename, FPackageName::GetMapPackageExtension());
-	}
-	else
-	{
+	FString Extension;
+	bool bFoundPath = FPackageName::TryConvertToMountedPath(InFilename, &Filename, &LongMapPackageName, nullptr /* ObjectName */, nullptr /* SubObjectName */, &Extension, nullptr /* OutFlexNameType */);
 #if PLATFORM_WINDOWS
+	if (!bFoundPath)
 	{
 		// Check if the Filename is actually from network drive and if so attempt to
 		// resolve to local path (if it's pointing to local machine's shared folder)
 		FString LocalFilename;
-		if ( FWindowsPlatformProcess::ResolveNetworkPath( Filename, LocalFilename ) )
+		if (FWindowsPlatformProcess::ResolveNetworkPath(InFilename, LocalFilename))
 		{
-			// Use local path if resolve succeeded
-			Filename = FString( *LocalFilename );
+			bFoundPath = FPackageName::TryConvertToMountedPath(LocalFilename, &Filename, &LongMapPackageName, nullptr /* ObjectName */, nullptr /* SubObjectName */, &Extension, nullptr /* OutFlexNameType */);
 		}
 	}
 #endif
-
-		if ( !FPackageName::TryConvertFilenameToLongPackageName(Filename, LongMapPackageName) )
-		{
-			FMessageDialog::Open(EAppMsgType::Ok, FText::Format(NSLOCTEXT("Editor", "MapLoad_FriendlyBadFilename", "Map load failed. The filename '{0}' is not within the game or engine content folders found in '{1}'."), FText::FromString(Filename), FText::FromString(FPaths::RootDir())));
-			return false;
-		}
+	if (!bFoundPath)
+	{
+		FMessageDialog::Open(EAppMsgType::Ok, FText::Format(NSLOCTEXT("Editor", "MapLoad_FriendlyBadFilename", "Map load failed. The filename '{0}' is not within the game or engine content folders found in '{1}'."), FText::FromString(Filename), FText::FromString(FPaths::RootDir())));
+		return false;
+	}
+	if (Extension.IsEmpty())
+	{
+		Extension = FPackageName::GetMapPackageExtension();
 	}
+	Filename += Extension;
 
 	// If a PIE world exists, warn the user that the PIE session will be terminated.
 	// Abort if the user refuses to terminate the PIE session.
diff --git a/Engine/Source/Editor/UnrealEd/Private/PlayLevel.cpp b/Engine/Source/Editor/UnrealEd/Private/PlayLevel.cpp
--- a/Engine/Source/Editor/UnrealEd/Private/PlayLevel.cpp
+++ b/Engine/Source/Editor/UnrealEd/Private/PlayLevel.cpp
@@ -2186,7 +2186,7 @@ UWorld* UEditorEngine::CreatePIEWorldByDuplication(FWorldContext &WorldContext,
 	UPackage* PlayWorldPackage = CreatePackage(*PlayWorldMapName);
 	PlayWorldPackage->SetPackageFlags(PKG_PlayInEditor);
 	PlayWorldPackage->PIEInstanceID = WorldContext.PIEInstance;
-	PlayWorldPackage->FileName = InPackage->FileName;
+	PlayWorldPackage->SetLoadedPath(InPackage->GetLoadedPath());
 	PlayWorldPackage->SetGuid( InPackage->GetGuid() );
 	PlayWorldPackage->MarkAsFullyLoaded();
 
diff --git a/Engine/Source/Editor/WorldBrowser/Private/Tiles/WorldTileCollectionModel.cpp b/Engine/Source/Editor/WorldBrowser/Private/Tiles/WorldTileCollectionModel.cpp
--- a/Engine/Source/Editor/WorldBrowser/Private/Tiles/WorldTileCollectionModel.cpp
+++ b/Engine/Source/Editor/WorldBrowser/Private/Tiles/WorldTileCollectionModel.cpp
@@ -1995,15 +1995,15 @@ bool FWorldTileCollectionModel::GenerateLODLevels(FLevelModelList InLevelList, i
 		FString SourceShortPackageName = FPackageName::GetShortName(SourceLongPackageName);
 		// Target PackageName for generated level: /LongPackageName+LOD/ShortPackageName_LOD[LODIndex]
 		const FString LODLevelPackageName = FString::Printf(TEXT("%sLOD/%s_LOD%d"),	*SourceLongPackageName, *SourceShortPackageName, TargetLODIndex+1);
-		// Target level filename
-		const FString LODLevelFileName = FPackageName::LongPackageNameToFilename(LODLevelPackageName) + FPackageName::GetMapPackageExtension();
 
 		// Create a package for a LOD level
 		UPackage* LODPackage = CreatePackage( *LODLevelPackageName);
 		LODPackage->FullyLoad();
 		LODPackage->Modify();
 		// This is a hack to avoid save file dialog when we will be saving LOD map package
-		LODPackage->FileName = FName(*LODLevelFileName);
+		FPackagePath LODLevelPackagePath = FPackagePath::FromPackageNameChecked(LODLevelPackageName);
+		LODLevelPackagePath.SetHeaderExtension(EPackageExtension::Map);
+		LODPackage->SetLoadedPath(LODLevelPackagePath);
 
 		// This is current actors offset from their original position
 		FVector ActorsOffset = FVector(TileModel->GetAbsoluteLevelPosition() - GetWorld()->OriginLocation);
@@ -2256,7 +2256,7 @@ bool FWorldTileCollectionModel::GenerateLODLevels(FLevelModelList InLevelList, i
 			// Save generated level
 			if (FEditorFileUtils::PromptToCheckoutLevels(false, LODWorld->PersistentLevel))
 			{
-				FEditorFileUtils::SaveLevel(LODWorld->PersistentLevel, LODLevelFileName);
+				FEditorFileUtils::SaveLevel(LODWorld->PersistentLevel, LODLevelPackagePath.GetLocalFullPath());
 				FAssetRegistryModule::AssetCreated(LODWorld);
 			}
 			
diff --git a/Engine/Source/Runtime/CoreUObject/Private/Blueprint/BlueprintSupport.cpp b/Engine/Source/Runtime/CoreUObject/Private/Blueprint/BlueprintSupport.cpp
--- a/Engine/Source/Runtime/CoreUObject/Private/Blueprint/BlueprintSupport.cpp
+++ b/Engine/Source/Runtime/CoreUObject/Private/Blueprint/BlueprintSupport.cpp
@@ -1263,7 +1263,7 @@ int32 FLinkerLoad::FindCDOExportIndex(UClass* LoadClass)
 	return INDEX_NONE;
 }
 
-UPackage* LoadPackageInternal(UPackage* InOuter, const TCHAR* InLongPackageName, uint32 LoadFlags, FLinkerLoad* ImportLinker, FArchive* InReaderOverride, const FLinkerInstancingContext* InstancingContext);
+UPackage* LoadPackageInternal(UPackage* InOuter, const FPackagePath& PackagePath, uint32 LoadFlags, FLinkerLoad* ImportLinker, FArchive* InReaderOverride, const FLinkerInstancingContext* InstancingContext, const FPackagePath* DiffPackagePath);
 
 void FLinkerLoad::ResolveDeferredDependencies(UStruct* LoadStruct)
 {
@@ -1383,7 +1383,7 @@ void FLinkerLoad::ResolveDeferredDependencies(UStruct* LoadStruct)
 			{
 				uint32 InternalLoadFlags = LoadFlags & (LOAD_NoVerify | LOAD_NoWarn | LOAD_Quiet);
 				// make sure LoadAllObjects() is called for this package
-				LoadPackageInternal(/*Outer =*/nullptr, *SourceLinker->Filename, InternalLoadFlags, this, nullptr/*InReaderOverride*/, nullptr/*InstancingContext*/); //-V595
+				LoadPackageInternal(/*Outer =*/nullptr, SourceLinker->GetPackagePath(), InternalLoadFlags, this, nullptr/*InReaderOverride*/, nullptr/*InstancingContext*/, nullptr /* DiffPackagePath */); //-V595
 			}
 
 			DEFERRED_DEPENDENCY_CHECK(ResolvingPlaceholderStack.Num() == 0);
diff --git a/Engine/Source/Runtime/CoreUObject/Private/Misc/PackageName.cpp b/Engine/Source/Runtime/CoreUObject/Private/Misc/PackageName.cpp
--- a/Engine/Source/Runtime/CoreUObject/Private/Misc/PackageName.cpp
+++ b/Engine/Source/Runtime/CoreUObject/Private/Misc/PackageName.cpp
@@ -5,34 +5,34 @@
 =============================================================================*/
 
 #include "Misc/PackageName.h"
+
+#include "Algo/Find.h"
 #include "Containers/StringView.h"
 #include "GenericPlatform/GenericPlatformFile.h"
+#include "HAL/CriticalSection.h"
 #include "HAL/FileManager.h"
-#include "Misc/Paths.h"
-#include "Stats/Stats.h"
-#include "Misc/CoreDelegates.h"
-#include "Misc/App.h"
-#include "Modules/ModuleManager.h"
-#include "UObject/Package.h"
-#include "UObject/PackageFileSummary.h"
-#include "UObject/Linker.h"
+#include "HAL/ThreadHeartBeat.h"
 #include "Interfaces/IPluginManager.h"
 #include "Internationalization/PackageLocalizationManager.h"
-#include "HAL/CriticalSection.h"
-#include "HAL/ThreadHeartBeat.h"
+#include "IO/IoDispatcher.h"
+#include "Misc/App.h"
 #include "Misc/AutomationTest.h"
+#include "Misc/CoreDelegates.h"
+#include "Misc/PackageSegment.h"
+#include "Misc/Paths.h"
 #include "Misc/PathViews.h"
 #include "Misc/ScopeRWLock.h"
 #include "Misc/StringBuilder.h"
-#include "IO/IoDispatcher.h"
+#include "Modules/ModuleManager.h"
 #include "ProfilingDebugging/LoadTimeTracker.h"
+#include "Stats/Stats.h"
+#include "UObject/Linker.h"
+#include "UObject/Package.h"
+#include "UObject/PackageFileSummary.h"
+#include "UObject/PackageResourceManager.h"
 
-DEFINE_LOG_CATEGORY_STATIC(LogPackageName, Log, All);
+DEFINE_LOG_CATEGORY(LogPackageName);
 
-FString FPackageName::AssetPackageExtension = TEXT(".uasset");
-FString FPackageName::MapPackageExtension = TEXT(".umap");
-FString FPackageName::TextAssetPackageExtension = TEXT(".utxt");
-FString FPackageName::TextMapPackageExtension = TEXT(".utxtmap");
 static FRWLock ContentMountPointCriticalSection;
 
 /** Event that is triggered when a new content path is mounted */
@@ -409,6 +409,7 @@ void FPackageName::InternalFilenameToLongPackageName(FStringView InFilename, FSt
 		}
 	}
 
+	FStringView Result;
 	if (!bIsValidLongPackageName)
 	{
 		Filename = IFileManager::Get().ConvertToRelativePath(*Filename);
@@ -422,9 +423,18 @@ void FPackageName::InternalFilenameToLongPackageName(FStringView InFilename, FSt
 				Filename += TEXT("/");
 			}
 		}
+		Result = FPathViews::GetBaseFilenameWithPath(Filename);
+	}
+	else
+	{
+		Result = FPathViews::GetBaseFilenameWithPath(Filename);
+		if (Result.Len() != Filename.Len())
+		{
+			UE_LOG(LogPackageName, Warning, TEXT("TryConvertFilenameToLongPackageName was passed an ObjectPath (%.*s) rather than a PackageName or FilePath; it will be converted to the PackageName. ")
+				TEXT("Accepting ObjectPaths is deprecated behavior and will be removed in a future release; TryConvertFilenameToLongPackageName will fail on ObjectPaths."), InFilename.Len(), InFilename.GetData());
+		}
 	}
 
-	FStringView Result = FPathViews::GetBaseFilenameWithPath(Filename);
 
 	{
 		FReadScopeLock ScopeLock(ContentMountPointCriticalSection);
@@ -549,6 +559,75 @@ FString FPackageName::LongPackageNameToFilename(const FString& InLongPackageName
 	return Result;
 }
 
+bool FPackageName::TryConvertToMountedPath(FStringView InPath, FString* OutLocalPathNoExtension, FString* OutPackageName, FString* OutObjectName, FString* OutSubObjectName, FString* OutExtension, EFlexNameType* OutFlexNameType, EErrorCode* OutFailureReason)
+{
+	auto ClearSuccessOutputs = [OutLocalPathNoExtension, OutPackageName, OutObjectName, OutSubObjectName, OutExtension, OutFlexNameType, OutFailureReason]()
+	{
+		if (OutLocalPathNoExtension) OutLocalPathNoExtension->Reset();
+		if (OutPackageName) OutPackageName->Reset();
+		if (OutObjectName) OutObjectName->Reset();
+		if (OutSubObjectName) OutSubObjectName->Reset();
+		if (OutExtension) OutExtension->Reset();
+		if (OutFlexNameType) *OutFlexNameType = EFlexNameType::Invalid;
+	};
+
+	TStringBuilder<256> MountPointPackageName;
+	TStringBuilder<256> MountPointFilePath;
+	TStringBuilder<256> PackageNameRelPath;
+	EFlexNameType FlexNameType;
+	EErrorCode FailureReason;
+	bool bResult = TryGetMountPointForPath(InPath, MountPointPackageName, MountPointFilePath, PackageNameRelPath, &FlexNameType, &FailureReason);
+	if (!bResult)
+	{
+		ClearSuccessOutputs();
+		if (OutFailureReason) *OutFailureReason = FailureReason;
+		return false;
+	}
+
+	FString Extension;
+	if (FlexNameType == EFlexNameType::LocalPath)
+	{
+		// Remove Extension from PackageNameRelPath and put it into OutExtension
+		int32 ExtensionStart;
+		FPackagePath::ParseExtension(PackageNameRelPath, &ExtensionStart);
+		Extension = FStringView(PackageNameRelPath).RightChop(ExtensionStart);
+		PackageNameRelPath.RemoveSuffix(PackageNameRelPath.Len() - ExtensionStart);
+	}
+	else
+	{
+		check(FlexNameType == EFlexNameType::PackageName || FlexNameType == EFlexNameType::ObjectPath);
+	}
+
+	TStringBuilder<256> ObjectPathOrPackageName;
+	ObjectPathOrPackageName << MountPointPackageName << PackageNameRelPath;
+	FStringView ClassName;
+	FStringView PackageName;
+	FStringView ObjectName;
+	FStringView SubObjectName;
+	SplitFullObjectPath(ObjectPathOrPackageName, ClassName, PackageName, ObjectName, SubObjectName);
+	if (ClassName.Len() > 0)
+	{
+		ClearSuccessOutputs();
+		if (OutFailureReason) *OutFailureReason = EErrorCode::PackageNameFullObjectPathNotAllowed;
+		return false;
+	}
+	if (!IsValidTextForLongPackageName(PackageName, &FailureReason))
+	{
+		ClearSuccessOutputs();
+		if (OutFailureReason) *OutFailureReason = FailureReason;
+		return false;
+	}
+	check(FStringView(PackageName).StartsWith(MountPointPackageName));
+
+	if (OutLocalPathNoExtension) *OutLocalPathNoExtension = FString(MountPointFilePath) + FStringView(PackageName).RightChop(MountPointPackageName.Len());
+	if (OutPackageName) *OutPackageName = PackageName;
+	if (OutObjectName) *OutObjectName = ObjectName;
+	if (OutSubObjectName) *OutSubObjectName = SubObjectName;
+	if (OutExtension) *OutExtension = Extension;
+	if (OutFlexNameType) *OutFlexNameType = FlexNameType;
+	if (OutFailureReason) *OutFailureReason = EErrorCode::PackageNameUnknown;
+	return true;
+}
 FString FPackageName::GetLongPackagePath(const FString& InLongPackageName)
 {
 	int32 IndexOfLastSlash = INDEX_NONE;
@@ -603,29 +682,58 @@ bool FPackageName::SplitLongPackageName(const FString& InLongPackageName, FStrin
 
 void FPackageName::SplitFullObjectPath(const FString& InFullObjectPath, FString& OutClassName, FString& OutPackageName, FString& OutObjectName, FString& OutSubObjectName)
 {
-	FString Sanitized = InFullObjectPath.TrimStartAndEnd();
-	const TCHAR* Cur = *Sanitized;
+	FStringView ClassName;
+	FStringView PackageName;
+	FStringView ObjectName;
+	FStringView SubObjectName;
+	SplitFullObjectPath(InFullObjectPath, ClassName, PackageName, ObjectName, SubObjectName);
+	OutClassName = ClassName;
+	OutPackageName = PackageName;
+	OutObjectName = ObjectName;
+	OutSubObjectName = SubObjectName;
+}
+
+void FPackageName::SplitFullObjectPath(FStringView InFullObjectPath, FStringView& OutClassName, FStringView& OutPackageName, FStringView& OutObjectName, FStringView& OutSubObjectName)
+{
+	FStringView Remaining = InFullObjectPath.TrimStartAndEnd();
 
-	auto ExtractBeforeDelim = [&Cur](TCHAR Delim, FString& OutString)
+	auto ExtractBeforeDelim = [&Remaining](TCHAR Delim, FStringView& OutStringView)
 	{
-		const TCHAR* Start = Cur;
-		while (*Cur != '\0' && *Cur != Delim)
+		int32 DelimIndex;
+		if (Remaining.FindChar(Delim, DelimIndex))
 		{
-			++Cur;
+			OutStringView = Remaining.Left(DelimIndex);
+			Remaining.RightChopInline(DelimIndex + 1);
+			return true;
 		}
-
-		OutString = FString(Cur - Start, Start);
-
-		if (*Cur == Delim)
+		else
 		{
-			++Cur;
+			OutStringView.Reset();
+			return false;
 		}
 	};
 
-	ExtractBeforeDelim(' ', OutClassName);
-	ExtractBeforeDelim('.', OutPackageName);
-	ExtractBeforeDelim(':', OutObjectName);
-	ExtractBeforeDelim('\0', OutSubObjectName);
+	ExtractBeforeDelim(' ', OutClassName); // If no space, then ClassName is empty and the remaining string is PackageName.ObjectName:SubObjectName
+	if (ExtractBeforeDelim('.', OutPackageName))
+	{
+		if (ExtractBeforeDelim(':', OutObjectName))
+		{
+			OutSubObjectName = Remaining;
+		}
+		else
+		{
+			// If no :, then the remaining string is ObjectName
+			OutObjectName = Remaining;
+			OutSubObjectName = FStringView();
+		}
+	}
+	else 
+	{
+		// If no '.', then the remaining string is PackageName
+		OutPackageName = Remaining;
+		OutObjectName = FStringView();
+		OutSubObjectName = FStringView();
+	}
 }
 
 FString FPackageName::GetLongPackageAssetName(const FString& InLongPackageName)
@@ -633,7 +741,18 @@ FString FPackageName::GetLongPackageAssetName(const FString& InLongPackageName)
 	return GetShortName(InLongPackageName);
 }
 
-bool FPackageName::DoesPackageNameContainInvalidCharacters(FStringView InLongPackageName, FText* OutReason /*= NULL*/)
+bool FPackageName::DoesPackageNameContainInvalidCharacters(FStringView InLongPackageName, FText* OutReason)
+{
+	EErrorCode Reason;
+	if (DoesPackageNameContainInvalidCharacters(InLongPackageName, &Reason))
+	{
+		if (OutReason) *OutReason = FormatErrorAsText(InLongPackageName, Reason);
+		return true;
+	}
+	return false;
+}
+
+bool FPackageName::DoesPackageNameContainInvalidCharacters(FStringView InLongPackageName, EErrorCode* OutReason /*= nullptr */)
 {
 	// See if the name contains invalid characters.
 	TStringBuilder<32> MatchedInvalidChars;
@@ -647,44 +766,42 @@ bool FPackageName::DoesPackageNameContainInvalidCharacters(FStringView InLongPac
 	}
 	if (MatchedInvalidChars.Len())
 	{
-		if (OutReason)
-		{
-			FFormatNamedArguments Args;
-			Args.Add( TEXT("IllegalNameCharacters"), FText::FromString(FString(MatchedInvalidChars)) );
-			*OutReason = FText::Format( NSLOCTEXT("Core", "PackageNameContainsInvalidCharacters", "Name may not contain the following characters: '{IllegalNameCharacters}'"), Args );
-		}
+		if (OutReason) *OutReason = EErrorCode::PackageNameContainsInvalidCharacters;
 		return true;
 	}
+	if (OutReason) *OutReason = EErrorCode::PackageNameUnknown;
 	return false;
 }
 
-bool FPackageName::IsValidLongPackageName(const FString& InLongPackageName, bool bIncludeReadOnlyRoots /*= false*/, FText* OutReason /*= NULL*/)
+bool FPackageName::IsValidTextForLongPackageName(FStringView InLongPackageName, FText* OutReason)
+{
+	EErrorCode Reason;
+	if (!IsValidTextForLongPackageName(InLongPackageName, &Reason))
+	{
+		if (OutReason) *OutReason = FormatErrorAsText(InLongPackageName, Reason);
+		return false;
+	}
+	return true;
+}
+
+bool FPackageName::IsValidTextForLongPackageName(FStringView InLongPackageName, EErrorCode* OutReason /*= nullptr */)
 {
 	// All package names must contain a leading slash, root, slash and name, at minimum theoretical length ("/A/B") is 4
 	if (InLongPackageName.Len() < PackageNameConstants::MinPackageNameLength)
 	{
-		if (OutReason)
-		{
-			*OutReason = FText::Format(NSLOCTEXT("Core", "LongPackageNames_PathTooShort", "Path should be no less than {0} characters long."), FText::AsNumber(PackageNameConstants::MinPackageNameLength));
-		}
+		if (OutReason) *OutReason = EErrorCode::LongPackageNames_PathTooShort;
 		return false;
 	}
 	// Package names start with a leading slash.
 	if (InLongPackageName[0] != '/')
 	{
-		if (OutReason)
-		{
-			*OutReason = NSLOCTEXT("Core", "LongPackageNames_PathWithNoStartingSlash", "Path should start with a '/'");
-		}
+		if (OutReason) *OutReason = EErrorCode::LongPackageNames_PathWithNoStartingSlash;
 		return false;
 	}
 	// Package names do not end with a trailing slash.
 	if (InLongPackageName[InLongPackageName.Len() - 1] == '/')
 	{
-		if (OutReason)
-		{
-			*OutReason = NSLOCTEXT("Core", "LongPackageNames_PathWithTrailingSlash", "Path may not end with a '/'");
-		}
+		if (OutReason) *OutReason = EErrorCode::LongPackageNames_PathWithTrailingSlash;
 		return false;
 	}
 	// Check for invalid characters
@@ -692,51 +809,83 @@ bool FPackageName::IsValidLongPackageName(const FString& InLongPackageName, bool
 	{
 		return false;
 	}
-	// Check valid roots
-	const FLongPackagePathsSingleton& Paths = FLongPackagePathsSingleton::Get();
-	TArray<FString> ValidRoots;
-	bool bValidRoot = false;
-	Paths.GetValidLongPackageRoots(ValidRoots, bIncludeReadOnlyRoots);
-	for (int32 RootIdx = 0; RootIdx < ValidRoots.Num(); ++RootIdx)
-	{
-		const FString& Root = ValidRoots[RootIdx];
-		if (InLongPackageName.StartsWith(Root))
-		{
-			bValidRoot = true;
-			break;
-		}
-	}
-	if (!bValidRoot && OutReason)
+	if (OutReason) *OutReason = EErrorCode::PackageNameUnknown;
+	return true;
+}
+
+bool FPackageName::IsValidLongPackageName(FStringView InLongPackageName, bool bIncludeReadOnlyRoots, FText* OutReason)
+{
+	EErrorCode Reason;
+	if (!IsValidLongPackageName(InLongPackageName, bIncludeReadOnlyRoots, &Reason))
 	{
-		if (ValidRoots.Num() == 0)
-		{
-			*OutReason = NSLOCTEXT("Core", "LongPackageNames_NoValidRoots", "No valid roots exist!");
-		}
-		else
+		if (OutReason)
 		{
-			FString ValidRootsString = TEXT("");
-			if (ValidRoots.Num() == 1)
-			{
-				ValidRootsString = FString::Printf(TEXT("'%s'"), *ValidRoots[0]);
-			}
-			else
+			if (Reason == EErrorCode::PackageNamePathNotMounted)
 			{
-				for (int32 RootIdx = 0; RootIdx < ValidRoots.Num(); ++RootIdx)
+				const FLongPackagePathsSingleton& Paths = FLongPackagePathsSingleton::Get();
+				TArray<FString> ValidRoots;
+				Paths.GetValidLongPackageRoots(ValidRoots, bIncludeReadOnlyRoots);
+				if (ValidRoots.Num() == 0)
+				{
+					*OutReason = NSLOCTEXT("Core", "LongPackageNames_NoValidRoots", "No valid roots exist!");
+				}
+				else
 				{
-					if (RootIdx < ValidRoots.Num() - 1)
+					FString ValidRootsString = TEXT("");
+					if (ValidRoots.Num() == 1)
 					{
-						ValidRootsString += FString::Printf(TEXT("'%s', "), *ValidRoots[RootIdx]);
+						ValidRootsString = FString::Printf(TEXT("'%s'"), *ValidRoots[0]);
 					}
 					else
 					{
-						ValidRootsString += FString::Printf(TEXT("or '%s'"), *ValidRoots[RootIdx]);
+						for (int32 RootIdx = 0; RootIdx < ValidRoots.Num(); ++RootIdx)
+						{
+							if (RootIdx < ValidRoots.Num() - 1)
+							{
+								ValidRootsString += FString::Printf(TEXT("'%s', "), *ValidRoots[RootIdx]);
+							}
+							else
+							{
+								ValidRootsString += FString::Printf(TEXT("or '%s'"), *ValidRoots[RootIdx]);
+							}
+						}
 					}
+					*OutReason = FText::Format( NSLOCTEXT("Core", "LongPackageNames_InvalidRoot", "Path does not start with a valid root. Path must begin with: {0}"), FText::FromString( ValidRootsString ) );
 				}
 			}
-			*OutReason = FText::Format( NSLOCTEXT("Core", "LongPackageNames_InvalidRoot", "Path does not start with a valid root. Path must begin with: {0}"), FText::FromString( ValidRootsString ) );
+			else
+			{
+				*OutReason = FormatErrorAsText(InLongPackageName, Reason);
+			}
 		}
+		return false;
 	}
-	return bValidRoot;
+	return true;
+}
+
+bool FPackageName::IsValidLongPackageName(FStringView InLongPackageName, bool bIncludeReadOnlyRoots /*= false*/, EErrorCode* OutReason /*= nullptr */)
+{
+	if (!IsValidTextForLongPackageName(InLongPackageName, OutReason))
+	{
+		return false;
+	}
+
+	// Check valid roots
+	const FLongPackagePathsSingleton& Paths = FLongPackagePathsSingleton::Get();
+	TArray<FString> ValidRoots;
+	bool bValidRoot = false;
+	Paths.GetValidLongPackageRoots(ValidRoots, bIncludeReadOnlyRoots);
+	for (int32 RootIdx = 0; RootIdx < ValidRoots.Num(); ++RootIdx)
+	{
+		const FString& Root = ValidRoots[RootIdx];
+		if (InLongPackageName.StartsWith(Root))
+		{
+			if (OutReason) *OutReason = EErrorCode::PackageNameUnknown;
+			return true;
+		}
+	}
+	if (OutReason) *OutReason = EErrorCode::PackageNamePathNotMounted;
+	return false;
 }
 
 bool FPackageName::IsValidObjectPath(const FString& InObjectPath, FText* OutReason)
@@ -853,6 +1002,165 @@ FName FPackageName::GetPackageMountPoint(const FString& InPackagePath, bool InWi
 	return FName();
 }
 
+bool FPackageName::TryConvertToMountedPathComponents(FStringView InFilePathOrPackageName, FStringBuilderBase& OutMountPointPackageName, FStringBuilderBase& OutMountPointFilePath, FStringBuilderBase& OutRelPath,
+	FStringBuilderBase& OutObjectName, EPackageExtension& OutExtension, FStringBuilderBase& OutCustomExtension, EFlexNameType* OutFlexNameType, EErrorCode* OutFailureReason)
+{
+	auto ClearSuccessOutputs = [&OutMountPointPackageName, &OutMountPointFilePath, &OutRelPath, &OutObjectName, &OutExtension, &OutCustomExtension, OutFlexNameType]()
+	{
+		OutMountPointPackageName.Reset();
+		OutMountPointFilePath.Reset();
+		OutRelPath.Reset();
+		OutObjectName.Reset();
+		OutExtension = EPackageExtension::Unspecified;
+		OutCustomExtension.Reset();
+		if (OutFlexNameType) *OutFlexNameType = EFlexNameType::Invalid;
+	};
+
+	EFlexNameType PathFlexNameType;
+	bool bFound = TryGetMountPointForPath(InFilePathOrPackageName, OutMountPointPackageName, OutMountPointFilePath, OutRelPath, &PathFlexNameType, OutFailureReason);
+	if (!bFound)
+	{
+		ClearSuccessOutputs();
+		return false;
+	}
+
+	OutObjectName.Reset();
+	OutExtension = EPackageExtension::Unspecified;
+	OutCustomExtension.Reset();
+
+	if (PathFlexNameType == EFlexNameType::LocalPath)
+	{
+		// Remove Extension from OutRelPath and put it into OutExtension
+		int32 ExtensionStart;
+		OutExtension = FPackagePath::ParseExtension(OutRelPath, &ExtensionStart);
+		if (OutExtension == EPackageExtension::Custom)
+		{
+			OutCustomExtension << FStringView(OutRelPath).RightChop(ExtensionStart);
+		}
+		OutRelPath.RemoveSuffix(OutRelPath.Len() - ExtensionStart);
+
+	}
+	else if (PathFlexNameType == EFlexNameType::ObjectPath)
+	{
+		// Legacy behavior; convert ObjectPaths to packageName
+		TStringBuilder<256> ObjectPath;
+		ObjectPath << OutMountPointPackageName << OutRelPath;
+		FStringView ClassName;
+		FStringView PackageName;
+		FStringView ObjectName;
+		FStringView SubObjectName;
+		SplitFullObjectPath(ObjectPath, ClassName, PackageName, ObjectName, SubObjectName);
+		if (ClassName.Len() > 0)
+		{
+			ClearSuccessOutputs();
+			if (OutFailureReason) *OutFailureReason = EErrorCode::PackageNameFullObjectPathNotAllowed;
+		}
+		if (!IsValidTextForLongPackageName(PackageName, OutFailureReason))
+		{
+			ClearSuccessOutputs();
+			return false;
+		}
+		OutObjectName << ObjectName;
+		if (SubObjectName.Len())
+		{
+			OutObjectName << SUBOBJECT_DELIMITER << SubObjectName;
+		}
+		check(FStringView(PackageName).StartsWith(OutMountPointPackageName));
+		int32 RelPathPackageNameLen = PackageName.Len() - OutMountPointPackageName.Len();
+		OutRelPath.RemoveSuffix(OutRelPath.Len() - RelPathPackageNameLen);
+	}
+	else
+	{
+		check(PathFlexNameType == EFlexNameType::PackageName);
+		TStringBuilder<256> PackagePath;
+		PackagePath << OutMountPointPackageName << OutRelPath;
+		if (!IsValidTextForLongPackageName(PackagePath, OutFailureReason))
+		{
+			ClearSuccessOutputs();
+			return false;
+		}
+	}
+
+	if (OutFlexNameType) *OutFlexNameType = PathFlexNameType;
+	if (OutFailureReason) *OutFailureReason = EErrorCode::PackageNameUnknown;
+	return true;
+}
+
+bool FPackageName::TryGetMountPointForPath(FStringView InFilePathOrPackageName, FStringBuilderBase& OutMountPointPackageName, FStringBuilderBase& OutMountPointFilePath, FStringBuilderBase& OutRelPath, EFlexNameType* OutFlexNameType, EErrorCode* OutFailureReason)
+{
+	OutMountPointPackageName.Reset();
+	OutMountPointFilePath.Reset();
+	OutRelPath.Reset();
+
+	if (InFilePathOrPackageName.IsEmpty())
+	{
+		if (OutFlexNameType)
+		{
+			*OutFlexNameType = EFlexNameType::Invalid;
+		}
+		if (OutFailureReason)
+		{
+			*OutFailureReason = EErrorCode::PackageNameEmptyPath;
+		}
+		return false;
+	}
+
+	FString PossibleAbsFilePath(FPaths::ConvertRelativePathToFull(FString(InFilePathOrPackageName)));
+	const FLongPackagePathsSingleton& Paths = FLongPackagePathsSingleton::Get();
+	FReadScopeLock ScopeLock(ContentMountPointCriticalSection);
+	for (const auto& Pair : Paths.ContentRootToPath)
+	{
+		FString RootFileAbsPath = FPaths::ConvertRelativePathToFull(Pair.ContentPath);
+		if (InFilePathOrPackageName.StartsWith(Pair.RootPath))
+		{
+			OutMountPointPackageName << Pair.RootPath;
+			OutMountPointFilePath << Pair.ContentPath;
+			FStringView RelPath(InFilePathOrPackageName.RightChop(Pair.RootPath.Len()));
+			OutRelPath << RelPath;
+			if (OutFlexNameType)
+			{
+				if (Algo::Find(RelPath, TEXT('.')))
+				{
+					*OutFlexNameType = EFlexNameType::ObjectPath;
+				}
+				else
+				{
+					*OutFlexNameType = EFlexNameType::PackageName;
+				}
+			}
+			if (OutFailureReason)
+			{
+				*OutFailureReason = EErrorCode::PackageNameUnknown;
+			}
+			return true;
+		}
+		else if (PossibleAbsFilePath.StartsWith(RootFileAbsPath))
+		{
+			OutMountPointPackageName << Pair.RootPath;
+			OutMountPointFilePath << Pair.ContentPath;
+			OutRelPath << FStringView(PossibleAbsFilePath).RightChop(RootFileAbsPath.Len());
+			if (OutFlexNameType)
+			{
+				*OutFlexNameType = EFlexNameType::LocalPath;
+			}
+			if (OutFailureReason)
+			{
+				*OutFailureReason = EErrorCode::PackageNameUnknown;
+			}
+			return true;
+		}
+	}
+	if (OutFlexNameType)
+	{
+		*OutFlexNameType = EFlexNameType::Invalid;
+	}
+	if (OutFailureReason)
+	{
+		*OutFailureReason = EErrorCode::PackageNamePathNotMounted;
+	}
+	return false;
+}
+
 FString FPackageName::ConvertToLongScriptPackageName(const TCHAR* InShortName)
 {
 	if (IsShortPackageName(FString(InShortName)))
@@ -892,111 +1200,103 @@ FName* FPackageName::FindScriptPackageName(FName InShortName)
 
 bool FPackageName::FindPackageFileWithoutExtension(const FString& InPackageFilename, FString& OutFilename, bool InAllowTextFormats)
 {
-	auto& FileManager = IFileManager::Get();
-
+	bool bExists = FindPackageFileWithoutExtension(InPackageFilename, OutFilename);
+	if (!InAllowTextFormats)
 	{
-		static const FString* PackageExtensions[] =
+		FPackagePath PackagePath = FPackagePath::FromLocalPath(OutFilename);
+		FOpenPackageResult Result = IPackageResourceManager::Get().OpenReadPackage(PackagePath);
+		if (!Result.Archive.IsValid() || Result.Format == EPackageFormat::Text)
 		{
-			&AssetPackageExtension,
-			&MapPackageExtension
-		};
-
-		// Loop through all known extensions and check if the file exists
-
-		for (int32 ExtensionIndex = 0; ExtensionIndex < UE_ARRAY_COUNT(PackageExtensions); ++ExtensionIndex)
-		{
-			FString   PackageFilename = InPackageFilename + *PackageExtensions[ExtensionIndex];
-			if (FileManager.FileExists(*PackageFilename))
-			{
-				// The package exists so exit. From now on InPackageFilename can be equal to OutFilename so
-				// don't attempt to use it anymore (case where &InPackageFilename == &OutFilename).
-				OutFilename = MoveTemp(PackageFilename);
-				return true;
-			}
+			return false;
 		}
 	}
+	return bExists;
+}
 
-#if WITH_TEXT_ARCHIVE_SUPPORT
-	if (InAllowTextFormats)
+bool FPackageName::FindPackageFileWithoutExtension(const FString& InPackageFilename, FString& OutFilename)
+{
+	FPackagePath PackagePath = FPackagePath::FromLocalPath(InPackageFilename);
+	if (IPackageResourceManager::Get().DoesPackageExist(PackagePath, &PackagePath))
 	{
-		static const FString* TextPackageExtensions[] =
-		{
-			&TextAssetPackageExtension,
-			&TextMapPackageExtension
-		};
-
-		for (int32 ExtensionIndex = 0; ExtensionIndex < UE_ARRAY_COUNT(TextPackageExtensions); ++ExtensionIndex)
-		{
-			FString   PackageFilename = InPackageFilename + *TextPackageExtensions[ExtensionIndex];
-			if (FileManager.FileExists(*PackageFilename))
-			{
-				// The package exists so exit. From now on InPackageFilename can be equal to OutFilename so
-				// don't attempt to use it anymore (case where &InPackageFilename == &OutFilename).
-				OutFilename = MoveTemp(PackageFilename);
-				return true;
-			}
-		}
+		OutFilename = PackagePath.GetLocalFullPath();
+		return true;
+	}
+	else
+	{
+		return false;
 	}
-#endif
-
-	return false;
 }
 
 bool FPackageName::FixPackageNameCase(FString& LongPackageName, FStringView Extension)
 {
-	// Find the matching long package root
-	const FLongPackagePathsSingleton& Paths = FLongPackagePathsSingleton::Get();
-	FReadScopeLock ScopeLock(ContentMountPointCriticalSection);
-	for (const FPathPair& Pair : Paths.ContentRootToPath)
+	FPackagePath PackagePath;
+	if (!FPackagePath::TryFromPackageName(LongPackageName, PackagePath))
 	{
-		if (LongPackageName.StartsWith(Pair.RootPath))
-		{
-			FString RelativePackageName = LongPackageName.Mid(Pair.RootPath.Len());
-			FString FileName = Pair.ContentPath / RelativePackageName;
-
-			int ExtensionLen = Extension.Len();
-			if(Extension.Len() > 0 && Extension[0] != '.')
-			{
-				FileName.AppendChar('.');
-				ExtensionLen++;
-			}
-
-			FileName += Extension;
-
-			FString CorrectFileName = IFileManager::Get().GetFilenameOnDisk(*FileName);
-			if(CorrectFileName.Len() >= RelativePackageName.Len() + ExtensionLen)
-			{
-				FString NewRelativePackageName = CorrectFileName.Mid(CorrectFileName.Len() - RelativePackageName.Len() - ExtensionLen, RelativePackageName.Len());
-				if(NewRelativePackageName == RelativePackageName)
-				{
-					LongPackageName.RemoveAt(Pair.RootPath.Len(), LongPackageName.Len() - Pair.RootPath.Len());
-					LongPackageName.Append(*NewRelativePackageName);
-					return true;
-				}
-			}
-			break;
-		}
+		return false;
 	}
-	return false;
+	if (!IPackageResourceManager::Get().TryMatchCaseOnDisk(PackagePath, &PackagePath))
+	{
+		return false;
+	}
+	TStringBuilder<256> DiskPackageName;
+	PackagePath.AppendPackageName(DiskPackageName);
+	check(FStringView(LongPackageName).Equals(DiskPackageName, ESearchCase::IgnoreCase));
+	LongPackageName = DiskPackageName;
+	return true;
 }
 
 bool FPackageName::DoesPackageExist(const FString& LongPackageName, const FGuid* Guid, FString* OutFilename, bool InAllowTextFormats)
 {
-	SCOPED_LOADTIMER(FPackageName_DoesPackageExist);
-
-	bool bFoundFile = false;
+	// Make sure interpreting LongPackageName as a filename is supported.
+	FPackagePath PackagePath;
+	{
+		SCOPED_LOADTIMER(FPackageName_DoesPackageExist);
+		TStringBuilder<64> PackageNameRoot;
+		TStringBuilder<64> FilePathRoot;
+		TStringBuilder<256> RelPath;
+		TStringBuilder<64> UnusedObjectName; // DoesPackageExist accepts ObjectPaths and ignores the ObjectName portion and uses only the PackageName
+		TStringBuilder<16> CustomExtension;
+		EPackageExtension Extension;
+		EErrorCode FailureReason;
+		if (!FPackageName::TryConvertToMountedPathComponents(LongPackageName, PackageNameRoot, FilePathRoot, RelPath, UnusedObjectName, Extension, CustomExtension, nullptr /* OutFlexNameType */, &FailureReason))
+		{
+			FString Message = FString::Printf(TEXT("Illegal call to DoesPackageExist: %s"), *FormatErrorAsString(LongPackageName, FailureReason));
+			UE_LOG(LogPackageName, Error, TEXT("%s"), *Message);
+			ensureMsgf(false, TEXT("%s"), *Message);
+			return false;
+		}
+		PackagePath = FPackagePath::FromMountedComponents(PackageNameRoot, FilePathRoot, RelPath, Extension, CustomExtension);
+	}
+	if (!DoesPackageExist(PackagePath, Guid, false /* bMatchCaseOnDisk */, &PackagePath))
+	{
+		return false;
+	}
+	if (!InAllowTextFormats && IsTextPackageExtension(PackagePath.GetHeaderExtension()))
+	{
+		return false;
+	}
+	if (OutFilename)
+	{
+		*OutFilename = PackagePath.GetLocalFullPath();
+	}
+	return true;
+}
 
-	// Make sure passing filename as LongPackageName is supported.
-	FString PackageName;
-	FText Reason;
+bool FPackageName::DoesPackageExist(const FPackagePath& PackagePath, FPackagePath* OutPackagePath)
+{
+	return DoesPackageExist(PackagePath, nullptr, false, OutPackagePath);
+}
 
-	if (!FPackageName::TryConvertFilenameToLongPackageName(LongPackageName, PackageName))
+bool FPackageName::DoesPackageExist(const FPackagePath& PackagePath, const FGuid* Guid, bool bMatchCaseOnDisk, FPackagePath* OutPackagePath)
+{
+	// DoesPackageExist returns false for local filenames that are in unmounted directories, even if those files exist on the local disk
+	if (!PackagePath.IsMountedPath())
 	{
-		verify(!FPackageName::IsValidLongPackageName(LongPackageName, true, &Reason));
-		UE_LOG(LogPackageName, Error, TEXT("Illegal call to DoesPackageExist: '%s' is not a standard unreal filename or a long path name. Reason: %s"), *LongPackageName, *Reason.ToString());
-		ensureMsgf(false, TEXT("Illegal call to DoesPackageExist: '%s' is not a standard unreal filename or a long path name. Reason: %s"), *LongPackageName, *Reason.ToString());
 		return false;
 	}
+	TStringBuilder<256> PackageName;
+	PackagePath.AppendPackageName(PackageName);
+
 	// Once we have the real Package Name, we can exit early if it's a script package - they exist only in memory.
 	if (IsScriptPackage(PackageName))
 	{
@@ -1008,9 +1308,10 @@ bool FPackageName::DoesPackageExist(const FString& LongPackageName, const FGuid*
 		return false;
 	}
 
-	if ( !FPackageName::IsValidLongPackageName( PackageName, true, &Reason ) )
+	FText Reason;
+	if ( !FPackageName::IsValidTextForLongPackageName( PackageName, &Reason ) )
 	{
-		UE_LOG(LogPackageName, Error, TEXT( "DoesPackageExist: DoesPackageExist FAILED: '%s' is not a standard unreal filename or a long path name. Reason: %s"), *LongPackageName, *Reason.ToString() );
+		UE_LOG(LogPackageName, Error, TEXT( "DoesPackageExist: DoesPackageExist FAILED: '%s' is not a long packagename name. Reason: %s"), PackageName.ToString(), *Reason.ToString() );
 		return false;
 	}
 
@@ -1019,51 +1320,65 @@ bool FPackageName::DoesPackageExist(const FString& LongPackageName, const FGuid*
 	{
 		if (DoesPackageExistOverrideDelegate.Execute(FName(*PackageName)))
 		{
-			if (OutFilename)
+			if (OutPackagePath)
 			{
-				*OutFilename = LongPackageNameToFilename(PackageName, TEXT(""));
+				*OutPackagePath = PackagePath;
 			}
-
 			return true;
 		}
 
 		return false;
 	}
 
-	// Convert to filename (no extension yet).
-	FString Filename = LongPackageNameToFilename(PackageName, TEXT(""));
-
-	// Find the filename (with extension).
-	bFoundFile = FindPackageFileWithoutExtension(Filename, Filename, InAllowTextFormats);
-
 	// On consoles, we don't support package downloading, so no need to waste any extra cycles/disk io dealing with it
-	if (!FPlatformProperties::RequiresCookedData() && bFoundFile && Guid != NULL)
+	if (!FPlatformProperties::RequiresCookedData() && Guid != nullptr)
 	{
 		// @todo: If we could get to list of linkers here, it would be faster to check
 		// then to open the file and read it
-		FArchive* PackageReader = IFileManager::Get().CreateFileReader(*Filename);
-		// This had better open
-		check(PackageReader != NULL);
-
+		FPackagePath LocalPackagePath;
+		FOpenPackageResult Result = IPackageResourceManager::Get().OpenReadPackage(PackagePath, &LocalPackagePath);
+		TUniquePtr<FArchive>& PackageReader = Result.Archive;
+		if (!PackageReader || Result.Format != EPackageFormat::Binary)
+		{
+			if (PackageReader)
+			{
+				UE_LOG(LogPackageName, Error, TEXT("DoesPackageExist: DoesPackageExist with Guid FAILED: '%s' exists on disk with TextFormat, and we cannot read guids from TextFormat packages."),
+					*PackagePath.GetDebugName());
+			}
+			return false;
+		}
 		// Read in the package summary
 		FPackageFileSummary Summary;
 		*PackageReader << Summary;
 
 		// Compare Guids
-		if (Summary.Guid != *Guid)
+		if (PackageReader->IsError() || Summary.Guid != *Guid)
 		{
-			bFoundFile = false;
+			return false;
 		}
+		PackageReader.Reset();
 
-		// Close package file
-		delete PackageReader;
+		if (bMatchCaseOnDisk)
+		{
+			IPackageResourceManager::Get().TryMatchCaseOnDisk(LocalPackagePath, &LocalPackagePath);
+		}
+		if (OutPackagePath)
+		{
+			*OutPackagePath = LocalPackagePath;
+		}
+		return true;
 	}
-
-	if (OutFilename && bFoundFile)
+	else
+	{
+		if (bMatchCaseOnDisk)
+		{
+			return IPackageResourceManager::Get().TryMatchCaseOnDisk(PackagePath, OutPackagePath);
+		}
+		else
 	{
-		*OutFilename = Filename;
+			return IPackageResourceManager::Get().DoesPackageExist(PackagePath, OutPackagePath);
+		}
 	}
-	return bFoundFile;
 }
 
 bool FPackageName::SearchForPackageOnDisk(const FString& PackageName, FString* OutLongPackageName, FString* OutFilename)
@@ -1096,62 +1411,74 @@ bool FPackageName::SearchForPackageOnDisk(const FString& PackageName, FString* O
 	else
 	{
 		// Attempt to find package by its short name by searching in the known content paths.
-		TArray<FString> Paths;		
+		TArray<TPair<FString,FString>> RootsNameAndFile;
 		{
 			TArray<FString> RootContentPaths;
 			FPackageName::QueryRootContentPaths( RootContentPaths );
 			for( TArray<FString>::TConstIterator RootPathIt( RootContentPaths ); RootPathIt; ++RootPathIt )
 			{
-				const FString& RootPath = *RootPathIt;
-				const FString& ContentFolder = FPackageName::LongPackageNameToFilename(RootPath, TEXT(""));
-				Paths.Add( ContentFolder );
+				const FString& RootPackageName = *RootPathIt;
+				const FString& RootFilePath = FPackageName::LongPackageNameToFilename(RootPackageName, TEXT(""));
+				RootsNameAndFile.Add(TPair<FString,FString>(RootPackageName, RootFilePath));
 			}
 		}
 
-		const FString PackageWildcard = (PackageName.Find(TEXT("."), ESearchCase::CaseSensitive) != INDEX_NONE ? PackageName : PackageName + TEXT(".*"));
-		TArray<FString> Results;
+		int32 ExtensionStart;
+		EPackageExtension RequiredExtension = FPackagePath::ParseExtension(PackageName, &ExtensionStart);
+		if (RequiredExtension == EPackageExtension::Custom || ExtensionToSegment(RequiredExtension) != EPackageSegment::Header)
+		{
+			UE_LOG(LogPackageName, Warning, TEXT("SearchForPackageOnDisk: Invalid extension in packagename %s. Searching for any header extension instead."), *PackageName);
+			RequiredExtension = EPackageExtension::Unspecified;
+		}
+		FStringView PackageWildCard(FStringView(PackageName).Left(ExtensionStart));
 
-		for (int32 PathIndex = 0; PathIndex < Paths.Num() && !bResult; ++PathIndex)
+		FPackagePath FirstResult;
+		TArray<FPackagePath> FoundResults;
+		IPackageResourceManager& PackageResourceManager = IPackageResourceManager::Get();
+		for (const TPair<FString, FString>& RootNameAndFile : RootsNameAndFile)
 		{
+			const FString& RootPackageName = RootNameAndFile.Get<0>();
+			const FString& RootFilePath = RootNameAndFile.Get<1>();
+			check(RootPackageName.EndsWith(TEXT("/")));
+			check(RootFilePath.EndsWith(TEXT("/")));
 			// Search directly on disk. Very slow!
-			IFileManager::Get().FindFilesRecursive(Results, *Paths[PathIndex], *PackageWildcard, true, false);
+			FoundResults.Reset();
+			PackageResourceManager.FindPackagesRecursive(FoundResults, RootPackageName, RootFilePath, FStringView(), PackageWildCard);
 
-			for (int32 FileIndex = 0; FileIndex < Results.Num(); ++FileIndex)
+			for (const FPackagePath& FoundPackagePath: FoundResults)
 			{			
-				FString Filename(Results[FileIndex]);
-				if (IsPackageFilename(Results[FileIndex]))
+				FStringView UnusedCustomExtension;
+				if (RequiredExtension != EPackageExtension::Unspecified && FoundPackagePath.GetHeaderExtension() != RequiredExtension)
+				{
+					continue;
+				}
+
+				bResult = true;
+				if (OutLongPackageName || OutFilename)
 				{
-					// Convert to long package name.
-					FString LongPackageName;
-					if (TryConvertFilenameToLongPackageName(Filename, LongPackageName))
+					if (!FirstResult.IsEmpty())
+					{
+						UE_LOG(LogPackageName, Warning, TEXT("SearchForPackageOnDisk: Found ambiguous long package name for '%s'. Returning '%s', but could also be '%s'."), *PackageName,
+							*FirstResult.GetDebugNameWithExtension(), *FoundPackagePath.GetDebugNameWithExtension());
+					}
+					else
 					{
+						FirstResult = FoundPackagePath;
 						if (OutLongPackageName)
 						{
-							if (bResult)
-							{
-								UE_LOG(LogPackageName, Warning, TEXT("Found ambiguous long package name for '%s'. Returning '%s', but could also be '%s'."), *PackageName, **OutLongPackageName, *LongPackageName );
-							}
-							else
-							{
-								*OutLongPackageName = LongPackageName;
-							}
+							*OutLongPackageName = FoundPackagePath.GetPackageName();
 						}
 						if (OutFilename)
 						{
-							FPaths::MakeStandardFilename(Filename);
-							if (bResult)
-							{
-								UE_LOG(LogPackageName, Warning, TEXT("Found ambiguous file name for '%s'. Returning '%s', but could also be '%s'."), *PackageName, **OutFilename, *Filename);
-							}
-							else
-							{
-								*OutFilename = Filename;
-							}
+							*OutFilename = FoundPackagePath.GetLocalFullPath();
 						}
-						bResult = true;
 					}
 				}
 			}
+			if (bResult)
+			{
+				break;
+			}
 		}
 	}
 	float ThisTime = FPlatformTime::Seconds() - StartTime;
@@ -1354,11 +1681,17 @@ bool FPackageName::IsTextPackageExtension(const TCHAR* Ext)
 	return IsTextAssetPackageExtension(Ext) || IsTextMapPackageExtension(Ext);
 }
 
+bool FPackageName::IsTextPackageExtension(EPackageExtension Extension)
+{
+	return Extension == EPackageExtension::TextAsset || Extension == EPackageExtension::TextMap;
+}
+
 bool FPackageName::IsTextAssetPackageExtension(const TCHAR* Ext)
 {
-	if (*Ext != TEXT('.'))
+	FStringView TextAssetPackageExtension(LexToString(EPackageExtension::TextAsset));
+	if (*Ext != TEXT('.') && *Ext != TEXT('\0'))
 	{
-		return (TextAssetPackageExtension.EndsWith(Ext));
+		return (TextAssetPackageExtension.RightChop(1) == Ext);
 	}
 	else
 	{
@@ -1368,9 +1701,10 @@ bool FPackageName::IsTextAssetPackageExtension(const TCHAR* Ext)
 
 bool FPackageName::IsTextMapPackageExtension(const TCHAR* Ext)
 {
-	if (*Ext != TEXT('.'))
+	FStringView TextMapPackageExtension(LexToString(EPackageExtension::TextMap));
+	if (*Ext != TEXT('.') && *Ext != TEXT('\0'))
 	{
-		return (TextMapPackageExtension.EndsWith(Ext));
+		return (TextMapPackageExtension.RightChop(1) == Ext);
 	}
 	else
 	{
@@ -1385,9 +1719,10 @@ bool FPackageName::IsPackageExtension( const TCHAR* Ext )
 
 bool FPackageName::IsAssetPackageExtension(const TCHAR* Ext)
 {
+	FStringView AssetPackageExtension(LexToString(EPackageExtension::Asset));
 	if (*Ext != TEXT('.'))
 	{
-		return (AssetPackageExtension.EndsWith(Ext));
+		return (AssetPackageExtension.RightChop(1) == Ext);
 	}
 	else
 	{
@@ -1397,9 +1732,10 @@ bool FPackageName::IsAssetPackageExtension(const TCHAR* Ext)
 
 bool FPackageName::IsMapPackageExtension(const TCHAR* Ext)
 {
+	FStringView MapPackageExtension(LexToString(EPackageExtension::Map));
 	if (*Ext != TEXT('.'))
 	{
-		return (MapPackageExtension.EndsWith(Ext));
+		return (MapPackageExtension.RightChop(1) == Ext);
 	}
 	else
 	{
@@ -1411,27 +1747,13 @@ bool FPackageName::FindPackagesInDirectory( TArray<FString>& OutPackages, const
 {
 	UE_CLOG(FIoDispatcher::IsInitialized(), LogPackageName, Error, TEXT("Can't search for packages using the filesystem when I/O dispatcher is enabled"));
 
-	FString LocalPathToRootDir;
-	if (!FPackageName::TryConvertLongPackageNameToFilename(RootDir / TEXT(""), LocalPathToRootDir))
-	{
-		LocalPathToRootDir = RootDir;
-	}
-	LocalPathToRootDir = FPaths::ConvertRelativePathToFull(MoveTemp(LocalPathToRootDir));
-
-	// Find all files in RootDir, then filter by extension (we have two package extensions so they can't
-	// be included in the search wildcard.
-	TArray<FString> AllFiles;
-	IFileManager::Get().FindFilesRecursive(AllFiles, *LocalPathToRootDir, TEXT("*.*"), true, false);
 	// Keep track if any package has been found. Can't rely only on OutPackages.Num() > 0 as it may not be empty.
 	const int32 PreviousPackagesCount = OutPackages.Num();
-	for (int32 FileIndex = 0; FileIndex < AllFiles.Num(); FileIndex++)
+	IteratePackagesInDirectory(RootDir, [&OutPackages](const TCHAR* PackageFilename) -> bool
 	{
-		const FString& Filename = AllFiles[FileIndex];
-		if (IsPackageFilename(Filename))
-		{
-			OutPackages.Add(Filename);
-		}
-	}
+		OutPackages.Add(PackageFilename);
+		return true;
+	});
 	return OutPackages.Num() > PreviousPackagesCount;
 }
 
@@ -1456,55 +1778,46 @@ bool FPackageName::FindPackagesInDirectories(TArray<FString>& OutPackages, const
 	return Packages.Num() > 0;
 }
 
-
 void FPackageName::IteratePackagesInDirectory(const FString& RootDir, const FPackageNameVisitor& Callback)
 {
-	class FPackageVisitor : public IPlatformFile::FDirectoryVisitor
+	auto LocalCallback = [&Callback](const FPackagePath& PackagePath) -> bool
 	{
-	public:
-		const FPackageNameVisitor& Callback;
-		explicit FPackageVisitor(const FPackageNameVisitor& InCallback)
-			: Callback(InCallback)
-		{
-		}
-		virtual bool Visit(const TCHAR* FilenameOrDirectory, bool bIsDirectory) override
-		{
-			bool Result = true;
-			if (!bIsDirectory && IsPackageFilename(FilenameOrDirectory))
-			{
-				Result = Callback(FilenameOrDirectory);
-			}
-			return Result;
-		}
+		return Callback(*PackagePath.GetLocalFullPath());
 	};
 
-	FPackageVisitor PackageVisitor(Callback);
-	IFileManager::Get().IterateDirectoryRecursively(*RootDir, PackageVisitor);
+	TStringBuilder<256> PackageNameRoot;
+	TStringBuilder<256> FilePathRoot;
+	TStringBuilder<256> RelRootDir;
+	if (TryGetMountPointForPath(RootDir, PackageNameRoot, FilePathRoot, RelRootDir))
+	{
+		IPackageResourceManager::Get().IteratePackagesInPath(PackageNameRoot, FilePathRoot, RelRootDir, LocalCallback);
+	}
+	else
+	{
+		// Searching a localonly path
+		IPackageResourceManager::Get().IteratePackagesInLocalOnlyDirectory(RootDir, LocalCallback);
+	}
 }
 
 void FPackageName::IteratePackagesInDirectory(const FString& RootDir, const FPackageNameStatVisitor& Callback)
 {
-	class FPackageStatVisitor : public IPlatformFile::FDirectoryStatVisitor
+	auto LocalCallback = [&Callback](const FPackagePath& PackagePath, const FFileStatData& StatData) -> bool
 	{
-	public:
-		const FPackageNameStatVisitor& Callback;
-		explicit FPackageStatVisitor(const FPackageNameStatVisitor& InCallback)
-			: Callback(InCallback)
-		{
-		}
-		virtual bool Visit(const TCHAR* FilenameOrDirectory, const FFileStatData& StatData) override
-		{
-			bool Result = true;
-			if (!StatData.bIsDirectory && IsPackageFilename(FilenameOrDirectory))
-			{
-				Result = Callback(FilenameOrDirectory, StatData);
-			}
-			return Result;
-		}
+		return Callback(*PackagePath.GetLocalFullPath(), StatData);
 	};
 
-	FPackageStatVisitor PackageVisitor(Callback);
-	IFileManager::Get().IterateDirectoryStatRecursively(*RootDir, PackageVisitor);
+	TStringBuilder<256> PackageNameRoot;
+	TStringBuilder<256> FilePathRoot;
+	TStringBuilder<256> RelRootDir;
+	if (TryGetMountPointForPath(RootDir, PackageNameRoot, FilePathRoot, RelRootDir))
+	{
+		IPackageResourceManager::Get().IteratePackagesStatInPath(PackageNameRoot, FilePathRoot, RelRootDir, LocalCallback);
+	}
+	else
+	{
+		// Searching a localonly path
+		IPackageResourceManager::Get().IteratePackagesStatInLocalOnlyDirectory(RootDir, LocalCallback);
+	}
 }
 
 void FPackageName::QueryRootContentPaths(TArray<FString>& OutRootContentPaths, bool bIncludeReadOnlyRoots, bool bWithoutLeadingSlashes, bool bWithoutTrailingSlashes)
@@ -1736,6 +2049,46 @@ bool FPackageName::IsLocalizedPackage(FStringView InPackageName)
 	return Remaining.StartsWith(TEXT("L10N"_SV), ESearchCase::IgnoreCase) && (Remaining.Len() == 4 || Remaining[4] == '/');
 }
 
+FString FPackageName::FormatErrorAsString(FStringView InPath, EErrorCode ErrorCode)
+{
+	FText ErrorText = FormatErrorAsText(InPath, ErrorCode);
+	return ErrorText.ToString();
+}
+
+FText FPackageName::FormatErrorAsText(FStringView InPath, EErrorCode ErrorCode)
+{
+	FFormatNamedArguments Args;
+	Args.Add(TEXT("InPath"), FText::FromString(FString(InPath)));
+	switch (ErrorCode)
+	{
+	case EErrorCode::PackageNameUnknown:
+		return FText::Format(NSLOCTEXT("Core", "PackageNameUnknownError", "Input '{InPath}' caused undocumented internal error."), Args);
+	case EErrorCode::PackageNameEmptyPath:
+		return FText::Format(NSLOCTEXT("Core", "PackageNameEmptyPath", "Input '{InPath}' was empty."), Args);
+	case EErrorCode::PackageNamePathNotMounted:
+		return FText::Format(NSLOCTEXT("Core", "PackageNamePathNotMounted", "Input '{InPath}' is not a child of an existing mount point."), Args);
+	case EErrorCode::PackageNameFullObjectPathNotAllowed:
+		return FText::Format(NSLOCTEXT("Core", "PackageNameFullObjectPathNotAllowed", "Input '{InPath}' is an unallowed FullObjectPath \"<ClassName> <PackageName>.<ObjectName>:<SubObjectName>\". Only partial ObjectPaths \"<PackageName>.<ObjectName>:<SubObjectName>\" are allowed."), Args);
+	case EErrorCode::PackageNameContainsInvalidCharacters:
+		Args.Add(TEXT("IllegalNameCharacters"), FText::FromString(FString(INVALID_LONGPACKAGE_CHARACTERS)));
+		return FText::Format(NSLOCTEXT("Core", "PackageNameContainsInvalidCharacters", "Input '{InPath}' contains one of the invalid characters for LongPackageNames: '{IllegalNameCharacters}'."), Args);
+	case EErrorCode::LongPackageNames_PathTooShort:
+	{
+		// This has to be an FFormatOrderedArguments until we change the localized text string for it.
+		FFormatOrderedArguments OrderedArgs;
+		OrderedArgs.Add(FText::AsNumber(PackageNameConstants::MinPackageNameLength));
+		OrderedArgs.Add(FText::FromString(FString(InPath)));
+		return FText::Format(NSLOCTEXT("Core", "LongPackageNames_PathTooShort", "Input '{1}' contains fewer than the minimum number of characters {0} for LongPackageNames."), OrderedArgs);
+	}
+	case EErrorCode::LongPackageNames_PathWithNoStartingSlash:
+		return FText::Format(NSLOCTEXT("Core", "LongPackageNames_PathWithNoStartingSlash", "Input '{InPath}' does not start with a '/', which is required for LongPackageNames."), Args);
+	case EErrorCode::LongPackageNames_PathWithTrailingSlash:
+		return FText::Format(NSLOCTEXT("Core", "LongPackageNames_PathWithTrailingSlash", "Input '{InPath}' ends with a '/', which is invalid for LongPackageNames."), Args);
+	default:
+		UE_LOG(LogPackageName, Warning, TEXT("FPackageName::FormatErrorAsText: Invalid ErrorCode %d"), static_cast<int32>(ErrorCode));
+		return FText::Format(NSLOCTEXT("Core", "PackageNameUnknownError", "Input '{InPath} caused undocumented internal error."), Args);
+	}
+}
 
 #if WITH_DEV_AUTOMATION_TESTS
 
diff --git a/Engine/Source/Runtime/CoreUObject/Private/Serialization/ArchiveStackTrace.cpp b/Engine/Source/Runtime/CoreUObject/Private/Serialization/ArchiveStackTrace.cpp
--- a/Engine/Source/Runtime/CoreUObject/Private/Serialization/ArchiveStackTrace.cpp
+++ b/Engine/Source/Runtime/CoreUObject/Private/Serialization/ArchiveStackTrace.cpp
@@ -1016,7 +1016,7 @@ FLinkerLoad* FArchiveStackTrace::CreateLinkerForPackage(FUObjectSerializeContext
 	}
 	// Create an archive for the linker. The linker will take ownership of it.
 	FLargeMemoryReader* PackageReader = new FLargeMemoryReader(PackageData.Data, PackageData.Size, ELargeMemoryReaderFlags::None, *InPackageName);	
-	FLinkerLoad* Linker = FLinkerLoad::CreateLinker(LoadContext, Package, *InFilename, LOAD_NoVerify, PackageReader);
+	FLinkerLoad* Linker = FLinkerLoad::CreateLinker(LoadContext, Package, FPackagePath::FromLocalPath(InFilename), LOAD_NoVerify, PackageReader);
 
 	if (Linker && Package)
 	{
