PackageResourceManager and PackagePaths

diff --git a/Engine/Source/Runtime/CoreUObject/Private/Serialization/AsyncLoading.cpp b/Engine/Source/Runtime/CoreUObject/Private/Serialization/AsyncLoading.cpp
--- a/Engine/Source/Runtime/CoreUObject/Private/Serialization/AsyncLoading.cpp
+++ b/Engine/Source/Runtime/CoreUObject/Private/Serialization/AsyncLoading.cpp
@@ -9,6 +9,7 @@
 #include "HAL/FileManager.h"
 #include "HAL/Event.h"
 #include "HAL/RunnableThread.h"
+#include "Misc/PackageSegment.h"
 #include "Misc/ScopeLock.h"
 #include "Stats/StatsMisc.h"
 #include "Misc/CoreStats.h"
@@ -23,6 +24,7 @@
 #include "UObject/CoreRedirects.h"
 #include "UObject/SoftObjectPath.h"
 #include "UObject/LinkerLoad.h"
+#include "UObject/PackageResourceManager.h"
 #include "Serialization/DeferredMessageLog.h"
 #include "UObject/UObjectThreadContext.h"
 #include "UObject/LinkerManager.h"
@@ -298,6 +300,7 @@ DECLARE_MEMORY_STAT(TEXT("FAsyncArchive Buffers"), STAT_FAsyncArchiveMem, STATGR
 
 FORCEINLINE void FAsyncArchive::LogItem(const TCHAR* Item, int64 Offset, int64 Size, double StartTime)
 {
+	FString FileName = PackagePath.GetLocalFullPath();
 	if (UE_LOG_ACTIVE(LogAsyncArchive, Verbose)
 #if defined(ASYNC_WATCH_FILE)
 		|| FileName.Contains(TEXT(ASYNC_WATCH_FILE))
@@ -1299,7 +1302,7 @@ FScopedAsyncPackageEvent::FScopedAsyncPackageEvent(FAsyncPackage* InPackage)
 	ThreadContext.AsyncPackage = Package;
 
 	Package->BeginAsyncLoad();
-	FExclusiveLoadPackageTimeTracker::PushLoadPackage(Package->Desc.NameToLoad);
+	FExclusiveLoadPackageTimeTracker::PushLoadPackage(Package->Desc.PackagePath.GetPackageFName());
 }
 
 FScopedAsyncPackageEvent::~FScopedAsyncPackageEvent()
@@ -1706,7 +1709,7 @@ EAsyncPackageState::Type FAsyncPackage::LoadImports_Event()
 		{
 			if (!IsFullyLoadedObj(Import->XObject))
 			{
-				UE_LOG(LogStreaming, Verbose, TEXT("FAsyncPackage::LoadImports for %s: import %s was found but was not fully loaded yet."), *Desc.NameToLoad.ToString(), *OriginalImport->ObjectName.ToString());
+				UE_LOG(LogStreaming, Verbose, TEXT("FAsyncPackage::LoadImports for %s: import %s was found but was not fully loaded yet."), *Desc.PackagePath.GetDebugName(), *OriginalImport->ObjectName.ToString());
 			}
 			else
 			{
@@ -1718,7 +1721,7 @@ EAsyncPackageState::Type FAsyncPackage::LoadImports_Event()
 
 		if (!GProcessPrestreamingRequests && bIsPrestreamRequest)
 		{
-			UE_LOG(LogStreaming, Display, TEXT("%s is NOT prestreaming %s"), *Desc.NameToLoad.ToString(), *Import->ObjectName.ToString());
+			UE_LOG(LogStreaming, Display, TEXT("%s is NOT prestreaming %s"), *Desc.PackagePath.GetDebugName(), *Import->ObjectName.ToString());
 			Import->bImportFailed = true;
 			continue;
 		}
@@ -1781,21 +1784,22 @@ EAsyncPackageState::Type FAsyncPackage::LoadImports_Event()
 				// This can happen with editor only classes, not sure if this should be a warning or a silent continue
 				if (!GIsInitialLoad)
 				{
-					UE_LOG(LogStreaming, Warning, TEXT("FAsyncPackage::LoadImports for %s: Skipping import %s, depends on missing native class"), *Desc.NameToLoad.ToString(), *Linker->GetImportFullName(LocalImportIndex));
+					UE_LOG(LogStreaming, Warning, TEXT("FAsyncPackage::LoadImports for %s: Skipping import %s, depends on missing native class"), *Desc.PackagePath.GetDebugName(), *Linker->GetImportFullName(LocalImportIndex));
 				}
 			}
 			else if (!ExistingPackage || bForcePackageLoad)
 			{
 				// The package doesn't exist and this import is not in the dependency list so add it now.
 				check(!FPackageName::IsShortPackageName(ImportPackageFName));
-				UE_LOG(LogStreaming, Verbose, TEXT("FAsyncPackage::LoadImports for %s: Loading %s"), *Desc.NameToLoad.ToString(), *ImportPackageFName.ToString());
-				const FAsyncPackageDesc Info(INDEX_NONE, ImportPackageFName, ImportPackageToLoad);
+				UE_LOG(LogStreaming, Verbose, TEXT("FAsyncPackage::LoadImports for %s: Loading %s"), *Desc.PackagePath.GetDebugName(), *ImportPackageFName.ToString());
+				FPackagePath ImportPackagePath = FPackagePath::FromPackageNameChecked(ImportPackageToLoad);
+				const FAsyncPackageDesc Info(INDEX_NONE, ImportPackageFName, ImportPackagePath);
 				PendingPackage = new FAsyncPackage(AsyncLoadingThread, Info, EDLBootNotificationManager);
 				PendingPackage->Desc.Priority = Desc.Priority;
 				PendingPackage->Desc.SetInstancingContext(Linker->GetInstancingContext());
 				if (bIsPrestreamRequest)
 				{
-					UE_LOG(LogStreaming, Display, TEXT("%s is prestreaming %s"), *Desc.NameToLoad.ToString(), *ImportPackageToLoad.ToString());
+					UE_LOG(LogStreaming, Display, TEXT("%s is prestreaming %s"), *Desc.PackagePath.GetDebugName(), *ImportPackagePath.GetDebugName());
 				}
 				TRACE_LOADTIME_ASYNC_PACKAGE_IMPORT_DEPENDENCY(this, PendingPackage);
 #if !UE_BUILD_SHIPPING
@@ -1972,7 +1976,7 @@ EAsyncPackageState::Type FAsyncPackage::SetupImports_Event()
 					if (!ImportPackage)
 					{
 						Import.bImportFailed = true;
-						UE_CLOG(!FLinkerLoad::IsKnownMissingPackage(Import.ObjectName), LogStreaming, Error, TEXT("Missing native package (%s) for import of package %s"), *Import.ObjectName.ToString(), *Desc.NameToLoad.ToString());
+						UE_CLOG(!FLinkerLoad::IsKnownMissingPackage(Import.ObjectName), LogStreaming, Error, TEXT("Missing native package (%s) for import of package %s"), *Import.ObjectName.ToString(), *Desc.PackagePath.GetDebugName());
 					}
 					else
 					{
@@ -2024,7 +2028,7 @@ EAsyncPackageState::Type FAsyncPackage::SetupImports_Event()
 				if (!ImportPackage)
 				{
 					Import.bImportFailed = true;
-					UE_CLOG(!FLinkerLoad::IsKnownMissingPackage(ImportPackageName), LogStreaming, Error, TEXT("Missing native package (%s) for import of %s in %s."), *ImportPackageName.ToString(), *Import.ObjectName.ToString(), *Desc.NameToLoad.ToString());
+					UE_CLOG(!FLinkerLoad::IsKnownMissingPackage(ImportPackageName), LogStreaming, Error, TEXT("Missing native package (%s) for import of %s in %s."), *ImportPackageName.ToString(), *Import.ObjectName.ToString(), *Desc.PackagePath.GetDebugName());
 				}
 				else if (OuterMostImport.OuterIndex.IsNull())
 				{
@@ -2575,7 +2579,7 @@ void FAsyncPackage::LinkImport(int32 LocalImportIndex)
 			if (!ImportPackage)
 			{
 				Import.bImportFailed = true;
-				UE_CLOG(!FLinkerLoad::IsKnownMissingPackage(OuterMostImport.ObjectName), LogStreaming, Error, TEXT("Missing native package (%s) for import of %s in %s."), *OuterMostImport.ObjectName.ToString(), *Import.ObjectName.ToString(), *Desc.NameToLoad.ToString());
+				UE_CLOG(!FLinkerLoad::IsKnownMissingPackage(OuterMostImport.ObjectName), LogStreaming, Error, TEXT("Missing native package (%s) for import of %s in %s."), *OuterMostImport.ObjectName.ToString(), *Import.ObjectName.ToString(), *Desc.PackagePath.GetDebugName());
 			}
 			else
 			{
@@ -2604,12 +2608,12 @@ void FAsyncPackage::LinkImport(int32 LocalImportIndex)
 										return;
 									}
 									Outer = OuterImport.XObject;
-									UE_CLOG(!Outer, LogStreaming, Fatal, TEXT("Missing outer for import of (%s): %s in %s was not found, but the package exists."), *Desc.NameToLoad.ToString(), *OuterImport.ObjectName.ToString(), *ImportPackage->GetFullName());
+									UE_CLOG(!Outer, LogStreaming, Fatal, TEXT("Missing outer for import of (%s): %s in %s was not found, but the package exists."), *Desc.PackagePath.GetDebugName(), *OuterImport.ObjectName.ToString(), *ImportPackage->GetFullName());
 								}
 							}
 							//@todo FH: if we change how StaticFindObjectFast works with external package we will need to change this
 							Import.XObject = FLinkerLoad::FindImportFast(FindClass, Outer, Import.ObjectName);
-							UE_CLOG(!Import.XObject, LogStreaming, Fatal, TEXT("Missing import of (%s): %s in %s was not found, but the package exists."), *Desc.NameToLoad.ToString(), *Import.ObjectName.ToString(), *ImportPackage->GetFullName());
+							UE_CLOG(!Import.XObject, LogStreaming, Fatal, TEXT("Missing import of (%s): %s in %s was not found, but the package exists."), *Desc.PackagePath.GetDebugName(), *Import.ObjectName.ToString(), *ImportPackage->GetFullName());
 						}
 					}
 				}
@@ -3409,7 +3413,7 @@ void FAsyncPackage::MakeNextPrecacheRequestCurrent()
 	Read->WaitCompletion();
 
 	bool bReady = AsyncLoader->PrecacheForEvent(Read, CurrentBlockOffset, CurrentBlockBytes);
-	UE_CLOG(!bReady, LogStreaming, Warning, TEXT("Precache request should have been hot %s."), *Linker->Filename);
+	UE_CLOG(!bReady, LogStreaming, Warning, TEXT("Precache request should have been hot %s."), *Linker->GetPackagePath().GetDebugName());
 	for (int32 Index = Req.FirstExportCovered; Index <= Req.LastExportCovered; Index++)
 	{
 		verify(ExportIndexToPrecacheRequest.Remove(Index) == 1);
@@ -5497,7 +5501,7 @@ EAsyncPackageState::Type FAsyncPackage::TickAsyncPackage(bool InbUseTimeLimit, b
 		BeginAsyncLoad();
 
 		// We have begun loading a package that we know the name of. Let the package time tracker know.
-		FExclusiveLoadPackageTimeTracker::PushLoadPackage(Desc.NameToLoad);
+		FExclusiveLoadPackageTimeTracker::PushLoadPackage(Desc.PackagePath.GetPackageFName());
 
 		if (!GEventDrivenLoaderEnabled)
 		{
@@ -5654,7 +5658,7 @@ EAsyncPackageState::Type FAsyncPackage::CreateLinker()
 			Package = CreatePackage(*Desc.Name.ToString());
 			if (!Package)
 			{
-				UE_LOG(LogStreaming, Error, TEXT("Failed to create package %s requested by async loading code. NameToLoad: %s"), *Desc.Name.ToString(), *Desc.NameToLoad.ToString());
+				UE_LOG(LogStreaming, Error, TEXT("Failed to create package %s requested by async loading code. NameToLoad: %s"), *Desc.Name.ToString(), *Desc.PackagePath.GetDebugName());
 				bLoadHasFailed = true;
 				return EAsyncPackageState::TimeOut;
 			}
@@ -5671,7 +5675,7 @@ EAsyncPackageState::Type FAsyncPackage::CreateLinker()
 		}
 		FScopeCycleCounterUObject ConstructorScope(Package, GET_STATID(STAT_FAsyncPackage_CreateLinker));
 
-		if (Package->FileName == NAME_None && !Package->bHasBeenFullyLoaded)
+		if (Package->GetLoadedPath().IsEmpty() && !Package->bHasBeenFullyLoaded)
 		{
 			SCOPED_LOADTIMER(CreateLinker_SetFlags);
 			// We just created the package, so set ownership flag and set up package info
@@ -5682,7 +5686,7 @@ EAsyncPackageState::Type FAsyncPackage::CreateLinker()
 			Package->PIEInstanceID = Desc.PIEInstanceID;
 
 			// Always store package filename we loading from
-			Package->FileName = Desc.NameToLoad;
+			Package->SetLoadedPath(Desc.PackagePath);
 #if WITH_EDITORONLY_DATA
 			// Assume all packages loaded through async loading are required by runtime
 			Package->SetLoadedByEditorPropertiesOnly(false);
@@ -5712,11 +5716,22 @@ EAsyncPackageState::Type FAsyncPackage::CreateLinker()
 		if (!Linker)
 		{
 			// Process any package redirects
-			FString NameToLoad;
+			FPackagePath PackagePath(Desc.PackagePath);
+			FString NameToLoad = PackagePath.GetPackageName();
 			{
 				SCOPED_LOADTIMER(CreateLinker_GetRedirectedName);
-				const FCoreRedirectObjectName NewPackageName = FCoreRedirects::GetRedirectedName(ECoreRedirectFlags::Type_Package, FCoreRedirectObjectName(NAME_None, NAME_None, Desc.NameToLoad));
-				NameToLoad = NewPackageName.PackageName.ToString();
+				FName FNameToLoad(*NameToLoad);
+				const FCoreRedirectObjectName NewPackageName = FCoreRedirects::GetRedirectedName(ECoreRedirectFlags::Type_Package, FCoreRedirectObjectName(NAME_None, NAME_None, FNameToLoad));
+				if (NewPackageName.PackageName != FNameToLoad)
+				{
+					NameToLoad = NewPackageName.PackageName.ToString();
+					if (!FPackagePath::TryFromMountedName(NameToLoad, PackagePath))
+					{
+						UE_LOG(LogStreaming, Error, TEXT("PackagePath %s requested by async loading code (PackageName: %s) has invalid FCoreRedirects redirector to unmounted contentroot %s."), *Desc.PackagePath.GetPackageName(), *Desc.Name.ToString(), *NameToLoad);
+						bLoadHasFailed = true;
+						return EAsyncPackageState::TimeOut;
+					}
+				}
 			}
 
 			// The editor must not redirect packages for localization.
@@ -5724,31 +5739,42 @@ EAsyncPackageState::Type FAsyncPackage::CreateLinker()
 			{
 				SCOPED_LOADTIMER(CreateLinker_MassagePath);
 				// Allow delegates to resolve this path
-				NameToLoad = FPackageName::GetDelegateResolvedPackagePath(NameToLoad);
-				NameToLoad = FPackageName::GetLocalizedPackagePath(NameToLoad);
+				FString PackageNameToLoad = FPackageName::GetDelegateResolvedPackagePath(NameToLoad);
+				PackageNameToLoad = FPackageName::GetLocalizedPackagePath(PackageNameToLoad);
+				if (PackageNameToLoad != NameToLoad)
+				{
+					NameToLoad = MoveTemp(PackageNameToLoad);
+					if (!FPackagePath::TryFromMountedName(NameToLoad, PackagePath))
+					{
+						UE_LOG(LogStreaming, Error, TEXT("PackagePath %s requested by async loading code (PackageName: %s) has invalid Localization redirector to unmounted contentroot %s."), *Desc.PackagePath.GetPackageName(), *Desc.Name.ToString(), *NameToLoad);
+						bLoadHasFailed = true;
+						return EAsyncPackageState::TimeOut;
+					}
+				}
 			}
 
-			const FGuid* const Guid = Desc.Guid.IsValid() ? &Desc.Guid : nullptr;
-
-			FString PackageFileName;
-			bool DoesPackageExist = false;
+			FName NameToLoadFName = PackagePath.GetPackageFName();
+			if (GetConvertedDynamicPackageNameToTypeName().Contains(NameToLoadFName))
 			{
-				SCOPED_LOADTIMER(CreateLinker_DoesExist);
-				DoesPackageExist = FPackageName::DoesPackageExist(NameToLoad, Guid, &PackageFileName);
+				// Skip the existence and normalization checks
 			}
-
+			else
 			{
-				SCOPED_LOADTIMER(CreateLinker_MissingPackage);
+				const FGuid* const Guid = Desc.Guid.IsValid() ? &Desc.Guid : nullptr;
 
-				if (Desc.NameToLoad == NAME_None ||
-					(!GetConvertedDynamicPackageNameToTypeName().Contains(Desc.Name) &&
-						!DoesPackageExist))
+				bool bDoesPackageExist = false;
 				{
-					FName FailedLoadName = FName(*NameToLoad);
+					SCOPED_LOADTIMER(CreateLinker_DoesExist);
+					bDoesPackageExist = FPackageName::DoesPackageExist(PackagePath, Guid, false /* bMatchCaseOnDisk */, &PackagePath);
+				}
 
-					if (!FLinkerLoad::IsKnownMissingPackage(FailedLoadName))
+				if (!bDoesPackageExist)
+				{
+					SCOPED_LOADTIMER(CreateLinker_MissingPackage);
+
+					if (!FLinkerLoad::IsKnownMissingPackage(NameToLoadFName))
 					{
-						UE_LOG(LogStreaming, Error, TEXT("Couldn't find file for package %s requested by async loading code. NameToLoad: %s"), *Desc.Name.ToString(), *Desc.NameToLoad.ToString());
+						UE_LOG(LogStreaming, Error, TEXT("Couldn't find file for package %s requested by async loading code. NameToLoad: %s"), *Desc.Name.ToString(), *Desc.PackagePath.GetDebugName());
 
 #if !WITH_EDITORONLY_DATA
 						UE_CLOG(bUseTimeLimit, LogStreaming, Error, TEXT("This will hitch streaming because it ends up searching the disk instead of finding the file in the pak file."));
@@ -5781,7 +5807,7 @@ EAsyncPackageState::Type FAsyncPackage::CreateLinker()
 						}
 
 						// Add to known missing list so it won't error again
-						FLinkerLoad::AddKnownMissingPackage(FailedLoadName);
+						FLinkerLoad::AddKnownMissingPackage(NameToLoadFName);
 					}
 
 					bLoadHasFailed = true;
@@ -5804,7 +5830,7 @@ EAsyncPackageState::Type FAsyncPackage::CreateLinker()
 				FWeakAsyncPackagePtr WeakPtr(this);
 				FPrecacheCallbackHandler* PrecacheHandler = &AsyncLoadingThread.GetPrecacheHandler();
 				check(Package);
-				Linker = FLinkerLoad::CreateLinkerAsync(LoadContext, Package, *PackageFileName, LinkerFlags, Desc.GetInstancingContext()
+				Linker = FLinkerLoad::CreateLinkerAsync(LoadContext, Package, PackagePath, LinkerFlags, Desc.GetInstancingContext()
 					, TFunction<void()>(
 						[WeakPtr, PrecacheHandler]()
 						{
@@ -5829,7 +5855,7 @@ EAsyncPackageState::Type FAsyncPackage::CreateLinker()
 			}
 			else
 			{
-				Linker = FLinkerLoad::CreateLinkerAsync(LoadContext, Package, *PackageFileName, LinkerFlags, Desc.GetInstancingContext(), TFunction<void()>([]() {}));
+				Linker = FLinkerLoad::CreateLinkerAsync(LoadContext, Package, PackagePath, LinkerFlags, Desc.GetInstancingContext(), TFunction<void()>([]() {}));
 			}
 		}
 
@@ -5838,7 +5864,7 @@ EAsyncPackageState::Type FAsyncPackage::CreateLinker()
 		check(Linker->AsyncRoot == nullptr || Linker->AsyncRoot == this);
 		Linker->AsyncRoot = this;
 
-		UE_LOG(LogStreaming, Verbose, TEXT("FAsyncPackage::CreateLinker for %s finished."), *Desc.NameToLoad.ToString());
+		UE_LOG(LogStreaming, Verbose, TEXT("FAsyncPackage::CreateLinker for %s finished."), *Desc.PackagePath.GetDebugName());
 	}
 	return EAsyncPackageState::Complete;
 }
@@ -5924,7 +5950,8 @@ void FAsyncPackage::AddImportDependency(const FName& PendingImport, const FName&
 
 	if (!PackageToStream)
 	{
-		FAsyncPackageDesc Info(INDEX_NONE, PendingImport, PackageToLoad);
+		FPackagePath ImportPackagePath = FPackagePath::FromPackageNameChecked(!PackageToLoad.IsNone() ? PackageToLoad : PendingImport);
+		FAsyncPackageDesc Info(INDEX_NONE, PendingImport, ImportPackagePath);
 		Info.SetInstancingContext(MoveTemp(InstancingContext));
 		PackageToStream = new FAsyncPackage(AsyncLoadingThread, Info, EDLBootNotificationManager);
 
@@ -6076,7 +6103,7 @@ EAsyncPackageState::Type FAsyncPackage::LoadImports(FFlushTree* FlushTree)
 				}
 				else
 				{
-					UE_LOG(LogStreaming, Verbose, TEXT("FAsyncPackage::LoadImports for %s: Linker exists for %s"), *Desc.NameToLoad.ToString(), *ImportPackageFName.ToString());
+					UE_LOG(LogStreaming, Verbose, TEXT("FAsyncPackage::LoadImports for %s: Linker exists for %s"), *Desc.PackagePath.GetDebugName(), *ImportPackageFName.ToString());
 					// Only keep a reference to this package so that its linker doesn't go away too soon
 					PendingPackage->DependencyRefCount.Increment();
 					ReferencedImports.Add(PendingPackage);
@@ -6093,13 +6120,13 @@ EAsyncPackageState::Type FAsyncPackage::LoadImports(FFlushTree* FlushTree)
 			// The package doesn't exist and this import is not in the dependency list so add it now.
 			if (!FPackageName::IsShortPackageName(ImportPackageName))
 			{
-				UE_LOG(LogStreaming, Verbose, TEXT("FAsyncPackage::LoadImports for %s: Loading %s"), *Desc.NameToLoad.ToString(), *ImportPackageName);
+				UE_LOG(LogStreaming, Verbose, TEXT("FAsyncPackage::LoadImports for %s: Loading %s"), *Desc.PackagePath.GetDebugName(), *ImportPackageName);
 				AddImportDependency(ImportPackageFName, ImportToLoad, FlushTree, InstancingContext);
 			}
 			else
 			{
 				// This usually means there's a reference to a script package from another project
-				UE_LOG(LogStreaming, Warning, TEXT("FAsyncPackage::LoadImports for %s: Short package name in imports list: %s"), *Desc.NameToLoad.ToString(), *ImportPackageName);
+				UE_LOG(LogStreaming, Warning, TEXT("FAsyncPackage::LoadImports for %s: Short package name in imports list: %s"), *Desc.PackagePath.GetDebugName(), *ImportPackageName);
 			}
 		}
 		UpdateLoadPercentage();
@@ -6121,7 +6148,7 @@ void FAsyncPackage::ImportFullyLoadedCallback(const FName& InPackageName, UPacka
 {
 	if (Result != EAsyncLoadingResult::Canceled)
 	{
-		UE_LOG(LogStreaming, Verbose, TEXT("FAsyncPackage::LoadImports for %s: Loaded %s"), *Desc.NameToLoad.ToString(), *InPackageName.ToString());
+		UE_LOG(LogStreaming, Verbose, TEXT("FAsyncPackage::LoadImports for %s: Loaded %s"), *Desc.PackagePath.GetDebugName(), *InPackageName.ToString());
 		int32 Index = ContainsDependencyPackage(PendingImportedPackages, InPackageName);
 		if (Index != INDEX_NONE)
 		{
@@ -6275,7 +6302,7 @@ void FAsyncPackage::FreeReferencedImports()
 	for (int32 ReferenceIndex = 0; ReferenceIndex < ReferencedImports.Num(); ++ReferenceIndex)
 	{
 		FAsyncPackage& Ref = *ReferencedImports[ReferenceIndex];
-		UE_LOG(LogStreaming, Verbose, TEXT("FAsyncPackage::FreeReferencedImports for %s: Releasing %s (%d)"), *Desc.NameToLoad.ToString(), *Ref.GetPackageName().ToString(), Ref.GetDependencyRefCount());
+		UE_LOG(LogStreaming, Verbose, TEXT("FAsyncPackage::FreeReferencedImports for %s: Releasing %s (%d)"), *Desc.PackagePath.GetDebugName(), *Ref.GetPackageName().ToString(), Ref.GetDependencyRefCount());
 		int32 RefPackageDependencyRefCount = Ref.DependencyRefCount.Decrement();
 		check(RefPackageDependencyRefCount >= 0);
 	}
@@ -6893,10 +6920,8 @@ void FAsyncPackage::UpdateLoadPercentage()
 	LoadPercentage = FMath::Max(NewLoadPercentage, LoadPercentage);
 }
 
-int32 FAsyncLoadingThread::LoadPackage(const FString& InName, const FGuid* InGuid, const TCHAR* InPackageToLoadFrom, FLoadPackageAsyncDelegate InCompletionDelegate, EPackageFlags InPackageFlags, int32 InPIEInstanceID, int32 InPackagePriority, const FLinkerInstancingContext* InstancingContext)
+int32 FAsyncLoadingThread::LoadPackage(const FPackagePath& InPackagePath, FName InCustomName, FLoadPackageAsyncDelegate InCompletionDelegate, const FGuid* InGuid, EPackageFlags InPackageFlags, int32 InPIEInstanceID, int32 InPackagePriority, const FLinkerInstancingContext* InstancingContext)
 {
-	int32 RequestID = INDEX_NONE;
-
 	static bool bOnce = false;
 	if (!bOnce && GEventDrivenLoaderEnabled)
 	{
@@ -6904,82 +6929,39 @@ int32 FAsyncLoadingThread::LoadPackage(const FString& InName, const FGuid* InGui
 		FGCObject::StaticInit(); // otherwise this thing is created during async loading, but not associated with a package
 	}
 
-	// The comments clearly state that it should be a package name but we also handle it being a filename as this function is not perf critical
-	// and LoadPackage handles having a filename being passed in as well.
-	FString PackageName;
-	bool bValidPackageName = true;
-
-	if (FPackageName::IsValidLongPackageName(InName, /*bIncludeReadOnlyRoots*/true))
+	if (!InPackagePath.IsMountedPath())
 	{
-		PackageName = InName;
+		UE_LOG(LogStreaming, Warning, TEXT("Async loading code requires MountedPath PackagePaths (%s)."), *InPackagePath.GetDebugName());
+		InCompletionDelegate.ExecuteIfBound(NAME_None, nullptr, EAsyncLoadingResult::Failed);
+		return INDEX_NONE;
 	}
-	// PackageName got populated by the conditional function
-	else if (!(FPackageName::IsPackageFilename(InName) && FPackageName::TryConvertFilenameToLongPackageName(InName, PackageName)))
-	{
-		// PackageName may get populated by the conditional function
-		FString ClassName;
-
-		if (!FPackageName::ParseExportTextPath(PackageName, &ClassName, &PackageName))
-		{
-			UE_LOG(LogStreaming, Warning, TEXT("LoadPackageAsync failed to begin to load a package because the supplied package name ")
-					TEXT("was neither a valid long package name nor a filename of a map within a content folder: '%s' (%s)"),
-					*PackageName, *InName);
 
-			bValidPackageName = false;
-		}
+	FName PackageName = InCustomName.IsNone() ? InPackagePath.GetPackageFName() : InCustomName;
+	if ( FCoreDelegates::OnAsyncLoadPackage.IsBound() )
+	{
+		FCoreDelegates::OnAsyncLoadPackage.Broadcast(PackageName.ToString());
 	}
 
-	FString PackageNameToLoad(InPackageToLoadFrom);
+	// Generate new request ID and add it immediately to the global request list (it needs to be there before we exit
+	// this function, otherwise it would be added when the packages are being processed on the async thread).
+	int32 RequestID = GPackageRequestID.Increment();
+	TRACE_LOADTIME_BEGIN_REQUEST(RequestID);
+	AddPendingRequest(RequestID);
 
-	if (bValidPackageName)
+	// Allocate delegate on Game Thread, it is not safe to copy delegates by value on other threads
+	TUniquePtr<FLoadPackageAsyncDelegate> CompletionDelegatePtr;
+	if (InCompletionDelegate.IsBound())
 	{
-		if (PackageNameToLoad.IsEmpty())
-		{
-			PackageNameToLoad = PackageName;
-		}
-		// Make sure long package name is passed to FAsyncPackage so that it doesn't attempt to 
-		// create a package with short name.
-		if (FPackageName::IsShortPackageName(PackageNameToLoad))
-		{
-			UE_LOG(LogStreaming, Warning, TEXT("Async loading code requires long package names (%s)."), *PackageNameToLoad);
-
-			bValidPackageName = false;
-		}
+		CompletionDelegatePtr = MakeUnique<FLoadPackageAsyncDelegate>(MoveTemp(InCompletionDelegate));
 	}
 
-	if (bValidPackageName)
-	{
-		if ( FCoreDelegates::OnAsyncLoadPackage.IsBound() )
-		{
-			FCoreDelegates::OnAsyncLoadPackage.Broadcast(InName);
-		}
-
-		// Generate new request ID and add it immediately to the global request list (it needs to be there before we exit
-		// this function, otherwise it would be added when the packages are being processed on the async thread).
-		RequestID = GPackageRequestID.Increment();
-		TRACE_LOADTIME_BEGIN_REQUEST(RequestID);
-		AddPendingRequest(RequestID);
-
-		// Allocate delegate on Game Thread, it is not safe to copy delegates by value on other threads
-		TUniquePtr<FLoadPackageAsyncDelegate> CompletionDelegatePtr;
-		if (InCompletionDelegate.IsBound())
-		{
-			CompletionDelegatePtr = MakeUnique<FLoadPackageAsyncDelegate>(MoveTemp(InCompletionDelegate));
-		}
-
-		// Add new package request
-		FAsyncPackageDesc PackageDesc(RequestID, *PackageName, *PackageNameToLoad, InGuid ? *InGuid : FGuid(), MoveTemp(CompletionDelegatePtr), InPackageFlags, InPIEInstanceID, InPackagePriority);
-		if (InstancingContext)
-		{
-			PackageDesc.SetInstancingContext(*InstancingContext);
-		}
-		QueuePackage(PackageDesc);
-	}
-	else
+	// Add new package request
+	FAsyncPackageDesc PackageDesc(RequestID, PackageName, InPackagePath, InGuid ? *InGuid : FGuid(), MoveTemp(CompletionDelegatePtr), InPackageFlags, InPIEInstanceID, InPackagePriority);
+	if (InstancingContext)
 	{
-		InCompletionDelegate.ExecuteIfBound(FName(*InName), nullptr, EAsyncLoadingResult::Failed);
+		PackageDesc.SetInstancingContext(*InstancingContext);
 	}
-
+	QueuePackage(PackageDesc);
 	return RequestID;
 }
 
@@ -7152,7 +7134,7 @@ static FAutoConsoleCommand GDumpSerializeCmd(
 
 static FCriticalSection SummaryRacePreventer;
 
-FAsyncArchive::FAsyncArchive(const TCHAR* InFileName, FLinkerLoad* InOwner, TFunction<void()>&& InSummaryReadyCallback)
+FAsyncArchive::FAsyncArchive(const FPackagePath& InPackagePath, FLinkerLoad* InOwner, TFunction<void()>&& InSummaryReadyCallback)
 	: Handle(nullptr)
 	, SizeRequestPtr(nullptr)
 	, EditorPrecacheRequestPtr(nullptr)
@@ -7171,7 +7153,7 @@ FAsyncArchive::FAsyncArchive(const TCHAR* InFileName, FLinkerLoad* InOwner, TFun
 	, HeaderSizeWhenReadingExportsFromSplitFile(0)
 	, LoadPhase(ELoadPhase::WaitingForSize)
 	, bCookedForEDLInEditor(false)
-	, FileName(InFileName)
+	, PackagePath(InPackagePath)
 	, OpenTime(FPlatformTime::Seconds())
 	, SummaryReadTime(0.0)
 	, ExportReadTime(0.0)
@@ -7179,8 +7161,8 @@ FAsyncArchive::FAsyncArchive(const TCHAR* InFileName, FLinkerLoad* InOwner, TFun
 	, OwnerLinker(InOwner)
 {
 	LogItem(TEXT("Open"));
-	Handle = FPlatformFileManager::Get().GetPlatformFile().OpenAsyncRead(InFileName);
-	check(Handle); // this generally cannot fail because it is async
+	Handle = IPackageResourceManager::Get().OpenAsyncReadPackage(PackagePath, EPackageSegment::Header);
+	check(Handle); // OpenAsyncReadPackage guarantees a non-null return value; the handle will fail to read later if the path does not exist
 
 	ReadCallbackFunction = [this](bool bWasCancelled, IAsyncReadRequest* Request)
 	{
@@ -7285,7 +7267,7 @@ void FAsyncArchive::ReadCallback(bool bWasCancelled, IAsyncReadRequest* Request)
 				// we need to be sure that we can at least get the size from the initial request. This is an early warning that custom versions are starting to get too big, relocate the total size to be at offset 4!
 				checkf(Ar.Tell() < FMaxPackageSummarySize::Value / 2,
 					TEXT("The initial read request was too small (%d) compared to package %s header size (%lld). Try increasing s.MaxPackageSummarySize value in DefaultEngine.ini."),
-					FMaxPackageSummarySize::Value, *FileName, Ar.Tell());
+					FMaxPackageSummarySize::Value, *PackagePath.GetDebugName(), Ar.Tell());
 				
 				// Support for cooked EDL packages in the editor
 				bCookedForEDLInEditor = !FPlatformProperties::RequiresCookedData() && (Sum.PackageFlags & PKG_FilterEditorOnly) && Sum.PreloadDependencyCount > 0 && Sum.PreloadDependencyOffset > 0;
@@ -7369,7 +7351,7 @@ void FAsyncArchive::FlushCache()
 			float(1000.0 * (SummaryReadTime - OpenTime)),
 			float(1000.0 * (ExportReadTime - SummaryReadTime)),
 			float(1000.0 * (Now - ExportReadTime)),
-			*FileName);
+			*PackagePath.GetDebugName());
 #if defined(ASYNC_WATCH_FILE)
 		if (FileName.Contains(TEXT(ASYNC_WATCH_FILE)))
 		{
@@ -7441,7 +7423,8 @@ void FAsyncArchive::Seek(int64 InPos)
 			FirstExportStarting();
 		}
 	}
-	checkf(InPos >= 0 && InPos <= TotalSizeOrMaxInt64IfNotReady(), TEXT("Bad position in FAsyncArchive::Seek. Filename:%s InPos:%lu, Size:%lu"), *FileName, InPos, TotalSizeOrMaxInt64IfNotReady());
+	checkf(InPos >= 0 && InPos <= TotalSizeOrMaxInt64IfNotReady(), TEXT("Bad position in FAsyncArchive::Seek. Filename:%s InPos:%lu, Size:%lu"),
+		*PackagePath.GetDebugName(), InPos, TotalSizeOrMaxInt64IfNotReady());
 #if DEVIRTUALIZE_FLinkerLoad_Serialize
 	SetPosAndUpdatePrecacheBuffer(InPos);
 #else
@@ -7644,7 +7627,7 @@ bool FAsyncArchive::PrecacheInternal(int64 RequestOffset, int64 RequestSize, boo
 		if (ReadRequestPtr)
 		{
 			// this one does not have what we need
-			UE_LOG(LogStreaming, Warning, TEXT("FAsyncArchive::PrecacheInternal Canceled read for %s  Offset = %lld   Size = %lld"), *FileName, RequestOffset, ReadRequestSize);
+			UE_LOG(LogStreaming, Warning, TEXT("FAsyncArchive::PrecacheInternal Canceled read for %s  Offset = %lld   Size = %lld"), *PackagePath.GetDebugName(), RequestOffset, ReadRequestSize);
 			CancelRead();
 		}
 	}
@@ -7729,10 +7712,9 @@ void FAsyncArchive::FirstExportStarting()
 		}
 
 		HeaderSizeWhenReadingExportsFromSplitFile = HeaderSize;
-		FileName = FPaths::GetBaseFilename(FileName, false) + TEXT(".uexp");
 
-		Handle = FPlatformFileManager::Get().GetPlatformFile().OpenAsyncRead(*FileName);
-		check(Handle); // this generally cannot fail because it is async
+		Handle = IPackageResourceManager::Get().OpenAsyncReadPackage(PackagePath, EPackageSegment::Exports);
+		check(Handle); // OpenAsyncReadPackage guarantees a non-null return value; the handle will fail to read later if the path does not exist
 
 		check(!SizeRequestPtr);
 		SizeRequestPtr = Handle->SizeRequest();
@@ -7750,12 +7732,11 @@ IAsyncReadRequest* FAsyncArchive::MakeEventDrivenPrecacheRequest(int64 Offset, i
 	{
 		// we need to avoid tearing down the old file and requests until we have the one in flight
 		HeaderSizeWhenReadingExportsFromSplitFile = HeaderSize;
-		FString NewFileName = FPaths::GetBaseFilename(FileName, false) + TEXT(".uexp");
 		IAsyncReadFileHandle* NewHandle;
 		{
 			double StartTime = FPlatformTime::Seconds();
-			NewHandle = FPlatformFileManager::Get().GetPlatformFile().OpenAsyncRead(*NewFileName);
-		check(NewHandle); // this generally cannot fail because it is async
+			NewHandle = IPackageResourceManager::Get().OpenAsyncReadPackage(PackagePath, EPackageSegment::Exports);
+			check(NewHandle); // OpenAsyncReadPackage guarantees a non-null return value; the handle will fail to read later if the path does not exist
 			LogItem(TEXT("Open UExp"), Offset - HeaderSizeWhenReadingExportsFromSplitFile, BytesToRead, StartTime);
 		}
 		{
@@ -7766,26 +7747,25 @@ IAsyncReadRequest* FAsyncArchive::MakeEventDrivenPrecacheRequest(int64 Offset, i
 			EAsyncIOPriorityAndFlags Prio = NewHandle->UsesCache() ? GetAsyncIOPrecachePriorityAndFlags() : GetAsyncIOPriority();
 
 			IAsyncReadRequest* Precache = NewHandle->ReadRequest(Offset - HeaderSizeWhenReadingExportsFromSplitFile, BytesToRead, Prio, CompleteCallback);
-		FlushCache();
-		if (Handle)
-		{
-			delete Handle;
-			Handle = nullptr;
-		}
-		Handle = NewHandle;
-		FileName = NewFileName;
+			FlushCache();
+			if (Handle)
+			{
+				delete Handle;
+				Handle = nullptr;
+			}
+			Handle = NewHandle;
 
-		FirstExportStarting();
+			FirstExportStarting();
 
-		check(!SizeRequestPtr);
-		SizeRequestPtr = Handle->SizeRequest();
-		if (SizeRequestPtr->PollCompletion())
-		{
-			TotalSize(); // complete the request
+			check(!SizeRequestPtr);
+			SizeRequestPtr = Handle->SizeRequest();
+			if (SizeRequestPtr->PollCompletion())
+			{
+				TotalSize(); // complete the request
+			}
+			LogItem(TEXT("First Precache"), Offset - HeaderSizeWhenReadingExportsFromSplitFile, BytesToRead, StartTime);
+			return Precache;
 		}
-		LogItem(TEXT("First Precache"), Offset - HeaderSizeWhenReadingExportsFromSplitFile, BytesToRead, StartTime);
-		return Precache;
-	}
 	}
 	double StartTime = FPlatformTime::Seconds();
 	check(Offset - HeaderSizeWhenReadingExportsFromSplitFile >= 0);
@@ -7953,7 +7933,7 @@ void FAsyncArchive::Serialize(void* Data, int64 Count)
 	}
 #endif
 	// Ensure we aren't reading beyond the end of the file
-	checkf(CurrentPos + Count <= TotalSizeOrMaxInt64IfNotReady(), TEXT("Seeked past end of file %s (%lld / %lld)"), *FileName, CurrentPos + Count, TotalSize());
+	checkf(CurrentPos + Count <= TotalSizeOrMaxInt64IfNotReady(), TEXT("Seeked past end of file %s (%lld / %lld)"), *PackagePath.GetDebugName(), CurrentPos + Count, TotalSize());
 
 	int64 BeforeBlockOffset = 0;
 	int64 BeforeBlockSize = 0;
@@ -8000,7 +7980,8 @@ void FAsyncArchive::Serialize(void* Data, int64 Count)
 	}
 	if (BeforeBlockSize)
 	{
-		UE_CLOG(GEventDrivenLoaderEnabled, LogAsyncArchive, Warning, TEXT("FAsyncArchive::Serialize Backwards streaming in %s  CurrentPos = %lld   BeforeBlockOffset = %lld"), *FileName, CurrentPos, BeforeBlockOffset);
+		UE_CLOG(GEventDrivenLoaderEnabled, LogAsyncArchive, Warning, TEXT("FAsyncArchive::Serialize Backwards streaming in %s  CurrentPos = %lld   BeforeBlockOffset = %lld"),
+			*PackagePath.GetDebugName(), CurrentPos, BeforeBlockOffset);
 		LogItem(TEXT("Sync Before Block"), BeforeBlockOffset, BeforeBlockSize);
 		if (!PrecacheInternal(BeforeBlockOffset, BeforeBlockSize))
 		{
diff --git a/Engine/Source/Runtime/CoreUObject/Private/Serialization/AsyncLoading.h b/Engine/Source/Runtime/CoreUObject/Private/Serialization/AsyncLoading.h
--- a/Engine/Source/Runtime/CoreUObject/Private/Serialization/AsyncLoading.h
+++ b/Engine/Source/Runtime/CoreUObject/Private/Serialization/AsyncLoading.h
@@ -400,11 +400,11 @@ struct FAsyncPackage : public FGCObject
 	}
 
 	/**
-	 * Returns the name to load of the package.
+	 * Returns the PackagePath from which to load the package.
 	 */
-	FORCEINLINE const FName& GetPackageNameToLoad() const
+	FORCEINLINE const FPackagePath& GetPackagePath() const
 	{
-		return Desc.NameToLoad;
+		return Desc.PackagePath;
 	}
 
 	void AddCompletionCallback(TUniquePtr<FLoadPackageAsyncDelegate>&& Callback, bool bInternal);
diff --git a/Engine/Source/Runtime/CoreUObject/Private/Serialization/AsyncLoading2.cpp b/Engine/Source/Runtime/CoreUObject/Private/Serialization/AsyncLoading2.cpp
--- a/Engine/Source/Runtime/CoreUObject/Private/Serialization/AsyncLoading2.cpp
+++ b/Engine/Source/Runtime/CoreUObject/Private/Serialization/AsyncLoading2.cpp
@@ -271,14 +271,14 @@ if ((ELogVerbosity::Type(ELogVerbosity::Verbosity & ELogVerbosity::VerbosityMask
 		UE_LOG(LogStreaming, Verbosity, LogDesc TEXT(": %s (0x%llX) %s (0x%llX) - ") Format, \
 			*(PackageDesc).CustomPackageName.ToString(), \
 			(PackageDesc).CustomPackageId.ValueForDebugging(), \
-			*(PackageDesc).DiskPackageName.ToString(), \
+			*(PackageDesc).DiskPackagePath.GetDebugName(), \
 			(PackageDesc).DiskPackageId.ValueForDebugging(), \
 			##__VA_ARGS__); \
 	} \
 	else \
 	{ \
 		UE_LOG(LogStreaming, Verbosity, LogDesc TEXT(": %s (0x%llX) - ") Format, \
-			*(PackageDesc).DiskPackageName.ToString(), \
+			*(PackageDesc).DiskPackagePath.GetDebugName(), \
 			(PackageDesc).DiskPackageId.ValueForDebugging(), \
 			##__VA_ARGS__); \
 	} \
@@ -366,17 +366,16 @@ struct FAsyncPackageDesc2
 	const FPackageStoreEntry* StoreEntry;
 	// The disk package id corresponding to the StoreEntry
 	// It is used by the loader for io chunks and to handle ref tracking of loaded packages and import objects.
-	FPackageId DiskPackageId; 
+	FPackageId DiskPackageId;
 	// The custom package id is only set for temp packages with a valid but "fake" CustomPackageName,
 	// if set, it will be used as key when tracking active async packages in AsyncPackageLookup
 	FPackageId CustomPackageId;
-	// The disk package name from the LoadPackage call, or none for imported packages
-	// up until the package summary has been serialized
-	FName DiskPackageName;
-	// The custom package name from the LoadPackage call is only used for temp packages,
-	// if set, it will be used as the runtime UPackage name
+	// The PackagePath corresponding to the StoreEntry, the identifier used to tell the PackageResourceManager which package to load from disk
+	// Will be empty for imported packages until the package summary has been serialized
+	FPackagePath DiskPackagePath;
+	// A custom package name, only used for temp packages. If set it will be used as the runtime UPackage name
 	FName CustomPackageName;
-	// Set from the package summary,
+	// Set from the package summary, this is the base package corresponding to a localized package, and if non-empty it is used for the package name rather than the DiskPackagePath's name.
 	FName SourcePackageName;
 	/** Delegate called on completion of loading. This delegate can only be created and consumed on the game thread */
 	TUniquePtr<FLoadPackageAsyncDelegate> PackageLoadedDelegate;
@@ -386,7 +385,7 @@ struct FAsyncPackageDesc2
 		int32 InPriority,
 		FPackageId InPackageIdToLoad,
 		const FPackageStoreEntry* InStoreEntry,
-		FName InDiskPackageName = FName(),
+		const FPackagePath& InDiskPackagePath = FPackagePath(),
 		FPackageId InPackageId = FPackageId(),
 		FName InCustomName = FName(),
 		TUniquePtr<FLoadPackageAsyncDelegate>&& InCompletionDelegate = TUniquePtr<FLoadPackageAsyncDelegate>())
@@ -395,7 +394,7 @@ struct FAsyncPackageDesc2
 		, StoreEntry(InStoreEntry)
 		, DiskPackageId(InPackageIdToLoad)
 		, CustomPackageId(InPackageId)
-		, DiskPackageName(InDiskPackageName)
+		, DiskPackagePath(InDiskPackagePath)
 		, CustomPackageName(InCustomName)
 		, PackageLoadedDelegate(MoveTemp(InCompletionDelegate))
 	{
@@ -408,7 +407,7 @@ struct FAsyncPackageDesc2
 		, StoreEntry(OldPackage.StoreEntry)
 		, DiskPackageId(OldPackage.DiskPackageId)
 		, CustomPackageId(OldPackage.CustomPackageId)
-		, DiskPackageName(OldPackage.DiskPackageName)
+		, DiskPackagePath(OldPackage.DiskPackagePath)
 		, CustomPackageName(OldPackage.CustomPackageName)
 		, SourcePackageName(OldPackage.SourcePackageName)
 	{
@@ -421,14 +420,6 @@ struct FAsyncPackageDesc2
 		PackageLoadedDelegate = MoveTemp(InPackageLoadedDelegate);
 	}
 
-	void SetDiskPackageName(FName SerializedDiskPackageName, FName SerializedSourcePackageName = FName())
-	{
-		check(DiskPackageName.IsNone() || DiskPackageName == SerializedDiskPackageName);
-		check(SourcePackageName.IsNone() || SourcePackageName == SerializedSourcePackageName);
-		DiskPackageName = SerializedDiskPackageName;
-		SourcePackageName = SerializedSourcePackageName;
-	}
-
 	bool CanBeImported() const
 	{
 		return CustomPackageName.IsNone();
@@ -450,7 +441,7 @@ struct FAsyncPackageDesc2
 			return SourcePackageName;
 		}
 		// normal packages
-		return DiskPackageName;
+		return DiskPackagePath.GetPackageFName();
 	}
 
 	/**
@@ -2383,10 +2374,10 @@ public:
 	virtual void ShutdownLoading() override;
 
 	virtual int32 LoadPackage(
-		const FString& InPackageName,
-		const FGuid* InGuid,
-		const TCHAR* InPackageToLoadFrom,
+		const FPackagePath& InPackagePath,
+		FName InCustomName,
 		FLoadPackageAsyncDelegate InCompletionDelegate,
+		const FGuid* InGuid,
 		EPackageFlags InPackageFlags,
 		int32 InPIEInstanceID,
 		int32 InPackagePriority,
@@ -2526,7 +2517,7 @@ private:
 	{
 		TRACE_CPUPROFILER_EVENT_SCOPE(CreateAsyncPackage);
 		UE_ASYNC_PACKAGE_DEBUG(Desc);
-		checkf(Desc.StoreEntry, TEXT("No package store entry for package %s"), *Desc.DiskPackageName.ToString());
+		checkf(Desc.StoreEntry, TEXT("No package store entry for package %s"), *Desc.DiskPackagePath.GetDebugName());
 
 		FAsyncPackageData Data;
 		Data.ExportCount = Desc.StoreEntry->ExportCount;
@@ -2660,7 +2651,7 @@ void FAsyncLoadingThread2::QueuePackage(FAsyncPackageDesc2& Package)
 {
 	//TRACE_CPUPROFILER_EVENT_SCOPE(QueuePackage);
 	UE_ASYNC_PACKAGE_DEBUG(Package);
-	checkf(Package.StoreEntry, TEXT("No package store entry for package %s"), *Package.DiskPackageName.ToString());
+	checkf(Package.StoreEntry, TEXT("No package store entry for package %s"), *Package.DiskPackagePath.GetDebugName());
 	{
 		FScopeLock QueueLock(&QueueCritical);
 		++QueuedPackagesCounter;
@@ -2680,7 +2671,7 @@ FAsyncPackage2* FAsyncLoadingThread2::FindOrInsertPackage(FAsyncPackageDesc2* De
 		if (!Package)
 		{
 			Package = CreateAsyncPackage(*Desc);
-			checkf(Package, TEXT("Failed to create async package %s"), *Desc->DiskPackageName.ToString());
+			checkf(Package, TEXT("Failed to create async package %s"), *Desc->DiskPackagePath.GetDebugName());
 			Package->AddRef();
 			AsyncPackageLookup.Add(Desc->GetAsyncPackageId(), Package);
 			bInserted = true;
@@ -2728,7 +2719,7 @@ bool FAsyncLoadingThread2::CreateAsyncPackagesFromQueue()
 		{
 			bool bInserted;
 			FAsyncPackage2* Package = FindOrInsertPackage(PackageDesc, bInserted);
-			checkf(Package, TEXT("Failed to find or insert imported package %s"), *PackageDesc->DiskPackageName.ToString());
+			checkf(Package, TEXT("Failed to find or insert imported package %s"), *PackageDesc->DiskPackagePath.GetDebugName());
 
 			if (bInserted)
 			{
@@ -3395,14 +3386,14 @@ EAsyncPackageState::Type FAsyncPackage2::Event_ProcessPackageSummary(FAsyncPacka
 
 		{
 			FName PackageName = Package->NameMap.GetName(PackageSummary->Name);
+			FPackagePath PackagePath = FPackagePath::FromPackageNameChecked(PackageName);
+			check(Package->Desc.DiskPackagePath.IsEmpty() || Package->Desc.DiskPackagePath == PackagePath);
+			Package->Desc.DiskPackagePath = MoveTemp(PackagePath);
 			if (PackageSummary->SourceName != PackageSummary->Name)
 			{
 				FName SourcePackageName = Package->NameMap.GetName(PackageSummary->SourceName);
-				Package->Desc.SetDiskPackageName(PackageName, SourcePackageName);
-			}
-			else
-			{
-				Package->Desc.SetDiskPackageName(PackageName);
+				check(Package->Desc.SourcePackageName.IsNone() || Package->Desc.SourcePackageName == SourcePackageName);
+				Package->Desc.SourcePackageName = SourcePackageName;
 			}
 		}
 
@@ -3529,7 +3520,7 @@ EAsyncPackageState::Type FAsyncPackage2::Event_ProcessExportBundle(FAsyncPackage
 				const int64 Pos = Ar.Tell();
 				checkf(CookedSerialSize <= uint64(Ar.TotalSize() - Pos),
 					TEXT("Package %s: Expected read size: %llu - Remaining archive size: %llu"),
-					*Package->Desc.DiskPackageName.ToString(), CookedSerialSize, uint64(Ar.TotalSize() - Pos));
+					*Package->Desc.DiskPackagePath.GetDebugName(), CookedSerialSize, uint64(Ar.TotalSize() - Pos));
 
 				const bool bSerialized = Package->EventDrivenSerializeExport(BundleEntry->LocalExportIndex, Ar);
 				if (!bSerialized)
@@ -3538,7 +3529,7 @@ EAsyncPackageState::Type FAsyncPackage2::Event_ProcessExportBundle(FAsyncPackage
 				}
 				checkf(CookedSerialSize == uint64(Ar.Tell() - Pos),
 					TEXT("Package %s: Expected read size: %llu - Actual read size: %llu"),
-					*Package->Desc.DiskPackageName.ToString(), CookedSerialSize, uint64(Ar.Tell() - Pos));
+					*Package->Desc.DiskPackagePath.GetDebugName(), CookedSerialSize, uint64(Ar.Tell() - Pos));
 
 				Ar.ExportBufferEnd();
 
@@ -3601,7 +3592,7 @@ UObject* FAsyncPackage2::EventDrivenIndexToObject(FPackageObjectIndex Index, boo
 		UE_CLOG(!Result, LogStreaming, Warning, TEXT("Missing %s import 0x%llX for package %s"),
 			Index.IsScriptImport() ? TEXT("script") : TEXT("package"),
 			Index.Value(),
-			*Desc.DiskPackageName.ToString());
+			*Desc.DiskPackagePath.GetDebugName());
 	}
 #if DO_CHECK
 	if (bCheckSerialized && !IsFullyLoadedObj(Result))
@@ -4315,7 +4306,7 @@ bool FAsyncPackage2::AreAllDependenciesFullyLoaded(TSet<FPackageId>& VisitedPack
 	{
 		FAsyncPackage2* AsyncRoot = AsyncLoadingThread.GetAsyncPackage(PackageId);
 		UE_LOG(LogStreaming, Verbose, TEXT("AreAllDependenciesFullyLoaded: '%s' doesn't have all exports processed by DeferredPostLoad"),
-			*AsyncRoot->Desc.DiskPackageName.ToString());
+			*AsyncRoot->Desc.DiskPackagePath.GetDebugName());
 	}
 	return bLoaded;
 }
@@ -5175,7 +5166,7 @@ FAsyncPackage2::FAsyncPackage2(
 , GraphAllocator(InGraphAllocator)
 , ImportStore(AsyncLoadingThread.GlobalPackageStore, Desc)
 {
-	TRACE_LOADTIME_NEW_ASYNC_PACKAGE(this, Desc.DiskPackageName);
+	TRACE_LOADTIME_NEW_ASYNC_PACKAGE(this, Desc.DiskPackagePath.GetPackageName());
 	AddRequestID(Desc.RequestID);
 
 	ExportBundlesSize = Desc.StoreEntry->ExportBundlesSize;
@@ -5215,13 +5206,13 @@ FAsyncPackage2::~FAsyncPackage2()
 	UE_ASYNC_PACKAGE_LOG(Verbose, Desc, TEXT("AsyncThread: Deleted"), TEXT("Package deleted."));
 
 	checkf(RefCount == 0, TEXT("RefCount is not 0 when deleting package %s"),
-		*Desc.DiskPackageName.ToString());
+		*Desc.DiskPackagePath.GetDebugName());
 
 	checkf(RequestIDs.Num() == 0, TEXT("MarkRequestIDsAsComplete() has not been called for package %s"),
-		*Desc.DiskPackageName.ToString());
+		*Desc.DiskPackagePath.GetDebugName());
 	
 	checkf(ConstructedObjects.Num() == 0, TEXT("ClearConstructedObjects() has not been called for package %s"),
-		*Desc.DiskPackageName.ToString());
+		*Desc.DiskPackagePath.GetDebugName());
 }
 
 void FAsyncPackage2::ReleaseRef()
@@ -5411,7 +5402,7 @@ void FAsyncPackage2::CreateUPackage(const FPackageSummary* PackageSummary)
 			bCreatedLinkerRoot = true;
 		}
 		LinkerRoot->SetFlags(RF_Public | RF_WasLoaded);
-		LinkerRoot->FileName = Desc.DiskPackageName;
+		LinkerRoot->SetLoadedPath(Desc.DiskPackagePath);
 		LinkerRoot->SetCanBeImportedFlag(Desc.CanBeImported());
 		LinkerRoot->SetPackageId(Desc.DiskPackageId);
 		LinkerRoot->SetPackageFlagsTo(PackageSummary->PackageFlags);
@@ -5559,7 +5550,7 @@ void FAsyncPackage2::AddCompletionCallback(TUniquePtr<FLoadPackageAsyncDelegate>
 	CompletionCallbacks.Emplace(MoveTemp(Callback));
 }
 
-int32 FAsyncLoadingThread2::LoadPackage(const FString& InName, const FGuid* InGuid, const TCHAR* InPackageToLoadFrom, FLoadPackageAsyncDelegate InCompletionDelegate, EPackageFlags InPackageFlags, int32 InPIEInstanceID, int32 InPackagePriority, const FLinkerInstancingContext*)
+int32 FAsyncLoadingThread2::LoadPackage(const FPackagePath& InPackagePath, FName InCustomName, FLoadPackageAsyncDelegate InCompletionDelegate, const FGuid* InGuid, EPackageFlags InPackageFlags, int32 InPIEInstanceID, int32 InPackagePriority, const FLinkerInstancingContext*)
 {
 	TRACE_CPUPROFILER_EVENT_SCOPE(LoadPackage);
 
@@ -5570,65 +5561,45 @@ int32 FAsyncLoadingThread2::LoadPackage(const FString& InName, const FGuid* InGu
 	}
 
 	int32 RequestID = INDEX_NONE;
-
-	// happy path where all inputs are actual package names
-	const FName Name = FName(*InName);
-	FName DiskPackageName = InPackageToLoadFrom ? FName(InPackageToLoadFrom) : Name;
-	bool bHasCustomPackageName = Name != DiskPackageName;
-
-	// Verify PackageToLoadName, or fixup to handle any input string that can be converted to a long package name.
-	FPackageId DiskPackageId = FPackageId::FromName(DiskPackageName);
-	const FPackageStoreEntry* StoreEntry = GlobalPackageStore.FindStoreEntry(DiskPackageId);
-	if (!StoreEntry)
+	auto CallCompletionWithError = [&InCompletionDelegate, this](FName PackageName)
 	{
-		FString PackageNameStr = DiskPackageName.ToString();
-		if (!FPackageName::IsValidLongPackageName(PackageNameStr))
+		if (InCompletionDelegate.IsBound())
 		{
-			FString NewPackageNameStr;
-			if (FPackageName::TryConvertFilenameToLongPackageName(PackageNameStr, NewPackageNameStr))
-			{
-				DiskPackageName = *NewPackageNameStr;
-				DiskPackageId = FPackageId::FromName(DiskPackageName);
-				StoreEntry = GlobalPackageStore.FindStoreEntry(DiskPackageId);
-				bHasCustomPackageName &= Name != DiskPackageName;
-			}
+			// Queue completion callback and execute at next process loaded packages call to maintain behavior compatibility with old loader
+			FQueuedFailedPackageCallback& QueuedFailedPackageCallback = QueuedFailedPackageCallbacks.AddDefaulted_GetRef();
+			QueuedFailedPackageCallback.PackageName = PackageName;
+			QueuedFailedPackageCallback.Callback.Reset(new FLoadPackageAsyncDelegate(InCompletionDelegate));
 		}
+	};
+
+	FName DiskPackageName = InPackagePath.GetPackageFName();
+	FName InPackageName = InCustomName.IsNone() ? DiskPackageName : InCustomName;
+	if (DiskPackageName.IsNone())
+	{
+		FAsyncPackageDesc2 PackageDesc(RequestID, InPackagePriority, FPackageId(), nullptr /* StoreEntry */, InPackagePath, FPackageId(), NAME_None);
+		UE_ASYNC_PACKAGE_LOG(Warning, PackageDesc, TEXT("LoadPackage: SkipPackage"), TEXT("InPackagePath is not a mounted PackagePath."));
+		CallCompletionWithError(InPackageName);
+		return RequestID;
 	}
 
-	// Verify CustomPackageName, or fixup to handle any input string that can be converted to a long package name.
-	// CustomPackageName must not be an existing disk package name,
-	// that could cause missing or incorrect import objects for other packages.
-	FName CustomPackageName;
+	FPackagePath DiskPackagePath(InPackagePath);
+
+	bool bCustomNameIsValid = true;
+	FName CustomPackageName(InCustomName != DiskPackageName ? InCustomName : NAME_None);
 	FPackageId CustomPackageId;
-	if (bHasCustomPackageName)
+	if (!CustomPackageName.IsNone())
 	{
-		FPackageId PackageId = FPackageId::FromName(Name);
-		if (!GlobalPackageStore.FindStoreEntry(PackageId))
+		FPackageId PackageId = FPackageId::FromName(CustomPackageName);
+		// CustomName must not be an existing disk package name because that could cause missing or incorrect import objects for other packages.
+		bCustomNameIsValid = !GlobalPackageStore.FindStoreEntry(PackageId);
+		if (bCustomNameIsValid)
 		{
-			FString PackageNameStr = Name.ToString();
-			if (FPackageName::IsValidLongPackageName(PackageNameStr))
-			{
-				CustomPackageName = Name;
-				CustomPackageId = PackageId;
-			}
-			else
-			{
-				FString NewPackageNameStr;
-				if (FPackageName::TryConvertFilenameToLongPackageName(PackageNameStr, NewPackageNameStr))
-				{
-					PackageId = FPackageId::FromName(FName(*NewPackageNameStr));
-					if (!GlobalPackageStore.FindStoreEntry(PackageId))
-					{
-						CustomPackageName = *NewPackageNameStr;
-						CustomPackageId = PackageId;
-					}
-				}
-			}
+			CustomPackageId = PackageId;
 		}
 	}
-	check(CustomPackageId.IsValid() == !CustomPackageName.IsNone());
 
-	bool bCustomNameIsValid = (!bHasCustomPackageName && CustomPackageName.IsNone()) || (bHasCustomPackageName && !CustomPackageName.IsNone());
+	FPackageId DiskPackageId = FPackageId::FromName(DiskPackageName);
+	const FPackageStoreEntry* StoreEntry = GlobalPackageStore.FindStoreEntry(DiskPackageId);
 	bool bDiskPackageIdIsValid = !!StoreEntry;
 	if (!bDiskPackageIdIsValid)
 	{
@@ -5641,10 +5612,10 @@ int32 FAsyncLoadingThread2::LoadPackage(const FString& InName, const FGuid* InGu
 			{
 				CustomPackageName = Package->Desc.CustomPackageName;
 				CustomPackageId = Package->Desc.CustomPackageId;
-				bHasCustomPackageName = bCustomNameIsValid = true;
 			}
-			DiskPackageName = Package->Desc.DiskPackageName;
 			DiskPackageId = Package->Desc.DiskPackageId;
+			DiskPackagePath = Package->Desc.DiskPackagePath;
+			DiskPackageName = DiskPackagePath.GetPackageFName();
 			StoreEntry = Package->Desc.StoreEntry;
 			bDiskPackageIdIsValid = true;
 		}
@@ -5654,7 +5625,7 @@ int32 FAsyncLoadingThread2::LoadPackage(const FString& InName, const FGuid* InGu
 	{
 		if (FCoreDelegates::OnAsyncLoadPackage.IsBound())
 		{
-			FCoreDelegates::OnAsyncLoadPackage.Broadcast(InName);
+			FCoreDelegates::OnAsyncLoadPackage.Broadcast(InPackageName.ToString());
 		}
 
 		// Generate new request ID and add it immediately to the global request list (it needs to be there before we exit
@@ -5678,15 +5649,15 @@ int32 FAsyncLoadingThread2::LoadPackage(const FString& InName, const FGuid* InGu
 #endif
 
 		// Add new package request
-		FAsyncPackageDesc2 PackageDesc(RequestID, InPackagePriority, DiskPackageId, StoreEntry, DiskPackageName, CustomPackageId, CustomPackageName, MoveTemp(CompletionDelegatePtr));
+		FAsyncPackageDesc2 PackageDesc(RequestID, InPackagePriority, DiskPackageId, StoreEntry, DiskPackagePath, CustomPackageId, CustomPackageName, MoveTemp(CompletionDelegatePtr));
 
 		// Fixup for redirected packages since the slim StoreEntry itself has been stripped from both package names and package ids
 		FPackageId RedirectedDiskPackageId = GlobalPackageStore.GetRedirectedPackageId(DiskPackageId);
 		if (RedirectedDiskPackageId.IsValid())
 		{
 			PackageDesc.DiskPackageId = RedirectedDiskPackageId;
-			PackageDesc.SourcePackageName = PackageDesc.DiskPackageName;
-			PackageDesc.DiskPackageName = FName();
+			PackageDesc.SourcePackageName = PackageDesc.DiskPackagePath.GetPackageFName();
+			PackageDesc.DiskPackagePath = FPackagePath();
 		}
 
 		QueuePackage(PackageDesc);
@@ -5695,7 +5666,7 @@ int32 FAsyncLoadingThread2::LoadPackage(const FString& InName, const FGuid* InGu
 	}
 	else
 	{
-		FAsyncPackageDesc2 PackageDesc(RequestID, InPackagePriority, DiskPackageId, StoreEntry, DiskPackageName, CustomPackageId, CustomPackageName);
+		FAsyncPackageDesc2 PackageDesc(RequestID, InPackagePriority, DiskPackageId, StoreEntry, DiskPackagePath, CustomPackageId, CustomPackageName);
 		if (!bDiskPackageIdIsValid)
 		{
 			UE_ASYNC_PACKAGE_LOG(Warning, PackageDesc, TEXT("LoadPackage: SkipPackage"),
@@ -5706,13 +5677,7 @@ int32 FAsyncLoadingThread2::LoadPackage(const FString& InName, const FGuid* InGu
 			UE_ASYNC_PACKAGE_LOG(Warning, PackageDesc, TEXT("LoadPackage: SkipPackage"), TEXT("The custom package name is invalid"));
 		}
 
-		if (InCompletionDelegate.IsBound())
-		{
-			// Queue completion callback and execute at next process loaded packages call to maintain behavior compatibility with old loader
-			FQueuedFailedPackageCallback& QueuedFailedPackageCallback = QueuedFailedPackageCallbacks.AddDefaulted_GetRef();
-			QueuedFailedPackageCallback.PackageName = Name;
-			QueuedFailedPackageCallback.Callback.Reset(new FLoadPackageAsyncDelegate(InCompletionDelegate));
-		}
+		CallCompletionWithError(InPackageName);
 	}
 
 	return RequestID;
diff --git a/Engine/Source/Runtime/CoreUObject/Private/Serialization/AsyncLoadingPrivate.h b/Engine/Source/Runtime/CoreUObject/Private/Serialization/AsyncLoadingPrivate.h
--- a/Engine/Source/Runtime/CoreUObject/Private/Serialization/AsyncLoadingPrivate.h
+++ b/Engine/Source/Runtime/CoreUObject/Private/Serialization/AsyncLoadingPrivate.h
@@ -20,7 +20,7 @@ public:
 		ProcessingExports,
 	};
 
-	FAsyncArchive(const TCHAR* InFileName, FLinkerLoad* InOwner, TFunction<void()>&& InSummaryReadyCallback);
+	FAsyncArchive(const FPackagePath& InPackagePath, FLinkerLoad* InOwner, TFunction<void()>&& InSummaryReadyCallback);
 	virtual ~FAsyncArchive ();
 
 	/** Archive overrides */
@@ -41,7 +41,7 @@ public:
 	virtual void FlushCache() override;
 	virtual FString GetArchiveName() const override 
 	{
-		return FileName;
+		return PackagePath.GetDebugName();
 	}
 
 	/** AsyncArchive interface */
@@ -117,8 +117,8 @@ private:
 	bool bCookedForEDLInEditor;
 
 	FAsyncFileCallBack ReadCallbackFunction;
-	/** Cached filename for debugging.												*/
-	FString	FileName;
+	/** Cached PackagePath for debugging.												*/
+	FPackagePath PackagePath;
 	double OpenTime;
 	double SummaryReadTime;
 	double ExportReadTime;
diff --git a/Engine/Source/Runtime/CoreUObject/Private/Serialization/AsyncLoadingThread.h b/Engine/Source/Runtime/CoreUObject/Private/Serialization/AsyncLoadingThread.h
--- a/Engine/Source/Runtime/CoreUObject/Private/Serialization/AsyncLoadingThread.h
+++ b/Engine/Source/Runtime/CoreUObject/Private/Serialization/AsyncLoadingThread.h
@@ -266,10 +266,10 @@ public:
 	void StartThread() override;
 
 	int32 LoadPackage(
-			const FString& InPackageName,
-			const FGuid* InGuid,
-			const TCHAR* InPackageToLoadFrom,
+			const FPackagePath& InPackagePath,
+			FName InCustomName,
 			FLoadPackageAsyncDelegate InCompletionDelegate,
+			const FGuid* InGuid,
 			EPackageFlags InPackageFlags,
 			int32 InPIEInstanceID,
 			int32 InPackagePriority,
diff --git a/Engine/Source/Runtime/CoreUObject/Private/Serialization/AsyncPackage.h b/Engine/Source/Runtime/CoreUObject/Private/Serialization/AsyncPackage.h
--- a/Engine/Source/Runtime/CoreUObject/Private/Serialization/AsyncPackage.h
+++ b/Engine/Source/Runtime/CoreUObject/Private/Serialization/AsyncPackage.h
@@ -9,6 +9,7 @@
 #include "CoreMinimal.h"
 #include "UObject/UObjectGlobals.h"
 #include "Misc/Guid.h"
+#include "Misc/PackagePath.h"
 #include "Templates/UniquePtr.h"
 #include "UObject/LinkerInstancingContext.h"
 
@@ -18,8 +19,8 @@ struct FAsyncPackageDesc
 	int32 RequestID;
 	/** Name of the UPackage to create. */
 	FName Name;
-	/** Name of the package to load. */
-	FName NameToLoad;
+	/** PackagePath of the package to load. */
+	FPackagePath PackagePath;
 	/** GUID of the package to load, or the zeroed invalid GUID for "don't care" */
 	FGuid Guid;
 	/** Delegate called on completion of loading. This delegate can only be created and consumed on the game thread */
@@ -42,27 +43,24 @@ struct FAsyncPackageDesc
 	void SetInstancingContext(FLinkerInstancingContext) {}
 #endif 
 
-	FAsyncPackageDesc(int32 InRequestID, const FName& InName, FName InPackageToLoadFrom = NAME_None, const FGuid& InGuid = FGuid(), TUniquePtr<FLoadPackageAsyncDelegate>&& InCompletionDelegate = TUniquePtr<FLoadPackageAsyncDelegate>(), EPackageFlags InPackageFlags = PKG_None, int32 InPIEInstanceID = INDEX_NONE, TAsyncLoadPriority InPriority = 0)
+	FAsyncPackageDesc(int32 InRequestID, const FName& InName, const FPackagePath& InPackagePath, const FGuid& InGuid = FGuid(), TUniquePtr<FLoadPackageAsyncDelegate>&& InCompletionDelegate = TUniquePtr<FLoadPackageAsyncDelegate>(), EPackageFlags InPackageFlags = PKG_None, int32 InPIEInstanceID = INDEX_NONE, TAsyncLoadPriority InPriority = 0)
 		: RequestID(InRequestID)
 		, Name(InName)
-		, NameToLoad(InPackageToLoadFrom)
+		, PackagePath(InPackagePath)
 		, Guid(InGuid)
 		, PackageLoadedDelegate(MoveTemp(InCompletionDelegate))
 		, PackageFlags(InPackageFlags)
 		, Priority(InPriority)
 		, PIEInstanceID(InPIEInstanceID)
 	{
-		if (NameToLoad == NAME_None)
-		{
-			NameToLoad = Name;
-		}
+		check(!PackagePath.IsEmpty());
 	}
 
 	/** This constructor does not modify the package loaded delegate as this is not safe outside the game thread */
 	FAsyncPackageDesc(const FAsyncPackageDesc& OldPackage)
 		: RequestID(OldPackage.RequestID)
 		, Name(OldPackage.Name)
-		, NameToLoad(OldPackage.NameToLoad)
+		, PackagePath(OldPackage.PackagePath)
 		, Guid(OldPackage.Guid)
 		, PackageFlags(OldPackage.PackageFlags)
 		, Priority(OldPackage.Priority)
diff --git a/Engine/Source/Runtime/CoreUObject/Private/Serialization/AsyncPackageLoader.cpp b/Engine/Source/Runtime/CoreUObject/Private/Serialization/AsyncPackageLoader.cpp
--- a/Engine/Source/Runtime/CoreUObject/Private/Serialization/AsyncPackageLoader.cpp
+++ b/Engine/Source/Runtime/CoreUObject/Private/Serialization/AsyncPackageLoader.cpp
@@ -7,6 +7,7 @@
 #include "UObject/GCObject.h"
 #include "UObject/LinkerLoad.h"
 #include "Misc/CoreDelegates.h"
+#include "Misc/PackageName.h"
 #include "IO/IoDispatcher.h"
 #include "HAL/IConsoleManager.h"
 
@@ -656,18 +657,74 @@ bool IsAsyncLoadingSuspendedInternal()
 	return GetAsyncPackageLoader().IsAsyncLoadingSuspended();
 }
 
-int32 LoadPackageAsync(const FString& InName, const FGuid* InGuid /*= nullptr*/, const TCHAR* InPackageToLoadFrom /*= nullptr*/, FLoadPackageAsyncDelegate InCompletionDelegate /*= FLoadPackageAsyncDelegate()*/, EPackageFlags InPackageFlags /*= PKG_None*/, int32 InPIEInstanceID /*= INDEX_NONE*/, int32 InPackagePriority /*= 0*/, const FLinkerInstancingContext* InstancingContext /*=nullptr*/)
+static FPackagePath GetLoadPackageAsyncPackagePath(FStringView InPackageNameOrFilePath)
+{
+	FPackagePath PackagePath;
+	if (!FPackagePath::TryFromMountedName(InPackageNameOrFilePath, PackagePath))
+	{
+		// Legacy behavior from FAsyncLoadingThread::LoadPackage: handle asset strings with class references: ClassName'PackageName'. 
+		FStringView ExportTextPackagePath;
+		if (FPackageName::ParseExportTextPath(InPackageNameOrFilePath, nullptr /* OutClassName */, &ExportTextPackagePath))
+		{
+			if (FPackagePath::TryFromMountedName(ExportTextPackagePath, PackagePath))
+			{
+				UE_LOG(LogStreaming, Warning, TEXT("Deprecation warning: calling LoadPackage with the export text format of a package name (ClassName'PackageName') is deprecated and will be removed in a future release."));
+			}
+		}
+	}
+	return PackagePath;
+}
+
+int32 IAsyncPackageLoader::LoadPackage(
+	const FString& InPackageName,
+	const FGuid* InGuid,
+	const TCHAR* InPackageToLoadFrom,
+	FLoadPackageAsyncDelegate InCompletionDelegate,
+	EPackageFlags InPackageFlags,
+	int32 InPIEInstanceID,
+	int32 InPackagePriority,
+	const FLinkerInstancingContext* InstancingContext)
+{
+	FPackagePath PackagePath = GetLoadPackageAsyncPackagePath(InPackageToLoadFrom ? FStringView(InPackageToLoadFrom) : FStringView(InPackageName));
+	return LoadPackage(PackagePath, FName(InPackageName), InCompletionDelegate, InGuid, InPackageFlags, InPIEInstanceID, InPackagePriority, InstancingContext);
+}
+
+int32 LoadPackageAsync(const FPackagePath& InPackagePath, FName InPackageNameToCreate /* = NAME_None*/, FLoadPackageAsyncDelegate InCompletionDelegate /*= FLoadPackageAsyncDelegate()*/, const FGuid* InGuid /*= nullptr*/, EPackageFlags InPackageFlags /*= PKG_None*/, int32 InPIEInstanceID /*= INDEX_NONE*/, int32 InPackagePriority /*= 0*/, const FLinkerInstancingContext* InstancingContext /*=nullptr*/)
+{
+	LLM_SCOPE(ELLMTag::AsyncLoading);
+	UE_CLOG(!GAsyncLoadingAllowed && !IsInAsyncLoadingThread(), LogStreaming, Fatal, TEXT("Requesting async load of \"%s\" when async loading is not allowed (after shutdown). Please fix higher level code."), *InPackagePath.GetDebugName());
+	return GetAsyncPackageLoader().LoadPackage(InPackagePath, InPackageNameToCreate, InCompletionDelegate, InGuid, InPackageFlags, InPIEInstanceID, InPackagePriority, InstancingContext);
+}
+
+int32 LoadPackageAsync(const FString& InName, const FGuid* InGuid /* nullptr*/)
+{
+	LLM_SCOPE(ELLMTag::AsyncLoading);
+	FPackagePath PackagePath = GetLoadPackageAsyncPackagePath(InName);
+	return LoadPackageAsync(PackagePath, NAME_None /* InPackageNameToCreate */, FLoadPackageAsyncDelegate(), InGuid, PKG_None, INDEX_NONE, 0, nullptr);
+}
+
+int32 LoadPackageAsync(const FString& InName, FLoadPackageAsyncDelegate CompletionDelegate, int32 InPackagePriority /*= 0*/, EPackageFlags InPackageFlags /*= PKG_None*/, int32 InPIEInstanceID /*= INDEX_NONE*/)
 {
 	LLM_SCOPE(ELLMTag::AsyncLoading);
-	UE_CLOG(!GAsyncLoadingAllowed && !IsInAsyncLoadingThread(), LogStreaming, Fatal, TEXT("Requesting async load of \"%s\" when async loading is not allowed (after shutdown). Please fix higher level code."), *InName);
-	return GetAsyncPackageLoader().LoadPackage(InName, InGuid, InPackageToLoadFrom, InCompletionDelegate, InPackageFlags, InPIEInstanceID, InPackagePriority, InstancingContext);
+	FPackagePath PackagePath = GetLoadPackageAsyncPackagePath(InName);
+	return LoadPackageAsync(PackagePath, NAME_None /* InPackageNameToCreate */, CompletionDelegate, nullptr /* Guid */, InPackageFlags, InPIEInstanceID, InPackagePriority, nullptr);
 }
 
-int32 LoadPackageAsync(const FString& PackageName, FLoadPackageAsyncDelegate CompletionDelegate, int32 InPackagePriority /*= 0*/, EPackageFlags InPackageFlags /*= PKG_None*/, int32 InPIEInstanceID /*= INDEX_NONE*/)
+int32 LoadPackageAsync(const FString& InName, const FGuid* InGuid /*= nullptr*/, const TCHAR* InPackageToLoadFrom, FLoadPackageAsyncDelegate InCompletionDelegate /*= FLoadPackageAsyncDelegate()*/, EPackageFlags InPackageFlags /*= PKG_None*/, int32 InPIEInstanceID /*= INDEX_NONE*/, int32 InPackagePriority /*= 0*/, const FLinkerInstancingContext* InstancingContext /*=nullptr*/)
 {
-	const FGuid* Guid = nullptr;
-	const TCHAR* PackageToLoadFrom = nullptr;
-	return LoadPackageAsync(PackageName, Guid, PackageToLoadFrom, CompletionDelegate, InPackageFlags, InPIEInstanceID, InPackagePriority );
+	LLM_SCOPE(ELLMTag::AsyncLoading);
+	FPackagePath PackagePath = GetLoadPackageAsyncPackagePath(InPackageToLoadFrom ? FStringView(InPackageToLoadFrom) : FStringView(InName));
+	FName InPackageNameToCreate;
+	if (InPackageToLoadFrom)
+	{
+		// We're loading from PackageToLoadFrom. If InName is different, it is the PackageNameToCreate. InName might be a filepath, so use FPackagePath to convert it to a PackagePath
+		FPackagePath PackagePathToCreate;
+		if (FPackagePath::TryFromMountedName(InName, PackagePathToCreate))
+		{
+			InPackageNameToCreate = PackagePathToCreate.GetPackageFName();
+		}
+	}
+	return LoadPackageAsync(PackagePath, InPackageNameToCreate, InCompletionDelegate, InGuid, InPackageFlags, InPIEInstanceID, InPackagePriority, InstancingContext);
 }
 
 void CancelAsyncLoading()
diff --git a/Engine/Source/Runtime/CoreUObject/Private/Serialization/BulkData.cpp b/Engine/Source/Runtime/CoreUObject/Private/Serialization/BulkData.cpp
--- a/Engine/Source/Runtime/CoreUObject/Private/Serialization/BulkData.cpp
+++ b/Engine/Source/Runtime/CoreUObject/Private/Serialization/BulkData.cpp
@@ -4,6 +4,7 @@
 #include "Serialization/BulkData.h"
 #include "HAL/FileManager.h"
 #include "Misc/ScopeLock.h"
+#include "Misc/PackageSegment.h"
 #include "Misc/Paths.h"
 #include "Misc/PathViews.h"
 #include "Serialization/LargeMemoryWriter.h"
@@ -16,6 +17,7 @@
 #include "UObject/LinkerSave.h"
 #include "Interfaces/ITargetPlatform.h"
 #include "UObject/DebugSerializationFlags.h"
+#include "UObject/PackageResourceManager.h"
 #include "Serialization/AsyncLoadingPrivate.h"
 #include "Async/MappedFileHandle.h"
 #include "HAL/PlatformFileManager.h"
@@ -119,11 +121,11 @@ const void* FOwnedBulkDataPtr::GetPointer()
 	return AllocatedData ? AllocatedData : (MappedRegion ? MappedRegion->GetMappedPtr() : nullptr);
 }
 
-bool FUntypedBulkData::FAllocatedPtr::MapFile(const TCHAR *InFilename, int64 Offset, int64 Size)
+bool FUntypedBulkData::FAllocatedPtr::MapFile(const FPackagePath& InPackagePath, EPackageSegment InPackageSegment, int64 Offset, int64 Size)
 {
 	check(!MappedHandle && !MappedRegion); // It doesn't make sense to do this twice, but if need be, not hard to do
 
-	MappedHandle = FPlatformFileManager::Get().GetPlatformFile().OpenMapped(InFilename);
+	MappedHandle = IPackageResourceManager::Get().OpenMappedHandleToPackage(InPackagePath, InPackageSegment);
 
 	if (!MappedHandle)
 	{
@@ -246,7 +248,8 @@ FUntypedBulkData& FUntypedBulkData::operator=( const FUntypedBulkData& Other )
 	}
 	else // Otherwise setup the bulk so that the data can be loaded through LoadBulkDataWithFileReader()
 	{
-		Filename = Other.Filename;
+		PackagePath = Other.PackagePath;
+		PackageSegment = Other.PackageSegment;
 		BulkDataFlags = Other.BulkDataFlags;
 		ElementCount = Other.ElementCount;
 		BulkDataOffsetInFile = Other.BulkDataOffsetInFile;
@@ -402,7 +405,7 @@ bool FUntypedBulkData::CanLoadFromDisk() const
 	return AttachedAr != NULL;
 #else
 	bool bCanLoadFromDisk = false;
-	if (!Filename.IsEmpty())
+	if (!PackagePath.IsEmpty())
 	{
 		bCanLoadFromDisk = true;
 	}
@@ -416,7 +419,7 @@ bool FUntypedBulkData::CanLoadFromDisk() const
 
 bool FUntypedBulkData::DoesExist() const
 {
-	return IFileManager::Get().FileExists(*Filename);
+	return IPackageResourceManager::Get().DoesPackageExist(PackagePath, PackageSegment);
 }
 
 /**
@@ -715,7 +718,7 @@ bool FUntypedBulkData::LoadBulkDataWithFileReader()
 bool FUntypedBulkData::CanLoadBulkDataWithFileReader() const
 {
 #if WITH_EDITOR
-	return !Filename.IsEmpty();
+	return !PackagePath.IsEmpty();
 #else
 	return false;
 #endif
@@ -829,13 +832,13 @@ void FUntypedBulkData::AsyncLoadBulkData()
 {
 	BulkDataAsync.Reallocate(GetBulkDataSize(), BulkDataAlignment);
 
-	FArchive* FileReaderAr = IFileManager::Get().CreateFileReader(*Filename, FILEREAD_Silent);
-	checkf(FileReaderAr != NULL, TEXT("Attempted to load bulk data from an invalid filename '%s'."), *Filename);
+	FOpenPackageResult Result = IPackageResourceManager::Get().OpenReadPackage(PackagePath, PackageSegment);
+	checkf(Result.Archive.IsValid() && Result.Format == EPackageFormat::Binary, TEXT("Attempted to load bulk data from an invalid package '%s'%s."),
+		*PackagePath.GetDebugName(PackageSegment), (Result.Archive == nullptr ? TEXT("") : TEXT(": Package Format is Text which is not supported")));
 
 	// Seek to the beginning of the bulk data in the file.
-	FileReaderAr->Seek(BulkDataOffsetInFile);
-	SerializeBulkData(*FileReaderAr, BulkDataAsync.Get());
-	delete FileReaderAr;
+	Result.Archive->Seek(BulkDataOffsetInFile);
+	SerializeBulkData(*Result.Archive, BulkDataAsync.Get());
 }
 
 /*-----------------------------------------------------------------------------
@@ -849,7 +852,7 @@ void FUntypedBulkData::StartSerializingBulkData(FArchive& Ar, UObject* Owner, in
 
 	SerializeFuture = Async(EAsyncExecution::ThreadPool, [=]() 
 	{ 
-		UE_CLOG(GEventDrivenLoaderEnabled, LogSerialization, Error, TEXT("Attempt to stream bulk data with EDL enabled. This is not desireable. File %s"), *Filename);
+		UE_CLOG(GEventDrivenLoaderEnabled, LogSerialization, Error, TEXT("Attempt to stream bulk data with EDL enabled. This is not desireable. Package %s"), *PackagePath.GetDebugName(PackageSegment));
 		AsyncLoadBulkData();
 
 		return true;
@@ -890,7 +893,7 @@ bool FUntypedBulkData::ShouldStreamBulkData()
 
 	const bool bForceStream = !!(BulkDataFlags & BULKDATA_ForceStreamPayload);
 
-	return (FPlatformProperties::RequiresCookedData() && !Filename.IsEmpty() &&
+	return (FPlatformProperties::RequiresCookedData() && !PackagePath.IsEmpty() &&
 		FPlatformProcess::SupportsMultithreading() && IsInGameThread() &&
 		(bForceStream || GetBulkDataSize() > GMinimumBulkDataSizeForAsyncLoading) &&
 		GMinimumBulkDataSizeForAsyncLoading >= 0);
@@ -994,7 +997,8 @@ void FUntypedBulkData::Serialize( FArchive& Ar, UObject* Owner, int32 Idx, bool
 		// We're loading from the persistent archive.
 		if (Ar.IsLoading())
 		{
-			Filename = TEXT("");
+			PackagePath = FPackagePath();
+			PackageSegment = EPackageSegment::Header;
 
 			// @todo when Landscape (and others?) only Lock/Unlock once, we can enable this
 			if (false) // FPlatformProperties::RequiresCookedData())
@@ -1043,7 +1047,7 @@ void FUntypedBulkData::Serialize( FArchive& Ar, UObject* Owner, int32 Idx, bool
 				Linker = Owner->GetLinker();
 			}
 #else
-			FLinker* Linker = nullptr;
+			FLinkerLoad* Linker = nullptr;
 			if (Owner != nullptr)
 			{
 				Package = Owner->GetOutermost();
@@ -1060,7 +1064,7 @@ void FUntypedBulkData::Serialize( FArchive& Ar, UObject* Owner, int32 Idx, bool
 			}
 			BulkDataOffsetInFile += OffsetInFileFixup;
 
-			// We're allowing defered serialization.
+			// We're allowing deferred serialization.
 			FArchive* CacheableArchive = Ar.GetCacheableArchive();
 			if (Ar.IsAllowingLazyLoading() && Owner != NULL && CacheableArchive)
 			{				
@@ -1069,10 +1073,12 @@ void FUntypedBulkData::Serialize( FArchive& Ar, UObject* Owner, int32 Idx, bool
 				CacheableArchive->AttachBulkData(Owner, this);
 				check(!CacheableArchive->IsTextFormat());
 				AttachedAr = CacheableArchive;
-				Filename = Linker->Filename;
+				PackagePath = Linker->GetPackagePath();
+				PackageSegment = EPackageSegment::Header;
 #else
 				check(Linker);
-				Filename = Linker->Filename;
+				PackagePath = Linker->GetPackagePath();
+				PackageSegment = EPackageSegment::Header;
 
 				check(!Owner->GetLinker() || Owner->GetLinker() == Linker);
 #endif // WITH_EDITOR
@@ -1089,15 +1095,15 @@ void FUntypedBulkData::Serialize( FArchive& Ar, UObject* Owner, int32 Idx, bool
 						bool bWasMapped = false;
 						if (bAttemptFileMapping)
 						{
-							UE_LOG(LogSerialization, Error, TEXT("Attempt to file map inline bulk data. This is not desireable. File %s"), *Filename);
+							UE_LOG(LogSerialization, Error, TEXT("Attempt to file map inline bulk data. This is not desireable. File %s"), *PackagePath.GetDebugName(PackageSegment));
 
-							if (GEventDrivenLoaderEnabled && (Filename.EndsWith(TEXT(".uasset")) || Filename.EndsWith(TEXT(".umap"))))
+							if (GEventDrivenLoaderEnabled && PackageSegment == EPackageSegment::Header)
 							{
-								BulkDataOffsetInFile -= IFileManager::Get().FileSize(*Filename);
+								BulkDataOffsetInFile -= IPackageResourceManager::Get().FileSize(PackagePath, PackageSegment);
 								check(BulkDataOffsetInFile >= 0);
-								Filename = FPaths::GetBaseFilename(Filename, false) + BulkDataExt::Export;
+								PackageSegment = EPackageSegment::Exports;
 							}
-							bWasMapped = BulkData.MapFile(*Filename, BulkDataOffsetInFile, GetBulkDataSize());
+							bWasMapped = BulkData.MapFile(PackagePath, PackageSegment, BulkDataOffsetInFile, GetBulkDataSize());
 						}
 						if (bWasMapped)
 						{
@@ -1117,10 +1123,9 @@ void FUntypedBulkData::Serialize( FArchive& Ar, UObject* Owner, int32 Idx, bool
 				else if (BulkDataFlags & BULKDATA_DuplicateNonOptionalPayload)
 				{
 					// Load from optional payload instead if optional payload is available
-					FString OptionalFilename = FPathViews::ChangeExtension(Filename, BulkDataExt::Optional);
-					if (IFileManager::Get().FileExists(*OptionalFilename))
+					if (IPackageResourceManager::Get().DoesPackageExist(PackagePath, EPackageSegment::BulkDataOptional))
 					{
-						Filename = OptionalFilename;
+						PackageSegment = EPackageSegment::BulkDataOptional;
 						Ar << BulkDataFlags;
 						SetBulkDataFlags(BULKDATA_OptionalPayload);
 
@@ -1140,6 +1145,7 @@ void FUntypedBulkData::Serialize( FArchive& Ar, UObject* Owner, int32 Idx, bool
 					}
 					else
 					{
+						PackageSegment = EPackageSegment::BulkDataDefault;
 						// Skip all bulkdata info of the copy in uptnl.
 						uint32 DummyBulkDataFlags;
 						Ar << DummyBulkDataFlags;
@@ -1157,30 +1163,28 @@ void FUntypedBulkData::Serialize( FArchive& Ar, UObject* Owner, int32 Idx, bool
 
 						int64 DummyBulkDataOffsetInFile;
 						Ar << DummyBulkDataOffsetInFile;
-
-						Filename = FPathViews::ChangeExtension(Filename, BulkDataExt::Default);
 					}
 				}
 				else if (BulkDataFlags & BULKDATA_OptionalPayload)
 				{
-					Filename = FPathViews::ChangeExtension(Filename, BulkDataExt::Optional);
+					PackageSegment = EPackageSegment::BulkDataOptional;
 				}
 				else if (BulkDataFlags & BULKDATA_PayloadInSeperateFile)
 				{
 					if (BulkDataFlags & BULKDATA_MemoryMappedPayload)
 					{
-						Filename = FPathViews::ChangeExtension(Filename, BulkDataExt::MemoryMapped);
+						PackageSegment = EPackageSegment::BulkDataMemoryMapped;
 					}
 					else
 					{
-						Filename = FPathViews::ChangeExtension(Filename, BulkDataExt::Default);
+						PackageSegment = EPackageSegment::BulkDataDefault;
 					}
 
 
 					if (bAttemptFileMapping)
 					{
 						check(!FPlatformProperties::GetMemoryMappingAlignment() || IsAligned(BulkDataOffsetInFile, FPlatformProperties::GetMemoryMappingAlignment()));
-						bool bWasMapped = BulkData.MapFile(*Filename, BulkDataOffsetInFile, GetBulkDataSize());
+						bool bWasMapped = BulkData.MapFile(PackagePath, PackageSegment, BulkDataOffsetInFile, GetBulkDataSize());
 						if (!bWasMapped)
 						{
 							// we failed to map when requested, so they will be looking for the memory, do a sync load now.
@@ -1194,7 +1198,8 @@ void FUntypedBulkData::Serialize( FArchive& Ar, UObject* Owner, int32 Idx, bool
 			{
 				if (Owner && Owner->GetLinker())
 				{
-					Filename = Owner->GetLinker()->Filename;
+					PackagePath = Owner->GetLinker()->GetPackagePath();
+					PackageSegment = EPackageSegment::Header;
 				}
 				if (ShouldStreamBulkData())
 				{
@@ -1215,26 +1220,27 @@ void FUntypedBulkData::Serialize( FArchive& Ar, UObject* Owner, int32 Idx, bool
 						if (BulkDataFlags & BULKDATA_PayloadInSeperateFile)
 						{
 							// open separate bulk data file
-							UE_CLOG(GEventDrivenLoaderEnabled, LogSerialization, Error, TEXT("Attempt to sync load bulk data with EDL enabled (separate file). This is not desireable. File %s"), *Filename);
+							UE_CLOG(GEventDrivenLoaderEnabled, LogSerialization, Error, TEXT("Attempt to sync load bulk data with EDL enabled (separate file). This is not desireable. File %s"), *PackagePath.GetDebugName(PackageSegment));
 
-							if (GEventDrivenLoaderEnabled && (Filename.EndsWith(TEXT(".uasset")) || Filename.EndsWith(TEXT(".umap"))))
+							if (GEventDrivenLoaderEnabled && PackageSegment == EPackageSegment::Header)
 							{
-								BulkDataOffsetInFile -= IFileManager::Get().FileSize(*Filename);
+								BulkDataOffsetInFile -= IPackageResourceManager::Get().FileSize(PackagePath, PackageSegment);
 								check(BulkDataOffsetInFile >= 0);
-								Filename = FPaths::GetBaseFilename(Filename, false) + BulkDataExt::Export;
+								PackageSegment = EPackageSegment::Exports;
 							}
 
-							FArchive* TargetArchive = IFileManager::Get().CreateFileReader(*Filename);
+							FOpenPackageResult TargetArchive = IPackageResourceManager::Get().OpenReadPackage(PackagePath, PackageSegment);
+							checkf(TargetArchive.Archive.IsValid() && TargetArchive.Format == EPackageFormat::Binary, TEXT("Attempted to load bulk data from an invalid PackagePath '%s': %s."),
+								*PackagePath.GetDebugName(PackageSegment), (!TargetArchive.Archive.IsValid() ? TEXT("could not find package") : TEXT("package is a TextAsset which is not supported")));
+
 							// seek to the location in the file where the payload is stored
-							TargetArchive->Seek(BulkDataOffsetInFile);
+							TargetArchive.Archive->Seek(BulkDataOffsetInFile);
 							// serialize the payload
-							SerializeBulkData(*TargetArchive, BulkData.Get());
-							// cleanup file
-							delete TargetArchive;
+							SerializeBulkData(*TargetArchive.Archive, BulkData.Get());
 						}
 						else
 						{
-							UE_CLOG(GEventDrivenLoaderEnabled, LogSerialization, Error, TEXT("Attempt to sync load bulk data with EDL enabled. This is not desireable. File %s"), *Filename);
+							UE_CLOG(GEventDrivenLoaderEnabled, LogSerialization, Error, TEXT("Attempt to sync load bulk data with EDL enabled. This is not desireable. File %s"), *PackagePath.GetDebugName(PackageSegment));
 
 							// store the current file offset
 							int64 CurOffset = Ar.Tell();
@@ -1548,6 +1554,7 @@ void FUntypedBulkData::InitializeMemberVariables()
 	BulkDataSizeOnDisk = INDEX_NONE;
 	BulkDataAlignment = DEFAULT_ALIGNMENT;
 	LockStatus = LOCKSTATUS_Unlocked;
+	PackageSegment = EPackageSegment::Header;
 #if WITH_EDITOR
 	Linker = nullptr;
 	AttachedAr = nullptr;
@@ -1662,7 +1669,7 @@ void FUntypedBulkData::SerializeBulkData( FArchive& Ar, void* Data )
 
 IAsyncReadFileHandle* FUntypedBulkData::OpenAsyncReadHandle() const
 {
-	return FPlatformFileManager::Get().GetPlatformFile().OpenAsyncRead(*GetFilename());
+	return IPackageResourceManager::Get().OpenAsyncReadPackage(GetPackagePath(), GetPackageSegment());
 }
 
 IBulkDataIORequest* FUntypedBulkData::CreateStreamingRequest(EAsyncIOPriorityAndFlags Priority, FBulkDataIORequestCallBack* CompleteCallback, uint8* UserSuppliedMemory) const
@@ -1674,27 +1681,27 @@ IBulkDataIORequest* FUntypedBulkData::CreateStreamingRequest(EAsyncIOPriorityAnd
 
 IBulkDataIORequest* FUntypedBulkData::CreateStreamingRequest(int64 OffsetInBulkData, int64 BytesToRead, EAsyncIOPriorityAndFlags Priority, FBulkDataIORequestCallBack* CompleteCallback, uint8* UserSuppliedMemory) const
 {
-	check(Filename.IsEmpty() == false);
+	check(PackagePath.IsEmpty() == false);
 
-	// If we are loading from a .uexp file then we need to adjust the filename and offset stored by BulkData in order to use them
+	// If we are loading from a .uexp file then we need to adjust the segment and offset stored by BulkData in order to use them
 	// to access the data in the .uexp file. To keep the method const (due to a large number of places calling this, assuming that
 	// it is const) we take a copy of these data values which if needed can be adjusted and use them instead.
-	FString AdjustedFilename = Filename; 
+	EPackageSegment AdjustedPackageSegment = PackageSegment;
 	int64 AdjustedBulkDataOffsetInFile = BulkDataOffsetInFile;
 
-	// Fix up the Filename/Offset to work with streaming if EDL is enabled and the filename is still referencing a uasset or umap
-	if (GEventDrivenLoaderEnabled && (AdjustedFilename.EndsWith(TEXT(".uasset")) || AdjustedFilename.EndsWith(TEXT(".umap"))))
+	// Fix up the Segment/Offset to work with streaming if EDL is enabled and the segment is still referencing the header
+	if (GEventDrivenLoaderEnabled && AdjustedPackageSegment == EPackageSegment::Header)
 	{
-		AdjustedBulkDataOffsetInFile -= IFileManager::Get().FileSize(*AdjustedFilename);
-		AdjustedFilename = FPaths::GetBaseFilename(AdjustedFilename, false) + BulkDataExt::Export;
+		AdjustedBulkDataOffsetInFile -= IPackageResourceManager::Get().FileSize(PackagePath, AdjustedPackageSegment);
+		AdjustedPackageSegment = EPackageSegment::Exports;
 
-		UE_LOG(LogSerialization, Error, TEXT("Streaming from the .uexp file '%s' this MUST be in a ubulk instead for best performance."), *AdjustedFilename);
+		UE_LOG(LogSerialization, Error, TEXT("Streaming from the .uexp file '%s' this MUST be in a ubulk instead for best performance."), *PackagePath.GetDebugName(AdjustedPackageSegment));
 	}
 
-	UE_CLOG(IsStoredCompressedOnDisk(), LogSerialization, Fatal, TEXT("Package level compression is no longer supported (%s)."), *AdjustedFilename);
-	UE_CLOG(GetBulkDataSize() <= 0, LogSerialization, Error, TEXT("(%s) has invalid bulk data size."), *AdjustedFilename);
+	UE_CLOG(IsStoredCompressedOnDisk(), LogSerialization, Fatal, TEXT("Package level compression is no longer supported (%s)."), *PackagePath.GetDebugName(AdjustedPackageSegment));
+	UE_CLOG(GetBulkDataSize() <= 0, LogSerialization, Error, TEXT("(%s) has invalid bulk data size."), *PackagePath.GetDebugName(AdjustedPackageSegment));
 
-	IAsyncReadFileHandle* IORequestHandle = FPlatformFileManager::Get().GetPlatformFile().OpenAsyncRead(*AdjustedFilename);
+	IAsyncReadFileHandle* IORequestHandle = IPackageResourceManager::Get().OpenAsyncReadPackage(PackagePath, AdjustedPackageSegment);
 	check(IORequestHandle); // this generally cannot fail because it is async
 
 	if (IORequestHandle == nullptr)
@@ -1732,7 +1739,15 @@ FBulkDataStreamingToken FUntypedBulkData::CreateStreamingToken() const
 
 IBulkDataIORequest* FUntypedBulkData::CreateStreamingRequestForRange(const FString& Filename, const BulkDataRangeArray& RangeArray, EAsyncIOPriorityAndFlags Priority, FBulkDataIORequestCallBack* CompleteCallback)
 {
-	check(Filename.IsEmpty() == false);
+	FPackagePath PackagePath;
+	EPackageSegment PackageSegment;
+	PackagePath = FPackagePath::FromLocalPath(Filename, PackageSegment);
+	return CreateStreamingRequestForRange(PackagePath, PackageSegment, RangeArray, Priority, CompleteCallback);
+}
+
+IBulkDataIORequest* FUntypedBulkData::CreateStreamingRequestForRange(const FPackagePath& PackagePath, EPackageSegment PackageSegment, const BulkDataRangeArray& RangeArray, EAsyncIOPriorityAndFlags Priority, FBulkDataIORequestCallBack* CompleteCallback)
+{
+	check(PackagePath.IsEmpty() == false);
 	check(RangeArray.Num() > 0);
 
 	const FBulkDataStreamingToken& Start = *(RangeArray[0]);
@@ -1741,7 +1756,7 @@ IBulkDataIORequest* FUntypedBulkData::CreateStreamingRequestForRange(const FStri
 	check(Start.IsValid());
 	check(End.IsValid());
 
-	IAsyncReadFileHandle* IORequestHandle = FPlatformFileManager::Get().GetPlatformFile().OpenAsyncRead(*Filename);
+	IAsyncReadFileHandle* IORequestHandle = IPackageResourceManager::Get().OpenAsyncReadPackage(PackagePath, PackageSegment);
 	check(IORequestHandle); // this generally cannot fail because it is async
 
 	if (IORequestHandle == nullptr)
@@ -1882,7 +1897,7 @@ void FUntypedBulkData::WaitForAsyncLoading()
 	DECLARE_SCOPE_CYCLE_COUNTER(TEXT("FUntypedBulkData::WaitForAsyncLoading"), STAT_UBD_WaitForAsyncLoading, STATGROUP_Memory);
 	while (!SerializeFuture.WaitFor(FTimespan::FromMilliseconds(1000.0)))
 	{
-		UE_LOG(LogSerialization, Warning, TEXT("Waiting for %s bulk data (%lld) to be loaded longer than 1000ms"), *Filename, GetBulkDataSize());
+		UE_LOG(LogSerialization, Warning, TEXT("Waiting for %s bulk data (%lld) to be loaded longer than 1000ms"), *PackagePath.GetDebugName(PackageSegment), GetBulkDataSize());
 	}
 	check(BulkDataAsync);
 }
@@ -1923,20 +1938,22 @@ void FUntypedBulkData::LoadDataIntoMemory( void* Dest )
 		(BulkDataFlags & BULKDATA_PayloadInSeperateFile))
 	{
 		// The attached archive is a package cooked for EDL loaded in the editor so the actual bulk data sits in a separate ubulk file.
-		FString BulkDataFilename;
+		EPackageSegment BulkDataSegment;
+
 		if (BulkDataFlags & BULKDATA_MemoryMappedPayload)
 		{
-			BulkDataFilename = FPathViews::ChangeExtension(Filename, BulkDataExt::MemoryMapped);
+			BulkDataSegment = EPackageSegment::BulkDataMemoryMapped;
 		}
 		else
 		{
-			BulkDataFilename = FPathViews::ChangeExtension(Filename, BulkDataExt::Default);
+			BulkDataSegment = EPackageSegment::BulkDataDefault;
 		}
-		BulkDataArchive = IFileManager::Get().CreateFileReader(*BulkDataFilename, FILEREAD_Silent);
+		FOpenPackageResult OpenResult = IPackageResourceManager::Get().OpenReadPackage(PackagePath, BulkDataSegment);
+		BulkDataArchive = OpenResult.Archive.Release();
+		check(BulkDataArchive == nullptr || OpenResult.Format == EPackageFormat::Binary);
 	}
 
 	if (!BulkDataArchive)
-
 	{
 		BulkDataArchive = AttachedAr;
 	}
@@ -1988,11 +2005,11 @@ void FUntypedBulkData::LoadDataIntoMemory( void* Dest )
 			}
 		}
 	}
-	// if we weren't able to load via linker, load directly by filename
+	// if we weren't able to load via linker, load directly by PackagePath
 	if (!bWasLoadedSuccessfully)
 	{
-		// load from the specied filename when the linker has been cleared
-		checkf( Filename != TEXT(""), TEXT( "Attempted to load bulk data without a proper filename." ) );
+		// load from the specied PackagePath when the linker has been cleared
+		checkf(!PackagePath.IsEmpty(), TEXT( "Attempted to load bulk data without a proper PackagePath." ) );
 
 #if PLATFORM_SUPPORTS_TEXTURE_STREAMING
 		static auto CVarTextureStreamingEnabled = IConsoleManager::Get().FindTConsoleVariableDataInt(TEXT("r.TextureStreaming"));
@@ -2001,25 +2018,26 @@ void FUntypedBulkData::LoadDataIntoMemory( void* Dest )
 		// is it possible for streaming mips to be loaded in non streaming ways.
 		if (CVarTextureStreamingEnabled->GetValueOnAnyThread() != 0)
 		{
-			bool bIsBulkFile = Filename.EndsWith(BulkDataExt::Default);
-			UE_CLOG(GEventDrivenLoaderEnabled && bIsBulkFile && (IsInGameThread() || IsInAsyncLoadingThread()), LogSerialization, Error, TEXT("Attempt to sync load bulk data with EDL enabled (LoadDataIntoMemory). This is not desireable. File %s"), *Filename);
+			bool bIsBulkFile = PackageSegment == EPackageSegment::BulkDataDefault;
+			UE_CLOG(GEventDrivenLoaderEnabled && bIsBulkFile && (IsInGameThread() || IsInAsyncLoadingThread()), LogSerialization, Error,
+				TEXT("Attempt to sync load bulk data with EDL enabled (LoadDataIntoMemory). This is not desireable. File %s"), *PackagePath.GetDebugName(PackageSegment));
 		}
 #endif
 
-		if (GEventDrivenLoaderEnabled && (Filename.EndsWith(TEXT(".uasset")) || Filename.EndsWith(TEXT(".umap"))))
+		if (GEventDrivenLoaderEnabled && PackageSegment == EPackageSegment::Header)
 		{
-			BulkDataOffsetInFile -= IFileManager::Get().FileSize(*Filename);
+			BulkDataOffsetInFile -= IPackageResourceManager::Get().FileSize(PackagePath, PackageSegment);
 			check(BulkDataOffsetInFile >= 0);
-			Filename = FPaths::GetBaseFilename(Filename, false) + BulkDataExt::Export;
+			PackageSegment = EPackageSegment::Exports;
 		}
 
-		FArchive* Ar = IFileManager::Get().CreateFileReader(*Filename, FILEREAD_Silent);
-		checkf( Ar != NULL, TEXT( "Attempted to load bulk data from an invalid filename '%s'." ), *Filename );
+		FOpenPackageResult Result = IPackageResourceManager::Get().OpenReadPackage(PackagePath, PackageSegment);
+		checkf( Result.Archive.IsValid() && Result.Format == EPackageFormat::Binary, TEXT( "Attempted to load bulk data from an invalid PackagePath '%s'%s." ),
+			*PackagePath.GetDebugName(PackageSegment), (!Result.Archive.IsValid() ? TEXT("") : TEXT("; package is in non-binary format and this is not supported.")));
 	
 		// Seek to the beginning of the bulk data in the file.
-		Ar->Seek( BulkDataOffsetInFile );
-		SerializeBulkData( *Ar, Dest );
-		delete Ar;
+		Result.Archive->Seek( BulkDataOffsetInFile );
+		SerializeBulkData( *Result.Archive, Dest );
 	}
 #endif // WITH_EDITOR
 }
