PackageResourceManager and PackagePaths

diff --git a/Engine/Source/Runtime/CoreUObject/Private/Serialization/BulkData2.cpp b/Engine/Source/Runtime/CoreUObject/Private/Serialization/BulkData2.cpp
--- a/Engine/Source/Runtime/CoreUObject/Private/Serialization/BulkData2.cpp
+++ b/Engine/Source/Runtime/CoreUObject/Private/Serialization/BulkData2.cpp
@@ -7,6 +7,7 @@
 #include "Misc/CommandLine.h"
 #include "Misc/ConfigCacheIni.h"
 #include "Misc/PackageName.h"
+#include "Misc/PackageSegment.h"
 #include "Misc/Paths.h"
 #include "Misc/PathViews.h"
 #include "ProfilingDebugging/LoadTimeTracker.h"
@@ -14,6 +15,7 @@
 #include "UObject/LinkerLoad.h"
 #include "UObject/Object.h"
 #include "UObject/Package.h"
+#include "UObject/PackageResourceManager.h"
 #include "IO/IoDispatcher.h"
 #include "Async/Async.h"
 #include "Async/MappedFileHandle.h"
@@ -28,15 +30,6 @@ IMPLEMENT_TYPE_LAYOUT(FBulkDataBase);
 // Handy macro to validate FIoStatus return values
 #define CHECK_IOSTATUS( InIoStatus, InMethodName ) checkf(InIoStatus.IsOk(), TEXT("%s failed: %s"), InMethodName, *InIoStatus.ToString());
 
-namespace BulkDataExt
-{
-	// TODO: Maybe expose this and start using everywhere?
-	const FString Export		= TEXT(".uexp");	// Stored in the export data
-	const FString Default		= TEXT(".ubulk");	// Stored in a separate file
-	const FString MemoryMapped	= TEXT(".m.ubulk");	// Stored in a separate file aligned for memory mapping
-	const FString Optional		= TEXT(".uptnl");	// Stored in a separate file that is optional
-}
-
 namespace
 {
 	const uint16 InvalidBulkDataIndex = ~uint16(0);
@@ -76,6 +69,20 @@ namespace
 	const FIoFilenameHash FALLBACK_IO_FILENAME_HASH = INVALID_IO_FILENAME_HASH - 1;
 }
 
+FIoFilenameHash MakeIoFilenameHash(const FPackagePath& PackagePath)
+{
+	if (!PackagePath.IsEmpty())
+	{
+		FString BaseFileName = PackagePath.GetLocalBaseFilenameWithPath().ToLower();
+		const FIoFilenameHash Hash = FCrc::StrCrc32<TCHAR>(*BaseFileName);
+		return Hash != INVALID_IO_FILENAME_HASH ? Hash : FALLBACK_IO_FILENAME_HASH;
+	}
+	else
+	{
+		return INVALID_IO_FILENAME_HASH;
+	}
+}
+
 FIoFilenameHash MakeIoFilenameHash(const FString& Filename)
 {
 	if (!Filename.IsEmpty())
@@ -103,51 +110,46 @@ FIoFilenameHash MakeIoFilenameHash(const FIoChunkId& ChunkID)
 	}
 }
 
-// TODO: The code in the FileTokenSystem namespace is a temporary system so that FBulkDataBase can hold
-// all of it's info about where the data is on disk in a single 8byte value. This can all be removed when
+// TODO: The code in the PackageTokenSystem namespace is a temporary system so that FBulkDataBase can hold
+// all of it's info about where the data is on disk in a single integer value. This can all be removed when
 // we switch this over to the new packing system.
-namespace FileTokenSystem
+namespace PackageTokenSystem
 {
-	struct Data
-	{
-		FString PackageHeaderFilename;
-	};
-
-	// Internal to the FileTokenSystem namespace
+	// Internal to the PackageTokenSystem namespace
 	namespace
 	{
-		struct StringData
+		struct FPayloadData
 		{
-			FString Filename;
+			FPackagePath PackagePath;
 			uint16 RefCount;
 		};
 
 		/**
-		* Provides a ref counted PackageName->Filename look up table.
+		* Provides a ref counted PackageName->FPackagePath look up table.
 		*/
-		class FStringTable
+		class FPackageDataTable
 		{
 		public:
 			using KeyType = uint64;
 
-			void Add(const KeyType& Key, const FString& Filename)
+			void Add(const KeyType& Key, const FPackagePath& PackagePath)
 			{
-				if (StringData* ExistingEntry = Table.Find(Key))
+				if (FPayloadData* ExistingEntry = Table.Find(Key))
 				{
 					ExistingEntry->RefCount++;
-					checkf(ExistingEntry->Filename == Filename, TEXT("Filename mismatch!"));
+					checkf(ExistingEntry->PackagePath == PackagePath, TEXT("PackagePath mismatch!"));
 				}
 				else
 				{
-					StringData& NewEntry = Table.Emplace(Key);
-					NewEntry.Filename = Filename;
+					FPayloadData& NewEntry = Table.Emplace(Key);
+					NewEntry.PackagePath = PackagePath;
 					NewEntry.RefCount = 1;
 				}
 			}
 
 			bool Remove(const KeyType& Key)
 			{
-				if (StringData* ExistingEntry = Table.Find(Key))
+				if (FPayloadData* ExistingEntry = Table.Find(Key))
 				{
 					if (--ExistingEntry->RefCount == 0)
 					{
@@ -161,15 +163,15 @@ namespace FileTokenSystem
 
 			void IncRef(const KeyType& Key)
 			{
-				if (StringData* ExistingEntry = Table.Find(Key))
+				if (FPayloadData* ExistingEntry = Table.Find(Key))
 				{
 					ExistingEntry->RefCount++;
 				}
 			}
 
-			const FString& Resolve(const KeyType& Key)
+			const FPackagePath& Resolve(const KeyType& Key)
 			{
-				return Table.Find(Key)->Filename;
+				return Table.Find(Key)->PackagePath;
 			}
 
 			int32 Num() const
@@ -178,41 +180,41 @@ namespace FileTokenSystem
 			}
 
 		private:
-			TMap<KeyType, StringData> Table;
+			TMap<KeyType, FPayloadData> Table;
 		};
 	}
 
-	FStringTable StringTable;
+	FPackageDataTable PackageDataTable;
 
 	FRWLock TokenLock;
 
-	FBulkDataOrId::FileToken RegisterFileToken( const FName& PackageName, const FString& Filename)
+	FBulkDataOrId::FileToken RegisterToken( const FName& PackageName, const FPackagePath& PackagePath)
 	{
 		const uint64 Token = (uint64(PackageName.GetComparisonIndex().ToUnstableInt()) << 32) | uint64(PackageName.GetNumber());
 
 		FWriteScopeLock LockForScope(TokenLock);
-		StringTable.Add(Token, Filename);
+		PackageDataTable.Add(Token, PackagePath);
 
 		return Token;
 	}
 
-	void UnregisterFileToken(FBulkDataOrId::FileToken ID)
+	void UnRegisterToken(FBulkDataOrId::FileToken ID)
 	{
 		if (ID != FBulkDataBase::InvalidToken)
 		{
 			FWriteScopeLock LockForScope(TokenLock);
 
-			StringTable.Remove(ID);
+			PackageDataTable.Remove(ID);
 		}
 	}
 
-	FBulkDataOrId::FileToken CopyFileToken(FBulkDataOrId::FileToken ID)
+	FBulkDataOrId::FileToken CopyToken(FBulkDataOrId::FileToken ID)
 	{
 		if (ID != FBulkDataBase::InvalidToken)
 		{
 			FWriteScopeLock LockForScope(TokenLock);
 
-			StringTable.IncRef(ID);
+			PackageDataTable.IncRef(ID);
 
 			return ID;
 		}
@@ -222,15 +224,15 @@ namespace FileTokenSystem
 		}
 	}
 
-	FString GetFilename(FBulkDataOrId::FileToken ID)
+	FPackagePath GetPackagePath(FBulkDataOrId::FileToken ID)
 	{
 		if (ID == FBulkDataBase::InvalidToken)
 		{
-			return FString();
+			return FPackagePath();
 		}
 
 		FReadScopeLock LockForScope(TokenLock);
-		return StringTable.Resolve(ID);
+		return PackageDataTable.Resolve(ID);
 	}
 }
 
@@ -674,7 +676,7 @@ FBulkDataBase& FBulkDataBase::operator=(const FBulkDataBase& Other)
 	}
 	else
 	{
-		Data.Token = FileTokenSystem::CopyFileToken(Other.Data.Token);
+		Data.Token = PackageTokenSystem::CopyToken(Other.Data.Token);
 	}
 
 	BulkDataSize = Other.BulkDataSize;
@@ -701,9 +703,8 @@ FBulkDataBase& FBulkDataBase::operator=(const FBulkDataBase& Other)
 		}
 		else
 		{
-			const FString AssetFilename = FileTokenSystem::GetFilename(Data.Token);
-			const FString MemoryMappedFilename = ConvertFilenameFromFlags(AssetFilename);
-			MemoryMapBulkData(MemoryMappedFilename, BulkDataOffset, BulkDataSize);
+			FPackagePath PackagePath = PackageTokenSystem::GetPackagePath(Data.Token);
+			MemoryMapBulkData(PackagePath, GetPackageSegmentFromFlags(), BulkDataOffset, BulkDataSize);
 		}
 	}
 
@@ -719,8 +720,8 @@ FBulkDataBase::~FBulkDataBase()
 	FreeData();
 	if (!IsUsingIODispatcher())
 	{
-		FileTokenSystem::UnregisterFileToken(Data.Token);
-	}	
+		PackageTokenSystem::UnRegisterToken(Data.Token);
+	}
 }
 
 void FBulkDataBase::Serialize(FArchive& Ar, UObject* Owner, int32 /*Index*/, bool bAttemptFileMapping, int32 ElementSize)
@@ -771,17 +772,17 @@ void FBulkDataBase::Serialize(FArchive& Ar, UObject* Owner, int32 /*Index*/, boo
 			Ar << DummyValue;
 		}
 
-		// Assuming that Owner/Package/Linker are all valid, the old BulkData system would
+		// We assume that Owner/Package/Linker are all valid. The old BulkData system would
 		// generally fail if any of these were nullptr but had plenty of inconsistent checks
 		// scattered throughout.
 		checkf(Owner != nullptr, TEXT("FBulkDataBase::Serialize requires a valid Owner"));
 		const UPackage* Package = Owner->GetOutermost();
-		checkf(Package != nullptr, TEXT("FBulkDataBase::Serialize requires n Owner that returns a valid UPackage"));
+		checkf(Package != nullptr, TEXT("FBulkDataBase::Serialize requires an Owner that returns a valid UPackage"));
 
 		if (!IsInlined() && bUseIoDispatcher)
 		{
 			Data.PackageID = Package->GetPackageId().Value();
-			SetRuntimeBulkDataFlags(BULKDATA_UsesIoDispatcher); // Indicates that this BulkData should use the FIoChunkId rather than a filename
+			SetRuntimeBulkDataFlags(BULKDATA_UsesIoDispatcher); // Indicates that this BulkData should use the FIoChunkId rather than a PackagePath
 		}
 		else
 		{
@@ -789,31 +790,26 @@ void FBulkDataBase::Serialize(FArchive& Ar, UObject* Owner, int32 /*Index*/, boo
 			Data.Token = InvalidToken;
 		}
 
-		const FString* Filename = nullptr;
+		const FPackagePath* PackagePath = nullptr;
 		const FLinkerLoad* Linker = nullptr;
 
-		FString FallbackFilename;
-
 		if (bUseIoDispatcher == false)
 		{
 			Linker = FLinkerLoad::FindExistingLinkerForPackage(Package);
 
 			if (Linker != nullptr)
 			{
-				Filename = &Linker->Filename;
+				PackagePath = &Linker->GetPackagePath();
 			}
 			else if (!::ShouldAllowBulkDataInIoStore() && !IsInlined())
 			{
-				const TCHAR* PackageExtension = Package->ContainsMap() ? TEXT(".umap") : TEXT(".uasset");
-				if (FPackageName::TryConvertLongPackageNameToFilename(Package->FileName.ToString(), FallbackFilename, PackageExtension))
+				PackagePath = &Package->GetLoadedPath();
+				if (PackagePath->IsEmpty())
 				{
-					Filename = &FallbackFilename;
-				}
-				else
-				{	
 					// Note that this Bulkdata object will end up with an invalid token and will end up resolving to an empty file path!
-					UE_LOG(LogSerialization, Warning, TEXT("LongPackageNameToFilename failed to convert '%s'. Path does not map to any roots!"), *Package->FileName.ToString());
-				}				
+					UE_LOG(LogSerialization, Warning, TEXT("Could not get PackagePath from linker or package for package %s!"), *Package->GetName());
+					PackagePath = nullptr;
+				}
 			}
 		}
 
@@ -842,7 +838,7 @@ void FBulkDataBase::Serialize(FArchive& Ar, UObject* Owner, int32 /*Index*/, boo
 		{
 			if (IsDuplicateNonOptional())
 			{
-				ProcessDuplicateData(Ar, Package, Filename, BulkDataOffset);
+				ProcessDuplicateData(Ar, Package, PackagePath, BulkDataOffset);
 			}
 
 			// Set during cook
@@ -869,11 +865,11 @@ void FBulkDataBase::Serialize(FArchive& Ar, UObject* Owner, int32 /*Index*/, boo
 				}
 				else
 				{
-					// If we have no valid input file name then the package is broken anyway and we will not be able to find any memory mapped data!
-					if (Filename != nullptr)
+					// If we have no valid input packagepath then the package is broken anyway and we will not be able to find any memory mapped data!
+					if (PackagePath != nullptr)
 					{
-						const FString MemoryMappedFilename = ConvertFilenameFromFlags(*Filename);
-						if (!MemoryMapBulkData(MemoryMappedFilename, BulkDataOffset, BulkDataSize))
+						EPackageSegment PackageSegment = GetPackageSegmentFromFlags();
+						if (!MemoryMapBulkData(*PackagePath, PackageSegment, BulkDataOffset, BulkDataSize))
 						{
 							bShouldForceLoad = true; // Signal we want to force the BulkData to load
 						}
@@ -896,10 +892,10 @@ void FBulkDataBase::Serialize(FArchive& Ar, UObject* Owner, int32 /*Index*/, boo
 			}
 		}
 
-		// If we are not using the FIoDispatcher and we have a filename then we need to make sure we can retrieve it later!
-		if (bUseIoDispatcher == false && Filename != nullptr)
+		// If we are not using the FIoDispatcher and we have a PackagePath then we need to make sure we can retrieve it later!
+		if (bUseIoDispatcher == false && PackagePath != nullptr)
 		{
-			Data.Token = FileTokenSystem::RegisterFileToken(Package->GetFName(), *Filename);
+			Data.Token = PackageTokenSystem::RegisterToken(Package->GetFName(), *PackagePath);
 		}
 
 		if (bShouldForceLoad)
@@ -1111,10 +1107,8 @@ bool FBulkDataBase::DoesExist() const
 #if ALLOW_OPTIONAL_DATA
 	if (!IsUsingIODispatcher())
 	{
-		FString Filename = FileTokenSystem::GetFilename(Data.Token);
-		Filename = ConvertFilenameFromFlags(Filename);
-
-		return IFileManager::Get().FileExists(*Filename);
+		FPackagePath PackagePath = PackageTokenSystem::GetPackagePath(Data.Token);
+		return IPackageResourceManager::Get().DoesPackageExist(PackagePath, GetPackageSegmentFromFlags());
 	}
 	else
 	{
@@ -1188,7 +1182,7 @@ IAsyncReadFileHandle* FBulkDataBase::OpenAsyncReadHandle() const
 	}
 	else
 	{
-		return FPlatformFileManager::Get().GetPlatformFile().OpenAsyncRead(*GetFilename());
+		return IPackageResourceManager::Get().OpenAsyncReadPackage(GetPackagePath(), GetPackageSegmentFromFlags());
 	}
 }
 
@@ -1211,13 +1205,14 @@ IBulkDataIORequest* FBulkDataBase::CreateStreamingRequest(int64 OffsetInBulkData
 	}
 	else
 	{
-		const FString AssetFilename = FileTokenSystem::GetFilename(Data.Token);
-		const FString Filename = ConvertFilenameFromFlags(AssetFilename);
+		FPackagePath PackagePath = PackageTokenSystem::GetPackagePath(Data.Token);
+
+		EPackageSegment PackageSegment = GetPackageSegmentFromFlags();
 
-		UE_CLOG(IsStoredCompressedOnDisk(), LogSerialization, Fatal, TEXT("Package level compression is no longer supported (%s)."), *Filename);
-		UE_CLOG(BulkDataSize <= 0, LogSerialization, Error, TEXT("(%s) has invalid bulk data size."), *Filename);
+		UE_CLOG(IsStoredCompressedOnDisk(), LogSerialization, Fatal, TEXT("Package level compression is no longer supported (%s)."), *PackagePath.GetDebugName(PackageSegment));
+		UE_CLOG(BulkDataSize <= 0, LogSerialization, Error, TEXT("(%s) has invalid bulk data size."), *PackagePath.GetDebugName(PackageSegment));
 
-		IAsyncReadFileHandle* IORequestHandle = FPlatformFileManager::Get().GetPlatformFile().OpenAsyncRead(*Filename);
+		IAsyncReadFileHandle* IORequestHandle = IPackageResourceManager::Get().OpenAsyncReadPackage(PackagePath, PackageSegment);
 		checkf(IORequestHandle, TEXT("OpenAsyncRead failed")); // An assert as there shouldn't be a way for this to fail
 
 		if (IORequestHandle == nullptr)
@@ -1253,7 +1248,8 @@ IBulkDataIORequest* FBulkDataBase::CreateStreamingRequestForRange(const BulkData
 	{
 		const FBulkDataBase& End = *RangeArray[RangeArray.Num() - 1];
 
-		checkf(Start.CreateChunkId() == End.CreateChunkId(), TEXT("BulkData range does not come from the same file (%s vs %s)"), *Start.GetFilename(), *End.GetFilename());
+		checkf(Start.CreateChunkId() == End.CreateChunkId(), TEXT("BulkData range does not come from the same file (%s vs %s)"),
+			*Start.GetPackagePath().GetDebugName(), *End.GetPackagePath().GetDebugName());
 
 		const int64 ReadOffset = Start.GetBulkDataOffsetInFile();
 		const int64 ReadLength = (End.GetBulkDataOffsetInFile() + End.GetBulkDataSize()) - ReadOffset;
@@ -1269,7 +1265,8 @@ IBulkDataIORequest* FBulkDataBase::CreateStreamingRequestForRange(const BulkData
 	{	
 		const FBulkDataBase& End = *RangeArray[RangeArray.Num() - 1];
 
-		checkf(Start.GetFilename() == End.GetFilename(), TEXT("BulkData range does not come from the same file (%s vs %s)"), *Start.GetFilename(), *End.GetFilename());
+		checkf(Start.GetPackagePath() == End.GetPackagePath(), TEXT("BulkData range does not come from the same file (%s vs %s)"),
+			*Start.GetPackagePath().GetDebugName(), *End.GetPackagePath().GetDebugName());
 
 		const int64 ReadOffset = Start.GetBulkDataOffsetInFile();
 		const int64 ReadLength = (End.GetBulkDataOffsetInFile() + End.GetBulkDataSize()) - ReadOffset;
@@ -1310,7 +1307,7 @@ void FBulkDataBase::RemoveBulkData()
 
 	if (!IsUsingIODispatcher())
 	{ 
-		FileTokenSystem::UnregisterFileToken(Data.Token);
+		PackageTokenSystem::UnRegisterToken(Data.Token);
 		Data.Token = InvalidToken;
 	}
 
@@ -1380,8 +1377,7 @@ FIoFilenameHash FBulkDataBase::GetIoFilenameHash() const
 {
 	if (!IsUsingIODispatcher())
 	{
-		FString Filename = FileTokenSystem::GetFilename(Data.Token);
-		return MakeIoFilenameHash(Filename);
+		return MakeIoFilenameHash(PackageTokenSystem::GetPackagePath(Data.Token));
 	}
 	else
 	{
@@ -1390,19 +1386,28 @@ FIoFilenameHash FBulkDataBase::GetIoFilenameHash() const
 }
 
 FString FBulkDataBase::GetFilename() const
+{
+	return GetPackagePath().GetLocalFullPath(GetPackageSegmentFromFlags());
+}
+
+FPackagePath FBulkDataBase::GetPackagePath() const
 {
 	if (!IsUsingIODispatcher())
 	{
-		FString Filename = FileTokenSystem::GetFilename(Data.Token);
-		return ConvertFilenameFromFlags(Filename);
+		return PackageTokenSystem::GetPackagePath(Data.Token);
 	}
 	else
 	{
-		UE_LOG(LogBulkDataRuntime, Warning, TEXT("Attempting to get the filename for BulkData that uses the IoDispatcher, this will return an empty string"));
-		return FString("");
+		UE_LOG(LogBulkDataRuntime, Warning, TEXT("Attempting to get the PackagePath for BulkData that uses the IoDispatcher, this will return an empty PackagePath"));
+		return FPackagePath();
 	}
 }
 
+EPackageSegment FBulkDataBase::GetPackageSegment() const
+{
+	return GetPackageSegmentFromFlags();
+}
+
 bool FBulkDataBase::CanDiscardInternalData() const
 {	
 	// Data marked as single use should always be discarded
@@ -1440,7 +1445,7 @@ void FBulkDataBase::LoadDataDirectly(void** DstBuffer)
 
 	if (!IsIoDispatcherEnabled())
 	{
-		InternalLoadFromFileSystem(DstBuffer);
+		InternalLoadFromPackageResource(DstBuffer);
 	}
 	else if (IsUsingIODispatcher())
 	{
@@ -1468,7 +1473,7 @@ void FBulkDataBase::LoadDataAsynchronously(AsyncCallback&& Callback)
 		Async(EAsyncExecution::ThreadPool, [this, Callback]()
 			{
 				void* DataPtr = nullptr;
-				InternalLoadFromFileSystem(&DataPtr);
+				InternalLoadFromPackageResource(&DataPtr);
 
 				FIoBuffer Buffer(FIoBuffer::Wrap, DataPtr, GetBulkDataSize());
 				TIoStatusOr<FIoBuffer> Status(Buffer);
@@ -1488,41 +1493,41 @@ void FBulkDataBase::LoadDataAsynchronously(AsyncCallback&& Callback)
 	}
 }
 
-void FBulkDataBase::InternalLoadFromFileSystem(void** DstBuffer)
+void FBulkDataBase::InternalLoadFromPackageResource(void** DstBuffer)
 {
-	FString Filename = FileTokenSystem::GetFilename(Data.Token);
-
+	FPackagePath PackagePath = PackageTokenSystem::GetPackagePath(Data.Token);
 	int64 Offset = BulkDataOffset;
 
-	// Fix up the Filename/Offset to work with streaming if EDL is enabled and the filename is still referencing a uasset or umap
-	if (IsInlined() && (Filename.EndsWith(TEXT(".uasset")) || Filename.EndsWith(TEXT(".umap"))))
+	// Fix up the PackageSegment/Offset to work with streaming if inlined
+	EPackageSegment PackageSegment;
+	if (IsInlined())
 	{
-		Offset -= IFileManager::Get().FileSize(*Filename);
-		Filename = FPaths::GetBaseFilename(Filename, false) + BulkDataExt::Export;
+		Offset -= IPackageResourceManager::Get().FileSize(PackagePath, EPackageSegment::Header);
+		PackageSegment = EPackageSegment::Exports;
 	}
 	else
 	{
-		Filename = ConvertFilenameFromFlags(Filename);
+		PackageSegment = GetPackageSegmentFromFlags();
 	}
 
-	// If the data is inlined then we already loaded is during ::Serialize, this warning should help track cases where data is being discarded then re-requested.
+	// If the data is inlined then we already loaded it during ::Serialize, this warning should help track cases where data is being discarded then re-requested.
 	// Disabled at the moment
-	UE_CLOG(IsInlined(), LogSerialization, Warning, TEXT("Reloading inlined bulk data directly from disk, this is detrimental to loading performance. Filename: '%s'."), *Filename);
+	UE_CLOG(IsInlined(), LogSerialization, Warning, TEXT("Reloading inlined bulk data directly from disk, this is detrimental to loading performance. PackagePath: '%s'."),
+		*PackagePath.GetDebugName(PackageSegment));
 
-	FArchive* Ar = IFileManager::Get().CreateFileReader(*Filename, FILEREAD_Silent);
-	checkf(Ar != nullptr, TEXT("Failed to open the file to load bulk data from. Filename: '%s'."), *Filename);
+	FOpenPackageResult Result = IPackageResourceManager::Get().OpenReadPackage(PackagePath, PackageSegment);
+	checkf(Result.Archive.IsValid() && Result.Format == EPackageFormat::Binary, TEXT("Failed to open the file to load bulk data from. PackagePath: '%s': %s."),
+		*PackagePath.GetDebugName(PackageSegment), (!Result.Archive.IsValid() ? TEXT("could not find package") : TEXT("package is a TextAsset which is not supported")));
 
 	// Seek to the beginning of the bulk data in the file.
-	Ar->Seek(Offset);
+	Result.Archive->Seek(Offset);
 
 	if (*DstBuffer == nullptr)
 	{
 		*DstBuffer = FMemory::Malloc(BulkDataSize, 0);
 	}
 
-	SerializeBulkData(*Ar, *DstBuffer, BulkDataSize);
-
-	delete Ar;
+	SerializeBulkData(*Result.Archive, *DstBuffer, BulkDataSize);
 }
 
 void FBulkDataBase::InternalLoadFromIoStore(void** DstBuffer)
@@ -1564,7 +1569,7 @@ void FBulkDataBase::InternalLoadFromIoStoreAsync(void** DstBuffer, AsyncCallback
 	Batch.Issue();
 }
 
-void FBulkDataBase::ProcessDuplicateData(FArchive& Ar, const UPackage* Package, const FString* Filename, int64& InOutOffsetInFile)
+void FBulkDataBase::ProcessDuplicateData(FArchive& Ar, const UPackage* Package, const FPackagePath* PackagePath, int64& InOutOffsetInFile)
 {
 	// We need to load the optional bulkdata info as we might need to create a FIoChunkId based on it!
 	EBulkDataFlags NewFlags;
@@ -1588,12 +1593,12 @@ void FBulkDataBase::ProcessDuplicateData(FArchive& Ar, const UPackage* Package,
 	}
 	else
 	{
-		// If we have no valid input file name then the package is broken anyway and we will not be able to find the optional data!
-		if (Filename != nullptr)
+		// If we have no valid input packagepath then the package is broken anyway and we will not be able to find the optional data!
+		if (PackagePath != nullptr)
 		{
-			const FString OptionalDataFilename = FPathViews::ChangeExtension(*Filename, BulkDataExt::Optional);
+			EPackageSegment PackageSegment = EPackageSegment::BulkDataOptional;
 
-			if (IFileManager::Get().FileExists(*OptionalDataFilename))
+			if (IPackageResourceManager::Get().DoesPackageExist(*PackagePath, PackageSegment))
 			{
 				BulkDataFlags = NewFlags;
 				checkf(BulkDataSize == NewSizeOnDisk, TEXT("Size mistach between original data size (%lld)and duplicate data size (%lld)"), BulkDataSize, NewSizeOnDisk);
@@ -1657,14 +1662,14 @@ void FBulkDataBase::SerializeBulkData(FArchive& Ar, void* DstBuffer, int64 DataL
 	}
 }
 
-bool FBulkDataBase::MemoryMapBulkData(const FString& Filename, int64 OffsetInBulkData, int64 BytesToRead)
+bool FBulkDataBase::MemoryMapBulkData(const FPackagePath& PackagePath, EPackageSegment PackageSegment, int64 OffsetInBulkData, int64 BytesToRead)
 {
 	checkf(!IsBulkDataLoaded(), TEXT("Attempting to memory map BulkData that is already loaded"));
 
 	IMappedFileHandle* MappedHandle = nullptr;
 	IMappedFileRegion* MappedRegion = nullptr;
 
-	MappedHandle = FPlatformFileManager::Get().GetPlatformFile().OpenMapped(*Filename);
+	MappedHandle = IPackageResourceManager::Get().OpenMappedHandleToPackage(PackagePath, PackageSegment);
 
 	if (MappedHandle == nullptr)
 	{
@@ -1704,7 +1709,8 @@ void FBulkDataBase::FlushAsyncLoading()
 		{
 			if (FPlatformTime::ToMilliseconds64(FPlatformTime::Cycles64() - StartTime))
 			{
-				UE_LOG(LogSerialization, Warning, TEXT("Waiting for %s bulk data (%lld) to be loaded longer than 1000ms"), *GetFilename(), GetBulkDataSize());
+				UE_LOG(LogSerialization, Warning, TEXT("Waiting for %s bulk data (%lld) to be loaded longer than 1000ms"),
+					*GetPackagePath().GetDebugName(), GetBulkDataSize());
 				StartTime = FPlatformTime::Cycles64(); // Reset so we spam the log every second or so that we are stalled!
 			}
 
@@ -1714,30 +1720,30 @@ void FBulkDataBase::FlushAsyncLoading()
 	}	
 }
 
-FString FBulkDataBase::ConvertFilenameFromFlags(const FString& Filename) const
-{	
+EPackageSegment FBulkDataBase::GetPackageSegmentFromFlags() const
+{
 	if (IsOptional())
 	{
 		// Optional data should be tested for first as we in theory can have data that would
 		// be marked as inline, also marked as optional and in this case we should treat it as
 		// optional data first.
-		return FPathViews::ChangeExtension(Filename, BulkDataExt::Optional);
+		return EPackageSegment::BulkDataOptional;
 	}
 	else if (!IsInSeparateFile())
 	{
-		return Filename;
+		return EPackageSegment::Header;
 	}
 	else if (IsInlined())
 	{
-		return FPathViews::ChangeExtension(Filename, BulkDataExt::Export);
+		return EPackageSegment::Exports;
 	}
 	else if (IsFileMemoryMapped())
 	{
-		return FPathViews::ChangeExtension(Filename, BulkDataExt::MemoryMapped);
+		return EPackageSegment::BulkDataMemoryMapped;
 	}
 	else
 	{
-		return FPathViews::ChangeExtension(Filename, BulkDataExt::Default);
+		return EPackageSegment::BulkDataDefault;
 	}
 }
 
diff --git a/Engine/Source/Runtime/CoreUObject/Private/Serialization/LoadTimeTrace.cpp b/Engine/Source/Runtime/CoreUObject/Private/Serialization/LoadTimeTrace.cpp
--- a/Engine/Source/Runtime/CoreUObject/Private/Serialization/LoadTimeTrace.cpp
+++ b/Engine/Source/Runtime/CoreUObject/Private/Serialization/LoadTimeTrace.cpp
@@ -144,7 +144,7 @@ void FLoadTimeProfilerTracePrivate::OutputEndRequest(uint64 RequestId)
 		<< EndRequest.RequestId(RequestId);
 }
 
-void FLoadTimeProfilerTracePrivate::OutputNewAsyncPackage(const void* AsyncPackage, const FName& PackageName)
+void FLoadTimeProfilerTracePrivate::OutputNewAsyncPackage(const void* AsyncPackage, FName PackageName)
 {
 	TCHAR Buffer[FName::StringBufferSize];
 	uint16 NameSize = (PackageName.ToString(Buffer) + 1) * sizeof(TCHAR);
@@ -153,6 +153,14 @@ void FLoadTimeProfilerTracePrivate::OutputNewAsyncPackage(const void* AsyncPacka
 		<< NewAsyncPackage.Attachment(Buffer, NameSize);
 }
 
+void FLoadTimeProfilerTracePrivate::OutputNewAsyncPackage(const void* AsyncPackage, const FString& PackageName)
+{
+	uint16 NameSize = (PackageName.Len() + 1) * sizeof(TCHAR);
+	UE_TRACE_LOG(LoadTime, NewAsyncPackage, LoadTimeChannel, NameSize)
+		<< NewAsyncPackage.AsyncPackage(AsyncPackage)
+		<< NewAsyncPackage.Attachment(*PackageName, NameSize);
+}
+
 void FLoadTimeProfilerTracePrivate::OutputBeginLoadAsyncPackage(const void* AsyncPackage)
 {
 	UE_TRACE_LOG(LoadTime, BeginLoadAsyncPackage, LoadTimeChannel)
diff --git a/Engine/Source/Runtime/CoreUObject/Private/Serialization/LoadTimeTracePrivate.h b/Engine/Source/Runtime/CoreUObject/Private/Serialization/LoadTimeTracePrivate.h
--- a/Engine/Source/Runtime/CoreUObject/Private/Serialization/LoadTimeTracePrivate.h
+++ b/Engine/Source/Runtime/CoreUObject/Private/Serialization/LoadTimeTracePrivate.h
@@ -15,7 +15,8 @@ struct FLoadTimeProfilerTracePrivate
 	static void OutputStartAsyncLoading();
 	static void OutputSuspendAsyncLoading();
 	static void OutputResumeAsyncLoading();
-	static void OutputNewAsyncPackage(const void* AsyncPackage, const FName& PackageName);
+	static void OutputNewAsyncPackage(const void* AsyncPackage, FName PackageName);
+	static void OutputNewAsyncPackage(const void* AsyncPackage, const FString& PackageName);
 	static void OutputBeginLoadAsyncPackage(const void* AsyncPackage);
 	static void OutputEndLoadAsyncPackage(const void* AsyncPackage);
 	static void OutputDestroyAsyncPackage(const void* AsyncPackage);
diff --git a/Engine/Source/Runtime/CoreUObject/Private/UObject/Linker.cpp b/Engine/Source/Runtime/CoreUObject/Private/UObject/Linker.cpp
--- a/Engine/Source/Runtime/CoreUObject/Private/UObject/Linker.cpp
+++ b/Engine/Source/Runtime/CoreUObject/Private/UObject/Linker.cpp
@@ -23,6 +23,7 @@
 #include "UObject/UObjectThreadContext.h"
 #include "UObject/DebugSerializationFlags.h"
 #include "UObject/ObjectResource.h"
+#include "UObject/PackageResourceManager.h"
 #include "Algo/Transform.h"
 
 DEFINE_LOG_CATEGORY(LogLinker);
@@ -162,16 +163,14 @@ FName FLinker::GetExportClassName( int32 i )
 /*----------------------------------------------------------------------------
 	FLinker.
 ----------------------------------------------------------------------------*/
-FLinker::FLinker(ELinkerType::Type InType, UPackage* InRoot, const TCHAR* InFilename)
+FLinker::FLinker(ELinkerType::Type InType, UPackage* InRoot)
 : LinkerType(InType)
 , LinkerRoot( InRoot )
-, Filename( InFilename )
 , FilterClientButNotServer(false)
 , FilterServerButNotClient(false)
 , ScriptSHA(nullptr)
 {
 	check(LinkerRoot);
-	check(InFilename);
 
 	if( !GIsClient && GIsServer)
 	{
@@ -183,6 +182,25 @@ FLinker::FLinker(ELinkerType::Type InType, UPackage* InRoot, const TCHAR* InFile
 	}
 }
 
+FLinker::FLinker(ELinkerType::Type InType, UPackage* InRoot, const TCHAR* InFilename)
+: FLinker(InType, InRoot)
+{
+	PRAGMA_DISABLE_DEPRECATION_WARNINGS
+	Filename = InFilename;
+	PRAGMA_ENABLE_DEPRECATION_WARNINGS
+}
+
+/** Returns a descriptor of the PackagePath this Linker is reading from or writing to, usable for an identifier in warning and log messages */
+FString FLinker::GetDebugName() const
+{
+	// UE_DEPRECATED(5.0)
+	UE_LOG(LogLinker, Warning, TEXT("Filename and the default GetDebugName on FLinker is deprecated and will be removed in a future release. Subclasses should override GetDebugName."));
+	PRAGMA_DISABLE_DEPRECATION_WARNINGS
+	return Filename;
+	PRAGMA_ENABLE_DEPRECATION_WARNINGS
+}
+
+
 void FLinker::Serialize( FArchive& Ar )
 {
 	// This function is only used for counting memory, actual serialization uses a different path
@@ -429,11 +447,13 @@ void FLinker::GetScriptSHAKey(uint8* OutKey)
 	ScriptSHA->GetHash(OutKey);
 }
 
+PRAGMA_DISABLE_DEPRECATION_WARNINGS
 FLinker::~FLinker()
 {
 	// free any SHA memory
 	delete ScriptSHA;
 }
+PRAGMA_ENABLE_DEPRECATION_WARNINGS
 
 
 
@@ -463,7 +483,7 @@ void DeleteLoader(FLinkerLoad* Loader)
 	FLinkerManager::Get().RemoveLinker(Loader);
 }
 
-static void LogGetPackageLinkerError(FUObjectSerializeContext* LoadContext, const TCHAR* InFilename, const FText& InErrorMessage, UObject* InOuter, uint32 LoadFlags)
+static void LogGetPackageLinkerError(FUObjectSerializeContext* LoadContext, const FPackagePath& PackagePath, const FText& InErrorMessage, UObject* InOuter, uint32 LoadFlags)
 {
 	static FName NAME_LoadErrors("LoadErrors");
 	struct Local
@@ -487,7 +507,7 @@ static void LogGetPackageLinkerError(FUObjectSerializeContext* LoadContext, cons
 
 	FLinkerLoad* SerializedPackageLinker = LoadContext ? LoadContext->SerializedPackageLinker : nullptr;
 	UObject* SerializedObject = LoadContext ? LoadContext->SerializedObject : nullptr;
-	FString LoadingFile = InFilename ? InFilename : InOuter ? *InOuter->GetName() : TEXT("NULL");
+	FString LoadingFile = !PackagePath.IsEmpty() ? PackagePath.GetDebugName() : FString(TEXT("NULL"));
 	
 	FFormatNamedArguments Arguments;
 	Arguments.Add(TEXT("LoadingFile"), FText::FromString(LoadingFile));
@@ -501,7 +521,7 @@ static void LogGetPackageLinkerError(FUObjectSerializeContext* LoadContext, cons
 		{
 			LinkerToUse = SerializedObject->GetLinker();
 		}
-		FString LoadedByFile = LinkerToUse ? *LinkerToUse->Filename : SerializedObject->GetOutermost()->GetName();
+		FString LoadedByFile = LinkerToUse ? LinkerToUse->GetDebugName() : SerializedObject->GetOutermost()->GetName();
 		FullErrorMessage = FText::FromString(FAssetMsg::GetAssetLogString(*LoadedByFile, FullErrorMessage.ToString()));
 	}
 
@@ -526,13 +546,9 @@ static void LogGetPackageLinkerError(FUObjectSerializeContext* LoadContext, cons
 			// so any errors we cant get links out of we will just let be output to the output log (above)
 			// rather than clog up the message log
 
-			if(InFilename != NULL && InOuter != NULL)
+			if (!PackagePath.IsEmpty() && InOuter != NULL)
 			{
-				FString PackageName;
-				if (!FPackageName::TryConvertFilenameToLongPackageName(InFilename, PackageName))
-				{
-					PackageName = InFilename;
-				}
+				FString PackageName = PackagePath.GetPackageNameOrFallback();
 				FString OuterPackageName;
 				if (!FPackageName::TryConvertFilenameToLongPackageName(InOuter->GetPathName(), OuterPackageName))
 				{
@@ -583,48 +599,73 @@ static void LogGetPackageLinkerError(FUObjectSerializeContext* LoadContext, cons
 	}
 }
 
-/** Customized version of FPackageName::DoesPackageExist that takes dynamic native class packages into account */
-static bool DoesPackageExistForGetPackageLinker(const FString& LongPackageName, const FGuid* Guid, FString& OutFilename)
+static bool IsConvertedDynamicPackage(FName PackageName)
 {
-	if (
+	return
 #if WITH_EDITORONLY_DATA
-		GLinkerAllowDynamicClasses && 
+		GLinkerAllowDynamicClasses &&
 #endif
-		GetConvertedDynamicPackageNameToTypeName().Contains(*LongPackageName))
+		GetConvertedDynamicPackageNameToTypeName().Contains(PackageName);
+}
+
+FString GetPrestreamPackageLinkerName(const TCHAR* InLongPackageName, bool bSkipIfExists)
+{
+	if (!InLongPackageName)
 	{
-		OutFilename = FPackageName::LongPackageNameToFilename(LongPackageName);
-		return true;
+		return FString();
 	}
-	else
+
+	FPackagePath PackagePath;
+	if (!FPackagePath::TryFromMountedName(InLongPackageName, PackagePath))
 	{
-		return FPackageName::DoesPackageExist(LongPackageName, Guid, &OutFilename);
+		return FString();
 	}
+	FName PackageFName(PackagePath.GetPackageFName());
+	UPackage* ExistingPackage = bSkipIfExists ? FindObjectFast<UPackage>(nullptr, PackageFName) : nullptr;
+	if (ExistingPackage)
+	{
+		return FString(); // we won't load this anyway, don't prestream
+	}
+	if (!IsConvertedDynamicPackage(PackageFName) && !FPackageName::DoesPackageExist(PackagePath, &PackagePath))
+	{
+		return FString();
+	}
+
+	return PackagePath.GetLocalFullPath();
 }
 
-FString GetPrestreamPackageLinkerName(const TCHAR* InLongPackageName, bool bExistSkip)
+static FPackagePath GetPackagePath(UPackage* InOuter, const TCHAR* InPackageNameOrFileName)
 {
-	FString NewFilename;
-	if (InLongPackageName)
+	if (InPackageNameOrFileName)
 	{
-		FString PackageName(InLongPackageName);
-		if (!FPackageName::TryConvertFilenameToLongPackageName(InLongPackageName, PackageName))
-		{
-			return FString();
-		}
-		UPackage* ExistingPackage = bExistSkip ? FindObject<UPackage>(nullptr, *PackageName) : nullptr;
-		if (ExistingPackage)
+		FPackagePath PackagePath;
+		if (FPackagePath::TryFromMountedName(InPackageNameOrFileName, PackagePath))
 		{
-			return FString(); // we won't load this anyway, don't prestream
+			return PackagePath;
 		}
-		
-		const bool DoesNativePackageExist = DoesPackageExistForGetPackageLinker(PackageName, nullptr, NewFilename);
-
-		if ( !DoesNativePackageExist )
+		else
 		{
-			return FString();
+			UE_LOG(LogLinker, Warning, TEXT("GetPackagePath: Path \"%s\" is not in a mounted path; returning empty PackagePath. LoadPackage or GetPackageLinker will fail."), InPackageNameOrFileName);
+			return FPackagePath();
 		}
 	}
-	return NewFilename;
+	else if (InOuter)
+	{
+		// Resolve filename from package name.
+		return FPackagePath::FromPackageNameChecked(InOuter->GetName());
+	}
+	else
+	{
+		UE_LOG(LogLinker, Warning, TEXT("GetPackagePath was passed an empty PackageName."));
+		return FPackagePath();
+	}
+}
+
+COREUOBJECT_API FLinkerLoad* GetPackageLinker(UPackage* InOuter, const TCHAR* InLongPackageName, uint32 LoadFlags,
+	UPackageMap* Sandbox, FGuid* CompatibleGuid, FArchive* InReaderOverride, FUObjectSerializeContext** InOutLoadContext,
+	FLinkerLoad* ImportLinker, const FLinkerInstancingContext* InstancingContext)
+{
+	return GetPackageLinker(InOuter, GetPackagePath(InOuter, InLongPackageName), LoadFlags, Sandbox, CompatibleGuid, InReaderOverride, InOutLoadContext, ImportLinker, InstancingContext);
 }
 
 //
@@ -633,7 +674,7 @@ FString GetPrestreamPackageLinkerName(const TCHAR* InLongPackageName, bool bExis
 FLinkerLoad* GetPackageLinker
 (
 	UPackage*		InOuter,
-	const TCHAR*	InLongPackageName,
+	const FPackagePath& InPackagePath,
 	uint32			LoadFlags,
 	UPackageMap*	Sandbox,
 	FGuid*			CompatibleGuid,
@@ -658,52 +699,42 @@ FLinkerLoad* GetPackageLinker
 		return Result;
 	}
 
-	FString PackageNameToCreate;
-	UPackage* TargetPackage = nullptr;
-	if (!InLongPackageName)
+	FPackagePath PackagePath(InPackagePath);
+	FString PackageName = PackagePath.GetPackageName();
+	FName PackageFName(*PackageName);
+	if (PackageFName.IsNone())
 	{
-		// Resolve filename from package name.
-		if (!InOuter)
-		{
-			// try to recover from this instead of throwing, it seems recoverable just by doing this
-			LogGetPackageLinkerError(InExistingContext, InLongPackageName, LOCTEXT("PackageResolveFailed", "Can't resolve asset name"), InOuter, LoadFlags);
-			return nullptr;
-		}
-		PackageNameToCreate = InOuter->GetName();
-		TargetPackage = InOuter;
+		// try to recover from this instead of throwing, it seems recoverable just by doing this
+		LogGetPackageLinkerError(InExistingContext, InPackagePath, LOCTEXT("PackageResolveFailed", "Can't resolve asset name"), InOuter, LoadFlags);
+		return nullptr;
+	}
 
-		// Process any package redirects
+	// Process any package redirects
+	{
+		const FCoreRedirectObjectName NewPackageName = FCoreRedirects::GetRedirectedName(ECoreRedirectFlags::Type_Package, FCoreRedirectObjectName(NAME_None, NAME_None, PackageFName));
+		if (NewPackageName.PackageName != PackageFName)
 		{
-			const FCoreRedirectObjectName NewPackageName = FCoreRedirects::GetRedirectedName(ECoreRedirectFlags::Type_Package, FCoreRedirectObjectName(NAME_None, NAME_None, *PackageNameToCreate));
-			NewPackageName.PackageName.ToString(PackageNameToCreate);
+			PackageFName = NewPackageName.PackageName;
+			PackageName = PackageFName.ToString();
+			if (!FPackagePath::TryFromMountedName(PackageName, PackagePath))
+			{
+				LogGetPackageLinkerError(InExistingContext, InPackagePath, LOCTEXT("InvalidPackageRedirect", "PackagePath has invalid FCoreRedirects redirector to unmounted contentroot"), InOuter, LoadFlags);
+				return nullptr;
+			}
 		}
 	}
+
+	UPackage* TargetPackage = nullptr;
+	if (InOuter)
+	{
+		TargetPackage = InOuter;
+	}
 	else
 	{
-		if (!FPackageName::TryConvertFilenameToLongPackageName(InLongPackageName, PackageNameToCreate))
+		TargetPackage = FindObject<UPackage>(nullptr, *PackageName);
+		if (TargetPackage && TargetPackage->GetOuter() != nullptr)
 		{
-			// try to recover from this instead of throwing, it seems recoverable just by doing this
-			LogGetPackageLinkerError(InExistingContext, InLongPackageName, LOCTEXT("PackageResolveFailed", "Can't resolve asset name"), InOuter, LoadFlags);
-			return nullptr;
-		}
-
-		// Process any package redirects
-		{
-			const FCoreRedirectObjectName NewPackageName = FCoreRedirects::GetRedirectedName(ECoreRedirectFlags::Type_Package, FCoreRedirectObjectName(NAME_None, NAME_None, *PackageNameToCreate));
-			NewPackageName.PackageName.ToString(PackageNameToCreate);
-		}
-
-		if (InOuter)
-		{
-			TargetPackage = InOuter;
-		}
-		else
-		{
-			TargetPackage = FindObject<UPackage>(nullptr, *PackageNameToCreate);
-			if (TargetPackage && TargetPackage->GetOuter() != nullptr)
-			{
-				TargetPackage = nullptr;
-			}
+			TargetPackage = nullptr;
 		}
 	}
 
@@ -714,42 +745,74 @@ FLinkerLoad* GetPackageLinker
 	}
 
 	// The editor must not redirect packages for localization. We also shouldn't redirect script or in-memory packages (in-memory packages exited earlier so we don't need to check here).
-	FString PackageNameToLoad = PackageNameToCreate;
-	if (!(GIsEditor || FPackageName::IsScriptPackage(PackageNameToLoad)))
+	FString PackageNameToCreate = PackageName;
+	if (!(GIsEditor || FPackageName::IsScriptPackage(PackageName)))
 	{
 		// Allow delegates to resolve the path
-		PackageNameToLoad = FPackageName::GetDelegateResolvedPackagePath(PackageNameToLoad);
+		FString PackageNameToLoad = FPackageName::GetDelegateResolvedPackagePath(PackageName);
 		PackageNameToLoad = FPackageName::GetLocalizedPackagePath(PackageNameToLoad);
+		if (PackageNameToLoad != PackageName)
+		{
+			PackageName = MoveTemp(PackageNameToLoad);
+			PackageFName = FName(*PackageName);
+			if (!FPackagePath::TryFromMountedName(PackageName, PackagePath))
+			{
+				LogGetPackageLinkerError(InExistingContext, InPackagePath, LOCTEXT("InvalidPackageRedirect", "GetDelegateResolvedPackagePath or GetLocalizedPackagePath returned path to unmounted contentroot"), InOuter, LoadFlags);
+				return nullptr;
+			}
+		}
 	}
 
-	// Verify that the file exists.
-	FString NewFilename;
-	const bool bDoesPackageExist = DoesPackageExistForGetPackageLinker(PackageNameToLoad, CompatibleGuid, NewFilename);
-	if (!bDoesPackageExist)
+	// Skip the existence and normalization checks for dynamic packages
+	if (!IsConvertedDynamicPackage(PackageFName))
 	{
-		// Issue a warning if the caller didn't request nowarn/quiet, and the package isn't marked as known to be missing.
-		bool bIssueWarning = (LoadFlags & (LOAD_NoWarn | LOAD_Quiet)) == 0 && !FLinkerLoad::IsKnownMissingPackage(InLongPackageName);
-
-		if (bIssueWarning)
+		// We need to call DoesPackageExist for a few reasons:
+		// 1) Get the extension that the package is using
+		// 2) Normalize the capitalization of the packagename to match the case on disk
+		// 3) Early exit, avoiding the cost of creating and deleting a UPackage if the package does not exist on disk
+#if WITH_EDITORONLY_DATA
+		// If we're creating a package, we need to match the capitalization on disk so that the anyone doing source control operations based on the PackageName will send the proper capitalization to the possibly-casesensitive source control
+		bool bMatchCaseOnDisk = TargetPackage == nullptr;
+#else
+		bool bMatchCaseOnDisk = false;
+#endif
+		bool bPackageExists = FPackageName::DoesPackageExist(PackagePath, CompatibleGuid, bMatchCaseOnDisk, &PackagePath);
+		PackageName = PackagePath.GetPackageName();
+		if (!bPackageExists)
 		{
-			// try to recover from this instead of throwing, it seems recoverable just by doing this
-			LogGetPackageLinkerError(InExistingContext, InLongPackageName, LOCTEXT("FileNotFoundShort", "Can't find file."), InOuter, LoadFlags);
+			// Issue a warning if the caller didn't request nowarn/quiet, and the package isn't marked as known to be missing.
+			bool bIssueWarning = (LoadFlags & (LOAD_NoWarn | LOAD_Quiet)) == 0 && !FLinkerLoad::IsKnownMissingPackage(InPackagePath.GetPackageFName());
+			if (bIssueWarning)
+			{
+				// try to recover from this instead of throwing, it seems recoverable just by doing this
+				LogGetPackageLinkerError(InExistingContext, InPackagePath, LOCTEXT("FileNotFoundShort", "Can't find file."), InOuter, LoadFlags);
+			}
+			return nullptr;
 		}
-		return nullptr;
 	}
 
 	UPackage* CreatedPackage = nullptr;
 	if (!TargetPackage)
 	{
+		if (PackageNameToCreate.Equals(PackageName, ESearchCase::IgnoreCase))
+		{
+			PackageNameToCreate = PackageName; // Take the capitalization from PackageName, which was normalized in the DoesPackageExist call
+		}
+		else
+		{
 #if WITH_EDITORONLY_DATA
-		// Make sure the package name matches the name on disk
-		FPackageName::FixPackageNameCase(PackageNameToCreate, FPathViews::GetExtension(NewFilename));
+			// Make sure the package name matches the capitalization on disk
+			FPackagePath PackagePathToCreate = FPackagePath::FromPackageNameChecked(PackageNameToCreate);
+			IPackageResourceManager::Get().TryMatchCaseOnDisk(PackagePathToCreate, &PackagePathToCreate);
+			PackageNameToCreate = PackagePathToCreate.GetPackageName();
 #endif
+		}
+		// Make sure the package name matches the name on disk
 		// Create the package with the provided long package name.
 		CreatedPackage = CreatePackage(*PackageNameToCreate);
 		if (!CreatedPackage)
 		{
-			LogGetPackageLinkerError(InExistingContext, InLongPackageName, LOCTEXT("FilenameToPackageShort", "Can't convert filename to asset name"), InOuter, LoadFlags);
+			LogGetPackageLinkerError(InExistingContext, InPackagePath, LOCTEXT("FilenameToPackageShort", "Can't convert filename to asset name"), InOuter, LoadFlags);
 			return nullptr;
 		}
 		if (LoadFlags & LOAD_PackageForPIE)
@@ -760,6 +823,7 @@ FLinkerLoad* GetPackageLinker
 	}
 	if (InOuter != TargetPackage)
 	{
+		// See if the Linker is already loaded for the TargetPackage we've found
 		if (FLinkerLoad* Result = FLinkerLoad::FindExistingLinkerForPackage(TargetPackage))
 		{
 			if (InExistingContext)
@@ -787,10 +851,8 @@ FLinkerLoad* GetPackageLinker
 	}
 
 	// Create new linker.
-	// we will already have found the filename above
-	check(NewFilename.Len() > 0);
 	TRefCountPtr<FUObjectSerializeContext> LoadContext(InExistingContext ? InExistingContext : FUObjectThreadContext::Get().GetSerializeContext());
-	FLinkerLoad* Result = FLinkerLoad::CreateLinker(LoadContext, TargetPackage, *NewFilename, LoadFlags, InReaderOverride, ImportLinker ? &ImportLinker->GetInstancingContext() : InstancingContext);
+	FLinkerLoad* Result = FLinkerLoad::CreateLinker(LoadContext, TargetPackage, PackagePath, LoadFlags, InReaderOverride, ImportLinker ? &ImportLinker->GetInstancingContext() : InstancingContext);
 
 	if (!Result && CreatedPackage)
 	{
@@ -800,14 +862,14 @@ FLinkerLoad* GetPackageLinker
 	return Result;
 }
 
-FLinkerLoad* LoadPackageLinker(UPackage* InOuter, const TCHAR* InLongPackageName, uint32 LoadFlags, UPackageMap* Sandbox, FGuid* CompatibleGuid, FArchive* InReaderOverride, TFunctionRef<void(FLinkerLoad* LoadedLinker)> LinkerLoadedCallback)
+FLinkerLoad* LoadPackageLinker(UPackage* InOuter, const FPackagePath& PackagePath, uint32 LoadFlags, UPackageMap* Sandbox, FGuid* CompatibleGuid, FArchive* InReaderOverride, TFunctionRef<void(FLinkerLoad* LoadedLinker)> LinkerLoadedCallback)
 {
 	FLinkerLoad* Linker = nullptr;
 	TRefCountPtr<FUObjectSerializeContext> LoadContext(FUObjectThreadContext::Get().GetSerializeContext());
 	BeginLoad(LoadContext);
 	{
 		FUObjectSerializeContext* InOutLoadContext = LoadContext;
-		Linker = GetPackageLinker(InOuter, InLongPackageName, LoadFlags, Sandbox, CompatibleGuid, InReaderOverride, &InOutLoadContext);
+		Linker = GetPackageLinker(InOuter, PackagePath, LoadFlags, Sandbox, CompatibleGuid, InReaderOverride, &InOutLoadContext);
 		if (InOutLoadContext != LoadContext)
 		{
 			// The linker already existed and was associated with another context
@@ -822,9 +884,19 @@ FLinkerLoad* LoadPackageLinker(UPackage* InOuter, const TCHAR* InLongPackageName
 	return Linker;
 }
 
+FLinkerLoad* LoadPackageLinker(UPackage* InOuter, const FPackagePath& PackagePath, uint32 LoadFlags, UPackageMap* Sandbox, FGuid* CompatibleGuid, FArchive* InReaderOverride)
+{
+	return LoadPackageLinker(InOuter, PackagePath, LoadFlags, Sandbox, CompatibleGuid, InReaderOverride, [](FLinkerLoad* InLinker) {});
+}
+
+FLinkerLoad* LoadPackageLinker(UPackage* InOuter, const TCHAR* InLongPackageName, uint32 LoadFlags, UPackageMap* Sandbox, FGuid* CompatibleGuid, FArchive* InReaderOverride, TFunctionRef<void(FLinkerLoad* LoadedLinker)> LinkerLoadedCallback)
+{
+	return LoadPackageLinker(InOuter, GetPackagePath(InOuter, InLongPackageName), LoadFlags, Sandbox, CompatibleGuid, InReaderOverride, LinkerLoadedCallback);
+}
+
 FLinkerLoad* LoadPackageLinker(UPackage* InOuter, const TCHAR* InLongPackageName, uint32 LoadFlags, UPackageMap* Sandbox, FGuid* CompatibleGuid, FArchive* InReaderOverride)
 {
-	return LoadPackageLinker(InOuter, InLongPackageName, LoadFlags, Sandbox, CompatibleGuid, InReaderOverride, [](FLinkerLoad* InLinker) {});
+	return LoadPackageLinker(InOuter, GetPackagePath(InOuter, InLongPackageName), LoadFlags, Sandbox, CompatibleGuid, InReaderOverride, [](FLinkerLoad* InLinker) {});
 }
 
 
@@ -839,8 +911,7 @@ void ResetLoadersForSave(UPackage* Package, const TCHAR* Filename)
 	FLinkerLoad* Loader = FLinkerLoad::FindExistingLinkerForPackage(Package);
 	if( Loader )
 	{
-		// Compare absolute filenames to see whether we're trying to save over an existing file.
-		if( FPaths::ConvertRelativePathToFull(Filename) == FPaths::ConvertRelativePathToFull( Loader->Filename ) )
+		if ( FPackagePath::FromLocalPath(Filename) == Loader->GetPackagePath())
 		{
 			// Detach all exports from the linker and dissociate the linker.
 			ResetLoaders( Package );
@@ -855,7 +926,7 @@ void ResetLoadersForSave(TArrayView<FPackageSaveInfo> InPackages)
 		[](const FPackageSaveInfo& InPackageSaveInfo)
 		{
 			FLinkerLoad* Loader = FLinkerLoad::FindExistingLinkerForPackage(InPackageSaveInfo.Package);
-			return Loader && FPaths::ConvertRelativePathToFull(InPackageSaveInfo.Filename) == FPaths::ConvertRelativePathToFull(Loader->Filename);
+			return Loader && FPackagePath::FromLocalPath(InPackageSaveInfo.Filename) == Loader->GetPackagePath();
 		},
 		[](const FPackageSaveInfo& InPackageSaveInfo)
 		{
diff --git a/Engine/Source/Runtime/CoreUObject/Private/UObject/LinkerLoad.cpp b/Engine/Source/Runtime/CoreUObject/Private/UObject/LinkerLoad.cpp
--- a/Engine/Source/Runtime/CoreUObject/Private/UObject/LinkerLoad.cpp
+++ b/Engine/Source/Runtime/CoreUObject/Private/UObject/LinkerLoad.cpp
@@ -10,10 +10,11 @@
 #include "Misc/ScopedSlowTask.h"
 #include "Misc/ObjectThumbnail.h"
 #include "Misc/App.h"
-#include "UObject/UObjectHash.h"
-#include "UObject/Package.h"
 #include "UObject/MetaData.h"
 #include "UObject/ObjectRedirector.h"
+#include "UObject/Package.h"
+#include "UObject/PackageResourceManager.h"
+#include "UObject/UObjectHash.h"
 #include "Misc/PackageName.h"
 #include "Blueprint/BlueprintSupport.h"
 #include "Misc/PreloadableFile.h"
@@ -423,16 +424,21 @@ void FLinkerLoad::StaticInit(UClass* InUTexture2DStaticClass)
 	UTexture2DStaticClass = InUTexture2DStaticClass;
 }
 
+FLinkerLoad* FLinkerLoad::CreateLinker(FUObjectSerializeContext* LoadContext, UPackage* Parent, const TCHAR* Filename, uint32 LoadFlags, FArchive* InLoader, const FLinkerInstancingContext* InstancingContext)
+{
+	return CreateLinker(LoadContext, Parent, FPackagePath::FromLocalPath(Filename), LoadFlags, InLoader, InstancingContext);
+}
+
 /**
  * Creates and returns a FLinkerLoad object.
  *
- * @param	Parent		Parent object to load into, can be NULL (most likely case)
- * @param	Filename	Name of file on disk to load
- * @param	LoadFlags	Load flags determining behavior
+ * @param	Parent				Parent object to load into, can be NULL (most likely case)
+ * @param	PackagePath			PackagePath to load from IPackageResourceManager
+ * @param	LoadFlags			Load flags determining behavior
  *
- * @return	new FLinkerLoad object for Parent/ Filename
+ * @return	new FLinkerLoad object for Parent/ PackagePath
  */
-FLinkerLoad* FLinkerLoad::CreateLinker(FUObjectSerializeContext* LoadContext, UPackage* Parent, const TCHAR* Filename, uint32 LoadFlags, FArchive* InLoader /*= nullptr*/, const FLinkerInstancingContext* InstancingContext /*= nullptr*/)
+FLinkerLoad* FLinkerLoad::CreateLinker(FUObjectSerializeContext* LoadContext, UPackage* Parent, const FPackagePath& PackagePath, uint32 LoadFlags, FArchive* InLoader /*= nullptr*/, const FLinkerInstancingContext* InstancingContext /*= nullptr*/)
 {
 	check(LoadContext);
 
@@ -450,7 +456,7 @@ FLinkerLoad* FLinkerLoad::CreateLinker(FUObjectSerializeContext* LoadContext, UP
 	LoadFlags &= ~LOAD_DeferDependencyLoads;
 #endif // USE_CIRCULAR_DEPENDENCY_LOAD_DEFERRING
 
-	FLinkerLoad* Linker = CreateLinkerAsync(LoadContext, Parent, Filename, LoadFlags, InstancingContext,
+	FLinkerLoad* Linker = CreateLinkerAsync(LoadContext, Parent, PackagePath, LoadFlags, InstancingContext,
 		TFunction<void()>([](){})
 		);
 	{
@@ -483,6 +489,14 @@ FLinkerLoad* FLinkerLoad::CreateLinker(FUObjectSerializeContext* LoadContext, UP
 	return Linker;
 }
 
+void FLinkerLoad::SetPackagePath(const FPackagePath& InPackagePath)
+{
+	PackagePath = InPackagePath;
+	PRAGMA_DISABLE_DEPRECATION_WARNINGS
+	Filename = PackagePath.GetLocalFullPath();
+	PRAGMA_ENABLE_DEPRECATION_WARNINGS
+}
+
 void FLinkerLoad::SetLoader(FArchive* InLoader)
 {
 	Loader = InLoader;
@@ -619,13 +633,13 @@ FName FLinkerLoad::FindSubobjectRedirectName(const FName& Name, UClass* Class)
  * Creates a FLinkerLoad object for async creation. Tick has to be called manually till it returns
  * true in which case the returned linker object has finished the async creation process.
  *
- * @param	Parent		Parent object to load into, can be NULL (most likely case)
- * @param	Filename	Name of file on disk to load
- * @param	LoadFlags	Load flags determining behavior
+ * @param	Parent				Parent object to load into, can be NULL (most likely case)
+ * @param	PackagePath			PackagePath to load from IPackageResourceManager
+ * @param	LoadFlags			Load flags determining behavior
  *
- * @return	new FLinkerLoad object for Parent/ Filename
+ * @return	new FLinkerLoad object for Parent/ PackagePath
  */
-FLinkerLoad* FLinkerLoad::CreateLinkerAsync(FUObjectSerializeContext* LoadContext, UPackage* Parent, const TCHAR* Filename, uint32 LoadFlags, const FLinkerInstancingContext* InstancingContext
+FLinkerLoad* FLinkerLoad::CreateLinkerAsync(FUObjectSerializeContext* LoadContext, UPackage* Parent, const FPackagePath& PackagePath, uint32 LoadFlags, const FLinkerInstancingContext* InstancingContext
 	, TFunction<void()>&& InSummaryReadyCallback
 	)
 {
@@ -652,7 +666,7 @@ FLinkerLoad* FLinkerLoad::CreateLinkerAsync(FUObjectSerializeContext* LoadContex
 		{
 			LoadFlags |= LOAD_Async;
 		}
-		Linker = new FLinkerLoad(Parent, Filename, LoadFlags, InstancingContext ? *InstancingContext : FLinkerInstancingContext());
+		Linker = new FLinkerLoad(Parent, PackagePath, LoadFlags, InstancingContext ? *InstancingContext : FLinkerInstancingContext());
 		Linker->SetSerializeContext(LoadContext);
 		Parent->LinkerLoad = Linker;
 		if (GEventDrivenLoaderEnabled && Linker)
@@ -841,12 +855,12 @@ FLinkerLoad::ELinkerStatus FLinkerLoad::Tick( float InTimeLimit, bool bInUseTime
 /**
  * Private constructor, passing arguments through from CreateLinker.
  *
- * @param	Parent		Parent object to load into, can be NULL (most likely case)
- * @param	Filename	Name of file on disk to load
- * @param	LoadFlags	Load flags determining behavior
+ * @param	Parent				Parent object to load into, can be NULL (most likely case)
+ * @param	PackagePath			PackagePath to load from IPackageResourceManager
+ * @param	LoadFlags			Load flags determining behavior
  */
-FLinkerLoad::FLinkerLoad(UPackage* InParent, const TCHAR* InFilename, uint32 InLoadFlags, FLinkerInstancingContext InInstancingContext)
-: FLinker(ELinkerType::Load, InParent, InFilename)
+FLinkerLoad::FLinkerLoad(UPackage* InParent, const FPackagePath& InPackagePath, uint32 InLoadFlags, FLinkerInstancingContext InInstancingContext)
+: FLinker(ELinkerType::Load, InParent)
 , LoadFlags(InLoadFlags)
 , bHaveImportsBeenVerified(false)
 , bDynamicClassLinker(false)
@@ -857,6 +871,7 @@ FLinkerLoad::FLinkerLoad(UPackage* InParent, const TCHAR* InFilename, uint32 InL
 , bIsDestroyingLoader(false)
 , StructuredArchive(nullptr)
 , StructuredArchiveFormatter(nullptr)
+, PackagePath(InPackagePath)
 , Loader(nullptr)
 , InstancingContext(MoveTemp(InInstancingContext))
 , AsyncRoot(nullptr)
@@ -882,13 +897,18 @@ FLinkerLoad::FLinkerLoad(UPackage* InParent, const TCHAR* InFilename, uint32 InL
 , TickStartTime(0.0)
 #if WITH_EDITOR
 , bExportsDuplicatesFixed(false)
-,	LoadProgressScope( nullptr )
+, LoadProgressScope( nullptr )
 #endif // WITH_EDITOR
 #if USE_CIRCULAR_DEPENDENCY_LOAD_DEFERRING
-,	bForceBlueprintFinalization(false)
-,	DeferredCDOIndex(INDEX_NONE)
+, bForceBlueprintFinalization(false)
+, DeferredCDOIndex(INDEX_NONE)
 #endif // USE_CIRCULAR_DEPENDENCY_LOAD_DEFERRING
 {
+	if (InPackagePath.GetHeaderExtension() == EPackageExtension::Unspecified)
+	{
+		UE_LOG(LogPackageName, Error, TEXT("PackagePath %s is missing header extension when assigned to LinkerLoad"), *InPackagePath.GetDebugName());
+	}
+
 	INC_DWORD_STAT(STAT_LinkerCount);
 	INC_DWORD_STAT(STAT_LiveLinkerCount);
 #if !UE_BUILD_SHIPPING && !UE_BUILD_TEST
@@ -898,7 +918,7 @@ FLinkerLoad::FLinkerLoad(UPackage* InParent, const TCHAR* InFilename, uint32 InL
 
 #if WITH_EDITOR
 	// Check if the linker is instanced @todo: pass through a load flag?
-	FName PackageNameToLoad = *FPackageName::FilenameToLongPackageName(InFilename);
+	FName PackageNameToLoad = GetPackagePath().GetPackageFName();
 	if (LinkerRoot->GetFName() != PackageNameToLoad)
 	{
 		InstancingContext.AddMapping(PackageNameToLoad, LinkerRoot->GetFName());
@@ -1010,10 +1030,10 @@ FLinkerLoad::ELinkerStatus FLinkerLoad::CreateLoader(
 		if (LoadProgressScope)
 		{
 			static const FTextFormat LoadingFileTextFormat = NSLOCTEXT("Core", "LoadingFileWithFilename", "Loading file: {CleanFilename}...");
-		FFormatNamedArguments FeedbackArgs;
-		FeedbackArgs.Add( TEXT("CleanFilename"), FText::FromString( FPaths::GetCleanFilename( *Filename ) ) );
+			FFormatNamedArguments FeedbackArgs;
+			FeedbackArgs.Add( TEXT("CleanFilename"), FText::FromString( FPaths::GetCleanFilename( *GetDebugName() ) ) );
 			LoadProgressScope->DefaultMessage = FText::Format(LoadingFileTextFormat, FeedbackArgs);
-		LoadProgressScope->EnterProgressFrame();
+			LoadProgressScope->EnterProgressFrame();
 		}
 #endif
 
@@ -1030,94 +1050,104 @@ FLinkerLoad::ELinkerStatus FLinkerLoad::CreateLoader(
 		}
 		else
 		{
-#if WITH_TEXT_ARCHIVE_SUPPORT
-			if (Filename.EndsWith(FPackageName::GetTextAssetPackageExtension()) || Filename.EndsWith(FPackageName::GetTextMapPackageExtension()))
+			// If want to be able to load cooked data in the editor we need to use FAsyncArchive which supports EDL cooked packages,
+			// otherwise the generic file reader is faster in the editor so use that
+			bool bCanUseAsyncLoader = FPlatformProperties::RequiresCookedData() || GAllowCookedDataInEditorBuilds;
+			if (bCanUseAsyncLoader)
 			{
-				INC_DWORD_STAT(STAT_TextAssetLinkerCount);
-				DECLARE_SCOPE_CYCLE_COUNTER(TEXT("FLinkerLoad::CreateTextArchiveFormatter"), STAT_LinkerLoad_CreateTextArchiveFormatter, STATGROUP_LinkerLoad);
-				TRACE_CPUPROFILER_EVENT_SCOPE(FLinkerLoad::CreateTextArchiveFormatter);
-				Loader = IFileManager::Get().CreateFileReader(*Filename);
-				StructuredArchiveFormatter = new FJsonArchiveInputFormatter(*this, [this](const FString& InFullPath)
+				Loader = new FAsyncArchive(GetPackagePath(), this,
+					GEventDrivenLoaderEnabled ? Forward<TFunction<void()>>(InSummaryReadyCallback) : TFunction<void()>([]() {})
+				);
+				if (Loader->IsError())
 				{
-					FPackageIndex Index = FindOrCreateImportOrExport(InFullPath);
-					if (Index.IsNull())
-					{
-						return (UObject*)nullptr;
-					}
-					else if (Index.IsImport())
-					{
-						return CreateImport(Index.ToImport());
-					}
-					else
-					{
-						check(Index.IsExport());
-						return CreateExport(Index.ToExport());
-					}
-				});
+					DestroyLoader();
+					UE_LOG(LogLinker, Warning, TEXT("Error opening file '%s'."), *GetDebugName());
+					return LINKER_Failed;
+				}
 			}
 			else
-#endif
 			{
-				// If want to be able to load cooked data in the editor we need to use FAsyncArchive which supports EDL cooked packages,
-				// otherwise the generic file reader is faster in the editor so use that
-				bool bCanUseAsyncLoader = FPlatformProperties::RequiresCookedData() || GAllowCookedDataInEditorBuilds;
-				if (bCanUseAsyncLoader)
-				{
-					Loader = new FAsyncArchive(*Filename, this,
-							GEventDrivenLoaderEnabled ? Forward<TFunction<void()>>(InSummaryReadyCallback) : TFunction<void()>([]() {})
-						);
-				}
+				FOpenPackageResult OpenResult;
 #if WITH_EDITOR
-				else if (FLinkerLoad::GetPreloadingEnabled() && FLinkerLoad::TryGetPreloadedLoader(Loader, *Filename))
+				if (FLinkerLoad::GetPreloadingEnabled() && FLinkerLoad::TryGetPreloadedLoader(GetPackagePath(), OpenResult))
 				{
 					// Loader set by TryGetPreloadedLoader
 				}
-#endif
 				else
+#endif
 				{
-					Loader = IFileManager::Get().CreateFileReader(*Filename);
+					OpenResult = IPackageResourceManager::Get().OpenReadPackage(GetPackagePath());
 				}
-
-				if (!Loader)
+				Loader = OpenResult.Archive.Release();
+				if (!Loader || Loader->IsError())
 				{
-					UE_LOG(LogLinker, Warning, TEXT("Error opening file '%s'."), *Filename);
+					if (Loader)
+					{
+						DestroyLoader();
+					}
+					UE_LOG(LogLinker, Warning, TEXT("Error opening file '%s'."), *GetDebugName());
 					return LINKER_Failed;
 				}
 
-				if (Loader->IsError())
+#if WITH_TEXT_ARCHIVE_SUPPORT
+				if (OpenResult.Format == EPackageFormat::Text)
 				{
-					DestroyLoader();
-					UE_LOG(LogLinker, Warning, TEXT("Error opening file '%s'."), *Filename);
-					return LINKER_Failed;
+					INC_DWORD_STAT(STAT_TextAssetLinkerCount);
+					DECLARE_SCOPE_CYCLE_COUNTER(TEXT("FLinkerLoad::CreateTextArchiveFormatter"), STAT_LinkerLoad_CreateTextArchiveFormatter, STATGROUP_LinkerLoad);
+					TRACE_CPUPROFILER_EVENT_SCOPE(FLinkerLoad::CreateTextArchiveFormatter);
+					StructuredArchiveFormatter = new FJsonArchiveInputFormatter(*this, [this](const FString& InFullPath)
+					{
+						FPackageIndex Index = FindOrCreateImportOrExport(InFullPath);
+						if (Index.IsNull())
+						{
+							return (UObject*)nullptr;
+						}
+						else if (Index.IsImport())
+						{
+							return CreateImport(Index.ToImport());
+						}
+						else
+						{
+							check(Index.IsExport());
+							return CreateExport(Index.ToExport());
+						}
+					});
 				}
+				else
+#endif
+				{
+					check(OpenResult.Format == EPackageFormat::Binary);
+				}
+			}
+
 #if DEVIRTUALIZE_FLinkerLoad_Serialize
-				ActiveFPLB = Loader->ActiveFPLB; // make sure my fast past loading is using the FAA2 fast path buffer
+			ActiveFPLB = Loader->ActiveFPLB; // make sure my fast past loading is using the FAA2 fast path buffer
 #endif
 
-				bool bHasHashEntry = FSHA1::GetFileSHAHash(*Filename, NULL);
-				if ((LoadFlags & LOAD_MemoryReader) || bHasHashEntry)
+			bool bHasHashEntry = FSHA1::GetFileSHAHash(*GetPackagePath().GetLocalFullPath(), nullptr);
+			if ((LoadFlags & LOAD_MemoryReader) || bHasHashEntry)
+			{
+				// force preload into memory if file has an SHA entry
+				// Serialize data from memory instead of from disk.
+				uint32	BufferSize = Loader->TotalSize();
+				void* Buffer = FMemory::Malloc(BufferSize);
+				Loader->Serialize(Buffer, BufferSize);
+				DestroyLoader();
+				if (bHasHashEntry)
 				{
-					// force preload into memory if file has an SHA entry
-					// Serialize data from memory instead of from disk.
-					uint32	BufferSize = Loader->TotalSize();
-					void*	Buffer = FMemory::Malloc(BufferSize);
-					Loader->Serialize(Buffer, BufferSize);
-					DestroyLoader();
-					if (bHasHashEntry)
-					{
-						// create buffer reader and spawn SHA verify when it gets closed
-						Loader = new FBufferReaderWithSHA(Buffer, BufferSize, true, *Filename, true);
-					}
-					else
-					{
-						// create a buffer reader
-						Loader = new FBufferReader(Buffer, BufferSize, true, true);
-					}
+					// create buffer reader and spawn SHA verify when it gets closed
+					Loader = new FBufferReaderWithSHA(Buffer, BufferSize, true, *GetPackagePath().GetLocalFullPath(), true);
 				}
 				else
 				{
-					bIsAsyncLoader = bCanUseAsyncLoader;
+					// create a buffer reader
+					Loader = new FBufferReader(Buffer, BufferSize, true, true);
 				}
+				bIsAsyncLoader = false;
+			}
+			else
+			{
+				bIsAsyncLoader = bCanUseAsyncLoader;
 			}
 		} 
 
@@ -1157,7 +1187,7 @@ FLinkerLoad::ELinkerStatus FLinkerLoad::CreateLoader(
 				if (Size <= 0)
 				{
 					DestroyLoader();
-					UE_LOG(LogLinker, Warning, TEXT("Error opening file '%s'."), *Filename);
+					UE_LOG(LogLinker, Warning, TEXT("Error opening file '%s'."), *GetDebugName());
 					return LINKER_Failed;
 				}
 				// Precache up to one ECC block before serializing package file summary.
@@ -1191,20 +1221,22 @@ FLinkerLoad::ELinkerStatus FLinkerLoad::SerializePackageFileSummaryInternal()
 	// Check tag.
 	if (Summary.Tag != PACKAGE_FILE_TAG)
 	{
-		UE_LOG(LogLinker, Warning, TEXT("The file '%s' contains unrecognizable data, check that it is of the expected type."), *Filename);
+		UE_LOG(LogLinker, Warning, TEXT("The file '%s' contains unrecognizable data, check that it is of the expected type."), *GetDebugName());
 		return LINKER_Failed;
 	}
 
 	// Validate the summary.
 	if (Summary.GetFileVersionUE4() < VER_UE4_OLDEST_LOADABLE_PACKAGE)
 	{
-		UE_LOG(LogLinker, Warning, TEXT("The file %s was saved by a previous version which is not backwards compatible with this one. Min Required Version: %i  Package Version: %i"), *Filename, (int32)VER_UE4_OLDEST_LOADABLE_PACKAGE, Summary.GetFileVersionUE4());
+		UE_LOG(LogLinker, Warning, TEXT("The file %s was saved by a previous version which is not backwards compatible with this one. Min Required Version: %i  Package Version: %i"),
+			*GetDebugName(), (int32)VER_UE4_OLDEST_LOADABLE_PACKAGE, Summary.GetFileVersionUE4());
 		return LINKER_Failed;
 	}
 	// Don't load packages that are only compatible with an engine version newer than the current one.
 	if (GEnforcePackageCompatibleVersionCheck && !FEngineVersion::Current().IsCompatibleWith(Summary.CompatibleWithEngineVersion))
 	{
-		UE_LOG(LogLinker, Warning, TEXT("Asset '%s' has been saved with engine version newer than current and therefore can't be loaded. CurrEngineVersion: %s AssetEngineVersion: %s"), *Filename, *FEngineVersion::Current().ToString(), *Summary.CompatibleWithEngineVersion.ToString());
+		UE_LOG(LogLinker, Warning, TEXT("Asset '%s' has been saved with engine version newer than current and therefore can't be loaded. CurrEngineVersion: %s AssetEngineVersion: %s"),
+			*GetDebugName(), *FEngineVersion::Current().ToString(), *Summary.CompatibleWithEngineVersion.ToString());
 		return LINKER_Failed;
 	}
 
@@ -1228,20 +1260,21 @@ FLinkerLoad::ELinkerStatus FLinkerLoad::SerializePackageFileSummaryInternal()
 			}
 			FORCEINLINE operator bool() const { return bDoWarn; }
 		} ZeroEngineVersionWarningEnabled;
-		UE_CLOG(ZeroEngineVersionWarningEnabled, LogLinker, Warning, TEXT("Asset '%s' has been saved with empty engine version. The asset will be loaded but may be incompatible."), *Filename);
+		UE_CLOG(ZeroEngineVersionWarningEnabled, LogLinker, Warning, TEXT("Asset '%s' has been saved with empty engine version. The asset will be loaded but may be incompatible."), *GetDebugName());
 	}
 
 	// Don't load packages that were saved with package version newer than the current one.
 	if ((Summary.GetFileVersionUE4() > GPackageFileUE4Version) || (Summary.GetFileVersionLicenseeUE4() > GPackageFileLicenseeUE4Version))
 	{
-		UE_LOG(LogLinker, Warning, TEXT("Unable to load package (%s) PackageVersion %i, MaxExpected %i : LicenseePackageVersion %i, MaxExpected %i."), *Filename, Summary.GetFileVersionUE4(), GPackageFileUE4Version, Summary.GetFileVersionLicenseeUE4(), GPackageFileLicenseeUE4Version);
+		UE_LOG(LogLinker, Warning, TEXT("Unable to load package (%s) PackageVersion %i, MaxExpected %i : LicenseePackageVersion %i, MaxExpected %i."),
+			*GetDebugName(), Summary.GetFileVersionUE4(), GPackageFileUE4Version, Summary.GetFileVersionLicenseeUE4(), GPackageFileLicenseeUE4Version);
 		return LINKER_Failed;
 	}
 
 	// don't load packages that contain editor only data in builds that don't support that and vise versa
 	if (!FPlatformProperties::HasEditorOnlyData() && !(Summary.PackageFlags & PKG_FilterEditorOnly))
 	{
-		UE_LOG(LogLinker, Warning, TEXT("Unable to load package (%s). Package contains EditorOnly data which is not supported by the current build."), *Filename);
+		UE_LOG(LogLinker, Warning, TEXT("Unable to load package (%s). Package contains EditorOnly data which is not supported by the current build."), *GetDebugName());
 		return LINKER_Failed;
 	}
 
@@ -1253,7 +1286,7 @@ FLinkerLoad::ELinkerStatus FLinkerLoad::SerializePackageFileSummaryInternal()
 		{
 			UE_LOG(LogLinker, Warning,
 				TEXT("Unable to load package (%s). Package contains cooked data which is not supported by the current build. Enable 'Allow Cooked Content In The Editor' in Project Settings under 'Engine - Cooker' section to load it."),
-				*Filename);
+				*GetDebugName());
 			return LINKER_Failed;
 		}
 	}
@@ -1286,7 +1319,7 @@ FLinkerLoad::ELinkerStatus FLinkerLoad::SerializePackageFileSummaryInternal()
 
 		if (Tag != PACKAGE_FILE_TAG)
 		{
-			UE_LOG(LogLinker, Warning, TEXT("Unable to load package (%s). Post Tag is not valid. File might be corrupted."), *Filename);
+			UE_LOG(LogLinker, Warning, TEXT("Unable to load package (%s). Post Tag is not valid. File might be corrupted."), *GetDebugName());
 			return LINKER_Failed;
 		}
 
@@ -1309,7 +1342,7 @@ FLinkerLoad::ELinkerStatus FLinkerLoad::SerializePackageFileSummary()
 	{
 		if (Loader->IsError())
 		{
-			UE_LOG(LogLinker, Warning, TEXT("The file '%s' contains unrecognizable data, check that it is of the expected type."), *Filename);
+			UE_LOG(LogLinker, Warning, TEXT("The file '%s' contains unrecognizable data, check that it is of the expected type."), *GetDebugName());
 			return LINKER_Failed;
 		}
 		if (bIsAsyncLoader)
@@ -1359,18 +1392,20 @@ FLinkerLoad::ELinkerStatus FLinkerLoad::UpdateFromPackageFileSummary()
 			{
 				// Loading a package with custom integration that we don't know about!
 				// Temporarily just warn and continue. @todo: this needs to be fixed properly
-				UE_LOG(LogLinker, Warning, TEXT("Package %s was saved with a custom integration that is not present. Tag %s  Version %d"), *Filename, *Diff.Version->Key.ToString(), Diff.Version->Version);
+				UE_LOG(LogLinker, Warning, TEXT("Package %s was saved with a custom integration that is not present. Tag %s  Version %d"), *GetDebugName(),
+					*Diff.Version->Key.ToString(), Diff.Version->Version);
 			}
 			else if (Diff.Type == ECustomVersionDifference::Invalid)
 			{
-				UE_LOG(LogLinker, Error, TEXT("Package %s was saved with an invalid custom version. Tag %s  Version %d"), *Filename, *Diff.Version->Key.ToString(), Diff.Version->Version);
+				UE_LOG(LogLinker, Error, TEXT("Package %s was saved with an invalid custom version. Tag %s  Version %d"), *GetDebugName(), *Diff.Version->Key.ToString(), Diff.Version->Version);
 				return LINKER_Failed;
 			}
 			else if (Diff.Type == ECustomVersionDifference::Newer)
 			{
 				FCustomVersion LatestVersion = FCurrentCustomVersions::Get(Diff.Version->Key).GetValue();
 				// Loading a package with a newer custom version than the current one.
-				UE_LOG(LogLinker, Error, TEXT("Package %s was saved with a newer custom version than the current. Tag %s Name '%s' PackageVersion %d  MaxExpected %d"), *Filename, *Diff.Version->Key.ToString(), *LatestVersion.GetFriendlyName().ToString(), Diff.Version->Version, LatestVersion.Version);
+				UE_LOG(LogLinker, Error, TEXT("Package %s was saved with a newer custom version than the current. Tag %s Name '%s' PackageVersion %d  MaxExpected %d"), *GetDebugName(),
+					*Diff.Version->Key.ToString(), *LatestVersion.GetFriendlyName().ToString(), Diff.Version->Version, LatestVersion.Version);
 				return LINKER_Failed;
 			}
 		}
@@ -2314,7 +2349,7 @@ FLinkerLoad::ELinkerStatus FLinkerLoad::FindExistingExports()
 #if WITH_EDITOR
 		if (LoadProgressScope)
 		{
-		LoadProgressScope->EnterProgressFrame(1);
+			LoadProgressScope->EnterProgressFrame(1);
 		}
 		if( GIsEditor && GIsRunning )
 		{
@@ -2323,7 +2358,7 @@ FLinkerLoad::ELinkerStatus FLinkerLoad::FindExistingExports()
 			// the objects in those cases, so don't try to find existing exports.
 			//
 			bool bContainsMap			= LinkerRoot ? LinkerRoot->ContainsMap() : false;
-			bool bRequestFindExisting = FCoreUObjectDelegates::ShouldLoadOnTop.IsBound() ? !FCoreUObjectDelegates::ShouldLoadOnTop.Execute(Filename) : true;
+			bool bRequestFindExisting = FCoreUObjectDelegates::ShouldLoadOnTop.IsBound() ? !FCoreUObjectDelegates::ShouldLoadOnTop.Execute(GetPackagePath().GetLocalFullPath()) : true;
 			if( (!IsRunningCommandlet() && bRequestFindExisting && !bContainsMap) )
 			{
 				for (int32 ExportIndex = 0; ExportIndex < ExportMap.Num(); ExportIndex++)
@@ -2593,7 +2628,7 @@ FName FLinkerLoad::GetExportClassPackage( int32 i )
 
 FString FLinkerLoad::GetArchiveName() const
 {
-	return *Filename;
+	return GetPackagePath().GetDebugName();
 }
 
 
@@ -2722,7 +2757,7 @@ void FLinkerLoad::GatherImportDependencies(int32 ImportIndex, TSet<FDependencyRe
 				UE_LOG(LogLinker, Warning, TEXT("VerifyImportInner failed [(%x, %d), (%x, %d)] for %s with linker: %s"), 
 					Import.XObject, Import.XObject ? (Import.XObject->IsNative() ? 1 : 0) : 0, 
 					Import.SourceLinker, Import.SourceIndex, 
-					*GetImportFullName(ImportIndex), *this->Filename );
+					*GetImportFullName(ImportIndex), *GetDebugName() );
 			}
 			EndLoad(LoadContext);
 			return;
@@ -2891,7 +2926,7 @@ FLinkerLoad::EVerifyResult FLinkerLoad::VerifyImport(int32 ImportIndex)
 			// then do the throw here
 			if (bCrashOnFail)
 			{
-				UE_LOG(LogLinker, Fatal,  TEXT("Failed import: %s %s (file %s)"), *Import.ClassName.ToString(), *GetImportFullName(ImportIndex), *Import.SourceLinker->Filename );
+				UE_LOG(LogLinker, Fatal,  TEXT("Failed import: %s %s (file %s)"), *Import.ClassName.ToString(), *GetImportFullName(ImportIndex), *Import.SourceLinker->GetDebugName() );
 				return Result;
 			}
 			// otherwise just printout warnings, and if in the editor, popup the EdLoadWarnings box
@@ -2953,7 +2988,7 @@ FLinkerLoad::EVerifyResult FLinkerLoad::VerifyImport(int32 ImportIndex)
 }
 
 // Internal Load package call so that we can pass the linker that requested this package as an import dependency
-UPackage* LoadPackageInternal(UPackage* InOuter, const TCHAR* InLongPackageName, uint32 LoadFlags, FLinkerLoad* ImportLinker, FArchive* InReaderOverride, const FLinkerInstancingContext* InstancingContext);
+UPackage* LoadPackageInternal(UPackage* InOuter, const FPackagePath& PackagePath, uint32 LoadFlags, FLinkerLoad* ImportLinker, FArchive* InReaderOverride, const FLinkerInstancingContext* InstancingContext, const FPackagePath* DiffPackagePath);
 
 /**
  * Safely verify that an import in the ImportMap points to a good object. This decides whether or not
@@ -3017,7 +3052,9 @@ bool FLinkerLoad::VerifyImportInner(const int32 ImportIndex, FString& WarningSuf
 			{
 				Package = CreatePackage(*PackageToLoadInto.ToString());
 			}
-			Package = LoadPackageInternal(Package, *PackageToLoad.ToString(), InternalLoadFlags | LOAD_IsVerifying, this, nullptr, nullptr);
+			FPackagePath PackagePathToLoad = FPackagePath::FromPackageNameChecked(PackageToLoad);
+			Package = LoadPackageInternal(Package, PackagePathToLoad, InternalLoadFlags | LOAD_IsVerifying, this, nullptr /* InReaderOverride */,
+				nullptr /* InstancingContext */, nullptr /* DiffPackagePath */);
 		}
 
 #if WITH_EDITOR
@@ -3057,7 +3094,14 @@ bool FLinkerLoad::VerifyImportInner(const int32 ImportIndex, FString& WarningSuf
 		// or when circular dependency happen, get the linker so we are able to create the import properly at a later time.
 		// When loading editor data never consider the package fully loaded and resolve the linker anyway, for cooked data, assign the linker if one is associated witht the package
 		const bool bWasFullyLoaded = Package && Package->IsFullyLoaded() && FPlatformProperties::RequiresCookedData();
-		Import.SourceLinker = !bWasFullyLoaded ? GetPackageLinker(Package, nullptr, InternalLoadFlags, nullptr, nullptr, nullptr, &SerializeContext) : FindExistingLinkerForPackage(Package);
+		if (!bWasFullyLoaded)
+		{
+			Import.SourceLinker = GetPackageLinker(Package, FPackagePath::FromPackageNameChecked(Package->GetName()), InternalLoadFlags, nullptr, nullptr, nullptr, &SerializeContext);
+		}
+		else
+		{
+			Import.SourceLinker = FindExistingLinkerForPackage(Package);
+		}
 #if WITH_EDITORONLY_DATA
 		if (Import.SourceLinker && !Package->HasAnyFlags(RF_LoadCompleted))
 		{
@@ -3103,7 +3147,7 @@ bool FLinkerLoad::VerifyImportInner(const int32 ImportIndex, FString& WarningSuf
 	// Find or load the linker load that contains the FObjectExport for this import
 	if (Import.OuterIndex.IsNull() && Import.ClassName!=NAME_Package )
 	{
-		UE_LOG(LogLinker, Error, TEXT("%s has an inappropriate outermost, it was probably saved with a deprecated outer (file: %s)"), *ImportObjectName, *Filename);
+		UE_LOG(LogLinker, Error, TEXT("%s has an inappropriate outermost, it was probably saved with a deprecated outer (file: %s)"), *ImportObjectName, *GetDebugName());
 		Import.SourceLinker = NULL;
 		return false;
 	}
@@ -4573,7 +4617,7 @@ UObject* FLinkerLoad::CreateExport( int32 Index )
 		if ( LoadClass->HasAnyClassFlags(CLASS_Transient) && (Export.ObjectFlags&RF_ClassDefaultObject) == 0 && (Export.ObjectFlags&RF_ArchetypeObject) == 0 )
 		{
 			FFormatNamedArguments Arguments;
-			Arguments.Add(TEXT("PackageName"), FText::FromString(Filename));
+			Arguments.Add(TEXT("PackageName"), GetPackagePath().GetDebugNameText());
 			Arguments.Add(TEXT("ObjectName"), FText::FromName(Export.ObjectName));
 			Arguments.Add(TEXT("ClassName"), FText::FromString(LoadClass->GetPathName()));
 			//@todo - should this actually be an assertion?
@@ -5105,7 +5149,7 @@ UObject* FLinkerLoad::IndexToObject( FPackageIndex Index )
 					ErrorCaption,
 					GEngineIni);
 
-				UE_LOG( LogLinker, Error, TEXT("Invalid export object index=%d while reading %s. File is most likely corrupted. Please verify your installation."), Index.ToExport(), *Filename );
+				UE_LOG( LogLinker, Error, TEXT("Invalid export object index=%d while reading %s. File is most likely corrupted. Please verify your installation."), Index.ToExport(), *GetDebugName() );
 
 				if (GLog)
 				{
@@ -5118,7 +5162,7 @@ UObject* FLinkerLoad::IndexToObject( FPackageIndex Index )
 			}
 		#else
 			{
-				UE_CLOG( !ExportMap.IsValidIndex( Index.ToExport() ), LogLinker, Fatal, TEXT("Invalid export object index=%d while reading %s. File is most likely corrupted. Please verify your installation."), Index.ToExport(), *Filename );
+				UE_CLOG( !ExportMap.IsValidIndex( Index.ToExport() ), LogLinker, Fatal, TEXT("Invalid export object index=%d while reading %s. File is most likely corrupted. Please verify your installation."), Index.ToExport(), *GetDebugName() );
 			}
 		#endif
 
@@ -5140,7 +5184,7 @@ UObject* FLinkerLoad::IndexToObject( FPackageIndex Index )
 					ErrorCaption,
 					GEngineIni);
 
-				UE_LOG( LogLinker, Error, TEXT("Invalid import object index=%d while reading %s. File is most likely corrupted. Please verify your installation."), Index.ToImport(), *Filename );
+				UE_LOG( LogLinker, Error, TEXT("Invalid import object index=%d while reading %s. File is most likely corrupted. Please verify your installation."), Index.ToImport(), *GetDebugName() );
 
 				if (GLog)
 				{
@@ -5153,7 +5197,7 @@ UObject* FLinkerLoad::IndexToObject( FPackageIndex Index )
 			}
 		#else
 			{
-				UE_CLOG( !ImportMap.IsValidIndex( Index.ToImport() ), LogLinker, Fatal, TEXT("Invalid import object index=%d while reading %s. File is most likely corrupted. Please verify your installation."), Index.ToImport(), *Filename );
+				UE_CLOG( !ImportMap.IsValidIndex( Index.ToImport() ), LogLinker, Fatal, TEXT("Invalid import object index=%d while reading %s. File is most likely corrupted. Please verify your installation."), Index.ToImport(), *GetDebugName() );
 			}
 		#endif
 
@@ -5309,7 +5353,7 @@ void FLinkerLoad::DetachBulkData( FUntypedBulkData* BulkData, bool bEnsureBulkDa
 	int32 RemovedCount = BulkDataLoaders.Remove( BulkData );
 	if( RemovedCount!=1 )
 	{	
-		UE_LOG(LogLinker, Fatal, TEXT("Detachment inconsistency: %i (%s)"), RemovedCount, *Filename );
+		UE_LOG(LogLinker, Fatal, TEXT("Detachment inconsistency: %i (%s)"), RemovedCount, *GetDebugName() );
 	}
 	BulkData->DetachFromArchive( this, bEnsureBulkDataIsLoaded );
 }
@@ -5986,9 +6030,20 @@ void FLinkerLoad::SetPreloadingEnabled(bool bEnabled)
 }
 bool FLinkerLoad::TryGetPreloadedLoader(FArchive*& OutLoader, const TCHAR* FileName)
 {
-	OutLoader = FPreloadableFile::TryTakeArchive(FileName);
-	return OutLoader != nullptr;
+	FOpenPackageResult Result;
+	if (IPackageResourceManager::TryTakePreloadableArchive(FPackagePath::FromLocalPath(FileName), Result))
+	{
+		OutLoader = Result.Archive.Release();
+		return true;
+	}
+	return false;
 }
+
+bool FLinkerLoad::TryGetPreloadedLoader(const FPackagePath& InPackagePath, FOpenPackageResult& OutResult)
+{
+	return IPackageResourceManager::TryTakePreloadableArchive(InPackagePath, OutResult);
+}
+
 #endif
 
 #undef LOCTEXT_NAMESPACE
diff --git a/Engine/Source/Runtime/CoreUObject/Private/UObject/LinkerManager.cpp b/Engine/Source/Runtime/CoreUObject/Private/UObject/LinkerManager.cpp
--- a/Engine/Source/Runtime/CoreUObject/Private/UObject/LinkerManager.cpp
+++ b/Engine/Source/Runtime/CoreUObject/Private/UObject/LinkerManager.cpp
@@ -34,19 +34,19 @@ bool FLinkerManager::Exec(class UWorld* InWorld, const TCHAR* Cmd, FOutputDevice
 		UE_LOG(LogLinker, Display, TEXT("ObjectLoaders: %d"), ObjectLoaders.Num());
 		for (auto Linker : ObjectLoaders)
 		{
-			UE_LOG(LogLinker, Display, TEXT("%s"), *Linker->Filename);
+			UE_LOG(LogLinker, Display, TEXT("%s"), *Linker->GetDebugName());
 		}
 
 		UE_LOG(LogLinker, Display, TEXT("LoadersWithNewImports: %d"), LoadersWithNewImports.Num());
 		for (auto Linker : LoadersWithNewImports)
 		{
-			UE_LOG(LogLinker, Display, TEXT("%s"), *Linker->Filename);
+			UE_LOG(LogLinker, Display, TEXT("%s"), *Linker->GetDebugName());
 		}
 #if !UE_BUILD_SHIPPING && !UE_BUILD_TEST
 		UE_LOG(LogLinker, Display, TEXT("LiveLinkers: %d"), LiveLinkers.Num());
 		for (auto Linker : LiveLinkers)
 		{
-			UE_LOG(LogLinker, Display, TEXT("%s"), *Linker->Filename);
+			UE_LOG(LogLinker, Display, TEXT("%s"), *Linker->GetDebugName());
 		}
 #endif
 		return true;
@@ -67,7 +67,7 @@ bool FLinkerManager::Exec(class UWorld* InWorld, const TCHAR* Cmd, FOutputDevice
 			Ar.Logf
 				(
 				TEXT("%s (%s): Names=%i (%iK/%iK) Text=%i (%iK) Imports=%i (%iK) Exports=%i (%iK) Gen=%i Bulk=%i"),
-				*Linker->Filename,
+				*Linker->GetDebugName(),
 				*Linker->LinkerRoot->GetFullName(),
 				Linker->NameMap.Num(),
 				Linker->NameMap.Num() * sizeof(FName) / 1024,
diff --git a/Engine/Source/Runtime/CoreUObject/Private/UObject/LinkerSave.cpp b/Engine/Source/Runtime/CoreUObject/Private/UObject/LinkerSave.cpp
--- a/Engine/Source/Runtime/CoreUObject/Private/UObject/LinkerSave.cpp
+++ b/Engine/Source/Runtime/CoreUObject/Private/UObject/LinkerSave.cpp
@@ -22,9 +22,11 @@
 TMap<FString, TArray<uint8> > FLinkerSave::PackagesToScriptSHAMap;
 
 FLinkerSave::FLinkerSave(UPackage* InParent, const TCHAR* InFilename, bool bForceByteSwapping, bool bInSaveUnversioned)
-:	FLinker(ELinkerType::Save, InParent, InFilename)
+:	FLinker(ELinkerType::Save, InParent)
 ,	Saver(nullptr)
 {
+	SetFilename(InFilename);
+
 	if (FPlatformProperties::HasEditorOnlyData())
 	{
 		// Create file saver.
@@ -74,9 +76,10 @@ FLinkerSave::FLinkerSave(UPackage* InParent, const TCHAR* InFilename, bool bForc
 
 
 FLinkerSave::FLinkerSave(UPackage* InParent, FArchive *InSaver, bool bForceByteSwapping, bool bInSaveUnversioned)
-: FLinker(ELinkerType::Save, InParent, TEXT("$$Memory$$"))
+: FLinker(ELinkerType::Save, InParent)
 , Saver(nullptr)
 {
+	SetFilename(TEXT("$$Memory$$"));
 	if (FPlatformProperties::HasEditorOnlyData())
 	{
 		// Create file saver.
@@ -126,13 +129,14 @@ FLinkerSave::FLinkerSave(UPackage* InParent, FArchive *InSaver, bool bForceByteS
 }
 
 FLinkerSave::FLinkerSave(UPackage* InParent, bool bForceByteSwapping, bool bInSaveUnversioned )
-:	FLinker(ELinkerType::Save, InParent, TEXT("$$Memory$$"))
+:	FLinker(ELinkerType::Save, InParent)
 ,	Saver(nullptr)
 {
+	SetFilename(TEXT("$$Memory$$"));
 	if (FPlatformProperties::HasEditorOnlyData())
 	{
 		// Create file saver.
-		Saver = new FLargeMemoryWriter( 0, false, *InParent->FileName.ToString() );
+		Saver = new FLargeMemoryWriter( 0, false, *InParent->GetLoadedPath().GetDebugName() );
 		check(Saver);
 
 		UPackage* Package = dynamic_cast<UPackage*>(LinkerRoot);
@@ -270,6 +274,26 @@ void FLinkerSave::Serialize( void* V, int64 Length )
 	Saver->Serialize( V, Length );
 }
 	
+FString FLinkerSave::GetDebugName() const
+{
+	return GetFilename();
+}
+
+const FString& FLinkerSave::GetFilename() const
+{
+	// When the deprecated Filename is removed from FLinker, add a separate Filename variable to FLinkerSave
+	PRAGMA_DISABLE_DEPRECATION_WARNINGS
+	return Filename;
+	PRAGMA_ENABLE_DEPRECATION_WARNINGS
+}
+
+void FLinkerSave::SetFilename(FStringView InFilename)
+{
+	// When the deprecated Filename is removed from FLinker, add a separate Filename variable to FLinkerSave
+	PRAGMA_DISABLE_DEPRECATION_WARNINGS
+	Filename = FString(InFilename);
+	PRAGMA_ENABLE_DEPRECATION_WARNINGS
+}
 
 FString FLinkerSave::GetArchiveName() const
 {
diff --git a/Engine/Source/Runtime/CoreUObject/Private/UObject/Package.cpp b/Engine/Source/Runtime/CoreUObject/Private/UObject/Package.cpp
--- a/Engine/Source/Runtime/CoreUObject/Private/UObject/Package.cpp
+++ b/Engine/Source/Runtime/CoreUObject/Private/UObject/Package.cpp
@@ -7,6 +7,7 @@
 #include "Misc/PackageName.h"
 #include "UObject/LinkerLoad.h"
 #include "UObject/LinkerManager.h"
+#include "UObject/PackageResourceManager.h"
 #include "UObject/UObjectHash.h"
 #include "UObject/UObjectThreadContext.h"
 
@@ -201,6 +202,19 @@ void UPackage::FullyLoad()
 	}
 }
 
+const FPackagePath& UPackage::GetLoadedPath() const
+{
+	return LoadedPath;
+}
+
+void UPackage::SetLoadedPath(const FPackagePath& InPackagePath)
+{
+	LoadedPath = InPackagePath;
+	PRAGMA_DISABLE_DEPRECATION_WARNINGS
+	FileName = InPackagePath.GetPackageFName();
+	PRAGMA_ENABLE_DEPRECATION_WARNINGS
+}
+
 /** Tags generated objects with flags */
 void UPackage::TagSubobjects(EObjectFlags NewFlags)
 {
@@ -221,32 +235,51 @@ void UPackage::TagSubobjects(EObjectFlags NewFlags)
  */
 bool UPackage::IsFullyLoaded() const
 {
-	// Newly created packages aren't loaded and therefore haven't been marked as being fully loaded. They are treated as fully
-	// loaded packages though in this case, which is why we are looking to see whether the package exists on disk and assume it
-	// has been fully loaded if it doesn't.
-	if (!bHasBeenFullyLoaded && !HasAnyInternalFlags(EInternalObjectFlags::AsyncLoading) && FileSize == 0)
+	if (bHasBeenFullyLoaded)
 	{
+		return true;
+	}
+
+	// We set bHasBeenFullyLoaded to true when it is read for some special cases
+
+	if (FileSize != 0)
+	{
+		// If it has a filesize, it is a normal on-disk package, therefore is not a special case, and we respect the current 'false' value of bHasBeenFullyLoaded
+		return false;
+	}
+
+	if (HasAnyInternalFlags(EInternalObjectFlags::AsyncLoading))
+	{
+		// If it's in the middle of an async load, don't make any changes and respect the current 'false' value of bHasBeenFullyLoaded
+		return false;
+	}
+
+	if (HasAnyPackageFlags(PKG_CompiledIn))
+	{
+		// Native packages don't have a file size but are always considered fully loaded.
+		bHasBeenFullyLoaded = true;
+		return true;
+	}
+
+	if (!GetConvertedDynamicPackageNameToTypeName().Contains(GetFName())) // ConvertedDynamicPackageNames do not exist on disk, but are not a special case
+	{
+		// Newly created packages aren't loaded and therefore haven't been marked as being fully loaded. They are treated as fully
+		// loaded packages though in this case, which is why we are looking to see whether the package exists on disk and assume it
+		// has been fully loaded if it doesn't.
+		// Try to find matching package in package file cache. We use the LoadedPath here as it may be loaded into a temporary package
 		FString DummyFilename;
-		FString SourcePackageName = FileName != NAME_None ? FileName.ToString() : GetName();
-		// Try to find matching package in package file cache. We use the source package name here as it may be loaded into a temporary package
-		if (HasAnyPackageFlags(PKG_CompiledIn))
-		{
-			// Native packages don't have a file size but are always considered fully loaded.
-			bHasBeenFullyLoaded = true;
-		}
-		else if (	!GetConvertedDynamicPackageNameToTypeName().Contains(GetFName()) &&
-				(
-					!FPackageName::DoesPackageExist(*SourcePackageName, NULL, &DummyFilename ) ||
-					(GIsEditor && IFileManager::Get().FileSize(*DummyFilename) < 0) 
-				)
-			)
+		FPackagePath SourcePackagePath = !LoadedPath.IsEmpty() ? LoadedPath : FPackagePath::FromPackageNameChecked(GetName());
+		if (!FPackageName::DoesPackageExist(SourcePackagePath, &SourcePackagePath) ||
+			(GIsEditor && IPackageResourceManager::Get().FileSize(SourcePackagePath) < 0))
 		{
 			// Package has NOT been found, so we assume it's a newly created one and therefore fully loaded.
 			bHasBeenFullyLoaded = true;
+			return true;
 		}
 	}
 
-	return bHasBeenFullyLoaded;
+	// Not a special case; respect the current 'false' value of bHasBeenFullyLoaded
+	return false;
 }
 
 void UPackage::BeginDestroy()
diff --git a/Engine/Source/Runtime/CoreUObject/Private/UObject/SavePackage.cpp b/Engine/Source/Runtime/CoreUObject/Private/UObject/SavePackage.cpp
--- a/Engine/Source/Runtime/CoreUObject/Private/UObject/SavePackage.cpp
+++ b/Engine/Source/Runtime/CoreUObject/Private/UObject/SavePackage.cpp
@@ -2053,9 +2053,8 @@ FSavePackageResultStruct UPackage::Save(UPackage* InOuter, UObject* Base, EObjec
 
 	// if the in memory package filename is different the filename we are saving it to,
 	// regenerate a new persistent id for it.
-	FString PackageFilename(Filename);
-	bool bIsValidLongPackageName = FPackageName::TryConvertFilenameToLongPackageName(PackageFilename, PackageFilename);
-	if (!bIsCooking && !InOuter->FileName.IsNone() && InOuter->FileName.ToString() != PackageFilename && !(SaveFlags & SAVE_FromAutosave))
+	FPackagePath TargetPackagePath = FPackagePath::FromLocalPath(Filename);
+	if (!bIsCooking && !InOuter->GetLoadedPath().IsEmpty() && InOuter->GetLoadedPath() != TargetPackagePath && !(SaveFlags & SAVE_FromAutosave))
 	{
 		InOuter->SetPersistentGuid(FGuid::NewGuid());
 	}
@@ -2401,13 +2400,13 @@ FSavePackageResultStruct UPackage::Save(UPackage* InOuter, UObject* Base, EObjec
 					{
 						// The entire package will be serialized to memory and then compared against package on disk.
 						// Each difference will be log with its Serialize call stack trace
-						FArchive* Saver = new FArchiveStackTrace(FindAssetInPackage(InOuter), *InOuter->FileName.ToString(), true, InOutDiffMap);
+						FArchive* Saver = new FArchiveStackTrace(FindAssetInPackage(InOuter), *InOuter->GetLoadedPath().GetPackageName(), true, InOutDiffMap);
 						Linker = TUniquePtr<FLinkerSave>(new FLinkerSave(InOuter, Saver, bForceByteSwapping, bSaveUnversioned));
 					}
 					else if (TargetPlatform != nullptr && (SaveFlags & SAVE_DiffOnly))
 					{
 						// The entire package will be serialized to memory and then compared against package on disk
-						FArchive* Saver = new FArchiveStackTrace(FindAssetInPackage(InOuter), *InOuter->FileName.ToString(), false);
+						FArchive* Saver = new FArchiveStackTrace(FindAssetInPackage(InOuter), *InOuter->GetLoadedPath().GetPackageName(), false);
 						Linker = TUniquePtr<FLinkerSave>(new FLinkerSave(InOuter, Saver, bForceByteSwapping, bSaveUnversioned));
 					}
 					else if ((!!TargetPlatform) && FParse::Value(FCommandLine::Get(), TEXT("DiffCookedPackages="), DiffCookedPackagesPath))
@@ -2422,7 +2421,7 @@ FSavePackageResultStruct UPackage::Save(UPackage* InOuter, UObject* Base, EObjec
 						TestArchiveFilename.ReplaceInline(*CookedPath, *DiffCookedPackagesPath);
 
 						FArchive* TestArchive = IFileManager::Get().CreateFileReader(*TestArchiveFilename);
-						FArchive* Saver = new FDiffSerializeArchive(*InOuter->FileName.ToString(), TestArchive);
+						FArchive* Saver = new FDiffSerializeArchive(*InOuter->GetLoadedPath().GetPackageName(), TestArchive);
 						Linker = TUniquePtr<FLinkerSave>(new FLinkerSave(InOuter, Saver, bForceByteSwapping));
 					}
 					else
@@ -3030,7 +3029,7 @@ FSavePackageResultStruct UPackage::Save(UPackage* InOuter, UObject* Base, EObjec
 				if( Conform )
 				{
 					// Conform to previous generation of file.
-					UE_LOG(LogSavePackage, Log,  TEXT("Conformal save, relative to: %s, Generation %i"), *Conform->Filename, Conform->Summary.Generations.Num()+1 );
+					UE_LOG(LogSavePackage, Log,  TEXT("Conformal save, relative to: %s, Generation %i"), *Conform->GetDebugName(), Conform->Summary.Generations.Num()+1 );
 					Linker->Summary.Guid = Conform->Summary.Guid;
 #if WITH_EDITORONLY_DATA
 					Linker->Summary.PersistentGuid = Conform->Summary.PersistentGuid;					
@@ -4654,9 +4653,9 @@ FSavePackageResultStruct UPackage::Save(UPackage* InOuter, UObject* Base, EObjec
 		{
 			// if the save was successful, update the internal package filename path if we aren't currently cooking
 #if WITH_EDITOR
-			if (TargetPlatform == nullptr && bIsValidLongPackageName)
+			if (TargetPlatform == nullptr && !TargetPackagePath.IsMountedPath())
 			{
-				InOuter->FileName = *PackageFilename;
+				InOuter->SetLoadedPath(TargetPackagePath);
 			}
 #endif
 
diff --git a/Engine/Source/Runtime/CoreUObject/Private/UObject/SavePackage2.cpp b/Engine/Source/Runtime/CoreUObject/Private/UObject/SavePackage2.cpp
--- a/Engine/Source/Runtime/CoreUObject/Private/UObject/SavePackage2.cpp
+++ b/Engine/Source/Runtime/CoreUObject/Private/UObject/SavePackage2.cpp
@@ -665,7 +665,7 @@ ESavePackageResult CreateLinker(FSaveContext& SaveContext)
 
 			// The entire package will be serialized to memory and then compared against package on disk.
 			// Each difference will be log with its Serialize call stack trace if IsDiffCallstack is true
-			FArchive* Saver = new FArchiveStackTrace(SaveContext.GetAsset(), *SaveContext.GetPackage()->FileName.ToString(), SaveContext.IsDiffCallstack(), SaveContext.GetDiffMapPtr());
+			FArchive* Saver = new FArchiveStackTrace(SaveContext.GetAsset(), *SaveContext.GetPackage()->GetLoadedPath().GetPackageName(), SaveContext.IsDiffCallstack(), SaveContext.GetDiffMapPtr());
 			SaveContext.Linker = MakeUnique<FLinkerSave>(SaveContext.GetPackage(), Saver, SaveContext.IsForceByteSwapping(), SaveContext.IsSaveUnversioned());
 		}
 		else
diff --git a/Engine/Source/Runtime/CoreUObject/Private/UObject/SoftObjectPath.cpp b/Engine/Source/Runtime/CoreUObject/Private/UObject/SoftObjectPath.cpp
--- a/Engine/Source/Runtime/CoreUObject/Private/UObject/SoftObjectPath.cpp
+++ b/Engine/Source/Runtime/CoreUObject/Private/UObject/SoftObjectPath.cpp
@@ -718,7 +718,7 @@ bool FSoftObjectPathThreadContext::GetSerializationOptions(FName& OutPackageName
 		{
 			if (CurrentPackageName == NAME_None)
 			{
-				CurrentPackageName = FName(*FPackageName::FilenameToLongPackageName(Linker->Filename));
+				CurrentPackageName = Linker->GetPackagePath().GetPackageFName();
 			}
 			if (Archive == nullptr)
 			{
