PackageResourceManager and PackagePaths

diff --git a/Engine/Source/Runtime/CoreUObject/Private/UObject/UObjectGlobals.cpp b/Engine/Source/Runtime/CoreUObject/Private/UObject/UObjectGlobals.cpp
--- a/Engine/Source/Runtime/CoreUObject/Private/UObject/UObjectGlobals.cpp
+++ b/Engine/Source/Runtime/CoreUObject/Private/UObject/UObjectGlobals.cpp
@@ -33,6 +33,7 @@
 #include "Serialization/SerializedPropertyScope.h"
 #include "UObject/UnrealType.h"
 #include "UObject/ObjectRedirector.h"
+#include "UObject/PackageResourceManager.h"
 #include "UObject/UObjectAnnotation.h"
 #include "Serialization/DuplicatedObject.h"
 #include "Serialization/DuplicatedDataReader.h"
@@ -1048,30 +1049,23 @@ public:
 	}
 };
 
-// this class is a hack to work around calling private functions int he linker 
+// this class is a hack to work around calling private functions in the linker 
 // I just want to replace the Linkers loader with a custom one
 class FUnsafeLinkerLoad : public FLinkerLoad
 {
 public:
-	FUnsafeLinkerLoad(UPackage *Package, const TCHAR* FileName, const TCHAR* DiffFilename, uint32 LoadFlags) : FLinkerLoad(Package, FileName, LoadFlags)
+	FUnsafeLinkerLoad(UPackage *Package, const FPackagePath& PackagePath, const FPackagePath& DiffPackagePath, uint32 LoadFlags)
+		: FLinkerLoad(Package, PackagePath, LoadFlags)
 	{
 		Package->LinkerLoad = this;
-
-		/*while (CreateLoader(TFunction<void()>([]() {})) == FLinkerLoad::LINKER_TimedOut)
-		{
-		}*/
-
-
-		
-
 		while ( Tick(0.0, false, false, nullptr) == FLinkerLoad::LINKER_TimedOut ) 
 		{ 
 		}
 
-		FArchive* OtherFile = IFileManager::Get().CreateFileReader(DiffFilename);
-		FDiffFileArchive* DiffArchive = new FDiffFileArchive(GetLoader(), OtherFile);
+		FOpenPackageResult OtherFile = IPackageResourceManager::Get().OpenReadPackage(DiffPackagePath);
+		checkf(!OtherFile.Archive.IsValid() || OtherFile.Format == EPackageFormat::Binary, TEXT("Text format is not yet supported with DiffPackage"));
+		FDiffFileArchive* DiffArchive = new FDiffFileArchive(GetLoader(), OtherFile.Archive.Release());
 		SetLoader(DiffArchive);
-
 	}
 };
 
@@ -1085,51 +1079,41 @@ UE_TRACE_EVENT_BEGIN(CUSTOM_LOADTIMER_LOG, LoadPackageInternal, NoSync)
 	UE_TRACE_EVENT_FIELD(UE::Trace::WideString, PackageName)
 UE_TRACE_EVENT_END()
 
-UPackage* LoadPackageInternal(UPackage* InOuter, const TCHAR* InLongPackageNameOrFilename, uint32 LoadFlags, FLinkerLoad* ImportLinker, FArchive* InReaderOverride, const FLinkerInstancingContext* InstancingContext)
+UPackage* LoadPackageInternal(UPackage* InOuter, const FPackagePath& PackagePath, uint32 LoadFlags, FLinkerLoad* ImportLinker, FArchive* InReaderOverride,
+	const FLinkerInstancingContext* InstancingContext, const FPackagePath* DiffPackagePath)
 {
 	DECLARE_SCOPE_CYCLE_COUNTER(TEXT("LoadPackageInternal"), STAT_LoadPackageInternal, STATGROUP_ObjectVerbose);
 	SCOPED_CUSTOM_LOADTIMER(LoadPackageInternal)
-		ADD_CUSTOM_LOADTIMER_META(LoadPackageInternal, PackageName, InLongPackageNameOrFilename);
-
-	checkf(IsInGameThread(), TEXT("Unable to load %s. Objects and Packages can only be loaded from the game thread."), InLongPackageNameOrFilename);
+		ADD_CUSTOM_LOADTIMER_META(LoadPackageInternal, PackageName, *PackagePath.GetPackageNameOrFallback());
 
-	UPackage* Result = nullptr;
+	if (PackagePath.IsEmpty())
+	{
+		UE_LOG(LogUObjectGlobals, Warning, TEXT("Attempted to LoadPackage from empty PackagePath."));
+		return nullptr;
+	}
+	checkf(IsInGameThread(), TEXT("Unable to load %s. Objects and Packages can only be loaded from the game thread."), *PackagePath.GetDebugName());
 
 	if (FPlatformProperties::RequiresCookedData() && GEventDrivenLoaderEnabled
 		&& EVENT_DRIVEN_ASYNC_LOAD_ACTIVE_AT_RUNTIME
 		)
 	{
-		FString InName;
-		FString InPackageName;
-
-		if (FPackageName::IsPackageFilename(InLongPackageNameOrFilename))
-		{
-			FPackageName::TryConvertFilenameToLongPackageName(InLongPackageNameOrFilename, InPackageName);
-		}
-		else
-		{
-			InPackageName = InLongPackageNameOrFilename;
-		}
-
-		if (InOuter)
-		{
-			InName = InOuter->GetPathName();
-		}
-		else
+		checkf(!InOuter || !InOuter->GetOuter(), TEXT("Loading into subpackages is not implemented.")); // Subpackages are no longer supported in UE
+		FName PackageName(InOuter ? InOuter->GetFName() : PackagePath.GetPackageFName());
+		if (PackageName.IsNone())
 		{
-			InName = InPackageName;
+			UE_LOG(LogUObjectGlobals, Warning, TEXT("Attempted to LoadPackage from non-mounted path %s. This is not supported."), *PackagePath.GetDebugName());
+			return nullptr;
 		}
 
-		FName PackageFName(*InPackageName);
-		UE_TRACK_REFERENCING_PACKAGE_SCOPED(PackageFName, PackageAccessTrackingOps::NAME_Load);
+		UE_TRACK_REFERENCING_PACKAGE_SCOPED(PackageName, PackageAccessTrackingOps::NAME_Load);
 
 		{
 			if (FCoreDelegates::OnSyncLoadPackage.IsBound())
 			{
-				FCoreDelegates::OnSyncLoadPackage.Broadcast(InName);
+				FCoreDelegates::OnSyncLoadPackage.Broadcast(PackageName.ToString());
 			}
 
-			int32 RequestID = LoadPackageAsync(InName, nullptr, *InPackageName);
+			int32 RequestID = LoadPackageAsync(PackagePath, PackageName);
 
 			if (RequestID != INDEX_NONE)
 			{
@@ -1137,48 +1121,11 @@ UPackage* LoadPackageInternal(UPackage* InOuter, const TCHAR* InLongPackageNameO
 			}
 		}
 
-		Result = (InOuter ? InOuter : FindObjectFast<UPackage>(nullptr, PackageFName));
-		return Result;
+		return (InOuter ? InOuter : FindObjectFast<UPackage>(nullptr, PackageName));
 	}
 
-	FString FileToLoad;
-#if WITH_EDITOR
-	FString DiffFileToLoad;
-#endif
+	UPackage* Result = nullptr;
 
-#if WITH_EDITOR
-	if (LoadFlags & LOAD_ForFileDiff)
-	{
-		FString TempFilenames = InLongPackageNameOrFilename;
-		ensure(TempFilenames.Split(TEXT(";"), &FileToLoad, &DiffFileToLoad, ESearchCase::CaseSensitive));
-	}
-	else
-#endif
-	if (InLongPackageNameOrFilename && FCString::Strlen(InLongPackageNameOrFilename) > 0)
-	{
-		FileToLoad = InLongPackageNameOrFilename;
-	}
-	else if (InOuter)
-	{
-		FileToLoad = InOuter->GetName();
-	}
-
-	// Make sure we're trying to load long package names only.
-	if (FPackageName::IsShortPackageName(FileToLoad))
-	{
-		FString LongPackageName;
-		FName* ScriptPackageName = FPackageName::FindScriptPackageName(*FileToLoad);
-		if (ScriptPackageName)
-		{
-			UE_LOG(LogUObjectGlobals, Warning, TEXT("LoadPackage: %s is a short script package name."), InLongPackageNameOrFilename);
-			FileToLoad = ScriptPackageName->ToString();
-		}
-		else if (!FPackageName::SearchForPackageOnDisk(FileToLoad, &FileToLoad))
-		{
-			UE_LOG(LogUObjectGlobals, Warning, TEXT("LoadPackage can't find package %s."), *FileToLoad);
-			return NULL;
-		}
-	}
 #if WITH_EDITOR
 	// In the editor loading cannot be part of a transaction as it cannot be undone, and may result in recording half-loaded objects. So we suppress any active transaction while in this stack, and set the editor loading flag
 	TGuardValue<ITransaction*> SuppressTransaction(GUndo, nullptr);
@@ -1189,21 +1136,21 @@ UPackage* LoadPackageInternal(UPackage* InOuter, const TCHAR* InLongPackageNameO
 	if (ShouldCreateThrottledSlowTask())
 	{
 		static const FTextFormat LoadingPackageTextFormat = NSLOCTEXT("Core", "LoadingPackage_Scope", "Loading Package '{0}'");
-		SlowTask.Emplace(100, FText::Format(LoadingPackageTextFormat, FText::FromString(FileToLoad)));
+		SlowTask.Emplace(100, FText::Format(LoadingPackageTextFormat, PackagePath.GetDebugNameText()));
 		SlowTask->Visibility = ESlowTaskVisibility::Invisible;
 		SlowTask->EnterProgressFrame(10);
 	}
 
 	if (FCoreDelegates::OnSyncLoadPackage.IsBound())
 	{
-		FCoreDelegates::OnSyncLoadPackage.Broadcast(FileToLoad);
+		FCoreDelegates::OnSyncLoadPackage.Broadcast(PackagePath.GetPackageNameOrFallback());
 	}
 	
 	// Set up a load context
 	TRefCountPtr<FUObjectSerializeContext> LoadContext = FUObjectThreadContext::Get().GetSerializeContext();
 
 	// Try to load.
-	BeginLoad(LoadContext, InLongPackageNameOrFilename);
+	BeginLoad(LoadContext, *PackagePath.GetDebugName());
 
 	bool bFullyLoadSkipped = false;
 
@@ -1220,21 +1167,21 @@ UPackage* LoadPackageInternal(UPackage* InOuter, const TCHAR* InLongPackageNameO
 
 		// Create a new linker object which goes off and tries load the file.
 #if WITH_EDITOR
-		if (LoadFlags & LOAD_ForFileDiff)
+		if (DiffPackagePath)
 		{
 			// Create the package with the provided long package name.
 			if (!InOuter)
 			{
-				InOuter = CreatePackage(*FileToLoad);
+				InOuter = CreatePackage(*PackagePath.GetPackageName());
 			}
 			
-			new FUnsafeLinkerLoad(InOuter, *FileToLoad, *DiffFileToLoad, LOAD_ForDiff);
+			new FUnsafeLinkerLoad(InOuter, PackagePath, *DiffPackagePath, LOAD_ForDiff);
 		}
 #endif
 
 		{
 			FUObjectSerializeContext* InOutLoadContext = LoadContext;
-			Linker = GetPackageLinker(InOuter, *FileToLoad, LoadFlags, nullptr, nullptr, InReaderOverride, &InOutLoadContext, ImportLinker, InstancingContext);
+			Linker = GetPackageLinker(InOuter, PackagePath, LoadFlags, nullptr, nullptr, InReaderOverride, &InOutLoadContext, ImportLinker, InstancingContext);
 			if (InOutLoadContext != LoadContext && InOutLoadContext)
 			{
 				// The linker already existed and was associated with another context
@@ -1304,13 +1251,8 @@ UPackage* LoadPackageInternal(UPackage* InOuter, const TCHAR* InLongPackageNameO
 			Result->SetPackageFlags(PKG_ForDiffing);
 		}
 
-		// Save the filename we load from in Long package name form
-		{
-			// convert will succeed here, otherwise the linker will have been null
-			FString LongPackageFilename;
-			FPackageName::TryConvertFilenameToLongPackageName(FileToLoad, LongPackageFilename);
-			Result->FileName = FName(*LongPackageFilename);
-		}
+		// Save the PackagePath we loaded from
+		Result->SetLoadedPath(PackagePath);
 		
 		// is there a script SHA hash for this package?
 		uint8 SavedScriptSHA[20];
@@ -1377,7 +1319,7 @@ UPackage* LoadPackageInternal(UPackage* InOuter, const TCHAR* InLongPackageNameO
 			// compare SHA hash keys
 			if (FMemory::Memcmp(SavedScriptSHA, LoadedScriptSHA, 20) != 0)
 			{
-				appOnFailSHAVerification(*Linker->Filename, false);
+				appOnFailSHAVerification(*Linker->GetPackagePath().GetLocalFullPath(), false);
 			}
 		}
 
@@ -1444,22 +1386,86 @@ UPackage* LoadPackageInternal(UPackage* InOuter, const TCHAR* InLongPackageNameO
 	return Result;
 }
 
-UPackage* LoadPackage(UPackage* InOuter, const TCHAR* InLongPackageName, uint32 LoadFlags, FArchive* InReaderOverride, const FLinkerInstancingContext* InstancingContext)
+UPackage* LoadPackage(UPackage* InOuter, const TCHAR* InLongPackageNameOrFilename, uint32 LoadFlags, FArchive* InReaderOverride, const FLinkerInstancingContext* InstancingContext)
+{
+	FPackagePath PackagePath;
+	FPackagePath* DiffPackagePathPtr = nullptr;
+
+#if WITH_EDITOR
+	FPackagePath DiffPackagePath;
+	if (LoadFlags & LOAD_ForFileDiff)
+	{
+		FString TempFilenames = InLongPackageNameOrFilename;
+		FString FileToLoad;
+		FString DiffFileToLoad;
+		ensure(TempFilenames.Split(TEXT(";"), &FileToLoad, &DiffFileToLoad, ESearchCase::CaseSensitive));
+		PackagePath = FPackagePath::FromLocalPath(FileToLoad);
+		DiffPackagePath = FPackagePath::FromLocalPath(DiffFileToLoad);
+		DiffPackagePathPtr = &DiffPackagePath;
+	}
+	else
+#endif
+	if (InLongPackageNameOrFilename && InLongPackageNameOrFilename[0] != '\0')
+	{
+		FString BufferName;
+		// Make sure we're trying to load long package names only.
+		if (FPackageName::IsShortPackageName(FStringView(InLongPackageNameOrFilename)))
+		{
+			BufferName = InLongPackageNameOrFilename;
+			FName* ScriptPackageName = FPackageName::FindScriptPackageName(*BufferName);
+			if (ScriptPackageName)
+			{
+				UE_LOG(LogUObjectGlobals, Warning, TEXT("LoadPackage: %s is a short script package name."), InLongPackageNameOrFilename);
+				BufferName = ScriptPackageName->ToString();
+				InLongPackageNameOrFilename = *BufferName;
+			}
+			else if (FPackageName::SearchForPackageOnDisk(BufferName, &BufferName))
+			{
+				InLongPackageNameOrFilename = *BufferName;
+			}
+			else
+			{
+				UE_LOG(LogUObjectGlobals, Warning, TEXT("LoadPackage can't find package %s."), InLongPackageNameOrFilename);
+				return nullptr;
+			}
+		}
+
+		if (!FPackagePath::TryFromMountedName(InLongPackageNameOrFilename, PackagePath))
+		{
+			UE_LOG(LogUObjectGlobals, Warning, TEXT("LoadPackage can't find package %s."), InLongPackageNameOrFilename);
+			return nullptr;
+		}
+	}
+	else if (InOuter)
+	{
+		PackagePath = FPackagePath::FromPackageNameChecked(InOuter->GetName());
+	}
+	else
+	{
+		UE_LOG(LogUObjectGlobals, Warning, TEXT("Empty name passed to LoadPackage."));
+		return nullptr;
+	}
+	return LoadPackage(InOuter, PackagePath, LoadFlags, InReaderOverride, InstancingContext, DiffPackagePathPtr);
+}
+
+UPackage* LoadPackage(UPackage* InOuter, const FPackagePath& PackagePath, uint32 LoadFlags, FArchive* InReaderOverride, const FLinkerInstancingContext* InstancingContext, const FPackagePath* DiffPackagePath)
 {
 	COOK_STAT(LoadPackageStats::NumPackagesLoaded++);
 	COOK_STAT(FScopedDurationTimer LoadTimer(LoadPackageStats::LoadPackageTimeSec));
 	// Change to 1 if you want more detailed stats for loading packages, but at the cost of adding dynamic stats.
 #if	STATS && 0
 	static FString Package = TEXT( "Package" );
-	const FString LongName = Package / InLongPackageName;
+	const FString LongName = Package / PackagePath.GetPackageNameOrFallback();
 	const TStatId StatId = FDynamicStats::CreateStatId<FStatGroup_STATGROUP_UObjects>( LongName );
 	FScopeCycleCounter CycleCounter( StatId );
 #endif // STATS
 
 	// since we are faking the object name, this is basically a duplicate of LLM_SCOPED_TAG_WITH_OBJECT_IN_SET
-	FString FakePackageName = FString(TEXT("Package ")) + InLongPackageName;
-	LLM_SCOPED_TAG_WITH_STAT_NAME_IN_SET(FLowLevelMemTracker::Get().IsTagSetActive(ELLMTagSet::Assets) ? FDynamicStats::CreateMemoryStatId<FStatGroup_STATGROUP_LLMAssets>(FName(*FakePackageName)).GetName() : NAME_None, ELLMTagSet::Assets, ELLMTracker::Default);
-	return LoadPackageInternal(InOuter, InLongPackageName, LoadFlags, /*ImportLinker =*/ nullptr, InReaderOverride, InstancingContext);
+	LLM_SCOPED_TAG_WITH_STAT_NAME_IN_SET(FLowLevelMemTracker::Get().IsTagSetActive(ELLMTagSet::Assets) ?
+											FDynamicStats::CreateMemoryStatId<FStatGroup_STATGROUP_LLMAssets>(FName(*(FString(TEXT("Package ")) + PackagePath.GetPackageNameOrFallback()))).GetName() :
+											NAME_None,
+										 ELLMTagSet::Assets, ELLMTracker::Default);
+	return LoadPackageInternal(InOuter, PackagePath, LoadFlags, /*ImportLinker =*/ nullptr, InReaderOverride, InstancingContext, DiffPackagePath);
 }
 
 /**
diff --git a/Engine/Source/Runtime/CoreUObject/Private/UObject/UObjectLinker.cpp b/Engine/Source/Runtime/CoreUObject/Private/UObject/UObjectLinker.cpp
--- a/Engine/Source/Runtime/CoreUObject/Private/UObject/UObjectLinker.cpp
+++ b/Engine/Source/Runtime/CoreUObject/Private/UObject/UObjectLinker.cpp
@@ -113,7 +113,7 @@ void UObject::SetLinker( FLinkerLoad* LinkerLoad, int32 LinkerIndex, bool bShoul
 			TEXT("Detaching from existing linker for %s while object %s needs loading from linker %s."),
 			*Existing.Linker->GetArchiveName(),
 			*GetFullName(),
-			LinkerLoad ? *LinkerLoad->Filename : TEXT(""));
+			LinkerLoad ? *LinkerLoad->GetDebugName() : TEXT(""));
 		check(Existing.Linker->ExportMap[Existing.LinkerIndex].Object!=nullptr);
 		check(Existing.Linker->ExportMap[Existing.LinkerIndex].Object==this);
 		Existing.Linker->ExportMap[Existing.LinkerIndex].ResetObject();
@@ -140,9 +140,9 @@ void UObject::SetLinker( FLinkerLoad* LinkerLoad, int32 LinkerIndex, bool bShoul
 		UE_CLOG(Existing.Linker && LinkerLoad, LogUObjectLinker, Fatal,
 			TEXT("It is only legal to change linkers in the editor. Trying to change linker on %s from %s (Existing->LinkerRoot=%s) to %s (LinkerLoad->LinkerRoot=%s)"),
 			*GetFullName(),
-			*Existing.Linker->Filename,
+			*Existing.Linker->GetDebugName(),
 			*GetNameSafe(Existing.Linker->LinkerRoot),
-			*LinkerLoad->Filename,
+			*LinkerLoad->GetDebugName(),
 			*GetNameSafe(LinkerLoad->LinkerRoot));
 #endif
 	}
diff --git a/Engine/Source/Runtime/CoreUObject/Private/Virtualization/VirtualizedBulkData.cpp b/Engine/Source/Runtime/CoreUObject/Private/Virtualization/VirtualizedBulkData.cpp
--- a/Engine/Source/Runtime/CoreUObject/Private/Virtualization/VirtualizedBulkData.cpp
+++ b/Engine/Source/Runtime/CoreUObject/Private/Virtualization/VirtualizedBulkData.cpp
@@ -3,6 +3,7 @@
 #include "Virtualization/VirtualizedBulkData.h"
 
 #include "HAL/FileManager.h"
+#include "Misc/PackageSegment.h"
 #include "Misc/SecureHash.h"
 #include "Misc/ScopeLock.h"
 #include "Serialization/BulkData.h"
@@ -23,7 +24,8 @@ void FVirtualizedUntypedBulkData::CreateFromBulkData(FUntypedBulkData& InBulkDat
 
 	Reset();
 
-	Filename = InBulkData.GetFilename();
+	PackagePath = InBulkData.GetPackagePath();
+	PackageSegment = InBulkData.GetPackageSegment();
 
 	OffsetInFile = InBulkData.GetBulkDataOffsetInFile();
 	PayloadLength = InBulkData.GetBulkDataSize();
@@ -47,14 +49,15 @@ void FVirtualizedUntypedBulkData::Serialize(FArchive& Ar, UObject* Owner)
 		{
 			Ar << Key;
 			Ar << PayloadLength;
-			Ar << bIsVirtualized;		
+			Ar << bIsVirtualized;
 			Ar << bShouldCompressData;
 			Ar << bIsDataStoredAsCompressed;
-			Ar << Filename;
+			Ar << PackagePath;
+			Ar << PackageSegment;
 			Ar << OffsetInFile;
 
 			bool bKeepBuffer = !Ar.IsLoading();
-			if (!bIsVirtualized && Filename.IsEmpty())
+			if (!bIsVirtualized && PackagePath.IsEmpty())
 			{
 				// Data is held in memory so we need to serialize it to the archive
 				bKeepBuffer = SerializeData(Ar, Payload);
@@ -78,7 +81,7 @@ void FVirtualizedUntypedBulkData::Serialize(FArchive& Ar, UObject* Owner)
 		{
 			FLinkerSave* LinkerSave = Cast<FLinkerSave>(Ar.GetLinker());
 
-			const bool bCanAttemptVirtualization = LinkerSave != nullptr;			
+			const bool bCanAttemptVirtualization = LinkerSave != nullptr;
 
 			if (bCanAttemptVirtualization)
 			{
@@ -86,7 +89,7 @@ void FVirtualizedUntypedBulkData::Serialize(FArchive& Ar, UObject* Owner)
 			}
 
 			Ar << bIsVirtualized;
-					
+
 			if (!bIsVirtualized)
 			{
 				// Need to load the payload so that we can write it out
@@ -145,33 +148,35 @@ void FVirtualizedUntypedBulkData::Serialize(FArchive& Ar, UObject* Owner)
 				bIsDataStoredAsCompressed = false;
 				OffsetInFile = INDEX_NONE;
 
-				Filename.Empty();
+				PackagePath.Empty();
+				PackageSegment = EPackageSegment::Header;
 			}
 			else
 			{
 				Ar << bIsDataStoredAsCompressed;
 
-				// If we can lazy load then find the filename, otherwise we will want
+				// If we can lazy load then find the PackagePath, otherwise we will want
 				// to serialize immediately.
 				if (Ar.IsAllowingLazyLoading())
 				{
-					Filename = GetFilenameFromOwner(Owner);
+					PackagePath = GetPackagePathFromOwner(Owner, PackageSegment);
 				}
 				else
 				{
-					Filename.Empty();
+					PackagePath.Empty();
+					PackageSegment = EPackageSegment::Header;
 				}
 				
 				OffsetInFile = INDEX_NONE;
 				Ar << OffsetInFile;
 
-				if (Filename.IsEmpty())
+				if (PackagePath.IsEmpty())
 				{
-					// If we have no filename then we need to load the data immediately
+					// If we have no packagepath then we need to load the data immediately
 					// as we will not be able to load it on demand.
 					SerializeData(Ar, Payload);
-				}		
-			}	
+				}
+			}
 		}
 	}
 }
@@ -198,19 +203,20 @@ FSharedBufferConstPtr FVirtualizedUntypedBulkData::LoadFromDisk() const
 	FSharedBufferConstPtr PayloadFromDisk;
 
 	// Open a reader to the file
-	TUniquePtr<FArchive> Ar(IFileManager::Get().CreateFileReader(*Filename, FILEREAD_Silent));
-	if (Ar.IsValid())
+	FOpenPackageResult Result = IPackageResourceManager::Get().OpenReadPackage(PackagePath, PackageSegment);
+	if (Result.Archive.IsValid() && Result.Format == EPackageFormat::Binary)
 	{
 		// Move the correct location of the data in the file
-		Ar->Seek(OffsetInFile);
+		Result.Archive->Seek(OffsetInFile);
 
 		// Now we can actually serialize it
-		SerializeData(*Ar, PayloadFromDisk);
+		SerializeData(*Result.Archive, PayloadFromDisk);
 	}
 	else
 	{
-		UE_LOG(LogVirtualization, Error, TEXT("Could not open (%s) to read FVirtualizedUntypedBulkData"), *Filename);
-	}	
+		UE_LOG(LogVirtualization, Error, TEXT("Could not open (%s) to read FVirtualizedUntypedBulkData: %s."), *PackagePath.GetDebugNameWithExtension(PackageSegment),
+			(!Result.Archive.IsValid() ? TEXT("could not find package") : TEXT("package is a TextAsset which is not supported")));
+	}
 
 	return PayloadFromDisk;
 }
@@ -304,34 +310,35 @@ FSharedBufferConstPtr FVirtualizedUntypedBulkData::PullData() const
 	return PulledPayload;
 }
 
-FString FVirtualizedUntypedBulkData::GetFilenameFromOwner(UObject* Owner) const
+FPackagePath FVirtualizedUntypedBulkData::GetPackagePathFromOwner(UObject* Owner, EPackageSegment& OutPackageSegment) const
 {
+	OutPackageSegment = EPackageSegment::Header;
 	if (Owner != nullptr)
 	{
 		UPackage* Package = Owner->GetOutermost();
 		checkf(Package != nullptr, TEXT("Owner was not a valid UPackage!"));
 
-		FLinker* Linker = FLinkerLoad::FindExistingLinkerForPackage(Package);
+		FLinkerLoad* Linker = FLinkerLoad::FindExistingLinkerForPackage(Package);
 		checkf(Linker != nullptr, TEXT("UPackage did not have a valid FLinkerLoad!"));
 
-		return Linker->Filename;
+		return Linker->GetPackagePath();
 	}
-
 	else
 	{
-		return FString();
+		return FPackagePath();
 	}
 }
 
 bool FVirtualizedUntypedBulkData::CanUnloadData() const
 {
 	// We cannot unload the data if are unable to reload it from a file
-	return bIsVirtualized || Filename.IsEmpty() == false;
+	return bIsVirtualized || PackagePath.IsEmpty() == false;
 }
 
 void FVirtualizedUntypedBulkData::Reset()
 {
-	Filename.Empty();
+	PackagePath.Empty();
+	PackageSegment = EPackageSegment::Header;
 	OffsetInFile = INDEX_NONE;
 	Key.Invalidate();
 	Payload.Reset();
@@ -405,7 +412,8 @@ void FVirtualizedUntypedBulkData::UpdatePayload(const FSharedBufferConstRef& InP
 	bIsVirtualized = false;
 	bIsDataStoredAsCompressed = false;
 
-	Filename.Empty();
+	PackagePath.Empty();
+	PackageSegment = EPackageSegment::Header;
 	OffsetInFile = INDEX_NONE;
 
 	CalculateKey();
diff --git a/Engine/Source/Runtime/CoreUObject/Public/Misc/PackageName.h b/Engine/Source/Runtime/CoreUObject/Public/Misc/PackageName.h
--- a/Engine/Source/Runtime/CoreUObject/Public/Misc/PackageName.h
+++ b/Engine/Source/Runtime/CoreUObject/Public/Misc/PackageName.h
@@ -9,12 +9,50 @@
 #include "CoreMinimal.h"
 #include "Containers/ArrayView.h"
 #include "Containers/StringView.h"
+#include "Misc/PackagePath.h"
 
 struct FFileStatData;
 
+DECLARE_LOG_CATEGORY_EXTERN(LogPackageName, Log, All);
+
 class COREUOBJECT_API FPackageName
 {
 public:
+	/** Indicates the format the input was in in functions that take flexible input that can be a LocalPath, PackageName, or ObjectPath */
+	enum class EFlexNameType
+	{
+		Invalid,
+		PackageName,
+		LocalPath,
+		ObjectPath,
+	};
+	enum class EErrorCode
+	{
+		PackageNameUnknown,
+		PackageNameEmptyPath,
+		PackageNamePathNotMounted,
+		PackageNameFullObjectPathNotAllowed,
+		PackageNameContainsInvalidCharacters,
+		LongPackageNames_PathTooShort,
+		LongPackageNames_PathWithNoStartingSlash,
+		LongPackageNames_PathWithTrailingSlash,
+	};
+
+	/**
+	 * Return a user-readable string for the error code returned from a FPackageName function
+	 *
+	 * @param InPath the path that was passed to the convert function
+	 * @param ErrorCode The error code returned from the convert function
+	 */
+	static FString FormatErrorAsString(FStringView InPath, EErrorCode ErrorCode);
+
+	/**
+	 * Return a user-readable string for the error code returned from a FPackageName function
+	 *
+	 * @param InPath the path that was passed to the convert function
+	 * @param ErrorCode The error code returned from the convert function
+	 */
+	static FText FormatErrorAsText(FStringView InPath, EErrorCode ErrorCode);
 
 	/**
 	 * Helper function for converting short to long script package name (InputCore -> /Script/InputCore)
@@ -67,6 +105,20 @@ public:
 	 */
 	static bool TryConvertLongPackageNameToFilename(const FString& InLongPackageName, FString& OutFilename, const FString& InExtension = TEXT(""));
 
+	/**
+	 * Find the MountPoint for a LocalPath, LongPackageName, or ObjectPath and return its elements. Use this function instead of TryConvertFilenameToLongPackageName or
+	 * TryConvertLongPackageNameToFilename if you need to handle InPaths that might be ObjectPaths.
+	 * @param InPath					The LocalPath (with path,name,extension), PackageName, or ObjectPath we want to 
+	 * @param OutLocalPathNoExtension	If non-null, will be set to the LocalPath with path and basename but not extension, or empty string if input was not mounted
+	 * @param OutPackageName			If non-null, will be set to the LongPackageName, or empty string if input was not mounted
+	 * @param OutObjectName				If non-null, will be set to the ObjectName, or empty string if the input was a LocalPath or PackageName or was not mounted
+	 * @param OutSubObjectName			If non-null, will be set to the SubObjectName, or empty string if the input was a LocalPath or PackageName or was not mounted
+	 * @param OutExtension				If non-null, will be set to the LocalPath's extension, or empty string if the input was not a LocalPath or was not mounted
+	 * @param OutFlexNameType			If non-null, will be set to the FlexNameType of InPath.
+	 * @param OutFailureReason			If non-null, will be set to the reason InPath could not be converted, or to EErrorCode::Unknown if the function was successful.
+	 */
+	static bool TryConvertToMountedPath(FStringView InPath, FString* OutLocalPathNoExtension, FString* OutPackageName, FString* OutObjectName, FString* OutSubObjectName, FString* OutExtension, EFlexNameType* OutFlexNameType = nullptr, EErrorCode* OutFailureReason = nullptr);
+
 	/** 
 	 * Converts a long package name to a file name with the supplied extension.
 	 * Throws a fatal error if the conversion is not successful.
@@ -115,6 +167,7 @@ public:
 	 * @param OutSubObjectName			The extracted subobject name (SubObject)
 	 */
 	static void SplitFullObjectPath(const FString& InFullObjectPath, FString& OutClassName, FString& OutPackageName, FString& OutObjectName, FString& OutSubObjectName);
+	static void SplitFullObjectPath(FStringView InFullObjectPath, FStringView& OutClassName, FStringView& OutPackageName, FStringView& OutObjectName, FStringView& OutSubObjectName);
 
 	/** 
 	 * Returns true if the path starts with a valid root (i.e. /Game/, /Engine/, etc) and contains no illegal characters.
@@ -124,7 +177,19 @@ public:
 	 * @param OutReason					When returning false, this will provide a description of what was wrong with the name.
 	 * @return							true if a valid long package name
 	 */
-	static bool IsValidLongPackageName(const FString& InLongPackageName, bool bIncludeReadOnlyRoots = false, FText* OutReason = nullptr);
+	static bool IsValidLongPackageName(FStringView InLongPackageName, bool bIncludeReadOnlyRoots = false, EErrorCode* OutReason = nullptr);
+	static bool IsValidLongPackageName(FStringView InLongPackageName, bool bIncludeReadOnlyRoots, FText* OutReason );
+
+	/**
+	 * Report whether a given name is the proper format for a PackageName, without checking whether it is in one of the registered mount points
+	 *
+	 * @param InLongPackageName			The package name to test
+	 * @param OutReason					When returning false, this will provide a description of what was wrong with the name.
+	 * @return							true if valid text for a long package name
+	 */
+	static bool IsValidTextForLongPackageName(FStringView InLongPackageName, EErrorCode* OutReason = nullptr);
+	static bool IsValidTextForLongPackageName(FStringView InLongPackageName, FText* OutReason);
+
 
 	/**
 	 * Returns true if the path starts with a valid root (i.e. /Game/, /Engine/, etc) and contains no illegal characters.
@@ -140,7 +205,6 @@ public:
 	/**
 	 * Returns true if the path starts with a valid root (i.e. /Game/, /Engine/, etc).
 	 * 
-	 *
 	 * @param InObjectPath				The object path to test
 	 * @return							true if a valid object path
 	 */
@@ -230,11 +294,32 @@ public:
 	 * Checks if the package exists on disk.
 	 * 
 	 * @param LongPackageName Package name.
+	 * @param Guid If nonnull, and the package is found on disk but does not have this PackageGuid in its FPackageFileSummary::Guid, false is returned
 	 * @param OutFilename Package filename on disk.
 	 * @param InAllowTextFormats Detect text format packages as well as binary (priority to text)
 	 * @return true if the specified package name points to an existing package, false otherwise.
 	 **/
-	static bool DoesPackageExist(const FString& LongPackageName, const FGuid* Guid = NULL, FString* OutFilename = NULL, bool InAllowTextFormats = true);
+	static bool DoesPackageExist(const FString& LongPackageName, const FGuid* Guid = nullptr, FString* OutFilename = nullptr, bool InAllowTextFormats = true);
+
+	/**
+	 * Checks if the package exists on disk. PackagePath must be a mounted path, otherwise returns false
+	 * 
+	 * @param PackagePath Package package.
+	 * @param Guid If nonnull, and the package is found on disk but does not have this PackageGuid in its FPackageFileSummary::Guid, false is returned
+	 * @param bMatchCaseOnDisk If true, the OutPackagePath is modified to match the capitalization of the discovered file
+	 * @param OutPackagePath If nonnull and the package exists, set to a copy of PackagePath with the HeaderExtension set to the extension that exists on disk (and if bMatchCaseOnDisk is true, capitalization changed to match). If not found, this variable is not written
+	 * @return true if the specified package name points to an existing package, false otherwise.
+	 **/
+	static bool DoesPackageExist(const FPackagePath& PackagePath, const FGuid* Guid = nullptr, bool bMatchCaseOnDisk = false, FPackagePath* OutPackagePath = nullptr);
+
+	/**
+	 * Checks if the package exists on disk. PackagePath must be a mounted path, otherwise returns false
+	 * 
+	 * @param PackagePath Package package.
+	 * @param OutPackagePath If nonnull and the package exists, set to a copy of PackagePath with the HeaderExtension set to the extension that exists on disk. If not found, this variable is not written
+	 * @return true if the specified package name points to an existing package, false otherwise.
+	 **/
+	static bool DoesPackageExist(const FPackagePath& PackagePath, FPackagePath* OutPackagePath);
 
 	/**
 	 * Attempts to find a package given its short name on disk (very slow).
@@ -251,7 +336,7 @@ public:
 	 * @param ObjectPath Path to the object.
 	 * @param OutLongPackageName Converted object path.
 	 *
-	 * @returns True if succeeded. False otherwise.
+	 * @return True if succeeded. False otherwise.
 	 */
 	static bool TryConvertShortPackagePathToLongInObjectPath(const FString& ObjectPath, FString& ConvertedObjectPath);
 
@@ -260,7 +345,7 @@ public:
 	 *
 	 * @param ObjectPath Path to the object.
 	 *
-	 * @returns Normalized path (or empty path, if short object path was given and it wasn't found on the disk).
+	 * @return Normalized path (or empty path, if short object path was given and it wasn't found on the disk).
 	 */
 	static FString GetNormalizedObjectPath(const FString& ObjectPath);
 
@@ -271,7 +356,7 @@ public:
 	 *
 	 * @param InSourcePackagePath	Path to the source package.
 	 *
-	 * @returns Resolved package path, or the source package path if there is no resolution occurs.
+	 * @return Resolved package path, or the source package path if there is no resolution occurs.
 	 */
 	static FString GetDelegateResolvedPackagePath(const FString& InSourcePackagePath);
 
@@ -280,7 +365,7 @@ public:
 	 *
 	 * @param InLocalizedPackagePath Path to the localized package.
 	 *
-	 * @returns Source package path.
+	 * @return Source package path.
 	 */
 	static FString GetSourcePackagePath(const FString& InLocalizedPackagePath);
 
@@ -289,7 +374,7 @@ public:
 	 *
 	 * @param InSourcePackagePath	Path to the source package.
 	 *
-	 * @returns Localized package path, or the source package path if there is no suitable localized package.
+	 * @return Localized package path, or the source package path if there is no suitable localized package.
 	 */
 	static FString GetLocalizedPackagePath(const FString& InSourcePackagePath);
 
@@ -299,7 +384,7 @@ public:
 	 * @param InSourcePackagePath	Path to the source package.
 	 * @param InCultureName			Culture name to get the localized package for.
 	 *
-	 * @returns Localized package path, or the source package path if there is no suitable localized package.
+	 * @return Localized package path, or the source package path if there is no suitable localized package.
 	 */
 	static FString GetLocalizedPackagePath(const FString& InSourcePackagePath, const FString& InCultureName);
 
@@ -310,6 +395,7 @@ public:
 	 */
 	static FORCEINLINE const FString& GetAssetPackageExtension()
 	{
+		static FString AssetPackageExtension(LexToString(EPackageExtension::Asset));
 		return AssetPackageExtension;
 	}
 	/** 
@@ -319,6 +405,7 @@ public:
 	 */
 	static FORCEINLINE const FString& GetMapPackageExtension()
 	{
+		static FString MapPackageExtension(LexToString(EPackageExtension::Map));
 		return MapPackageExtension;
 	}
 
@@ -329,6 +416,7 @@ public:
 	*/
 	static FORCEINLINE const FString& GetTextAssetPackageExtension()
 	{
+		static FString TextAssetPackageExtension(LexToString(EPackageExtension::TextAsset));
 		return TextAssetPackageExtension;
 	}
 
@@ -339,6 +427,7 @@ public:
 	*/
 	static FORCEINLINE const FString& GetTextMapPackageExtension()
 	{
+		static FString TextMapPackageExtension(LexToString(EPackageExtension::TextMap));
 		return TextMapPackageExtension;
 	}
 
@@ -351,6 +440,14 @@ public:
 	 */
 	static bool IsTextPackageExtension(const TCHAR* Ext);
 
+	/**
+	 * Returns whether the passed in extension is a text header extension
+	 *
+	 * @param	Extension to test.
+	 * @return	True if Ext is either an text asset or a text map extension, otherwise false
+	 */
+	static bool IsTextPackageExtension(EPackageExtension Extension);
+
 	/**
 	 * Returns whether the passed in extension is a valid text asset package
 	 * extension. Extensions with and without trailing dots are supported.
@@ -405,7 +502,10 @@ public:
 	 */
 	static FORCEINLINE bool IsPackageFilename(const FString& Filename)
 	{
-		return Filename.EndsWith(AssetPackageExtension) || Filename.EndsWith(MapPackageExtension);
+		FStringView AssetPackageExtension(LexToString(EPackageExtension::Asset));
+		FStringView MapPackageExtension(LexToString(EPackageExtension::Map));
+		FStringView FilenameView(Filename);
+		return FilenameView.EndsWith(AssetPackageExtension) || FilenameView.EndsWith(MapPackageExtension);
 	}
 
 	/**
@@ -537,17 +637,21 @@ public:
 	/**
 	 * Checks if a package name contains characters that are invalid for package names.
 	 */
-	static bool DoesPackageNameContainInvalidCharacters(FStringView InLongPackageName, FText* OutReason = NULL);
-	
+	static bool DoesPackageNameContainInvalidCharacters(FStringView InLongPackageName, FText* OutReason);
+	static bool DoesPackageNameContainInvalidCharacters(FStringView InLongPackageName, EErrorCode* OutReason = nullptr);
+
 	/**
-	* Checks if a package can be found using known package extensions.
+	* Checks if a package can be found using known package extensions (header extensions only; files with the extensions of other segments are not returned).
 	*
 	* @param InPackageFilename Package filename without the extension.
 	* @param OutFilename If the package could be found, filename with the extension.
 	* @param InAllowTextFormats Detect text format packages as well as binary (priority to text)
 	* @return true if the package could be found on disk.
 	*/
-	static bool FindPackageFileWithoutExtension(const FString& InPackageFilename, FString& OutFilename, bool InAllowTextFormats = true);
+	static bool FindPackageFileWithoutExtension(const FString& InPackageFilename, FString& OutFilename);
+
+	UE_DEPRECATED(5.0, "Specifying AllowTextFormats is no longer supported. Text format is instead specified as a field on the result struct from IPackageResourceManager->OpenReadPackage.")
+	static bool FindPackageFileWithoutExtension(const FString& InPackageFilename, FString& OutFilename, bool InAllowTextFormats);
 
 	/**
 	 * Converts a long package name to the case it exists as on disk.
@@ -573,11 +677,6 @@ public:
 
 private:
 
-	static FString AssetPackageExtension;
-	static FString MapPackageExtension;
-	static FString TextAssetPackageExtension;
-	static FString TextMapPackageExtension;
-
 	/**
 	 * Internal function used to rename filename to long package name.
 	 *
@@ -586,6 +685,39 @@ private:
 	 */
 	static void InternalFilenameToLongPackageName(FStringView InFilename, FStringBuilderBase& OutPackageName);
 
+	/**
+	 * Internal helper to find a mount point that contains the given LocalFilePath or PackageName, and if found, return the MountPoint, RelativePath, and PackageExtension
+	 *
+	 * @param InPath The path to test, either a LocalFilePath, PackageName, or ObjectPath
+	 * @param OutMountPointPackageName If the MountPoint is found, the PackageName of the MountPoint is copied into this variable, otherwise it is set to empty string
+	 * @param OutMountPointFilePath If the MountPoint is found, the LocalFilePath of the MountPoint is copied into this variable, otherwise it is set to empty string
+	 * @param OutRelPath If the MountPoint is found, the RelativePath from the MountPoint to the PackageName is copied into this variable, otherwise it is set to empty string
+	 * @param OutObjectPath If the MountPoint is found and InPath was an ObjectPath, this variable is set to the object's relative path from the package: ObjectPath:SubObject1.SubObject2, otherwise it is set to the empty string
+	 * @param OutExtension If the MountPoint is found and InPath was a LocalFilePath, this variable is set to the EPackageExtension that was in the filepath, otherwise it is set to EPackageExtension::Unspecified
+	 * @param OutCustomExtension If the mount point was found and InPath was a filepath with extension and the extension was not a recognized package extension, this varialbe is set to the extension text (including .), otherwise it is set to the empty string
+	 * @param OutFlexNameType If non-null, it is set to whether InFilePathOrPackageName is a PackageName or Filename if the MountPoint is found, otherwise it is set to EFlexNameType::Invalid
+	 * @param OutFailureReason If non-null, it is set to the failurereason if the MountPoint is not found, otherwise it is set to EErrorCode::PackageNameUnknown
+	 * @return True if the MountPoint was found, else false
+	 */
+	static bool TryConvertToMountedPathComponents(FStringView InPath, FStringBuilderBase& OutMountPointPackageName, FStringBuilderBase& OutMountPointFilePath, FStringBuilderBase& OutRelPath,
+		FStringBuilderBase& OutObjectName, EPackageExtension& OutExtension, FStringBuilderBase& OutCustomExtension, EFlexNameType* OutFlexNameType = nullptr, EErrorCode* OutFailureReason = nullptr);
+
+	/**
+	 * Internal helper to find a mount point that contains the given LocalFilePath or PackageName, and if found, return the MountPoint and RelativePath
+	 *
+	 * @param InFilePathOrPackageName The path to test, either a LocalFilePath or a PackageName or an ObjectPath
+	 * @param OutMountPointPackageName If the MountPoint is found, the PackageName of the MountPoint is copied into this variable, otherwise it is set to empty string
+	 * @param OutMountPointFilePath If the MountPoint is found, the LocalFilePath of the MountPoint is copied into this variable, otherwise it is set to empty string
+	 * @param OutRelPath If the MountPoint is found, the RelativePath from the MountPoint to InFilePathOrPackageName is copied into this variable, otherwise it is set to empty string
+	 * 					 If InFilePathOrPackageName was a filepath, the extension is removed before copying it into OutRelpath
+	 *					 The OutRelPath is the same for both the LocalFilePath and the PackageName
+	 * @param OutFlexNameType If non-null, will be set to whether InFilePathOrPackageName was a PackageName or Filename if the MountPoint is found, otherwise it is set to EFlexNameType::Invalid
+	 * @param OutFailureReason If non-null, will be set to the reason InPath could not be converted, or to EErrorCode::Unknown if the function was successful.
+	 * @return True if the MountPoint was found, else false
+	 */
+	static bool TryGetMountPointForPath(FStringView InFilePathOrPackageName, FStringBuilderBase& OutMountPointPackageName, FStringBuilderBase& OutMountPointFilePath, FStringBuilderBase& OutRelPath,
+		EFlexNameType* OutFlexNameType = nullptr, EErrorCode* OutFailureReason = nullptr);
+
 	/** Event that is triggered when a new content path is mounted */
 	static FOnContentPathMountedEvent OnContentPathMountedEvent;
 
@@ -594,5 +726,7 @@ private:
 
 	/** Delegate used to check whether a package exist without using the filesystem. */
 	static FDoesPackageExistOverride DoesPackageExistOverrideDelegate;
+
+	friend class FPackagePath;
 };
 
diff --git a/Engine/Source/Runtime/CoreUObject/Public/Serialization/AsyncPackageLoader.h b/Engine/Source/Runtime/CoreUObject/Public/Serialization/AsyncPackageLoader.h
--- a/Engine/Source/Runtime/CoreUObject/Public/Serialization/AsyncPackageLoader.h
+++ b/Engine/Source/Runtime/CoreUObject/Public/Serialization/AsyncPackageLoader.h
@@ -3,6 +3,7 @@
 #pragma once
 
 #include "CoreMinimal.h"
+#include "Misc/PackagePath.h"
 #include "UObject/Object.h"
 #include "UObject/UObjectArray.h"
 #include "Stats/Stats2.h"
@@ -96,23 +97,34 @@ public:
 
 	virtual void StartThread() = 0;
 
+	UE_DEPRECATED(5.0, "Call the version that takes a FPackagePath instead")
+	COREUOBJECT_API int32 LoadPackage(
+		const FString& InPackageName,
+		const FGuid* InGuid,
+		const TCHAR* InPackageToLoadFrom,
+		FLoadPackageAsyncDelegate InCompletionDelegate,
+		EPackageFlags InPackageFlags,
+		int32 InPIEInstanceID,
+		int32 InPackagePriority,
+		const FLinkerInstancingContext* InstancingContext);
+
 	/**
 	 * Asynchronously load a package.
 	 *
-	 * @param	InName					Name of package to load
-	 * @param	InGuid					GUID of the package to load, or nullptr for "don't care"
-	 * @param	InPackageToLoadFrom		If non-null, this is another package name. We load from this package name, into a (probably new) package named InName
+	 * @param	PackagePath		Name of package to load. The package is created if it does not already exist.
+	 * @param	CustomName				If not none, this is the name of the package to load into (and create if not yet existing). If none, the name is take from PackagePath.
 	 * @param	InCompletionDelegate	Delegate to be invoked when the packages has finished streaming
+	 * @param	InGuid					GUID of the package to load, or nullptr for "don't care"
 	 * @param	InPackageFlags			Package flags used to construct loaded package in memory
 	 * @param	InPIEInstanceID			Play in Editor instance ID
 	 * @param	InPackagePriority		Loading priority
 	 * @return Unique ID associated with this load request (the same package can be associated with multiple IDs).
 	 */
 	virtual int32 LoadPackage(
-			const FString& InPackageName,
-			const FGuid* InGuid,
-			const TCHAR* InPackageToLoadFrom,
+			const FPackagePath& PackagePath,
+			FName CustomPackageName,
 			FLoadPackageAsyncDelegate InCompletionDelegate,
+			const FGuid* InGuid,
 			EPackageFlags InPackageFlags,
 			int32 InPIEInstanceID,
 			int32 InPackagePriority,
diff --git a/Engine/Source/Runtime/CoreUObject/Public/Serialization/BulkData.h b/Engine/Source/Runtime/CoreUObject/Public/Serialization/BulkData.h
--- a/Engine/Source/Runtime/CoreUObject/Public/Serialization/BulkData.h
+++ b/Engine/Source/Runtime/CoreUObject/Public/Serialization/BulkData.h
@@ -8,6 +8,7 @@
 #include "Async/Future.h"
 #include "Async/AsyncFileHandle.h"
 #include "GenericPlatform/GenericPlatformFile.h"
+#include "Misc/PackagePath.h"
 #include "Serialization/FileRegions.h"
 #include "BulkDataCommon.h"
 #include "BulkData2.h"
@@ -251,7 +252,7 @@ private:
 			bAllocated = false;
 		}
 
-		COREUOBJECT_API bool MapFile(const TCHAR *Filename, int64 Offset, int64 Size);
+		COREUOBJECT_API bool MapFile(const FPackagePath& PackagePath, EPackageSegment PackageSegment, int64 Offset, int64 Size);
 		COREUOBJECT_API void UnmapFile();
 
 		FOwnedBulkDataPtr* StealFileMapping()
@@ -493,19 +494,21 @@ public:
 	 */
 	void ClearBulkDataFlags( uint32 BulkDataFlagsToClear );
 
-	/** 
-	 * Returns the filename this bulkdata resides in
-	 *
-	 * @return Filename where this bulkdata can be loaded from
-	 **/
-	const FString& GetFilename() const { return Filename; }
+	UE_DEPRECATED(5.0, "Use GetPackagePath instead")
+	FString GetFilename() const { return GetPackagePath().GetLocalFullPath(GetPackageSegment()); }
+
+	/** Returns the PackagePath this bulkdata resides in */
+	const FPackagePath& GetPackagePath() const { return PackagePath; }
+
+	/** Returns which segment of its PackagePath this bulkdata resides in */
+	EPackageSegment GetPackageSegment() const { return PackageSegment; }
 
 	/** 
 	 * Returns the io filename hash associated with this bulk data.
 	 *
 	 * @return Hash or INVALID_IO_FILENAME_HASH if invalid.
 	 **/
-	FIoFilenameHash GetIoFilenameHash() const { return MakeIoFilenameHash(Filename); }
+	FIoFilenameHash GetIoFilenameHash() const { return MakeIoFilenameHash(PackagePath); }
 
 	/*-----------------------------------------------------------------------------
 		Data retrieval and manipulation.
@@ -690,13 +693,16 @@ public:
 	 * There is no way to validate this and it is up to the caller to make sure that it is correct.
 	 * The memory to be read into will be automatically allocated the size of which can be retrieved by calling IBulkDataIORequest::GetSize()
 	 *
-	 * @param Filename			The file to read from.
+	 * @param PackagePath		The PackagePath to read from.
+	 * @param PackageSegment	The PackageSegment to read from.
 	 * @param Start				The bulk data to start reading from.
 	 * @param End				The bulk data to finish reading from.
 	 * @param Priority			Priority and flags of the request. If this includes AIOP_FLAG_PRECACHE, then memory will never be returned. The request should always be canceled and waited for, even for a precache request.
 	 * @param CompleteCallback	Called from an arbitrary thread when the request is complete. Can be nullptr, if non-null, must remain valid until it is called. It will always be called.
 	 * @return					A request for the read. This is owned by the caller and must be deleted by the caller.
 	**/
+	static IBulkDataIORequest* CreateStreamingRequestForRange(const FPackagePath& PackagePath, EPackageSegment PackageSegment, const BulkDataRangeArray& RangeArray, EAsyncIOPriorityAndFlags Priority, FBulkDataIORequestCallBack* CompleteCallback);
+	UE_DEPRECATED(4.26, "Use version that takes a FPackagePath instead")
 	static IBulkDataIORequest* CreateStreamingRequestForRange(const FString& Filename, const BulkDataRangeArray& RangeArray, EAsyncIOPriorityAndFlags Priority, FBulkDataIORequestCallBack* CompleteCallback);
 #endif
 
@@ -819,8 +825,10 @@ private:
 	TFuture<bool> SerializeFuture;
 
 protected:
-	/** name of the package file containing the bulkdata */
-	FString				Filename;
+	/** name of the PackagePath containing the bulkdata */
+	FPackagePath PackagePath;
+	/** Which segment of the package's data in the PackageResourceManager contains the bulkdata */
+	EPackageSegment PackageSegment;
 #if WITH_EDITOR
 	/** Archive associated with bulk data for serialization																*/
 	FArchive*			AttachedAr;
diff --git a/Engine/Source/Runtime/CoreUObject/Public/Serialization/BulkData2.h b/Engine/Source/Runtime/CoreUObject/Public/Serialization/BulkData2.h
--- a/Engine/Source/Runtime/CoreUObject/Public/Serialization/BulkData2.h
+++ b/Engine/Source/Runtime/CoreUObject/Public/Serialization/BulkData2.h
@@ -5,8 +5,11 @@
 #include "CoreMinimal.h"
 #include "BulkDataCommon.h"
 #include "BulkDataBuffer.h"
+#include "Containers/StringView.h"
 #include "Async/AsyncFileHandle.h"
 #include "IO/IoDispatcher.h"
+#include "Misc/PackagePath.h"
+#include "Misc/PackageSegment.h"
 
 struct FOwnedBulkDataPtr;
 class IMappedFileHandle;
@@ -18,6 +21,8 @@ using FIoFilenameHash = uint32;
 const FIoFilenameHash INVALID_IO_FILENAME_HASH = 0;
 /** Helpers to create the hash from a filename. Returns IOFILENAMEHASH_NONE if and only if the filename is empty. */
 COREUOBJECT_API FIoFilenameHash MakeIoFilenameHash(const FString& Filename);
+/** Helpers to create the hash from a FPackagePath. Returns IOFILENAMEHASH_NONE if and only if the PackagePath is empty. */
+COREUOBJECT_API FIoFilenameHash MakeIoFilenameHash(const FPackagePath& Filename);
 /** Helpers to create the hash from a chunk id. Returns IOFILENAMEHASH_NONE if and only if the chunk id is invalid. */
 COREUOBJECT_API FIoFilenameHash MakeIoFilenameHash(const FIoChunkId& ChunkID);
 
@@ -208,8 +213,15 @@ public:
 	// Added for compatibility with the older BulkData system
 	int64 GetBulkDataOffsetInFile() const;
 
+	UE_DEPRECATED(5.0, "Use GetPackagePath instead")
 	FString GetFilename() const;
 
+	/** Returns the PackagePath this bulkdata resides in */
+	FPackagePath GetPackagePath() const;
+
+	/** Returns which segment of its PackagePath this bulkdata resides in */
+	EPackageSegment GetPackageSegment() const;
+
 	/** 
 	 * Returns the io filename hash associated with this bulk data.
 	 *
@@ -240,11 +252,11 @@ private:
 	 */
 	bool CanDiscardInternalData() const;
 
-	void ProcessDuplicateData(FArchive& Ar, const UPackage* Package, const FString* Filename, int64& InOutOffsetInFile);
+	void ProcessDuplicateData(FArchive& Ar, const UPackage* Package, const FPackagePath* PackagePath, int64& InOutOffsetInFile);
 	void SerializeDuplicateData(FArchive& Ar, EBulkDataFlags& OutBulkDataFlags, int64& OutBulkDataSizeOnDisk, int64& OutBulkDataOffsetInFile);
 	void SerializeBulkData(FArchive& Ar, void* DstBuffer, int64 DataLength);
 
-	bool MemoryMapBulkData(const FString& Filename, int64 OffsetInBulkData, int64 BytesToRead);
+	bool MemoryMapBulkData(const FPackagePath& PackagePath, EPackageSegment PackageSegment, int64 OffsetInBulkData, int64 BytesToRead);
 
 	// Methods for dealing with the allocated data
 	FORCEINLINE void* AllocateData(SIZE_T SizeInBytes) { return DataAllocation.AllocateData(this, SizeInBytes); }
@@ -256,7 +268,7 @@ private:
 	/** Blocking call that waits until any pending async load finishes */
 	void FlushAsyncLoading();
 	
-	FString ConvertFilenameFromFlags(const FString& Filename) const;
+	EPackageSegment GetPackageSegmentFromFlags() const;
 
 private:
 	using AsyncCallback = TFunction<void(TIoStatusOr<FIoBuffer>)>;
@@ -265,7 +277,7 @@ private:
 	void LoadDataAsynchronously(AsyncCallback&& Callback);
 
 	// Used by LoadDataDirectly/LoadDataAsynchronously
-	void InternalLoadFromFileSystem(void** DstBuffer);
+	void InternalLoadFromPackageResource(void** DstBuffer);
 	void InternalLoadFromIoStore(void** DstBuffer);
 	void InternalLoadFromIoStoreAsync(void** DstBuffer, AsyncCallback&& Callback);
 
diff --git a/Engine/Source/Runtime/CoreUObject/Public/UObject/Linker.h b/Engine/Source/Runtime/CoreUObject/Public/UObject/Linker.h
--- a/Engine/Source/Runtime/CoreUObject/Public/UObject/Linker.h
+++ b/Engine/Source/Runtime/CoreUObject/Public/UObject/Linker.h
@@ -11,6 +11,7 @@
 #include "Templates/RefCounting.h"
 
 class FReferenceCollector;
+class FPackagePath;
 
 DECLARE_LOG_CATEGORY_EXTERN(LogLinker, Log, All);
 
@@ -239,6 +240,7 @@ public:
 	TArray<FGatherableTextData> GatherableTextDataMap;
 
 	/** The name of the file for this package */
+	UE_DEPRECATED(5.0, "Use GetDebugName for logging identifiers. For other purposes, use GetPackagePath on LinkerLoad and GetFilename on LinkerSave.")
 	FString					Filename;
 
 	/** If true, filter out exports that are for clients but not servers */
@@ -251,6 +253,8 @@ public:
 	class FSHA1*			ScriptSHA;
 
 	/** Constructor. */
+	FLinker(ELinkerType::Type InType, UPackage* InRoot);
+	UE_DEPRECATED(5.0, "Linker's filename is deprecated; subclasses should create their own filename if required")
 	FLinker(ELinkerType::Type InType, UPackage* InRoot, const TCHAR* InFilename);
 
 	virtual ~FLinker();
@@ -299,6 +303,9 @@ public:
 		return LinkerType;
 	}
 
+	/** Returns a descriptor of the PackagePath this Linker is reading from or writing to, usable for an identifier in warning and log messages */
+	virtual FString GetDebugName() const;
+
 	/**
 	 * I/O function
 	*
@@ -618,7 +625,7 @@ COREUOBJECT_API void DeleteLoader(FLinkerLoad* Loader);
 /** 
  * Loads a linker for a package and returns it without loading any objects.
  * @param InOuter Package if known, can be null
- * @param InLongPackageName Name of the package to load
+ * @param PackagePath Package resource to load, must not be empty
  * @param LoadFlags Flags to pass to the new linker
  * @param Sandbox Additional sandbox for loading
  * @param CompatibleGuid Net GUID
@@ -626,13 +633,18 @@ COREUOBJECT_API void DeleteLoader(FLinkerLoad* Loader);
  * @param LinkerLoadedCallback Callback when the linker is loaded (or not found)
  * @return Pointer to the loaded linker or null if the file didn't exist
  */
+COREUOBJECT_API FLinkerLoad* LoadPackageLinker(UPackage* InOuter, const FPackagePath& PackagePath, uint32 LoadFlags, UPackageMap* Sandbox, FGuid* CompatibleGuid, FArchive* InReaderOverride, TFunctionRef<void(FLinkerLoad* LoadedLinker)> LinkerLoadedCallback);
+COREUOBJECT_API FLinkerLoad* LoadPackageLinker(UPackage* InOuter, const FPackagePath& PackagePath, uint32 LoadFlags = LOAD_None, UPackageMap* Sandbox = nullptr, FGuid* CompatibleGuid = nullptr, FArchive* InReaderOverride = nullptr);
+
+UE_DEPRECATED(5.0, "Use version that takes an FPackagePath instead")
 COREUOBJECT_API FLinkerLoad* LoadPackageLinker(UPackage* InOuter, const TCHAR* InLongPackageName, uint32 LoadFlags, UPackageMap* Sandbox, FGuid* CompatibleGuid, FArchive* InReaderOverride, TFunctionRef<void(FLinkerLoad* LoadedLinker)> LinkerLoadedCallback);
+UE_DEPRECATED(5.0, "Use version that takes an FPackagePath instead")
 COREUOBJECT_API FLinkerLoad* LoadPackageLinker(UPackage* InOuter, const TCHAR* InLongPackageName, uint32 LoadFlags = LOAD_None, UPackageMap* Sandbox = nullptr, FGuid* CompatibleGuid = nullptr, FArchive* InReaderOverride = nullptr);
 
 /** 
  * Gets a linker for a package and returns it without loading any objects. This call must be preceeded by BeginLoad and followed by EndLoad calls
  * @param InOuter Package if known, can be null
- * @param InLongPackageName Name of the package to load
+ * @param PackagePath Package resource to load, must not be empty
  * @param LoadFlags Flags to pass to the new linker
  * @param Sandbox Additional sandbox for loading
  * @param CompatibleGuid Net GUID
@@ -642,11 +654,14 @@ COREUOBJECT_API FLinkerLoad* LoadPackageLinker(UPackage* InOuter, const TCHAR* I
  * @param InstancingContext Optional instancing context to pass in if a linker is created
  * @return Pointer to the loaded linker or null if the file didn't exist
  */
+COREUOBJECT_API FLinkerLoad* GetPackageLinker(UPackage* InOuter, const FPackagePath& PackagePath, uint32 LoadFlags, UPackageMap* Sandbox, FGuid* CompatibleGuid, FArchive* InReaderOverride = nullptr, FUObjectSerializeContext** InOutLoadContext = nullptr, FLinkerLoad* ImportLinker = nullptr, const FLinkerInstancingContext* InstancingContext = nullptr);
+
+UE_DEPRECATED(5.0, "Use version that takes a PackagePath")
 COREUOBJECT_API FLinkerLoad* GetPackageLinker(UPackage* InOuter, const TCHAR* InLongPackageName, uint32 LoadFlags, UPackageMap* Sandbox, FGuid* CompatibleGuid, FArchive* InReaderOverride = nullptr, FUObjectSerializeContext** InOutLoadContext = nullptr, FLinkerLoad* ImportLinker = nullptr, const FLinkerInstancingContext* InstancingContext = nullptr);
 
 
 
-COREUOBJECT_API FString GetPrestreamPackageLinkerName(const TCHAR* InLongPackageName, bool bExistSkip = true);
+COREUOBJECT_API FString GetPrestreamPackageLinkerName(const TCHAR* InLongPackageName, bool bSkipIfExists = true);
 
 UE_DEPRECATED(4.25, "No longer used; use version that takes a UPackage* and call EnsureLoadingComplete separately.")
 COREUOBJECT_API void ResetLoadersForSave(UObject* InOuter, const TCHAR *Filename);
diff --git a/Engine/Source/Runtime/CoreUObject/Public/UObject/LinkerLoad.h b/Engine/Source/Runtime/CoreUObject/Public/UObject/LinkerLoad.h
--- a/Engine/Source/Runtime/CoreUObject/Public/UObject/LinkerLoad.h
+++ b/Engine/Source/Runtime/CoreUObject/Public/UObject/LinkerLoad.h
@@ -3,11 +3,13 @@
 #pragma once
 
 #include "CoreMinimal.h"
+#include "Misc/PackagePath.h"
 #include "Serialization/ArchiveUObject.h"
 #include "UObject/LazyObjectPtr.h"
+#include "UObject/Linker.h"
 #include "UObject/SoftObjectPtr.h"
 #include "UObject/ObjectResource.h"
-#include "UObject/Linker.h"
+#include "UObject/PackageResourceManager.h"
 
 class FLinkerPlaceholderBase;
 class IPakFile;
@@ -136,6 +138,17 @@ public:
 		return bIsAsyncLoader ? (FAsyncArchive*)Loader : nullptr;
 	}
 
+	virtual FString GetDebugName() const override
+	{
+		return GetPackagePath().GetDebugName();
+	}
+
+	/** Get the PackagePath being loaded. For linkers created from LoadPackage this will be a mounted and extension-specified path, but it may be unmounted and unspecified for other linkers */
+	const FPackagePath& GetPackagePath() const
+	{
+		return PackagePath;
+	}
+
 	FORCEINLINE const FLinkerInstancingContext& GetInstancingContext() const
 	{
 		return InstancingContext;
@@ -153,6 +166,11 @@ private:
 	FArchiveFormatterType* StructuredArchiveFormatter;
 	TOptional<FStructuredArchive::FRecord> StructuredArchiveRootRecord;
 	TArray<FStructuredArchiveChildReader*> ExportReaders;
+	/** The packagepath being loaded */
+	FPackagePath		PackagePath;
+
+	/** Set the packagepath being loaded */
+	void SetPackagePath(const FPackagePath& PackagePath);
 
 	/** The archive that actually reads the raw data from disk.																*/
 	FArchive*				Loader;
@@ -471,13 +489,15 @@ public:
 	 * Creates and returns a FLinkerLoad object.
 	 *
 	 * @param	Parent				Parent object to load into, can be NULL (most likely case)
-	 * @param	Filename			Name of file on disk to load
+	 * @param	PackagePath			Path of the package on disk to load
 	 * @param	LoadFlags			Load flags determining behavior
 	 * @param	InLoader			Loader archive override
 	 * @param	InstancingContext	Context to remap package name when loading a package on disk into a package with a different name
 	 *
 	 * @return	new FLinkerLoad object for Parent/ Filename
 	 */
+	COREUOBJECT_API static FLinkerLoad* CreateLinker(FUObjectSerializeContext* LoadContext, UPackage* Parent, const FPackagePath& PackagePath, uint32 LoadFlags, FArchive* InLoader = nullptr, const FLinkerInstancingContext* InstancingContext = nullptr);
+	UE_DEPRECATED(5.0, "Use version that takes an FPackagePath instead")
 	COREUOBJECT_API static FLinkerLoad* CreateLinker(FUObjectSerializeContext* LoadContext, UPackage* Parent, const TCHAR* Filename, uint32 LoadFlags, FArchive* InLoader = nullptr, const FLinkerInstancingContext* InstancingContext = nullptr);
 
 	void Verify();
@@ -664,6 +684,8 @@ public:
 #if WITH_EDITOR
 	COREUOBJECT_API static bool GetPreloadingEnabled();
 	COREUOBJECT_API static void SetPreloadingEnabled(bool bEnabled);
+	COREUOBJECT_API static bool TryGetPreloadedLoader(const FPackagePath& InPackagePath, FOpenPackageResult& OutResult);
+	UE_DEPRECATED(5.0, "Use version that takes a PackagePath instead")
 	COREUOBJECT_API static bool TryGetPreloadedLoader(FArchive*& OutLoader, const TCHAR* FileName);
 	private:
 		static bool bPreloadingEnabled;
@@ -909,13 +931,13 @@ private:
 	 * true in which case the returned linker object has finished the async creation process.
 	 *
 	 * @param	Parent				Parent object to load into, can be NULL (most likely case)
-	 * @param	Filename			Name of file on disk to load
+	 * @param	PackagePath			Path of the package data to load
 	 * @param	LoadFlags			Load flags determining behavior
 	 * @param	InstancingContext	Context to remap package name when loading a package on disk into a package with a different name
 	 *
 	 * @return	new FLinkerLoad object for Parent/ Filename
 	 */
-	COREUOBJECT_API static FLinkerLoad* CreateLinkerAsync(FUObjectSerializeContext* LoadContext, UPackage* Parent, const TCHAR* Filename, uint32 LoadFlags, const FLinkerInstancingContext* InstancingContext
+	COREUOBJECT_API static FLinkerLoad* CreateLinkerAsync(FUObjectSerializeContext* LoadContext, UPackage* Parent, const FPackagePath& PackagePath, uint32 LoadFlags, const FLinkerInstancingContext* InstancingContext
 		, TFunction<void()>&& InSummaryReadyCallback
 	);
 
@@ -936,11 +958,11 @@ protected: // Daniel L: Made this protected so I can override the constructor an
 	 * Private constructor, passing arguments through from CreateLinker.
 	 *
 	 * @param	Parent				Parent object to load into, can be NULL (most likely case)
-	 * @param	Filename			Name of file on disk to load
+	 * @param	PackagePath			Path of the package data to load
 	 * @param	LoadFlags			Load flags determining behavior
 	 * @param	InstancingContext	The instancing context for remapping imports if needed.
 	 */
-	FLinkerLoad(UPackage* InParent, const TCHAR* InFilename, uint32 InLoadFlags, FLinkerInstancingContext InstancingContext = FLinkerInstancingContext());
+	FLinkerLoad(UPackage* InParent, const FPackagePath& PackagePath, uint32 InLoadFlags, FLinkerInstancingContext InstancingContext = FLinkerInstancingContext());
 private:
 	/**
 	 * Returns whether the time limit allotted has been exceeded, if enabled.
diff --git a/Engine/Source/Runtime/CoreUObject/Public/UObject/LinkerSave.h b/Engine/Source/Runtime/CoreUObject/Public/UObject/LinkerSave.h
--- a/Engine/Source/Runtime/CoreUObject/Public/UObject/LinkerSave.h
+++ b/Engine/Source/Runtime/CoreUObject/Public/UObject/LinkerSave.h
@@ -127,6 +127,10 @@ public:
 	using FLinker::Serialize;
 	void Serialize( void* V, int64 Length );
 
+	// FLinker interface
+	virtual FString GetDebugName() const override;
+
+	// LinkerSave interface
 	/**
 	 * Closes and deletes the Saver (file, memory or custom writer) which will close any associated file handle.
 	 * Returns false if the owned saver contains errors after closing it, true otherwise.
@@ -137,4 +141,11 @@ public:
 	 * Sets a flag indicating that this archive contains data required to be gathered for localization.
 	 */
 	void ThisRequiresLocalizationGather();
+
+	/** Get the filename being saved to */
+	const FString& GetFilename() const;
+
+protected:
+	/** Set the filename being saved to */
+	void SetFilename(FStringView InFilename);
 };
diff --git a/Engine/Source/Runtime/CoreUObject/Public/UObject/Package.h b/Engine/Source/Runtime/CoreUObject/Public/UObject/Package.h
--- a/Engine/Source/Runtime/CoreUObject/Public/UObject/Package.h
+++ b/Engine/Source/Runtime/CoreUObject/Public/UObject/Package.h
@@ -9,12 +9,13 @@
 #include "UObject/PackageId.h"
 #include "UObject/LinkerSave.h"
 #include "Misc/Guid.h"
-#include "Misc/WorldCompositionUtility.h"
 #include "Misc/OutputDeviceError.h"
 #include "Misc/ObjectThumbnail.h"
+#include "Misc/PackagePath.h"
+#include "Misc/SecureHash.h"
+#include "Misc/WorldCompositionUtility.h"
 #include "Serialization/CustomVersion.h"
 #include "Templates/UniquePtr.h"
-#include "Misc/SecureHash.h"
 #include "Async/Future.h"
 
 class Error;
@@ -202,12 +203,16 @@ private:
 	
 	/** Globally unique id used to address I/O chunks within the package */
 	FPackageId PackageId;
+
+	/** The PackagePath this package was loaded from */
+	FPackagePath LoadedPath;
 public:
 
 	/** Editor only: PIE instance ID this package belongs to, INDEX_NONE otherwise */
 	int32 PIEInstanceID;		// TODO: strip from runtime?
 
 	/** The name of the file that this package was loaded from */
+	UE_DEPRECATED(5.0, "Use GetLoadedPath instead")
 	FName	FileName;
 
 	/** Linker load associated with this package */
@@ -332,6 +337,16 @@ public:
 	*/
 	void FullyLoad();
 
+	/**
+	 * Get the path this package was loaded from; may be different than packagename, and may not be set if the package was not loaded from disk
+	 */
+	const FPackagePath& GetLoadedPath() const;
+
+	/**
+	 * Set the path this package was loaded from; typically called only by the linker during load
+	 */
+	void SetLoadedPath(const FPackagePath& PackagePath);
+
 	/**
 	* Marks/Unmarks the package's bCanBeImported flag.
 	*/
diff --git a/Engine/Source/Runtime/CoreUObject/Public/UObject/UObjectGlobals.h b/Engine/Source/Runtime/CoreUObject/Public/UObject/UObjectGlobals.h
--- a/Engine/Source/Runtime/CoreUObject/Public/UObject/UObjectGlobals.h
+++ b/Engine/Source/Runtime/CoreUObject/Public/UObject/UObjectGlobals.h
@@ -12,6 +12,7 @@
 #include "UObject/PrimaryAssetId.h"
 #include "UObject/LinkerInstancingContext.h"
 #include "Misc/OutputDeviceRedirector.h"
+#include "Misc/PackagePath.h"
 #include "Containers/ArrayView.h"
 #include "Containers/StringView.h"
 #include "Templates/Function.h"
@@ -378,15 +379,31 @@ COREUOBJECT_API void StaticTick( float DeltaTime, bool bUseFullTimeLimit = true,
  * Loads a package and all contained objects that match context flags.
  *
  * @param	InOuter				Package to load new package into (usually nullptr or ULevel->GetOuter())
- * @param	InLongPackageName	Long package name to load
+ * @param	InLongPackageName	Long package name to load, if null the name is taken from InOuter
  * @param	LoadFlags			Flags controlling loading behavior, from the ELoadFlags enum
  * @param	InReaderOverride	Optional archive to use for reading package data
  * @param	InLoadContext		Additional context when called during serialization
+ * @param   InstancingContext   Additional context to map object names to their instanced counterpart when loading an instanced package
  *
  * @return	Loaded package if successful, nullptr otherwise
  */
 COREUOBJECT_API UPackage* LoadPackage( UPackage* InOuter, const TCHAR* InLongPackageName, uint32 LoadFlags, FArchive* InReaderOverride = nullptr, const FLinkerInstancingContext* InstancingContext = nullptr);
 
+/**
+ * Loads a package and all contained objects that match context flags.
+ *
+ * @param	InOuter				Package to load new package into (usually nullptr or ULevel->GetOuter())
+ * @param	InPackagePath		PackagePath to load, must be non-empty
+ * @param	LoadFlags			Flags controlling loading behavior, from the ELoadFlags enum
+ * @param	InReaderOverride	Optional archive to use for reading package data
+ * @param	InLoadContext		Additional context when called during serialization
+ * @param   InstancingContext   Additional context to map object names to their instanced counterpart when loading an instanced package
+ * @param   DiffPackagePath		An additional PackagePath to load and compare to the package at InPackagePath, used when diffing packages
+ *
+ * @return	Loaded package if successful, nullptr otherwise
+ */
+COREUOBJECT_API UPackage* LoadPackage(UPackage* InOuter, const FPackagePath& InPackagePath, uint32 LoadFlags, FArchive* InReaderOverride = nullptr, const FLinkerInstancingContext* InstancingContext = nullptr, const FPackagePath* DiffPackagePath = nullptr);
+
 /** Async package loading result */
 namespace EAsyncLoadingResult
 {
@@ -414,33 +431,47 @@ DECLARE_DELEGATE_ThreeParams(FLoadPackageAsyncDelegate, const FName& /*PackageNa
 
 /**
  * Asynchronously load a package and all contained objects that match context flags. Non-blocking.
- * This version is useful when loading multiple copies of the same package.
+ * Use this version to specify the PackagePath rather than having the other versions internally convert the InName to a PackagePath by searching the current package mount points.
+ * Use this version if you need to specify a packagename that is different from the packagename on disk; this is useful when loading multiple copies of the same package.
  *
- * @param	InName					Name of package to load
- * @param	InGuid					GUID of the package to load, or nullptr for "don't care"
- * @param	InPackageToLoadFrom		If non-null, this is another package name. We load from this package name, into a (probably new) package named InName
+ * @param	InPackagePath			PackagePath to load. Must be a mounted path. The package is created if it does not already exist.
+ * @param	InPackageNameToCreate	If not none, this is the name of the package to load the bytes on disk into (and create if not yet existing). If none, the name is taken from PackagePath.
  * @param	InCompletionDelegate	Delegate to be invoked when the packages has finished streaming
+ * @param	InGuid					GUID of the package to load, or nullptr for "don't care"
  * @param	InPackageFlags			Package flags used to construct loaded package in memory
  * @param	InPIEInstanceID			Play in Editor instance ID
  * @param	InPackagePriority		Loading priority
+ * @param   InstancingContext		Additional context to map object names to their instanced counterpart when loading an instanced package
  * @return Unique ID associated with this load request (the same package can be associated with multiple IDs).
  */
-COREUOBJECT_API int32 LoadPackageAsync(const FString& InName, const FGuid* InGuid = nullptr, const TCHAR* InPackageToLoadFrom = nullptr, FLoadPackageAsyncDelegate InCompletionDelegate = FLoadPackageAsyncDelegate(), EPackageFlags InPackageFlags = PKG_None, int32 InPIEInstanceID = INDEX_NONE, TAsyncLoadPriority InPackagePriority = 0, const FLinkerInstancingContext* InstancingContext = nullptr);
+COREUOBJECT_API int32 LoadPackageAsync(const FPackagePath& InPackagePath, FName InPackageNameToCreate = NAME_None, FLoadPackageAsyncDelegate InCompletionDelegate = FLoadPackageAsyncDelegate(), const FGuid* InGuid = nullptr, EPackageFlags InPackageFlags = PKG_None, int32 InPIEInstanceID = INDEX_NONE, TAsyncLoadPriority InPackagePriority = 0, const FLinkerInstancingContext* InstancingContext = nullptr);
 
 /**
  * Asynchronously load a package and all contained objects that match context flags. Non-blocking.
+ * Use this version for convenience when you just need to load a package without notification and with default behavior from a packagename/filename.
  *
- * @param	InName					Name of package to load
+ * @param	InName					PackageName or LocalFilePath of package to load. Must be a mounted name/path. The package is created if it does not already exist.
+ * @param	InGuid					GUID of the package to load, or nullptr for "don't care"
+ * @return Unique ID associated with this load request (the same package can be associated with multiple IDs).
+ */
+COREUOBJECT_API int32 LoadPackageAsync(const FString& InName, const FGuid* InGuid = nullptr);
+
+/**
+ * Asynchronously load a package and all contained objects that match context flags. Non-blocking.
+ * Use this version when you need to load a package with default behavior from a packagename/filename, and need to be notified when it is loaded.
+ *
+ * @param	InName					PackageName or LocalFilePath of package to load. Must be a mounted name/path. The package is created if it does not already exist.
  * @param	InCompletionDelegate	Delegate to be invoked when the packages has finished streaming
  * @param	InPackagePriority		Loading priority
  * @param	InPackageFlags			Package flags used to construct loaded package in memory
  * @param	InPIEInstanceID			Play in Editor instance ID
  * @return Unique ID associated with this load request (the same package can be associated with multiple IDs).
- *
- * @see FStreamableManager for an engine-level wrapper
  */
 COREUOBJECT_API int32 LoadPackageAsync(const FString& InName, FLoadPackageAsyncDelegate InCompletionDelegate, TAsyncLoadPriority InPackagePriority = 0, EPackageFlags InPackageFlags = PKG_None, int32 InPIEInstanceID = INDEX_NONE);
 
+UE_DEPRECATED(5.0, "Use version that takes a PackagePath")
+COREUOBJECT_API int32 LoadPackageAsync(const FString& InName, const FGuid* InGuid, const TCHAR* InPackageToLoadFrom, FLoadPackageAsyncDelegate InCompletionDelegate = FLoadPackageAsyncDelegate(), EPackageFlags InPackageFlags = PKG_None, int32 InPIEInstanceID = INDEX_NONE, TAsyncLoadPriority InPackagePriority = 0, const FLinkerInstancingContext* InstancingContext = nullptr);
+
 /**
 * Cancels all async package loading requests.
 */
diff --git a/Engine/Source/Runtime/CoreUObject/Public/Virtualization/VirtualizedBulkData.h b/Engine/Source/Runtime/CoreUObject/Public/Virtualization/VirtualizedBulkData.h
--- a/Engine/Source/Runtime/CoreUObject/Public/Virtualization/VirtualizedBulkData.h
+++ b/Engine/Source/Runtime/CoreUObject/Public/Virtualization/VirtualizedBulkData.h
@@ -6,6 +6,7 @@
 #include "HAL/Platform.h"
 #include "Memory/SharedBuffer.h"
 #include "Misc/Guid.h"
+#include "Misc/PackagePath.h"
 #include "Virtualization/VirtualizationManager.h"
 
 class FArchive;
@@ -123,7 +124,7 @@ private:
 	void PushData();
 	FSharedBufferConstPtr PullData() const;
 
-	FString GetFilenameFromOwner(UObject* Owner) const;
+	FPackagePath GetPackagePathFromOwner(UObject* Owner, EPackageSegment& OutPackageSegment) const;
 
 	bool CanUnloadData() const;
 
@@ -148,8 +149,11 @@ private:
 	/** If the data has been saved to disk rather than being Virtualized, was it stored in a compressed format? */
 	bool bIsDataStoredAsCompressed = false ;
 
-	/** File containing the payload (this will be empty if the payload does not come from a file)*/
-	FString Filename;
+	/** PackagePath containing the payload (this will be empty if the payload does not come from PackageResourceManager)*/
+	FPackagePath PackagePath;
+
+	/** PackageSegment to load with the packagepath (unused if the payload does not come from PackageResourceManager) */
+	EPackageSegment PackageSegment;
 
 	/** Offset of the payload in the file that contains it (INDEX_NONE if the payload does not come from a file)*/
 	int64 OffsetInFile = INDEX_NONE;
diff --git a/Engine/Source/Runtime/Engine/Classes/Engine/SkeletalMesh.h b/Engine/Source/Runtime/Engine/Classes/Engine/SkeletalMesh.h
--- a/Engine/Source/Runtime/Engine/Classes/Engine/SkeletalMesh.h
+++ b/Engine/Source/Runtime/Engine/Classes/Engine/SkeletalMesh.h
@@ -2110,7 +2110,9 @@ public:
 	//~ End UStreamableRenderAsset Interface.
 
 #if USE_BULKDATA_STREAMING_TOKEN
+	UE_DEPRECATED(5.0, "Use GetMipDataPackagePath instead")
 	bool GetMipDataFilename(const int32 MipIndex, FString& OutBulkDataFilename) const;
+	bool GetMipDataPackagePath(const int32 MipIndex, FPackagePath& OutPackagePath, EPackageSegment& OutPackageSegment) const;
 #endif
 
 	/**
diff --git a/Engine/Source/Runtime/Engine/Classes/Engine/StaticMesh.h b/Engine/Source/Runtime/Engine/Classes/Engine/StaticMesh.h
--- a/Engine/Source/Runtime/Engine/Classes/Engine/StaticMesh.h
+++ b/Engine/Source/Runtime/Engine/Classes/Engine/StaticMesh.h
@@ -1531,7 +1531,9 @@ public:
 	//~ End UStreamableRenderAsset Interface
 
 #if USE_BULKDATA_STREAMING_TOKEN
-	bool GetMipDataFilename(const int32 MipIndex, FString& BulkDataFilename) const;
+	UE_DEPRECATED(5.0, "Use GetMipDataPackagePath instead")
+	bool GetMipDataFilename(const int32 MipIndex, FString& OutBulkDataFilename) const;
+	bool GetMipDataPackagePath(const int32 MipIndex, FPackagePath& OutPackagePath, EPackageSegment& OutPackageSegment) const;
 #endif
 
 	/**
diff --git a/Engine/Source/Runtime/Engine/Private/Animation/AnimationStreaming.cpp b/Engine/Source/Runtime/Engine/Private/Animation/AnimationStreaming.cpp
--- a/Engine/Source/Runtime/Engine/Private/Animation/AnimationStreaming.cpp
+++ b/Engine/Source/Runtime/Engine/Private/Animation/AnimationStreaming.cpp
@@ -221,7 +221,8 @@ void FStreamingAnimationData::BeginPendingRequests(const TArray<uint32>& Indices
 				AnimationStreamingManager->OnAsyncFileCallback(this, ChunkIndex, ChunkSize, Req, bWasCancelled);
 			};
 
-			UE_LOG(LogAnimation, Warning, TEXT("Loading Stream Anim %s Chunk:%i Length: %.3f Offset:%i Size:%i File:%s\n"), *StreamableAnim->GetName(), ChunkIndex, Chunk.SequenceLength, Chunk.BulkData.GetBulkDataOffsetInFile(), Chunk.BulkData.GetBulkDataSize(), *Chunk.BulkData.GetFilename());
+			UE_LOG(LogAnimation, Warning, TEXT("Loading Stream Anim %s Chunk:%i Length: %.3f Offset:%i Size:%i File:%s\n"),
+				*StreamableAnim->GetName(), ChunkIndex, Chunk.SequenceLength, Chunk.BulkData.GetBulkDataOffsetInFile(), Chunk.BulkData.GetBulkDataSize(), *Chunk.BulkData.GetPackagePath().GetDebugName());
 			ChunkStorage.IORequest = Chunk.BulkData.CreateStreamingRequest(AsyncIOPriority, &AsyncFileCallBack, nullptr);
 			if (!ChunkStorage.IORequest)
 			{
diff --git a/Engine/Source/Runtime/Engine/Private/GameFramework/UpdateLevelVisibilityLevelInfo.cpp b/Engine/Source/Runtime/Engine/Private/GameFramework/UpdateLevelVisibilityLevelInfo.cpp
--- a/Engine/Source/Runtime/Engine/Private/GameFramework/UpdateLevelVisibilityLevelInfo.cpp
+++ b/Engine/Source/Runtime/Engine/Private/GameFramework/UpdateLevelVisibilityLevelInfo.cpp
@@ -26,7 +26,8 @@ FUpdateLevelVisibilityLevelInfo::FUpdateLevelVisibilityLevelInfo(const ULevel* c
 
 	// When packages are duplicated for PIE, they may not have a FileName.
 	// For now, just revert to the old behavior.
-	FileName = (LevelPackage->FileName == NAME_None) ? PackageName : LevelPackage->FileName;
+	FName LoadedPathName = LevelPackage->GetLoadedPath().GetPackageFName();
+	FileName = (LoadedPathName.IsNone()) ? PackageName : LoadedPathName;
 }
 
 bool FUpdateLevelVisibilityLevelInfo::NetSerialize(FArchive& Ar, UPackageMap* PackageMap, bool& bOutSuccess)
diff --git a/Engine/Source/Runtime/Engine/Private/Level.cpp b/Engine/Source/Runtime/Engine/Private/Level.cpp
--- a/Engine/Source/Runtime/Engine/Private/Level.cpp
+++ b/Engine/Source/Runtime/Engine/Private/Level.cpp
@@ -732,13 +732,14 @@ void ULevel::PostLoad()
 	{
 		UPackage* LevelPackage = GetPackage();
 		bool bPackageForPIE = LevelPackage->HasAnyPackageFlags(PKG_PlayInEditor);
-		bool bInstanced = !LevelPackage->FileName.IsNone() && (LevelPackage->FileName != LevelPackage->GetFName());
+		FName PackageResourceName = LevelPackage->GetLoadedPath().GetPackageFName();
+		bool bInstanced = !PackageResourceName.IsNone() && (PackageResourceName != LevelPackage->GetFName());
 
 		// if the level is instanced, create an instancing context for remapping the actor imports
 		FLinkerInstancingContext InstancingContext;
 		if (bInstanced)
 		{
-			InstancingContext.AddMapping(LevelPackage->FileName, LevelPackage->GetFName());
+			InstancingContext.AddMapping(PackageResourceName, LevelPackage->GetFName());
 		}
 
 		TArray<FString> ActorPackageNames = GetOnDiskExternalActorPackages();
@@ -2326,7 +2327,7 @@ TArray<UPackage*> ULevel::GetLoadedExternalActorPackages() const
 	{
 		for (TObjectIterator<UPackage> It; It; ++It)
 		{
-			if (!ActorPackages.Contains(*It) && It->FileName.ToString().Contains(ExternalActorsPath))
+			if (!ActorPackages.Contains(*It) && It->GetLoadedPath().GetPackageName().Contains(ExternalActorsPath))
 			{
 				ActorPackages.Add(*It);
 			}
@@ -2636,7 +2637,9 @@ bool ULevel::IsInstancedLevel() const
 {
 	UPackage* LevelPackage = GetOutermost();
 	static FName UnsavedPackageName("/Temp/Untitled");
-	return !LevelPackage->FileName.IsNone() && (LevelPackage->FileName != LevelPackage->GetFName()) && !LevelPackage->GetFName().IsEqual(UnsavedPackageName, ENameCase::IgnoreCase, /*bCompareNumber*/false);
+	const FPackagePath& PackagePath = LevelPackage->GetLoadedPath();
+	return !PackagePath.IsEmpty() && PackagePath.GetPackageFName() != LevelPackage->GetFName() 
+		&& !LevelPackage->GetFName().IsEqual(UnsavedPackageName, ENameCase::IgnoreCase, /*bCompareNumber*/false);
 }
 
 void ULevel::ApplyWorldOffset(const FVector& InWorldOffset, bool bWorldShift)
diff --git a/Engine/Source/Runtime/Engine/Private/LevelInstance/LevelInstanceActor.cpp b/Engine/Source/Runtime/Engine/Private/LevelInstance/LevelInstanceActor.cpp
--- a/Engine/Source/Runtime/Engine/Private/LevelInstance/LevelInstanceActor.cpp
+++ b/Engine/Source/Runtime/Engine/Private/LevelInstance/LevelInstanceActor.cpp
@@ -125,7 +125,7 @@ void ALevelInstance::LoadLevelInstance()
 			bool bForce = false;
 #if WITH_EDITOR
 			// When reinstancing or world wasn't ticked between changes. Avoid reloading the level but if the package changed, force the load
-			bForce = IsLoaded() && LevelInstanceSubsystem->GetLevelInstanceLevel(this)->GetPackage()->FileName != FName(*GetWorldAssetPackage());
+			bForce = IsLoaded() && LevelInstanceSubsystem->GetLevelInstanceLevel(this)->GetPackage()->GetLoadedPath() != FPackagePath::FromPackageNameChecked(GetWorldAssetPackage());
 #endif
 			LevelInstanceSubsystem->RequestLoadLevelInstance(this, bForce);
 		}
@@ -316,7 +316,7 @@ bool ALevelInstance::CheckForLoop(TSoftObjectPtr<UWorld> InLevelInstance, TArray
 			FName LongPackageName(*InLevelInstance.GetLongPackageName());
 			// Check to exclude NAME_None since Preview Levels are in the transient package
 			// Check the level we are spawned in to detect the loop (this will handle loops caused by LevelInstances and by regular level streaming)
-			if (LongPackageName != NAME_None && LevelInstanceActor->GetLevel()->GetPackage()->FileName == LongPackageName)
+			if (LongPackageName != NAME_None && LevelInstanceActor->GetLevel()->GetPackage()->GetLoadedPath() == FPackagePath::FromPackageNameChecked(LongPackageName))
 			{
 				bValid = false;
 				if (LoopStart)
diff --git a/Engine/Source/Runtime/Engine/Private/LevelStreaming.cpp b/Engine/Source/Runtime/Engine/Private/LevelStreaming.cpp
--- a/Engine/Source/Runtime/Engine/Private/LevelStreaming.cpp
+++ b/Engine/Source/Runtime/Engine/Private/LevelStreaming.cpp
@@ -1100,9 +1100,12 @@ bool ULevelStreaming::RequestLevel(UWorld* PersistentWorld, bool bAllowLevelLoad
 	if (bAllowLevelLoadRequests)
 	{
 		const FName DesiredPackageNameToLoad = bIsGameWorld ? GetLODPackageNameToLoad() : PackageNameToLoad;
-		const FString PackageNameToLoadFrom = DesiredPackageNameToLoad != NAME_None ? DesiredPackageNameToLoad.ToString() : DesiredPackageName.ToString();
+		FString NormalizedPackageName = (DesiredPackageNameToLoad.IsNone() ? DesiredPackageName : DesiredPackageNameToLoad).ToString();
+		// The PackageName might be an objectpath; convert it to a packagename if it is not one already
+		NormalizedPackageName = FPackageName::ObjectPathToPackageName(NormalizedPackageName);
+		FPackagePath PackagePath = FPackagePath::FromPackageNameChecked(NormalizedPackageName);
 
-		if (FPackageName::DoesPackageExist(PackageNameToLoadFrom))
+		if (FPackageName::DoesPackageExist(PackagePath, &PackagePath))
 		{
 			CurrentState = ECurrentState::Loading;
 			FWorldNotifyStreamingLevelLoading::Started(PersistentWorld);
@@ -1113,7 +1116,7 @@ bool ULevelStreaming::RequestLevel(UWorld* PersistentWorld, bool bAllowLevelLoad
 			// Kick off async load request.
 			STAT_ADD_CUSTOMMESSAGE_NAME( STAT_NamedMarker, *(FString( TEXT( "RequestLevel - " ) + DesiredPackageName.ToString() )) );
 			TRACE_BOOKMARK(TEXT("RequestLevel - %s"), *DesiredPackageName.ToString());
-			LoadPackageAsync(DesiredPackageName.ToString(), nullptr, *PackageNameToLoadFrom, FLoadPackageAsyncDelegate::CreateUObject(this, &ULevelStreaming::AsyncLevelLoadComplete), PackageFlags, PIEInstanceID, GetPriority());
+			LoadPackageAsync(PackagePath, DesiredPackageName, FLoadPackageAsyncDelegate::CreateUObject(this, &ULevelStreaming::AsyncLevelLoadComplete), nullptr, PackageFlags, PIEInstanceID, GetPriority());
 
 			// streamingServer: server loads everything?
 			// Editor immediately blocks on load and we also block if background level streaming is disabled.
@@ -1130,7 +1133,7 @@ bool ULevelStreaming::RequestLevel(UWorld* PersistentWorld, bool bAllowLevelLoad
 		}
 		else
 		{
-			UE_LOG(LogStreaming, Error,TEXT("Couldn't find file for package %s."), *PackageNameToLoadFrom);
+			UE_LOG(LogStreaming, Error,TEXT("Couldn't find file for package %s."), *PackagePath.GetDebugName());
 			CurrentState = ECurrentState::FailedToLoad;
 			return false;
 		}
diff --git a/Engine/Source/Runtime/Engine/Private/SkeletalMesh.cpp b/Engine/Source/Runtime/Engine/Private/SkeletalMesh.cpp
--- a/Engine/Source/Runtime/Engine/Private/SkeletalMesh.cpp
+++ b/Engine/Source/Runtime/Engine/Private/SkeletalMesh.cpp
@@ -6,10 +6,9 @@
 
 #include "Engine/SkeletalMesh.h"
 #include "Serialization/CustomVersion.h"
-#include "UObject/FrameworkObjectVersion.h"
 #include "Misc/App.h"
+#include "Misc/PackageSegment.h"
 #include "Modules/ModuleManager.h"
-#include "UObject/UObjectIterator.h"
 #include "EngineStats.h"
 #include "EngineGlobals.h"
 #include "RawIndexBuffer.h"
@@ -23,9 +22,15 @@
 #include "Components/SkeletalMeshComponent.h"
 #include "Engine/CollisionProfile.h"
 #include "ComponentReregisterContext.h"
+#include "UObject/CoreObjectVersion.h"
+#include "UObject/CoreRedirects.h"
 #include "UObject/EditorObjectVersion.h"
+#include "UObject/FrameworkObjectVersion.h"
+#include "UObject/NiagaraObjectVersion.h"
+#include "UObject/PackageResourceManager.h"
+#include "UObject/PropertyPortFlags.h"
 #include "UObject/RenderingObjectVersion.h"
-#include "UObject/CoreObjectVersion.h"
+#include "UObject/UObjectIterator.h"
 #include "EngineUtils.h"
 #include "EditorSupportDelegates.h"
 #include "GPUSkinVertexFactory.h"
@@ -46,14 +51,11 @@
 #include "SkeletalMeshTypes.h"
 #include "Rendering/SkeletalMeshVertexBuffer.h"
 #include "Rendering/SkeletalMeshRenderData.h"
-#include "UObject/PropertyPortFlags.h"
 #include "Templates/UniquePtr.h"
 #include "AnimationRuntime.h"
 #include "Animation/AnimSequence.h"
-#include "UObject/NiagaraObjectVersion.h"
 #include "Animation/SkinWeightProfile.h"
 #include "Streaming/SkeletalMeshUpdate.h"
-#include "UObject/CoreRedirects.h"
 #include "HAL/FileManager.h"
 
 #if WITH_EDITOR
@@ -1024,22 +1026,42 @@ int32 USkeletalMesh::CalcCumulativeLODSize(int32 NumLODs) const
 
 #if USE_BULKDATA_STREAMING_TOKEN
 bool USkeletalMesh::GetMipDataFilename(const int32 MipIndex, FString& OutBulkDataFilename) const
+{
+	FPackagePath PackagePath;
+	EPackageSegment PackageSegment;
+	if (GetMipDataPackagePath(MipIndex, PackagePath, PackageSegment))
+	{
+		OutBulkDataFilename = PackagePath.GetLocalFullPath(PackageSegment);
+		return true;
+	}
+	return false;
+}
+
+bool USkeletalMesh::GetMipDataPackagePath(const int32 MipIndex, FPackagePath& OutPackagePath, EPackageSegment& OutPackageSegment) const
 {
 	// TODO: this is slow. Should cache the name once per mesh
-	FString PackageName = GetOutermost()->FileName.ToString();
+	FPackagePath PackagePath = GetOutermost()->GetLoadedPath();
 	// Handle name redirection and localization
-	const FCoreRedirectObjectName RedirectedName =
-		FCoreRedirects::GetRedirectedName(
-			ECoreRedirectFlags::Type_Package,
-			FCoreRedirectObjectName(NAME_None, NAME_None, *PackageName));
-	FString LocalizedName;
-	LocalizedName = FPackageName::GetDelegateResolvedPackagePath(RedirectedName.PackageName.ToString());
-	LocalizedName = FPackageName::GetLocalizedPackagePath(LocalizedName);
-	bool bSucceed = FPackageName::DoesPackageExist(LocalizedName, nullptr, &OutBulkDataFilename);
+	{
+		FString PackageName = PackagePath.GetPackageName();
+		const FCoreRedirectObjectName RedirectedName =
+			FCoreRedirects::GetRedirectedName(
+				ECoreRedirectFlags::Type_Package,
+				FCoreRedirectObjectName(NAME_None, NAME_None, *PackageName));
+		FString LocalizedName;
+		LocalizedName = FPackageName::GetDelegateResolvedPackagePath(RedirectedName.PackageName.ToString());
+		LocalizedName = FPackageName::GetLocalizedPackagePath(LocalizedName);
+		if (LocalizedName != PackageName)
+		{
+			PackagePath = FPackagePath::FromPackageNameChecked(LocalizedName);
+		}
+	}
+	bool bSucceed = FPackageName::DoesPackageExist(PackagePath, &PackagePath);
 	check(bSucceed);
 	const FSkeletalMeshRenderData* SkelMeshRenderData = GetResourceForRendering();
 	const bool bLODIsOptional = SkelMeshRenderData ? (MipIndex < SkelMeshRenderData->LODRenderData.Num() - SkelMeshRenderData->NumNonOptionalLODs) : false;
-	OutBulkDataFilename = FPaths::ChangeExtension(OutBulkDataFilename, bLODIsOptional ? TEXT(".uptnl") : TEXT(".ubulk"));
+	OutPackagePath = MoveTemp(PackagePath);
+	OutPackageSegment = bLODIsOptional ? EPackageSegment::BulkDataOptional : EPackageSegment::BulkDataDefault;
 	return true;
 }
 #endif // USE_BULKDATA_STREAMING_TOKEN
@@ -1047,10 +1069,11 @@ bool USkeletalMesh::GetMipDataFilename(const int32 MipIndex, FString& OutBulkDat
 FIoFilenameHash USkeletalMesh::GetMipIoFilenameHash(const int32 MipIndex) const
 {
 #if USE_BULKDATA_STREAMING_TOKEN
-	FString MipFilename;
-	if (GetMipDataFilename(MipIndex, MipFilename))
+	FPackagePath PackagePath;
+	EPackageSegment PackageSegment;
+	if (GetMipDataPackagePath(MipIndex, PackagePath, PackageSegment))
 	{
-		return MakeIoFilenameHash(MipFilename);
+		return MakeIoFilenameHash(PackagePath);
 	}
 #else
 	if (SkeletalMeshRenderData && SkeletalMeshRenderData->LODRenderData.IsValidIndex(MipIndex))
@@ -1067,8 +1090,9 @@ FIoFilenameHash USkeletalMesh::GetMipIoFilenameHash(const int32 MipIndex) const
 bool USkeletalMesh::DoesMipDataExist(const int32 MipIndex) const
 {
 #if USE_BULKDATA_STREAMING_TOKEN
-	FString MipDataFilename;
-	return GetMipDataFilename(MipIndex, MipDataFilename) && IFileManager::Get().FileExists(*MipDataFilename);
+	FPackagePath PackagePath;
+	EPackageSegment PackageSegment;
+	return GetMipDataPackagePath(MipIndex, PackagePath, PackageSegment) && IPackageResourceManager::Get().DoesPackageExist(PackagePath, PackageSegment);
 #else
 	return SkeletalMeshRenderData && SkeletalMeshRenderData->LODRenderData.IsValidIndex(MipIndex) && SkeletalMeshRenderData->LODRenderData[MipIndex].StreamingBulkData.DoesExist();
 #endif
diff --git a/Engine/Source/Runtime/Engine/Private/StaticMesh.cpp b/Engine/Source/Runtime/Engine/Private/StaticMesh.cpp
--- a/Engine/Source/Runtime/Engine/Private/StaticMesh.cpp
+++ b/Engine/Source/Runtime/Engine/Private/StaticMesh.cpp
@@ -7,11 +7,10 @@
 #include "Engine/StaticMesh.h"
 #include "Serialization/MemoryWriter.h"
 #include "Misc/ConfigCacheIni.h"
+#include "Misc/PackageSegment.h"
 #include "Misc/ScopedSlowTask.h"
-#include "UObject/FrameworkObjectVersion.h"
 #include "Misc/App.h"
 #include "Modules/ModuleManager.h"
-#include "UObject/UObjectAnnotation.h"
 #include "RenderingThread.h"
 #include "VertexFactory.h"
 #include "LocalVertexFactory.h"
@@ -21,8 +20,11 @@
 #include "Engine/CollisionProfile.h"
 #include "Serialization/MemoryReader.h"
 #include "UObject/EditorObjectVersion.h"
-#include "UObject/RenderingObjectVersion.h"
+#include "UObject/FrameworkObjectVersion.h"
 #include "UObject/Package.h"
+#include "UObject/PackageResourceManager.h"
+#include "UObject/RenderingObjectVersion.h"
+#include "UObject/UObjectAnnotation.h"
 #include "EngineUtils.h"
 #include "Engine/AssetUserData.h"
 #include "StaticMeshResources.h"
@@ -6160,20 +6162,40 @@ int32 UStaticMesh::CalcCumulativeLODSize(int32 NumLODs) const
 
 #if USE_BULKDATA_STREAMING_TOKEN
 bool UStaticMesh::GetMipDataFilename(const int32 MipIndex, FString& OutBulkDataFilename) const
+{
+	FPackagePath PackagePath;
+	EPackageSegment PackageSegment;
+	if (GetMipDataPackagePath(MipIndex, PackagePath, PackageSegment))
+	{
+		OutBulkDataFilename = PackagePath.GetLocalFullPath(PackageSegment);
+		return true;
+	}
+	return false;
+}
+
+bool UStaticMesh::GetMipDataPackagePath(const int32 MipIndex, FPackagePath& OutPackagePath, EPackageSegment& OutPackageSegment) const
 {
 	// TODO: this is slow. Should cache the name once per mesh
-	FString PackageName = GetOutermost()->FileName.ToString();
+	FPackagePath PackagePath = GetOutermost()->GetLoadedPath();
 	// Handle name redirection and localization
-	const FCoreRedirectObjectName RedirectedName =
-		FCoreRedirects::GetRedirectedName(
-			ECoreRedirectFlags::Type_Package,
-			FCoreRedirectObjectName(NAME_None, NAME_None, *PackageName));
-	FString LocalizedName;
-	LocalizedName = FPackageName::GetDelegateResolvedPackagePath(RedirectedName.PackageName.ToString());
-	LocalizedName = FPackageName::GetLocalizedPackagePath(LocalizedName);
-	bool bSucceed = FPackageName::DoesPackageExist(LocalizedName, nullptr, &OutBulkDataFilename);
+	{
+		FString PackageName = PackagePath.GetPackageName();
+		const FCoreRedirectObjectName RedirectedName =
+			FCoreRedirects::GetRedirectedName(
+				ECoreRedirectFlags::Type_Package,
+				FCoreRedirectObjectName(NAME_None, NAME_None, *PackageName));
+		FString LocalizedName;
+		LocalizedName = FPackageName::GetDelegateResolvedPackagePath(RedirectedName.PackageName.ToString());
+		LocalizedName = FPackageName::GetLocalizedPackagePath(LocalizedName);
+		if (LocalizedName != PackageName)
+		{
+			PackagePath = FPackagePath::FromPackageNameChecked(LocalizedName);
+		}
+	}
+	bool bSucceed = FPackageName::DoesPackageExist(PackagePath, &PackagePath);
 	check(bSucceed);
-	OutBulkDataFilename = FPaths::ChangeExtension(OutBulkDataFilename, MipIndex < GetMinLOD().Default ? TEXT(".uptnl") : TEXT(".ubulk"));
+	OutPackagePath = MoveTemp(PackagePath);
+	OutPackageSegment = MipIndex < GetMinLOD().Default ? EPackageSegment::BulkDataOptional : EPackageSegment::BulkDataDefault;
 	return true;
 }
 #endif // USE_BULKDATA_STREAMING_TOKEN
@@ -6181,10 +6203,11 @@ bool UStaticMesh::GetMipDataFilename(const int32 MipIndex, FString& OutBulkDataF
 FIoFilenameHash UStaticMesh::GetMipIoFilenameHash(const int32 MipIndex) const
 {
 #if USE_BULKDATA_STREAMING_TOKEN
-	FString MipFilename;
-	if (GetMipDataFilename(MipIndex, MipFilename))
+	FPackagePath PackagePath;
+	EPackageSegment PackageSegment;
+	if (GetMipDataPackagePath(MipIndex, PackagePath, PackageSegment))
 	{
-		return MakeIoFilenameHash(MipFilename);
+		return MakeIoFilenameHash(PackagePath);
 	}
 #else
 	if (GetRenderData() && GetRenderData()->LODResources.IsValidIndex(MipIndex))
@@ -6201,8 +6224,9 @@ FIoFilenameHash UStaticMesh::GetMipIoFilenameHash(const int32 MipIndex) const
 bool UStaticMesh::DoesMipDataExist(const int32 MipIndex) const
 {
 #if USE_BULKDATA_STREAMING_TOKEN
-	FString MipDataFilename;
-	return GetMipDataFilename(MipIndex, MipDataFilename) && IFileManager::Get().FileExists(*MipDataFilename);
+	FPackagePath PackagePath;
+	EPackageSegment PackageSegment;
+	return GetMipDataPackagePath(MipIndex, PackagePath, PackageSegment) && IPackageResourceManager::Get().DoesPackageExist(PackagePath, PackageSegment);
 #else
 	return GetRenderData() && GetRenderData()->LODResources.IsValidIndex(MipIndex) && GetRenderData()->LODResources[MipIndex].StreamingBulkData.DoesExist();
 #endif
diff --git a/Engine/Source/Runtime/Engine/Private/StreamableManager.cpp b/Engine/Source/Runtime/Engine/Private/StreamableManager.cpp
--- a/Engine/Source/Runtime/Engine/Private/StreamableManager.cpp
+++ b/Engine/Source/Runtime/Engine/Private/StreamableManager.cpp
@@ -1122,9 +1122,19 @@ FStreamable* FStreamableManager::StreamInternal(const FSoftObjectPath& InTargetN
 				Package.LeftInline(FirstDot,false);
 			}
 
-			Existing->bAsyncLoadRequestOutstanding = true;
-			Existing->bLoadFailed = false;
-			int32 RequestId = LoadPackageAsync(Package, FLoadPackageAsyncDelegate::CreateSP(Handle, &FStreamableHandle::AsyncLoadCallbackWrapper, TargetName), Priority);
+			FPackagePath PackagePath;
+			if (!FPackagePath::TryFromPackageName(Package, PackagePath))
+			{
+				UE_LOG(LogStreamableManager, Error, TEXT("Failed attempt to load %s; it is not a valid LongPackageName"), *Package);
+				Existing->bLoadFailed = true;
+				Existing->bAsyncLoadRequestOutstanding = false;
+			}
+			else
+			{
+				Existing->bLoadFailed = false;
+				Existing->bAsyncLoadRequestOutstanding = true;
+				LoadPackageAsync(PackagePath, NAME_None /* PackageNameToCreate */, FLoadPackageAsyncDelegate::CreateSP(Handle, &FStreamableHandle::AsyncLoadCallbackWrapper, TargetName), nullptr /* InGuid */, PKG_None /* InPackageFlags */, Priority);
+			}
 		}
 	}
 	return Existing;
diff --git a/Engine/Source/Runtime/Engine/Private/Streaming/LandscapeMeshMobileUpdate.cpp b/Engine/Source/Runtime/Engine/Private/Streaming/LandscapeMeshMobileUpdate.cpp
--- a/Engine/Source/Runtime/Engine/Private/Streaming/LandscapeMeshMobileUpdate.cpp
+++ b/Engine/Source/Runtime/Engine/Private/Streaming/LandscapeMeshMobileUpdate.cpp
@@ -9,6 +9,7 @@ LandscapeMeshMobileUpdate.cpp: Helpers to stream in and out mobile landscape ver
 #include "LandscapeComponent.h"
 #include "LandscapeRenderMobile.h"
 #include "Streaming/RenderAssetUpdate.inl"
+#include "Misc/PackageSegment.h"
 
 template class TRenderAssetUpdate<FLandscapeMeshMobileUpdateContext>;
 
@@ -213,18 +214,20 @@ bool FLandscapeMeshMobileStreamIn_IO::HasPendingIORequests() const
 	return false;
 }
 
-FString FLandscapeMeshMobileStreamIn_IO::GetIOFilename(const FContext& Context)
+void FLandscapeMeshMobileStreamIn_IO::GetIOPackagePath(const FContext& Context, FPackagePath& OutPackagePath, EPackageSegment& OutPackageSegment)
 {
 	const ULandscapeLODStreamingProxy* LandscapeProxy = Context.LandscapeProxy;
 
 	if (!IsCancelled() && LandscapeProxy)
 	{
-		FString Filename;
-		verify(LandscapeProxy->GetMipDataFilename(PendingFirstLODIdx, Filename));
-		return Filename;
+		verify(LandscapeProxy->GetMipDataPackagePath(PendingFirstLODIdx, OutPackagePath, OutPackageSegment));
+	}
+	else
+	{
+		MarkAsCancelled();
+		OutPackagePath = FPackagePath();
+		OutPackageSegment = EPackageSegment::Header;
 	}
-	MarkAsCancelled();
-	return FString();
 }
 
 void FLandscapeMeshMobileStreamIn_IO::SetAsyncFileCallback(const FContext& Context)
@@ -249,7 +252,7 @@ void FLandscapeMeshMobileStreamIn_IO::SetAsyncFileCallback(const FContext& Conte
 	};
 }
 
-void FLandscapeMeshMobileStreamIn_IO::SetIORequest(const FContext& Context, const FString& IOFilename)
+void FLandscapeMeshMobileStreamIn_IO::SetIORequest(const FContext& Context, const FPackagePath& PackagePath, EPackageSegment PackageSegment)
 {
 	if (IsCancelled())
 	{
@@ -274,10 +277,11 @@ void FLandscapeMeshMobileStreamIn_IO::SetIORequest(const FContext& Context, cons
 #endif
 			if (BulkDataArray[0]->GetBulkDataSize() > 0)
 			{
-				check(!IOFilename.IsEmpty());
+				check(!PackagePath.IsEmpty());
 				TaskSynchronization.Increment();
 				IORequests[Index] = FBulkDataInterface::CreateStreamingRequestForRange(
-					STREAMINGTOKEN_PARAM(IOFilename)
+					STREAMINGTOKEN_PARAM(PackagePath)
+					STREAMINGTOKEN_PARAM(PackageSegment)
 					BulkDataArray,
 					bHighPrioIORequest ? AIOP_BelowNormal : AIOP_Low,
 					&AsyncFileCallback);
@@ -366,8 +370,10 @@ void FLandscapeMeshMobileStreamIn_IO_AsyncReallocate::DoInitiateIO(const FContex
 {
 	DECLARE_SCOPE_CYCLE_COUNTER(TEXT("LSMMStreamInIOAsyncRealloc_DoInitiateIO"), STAT_LSMMStreamInIOAsyncRealloc_DoInitiateIO, STATGROUP_StreamingDetails);
 	check(Context.CurrentThread == TT_Async);
-	const FString IOFilename = GetIOFilename(Context);
-	SetIORequest(Context, IOFilename);
+	FPackagePath PackagePath;
+	EPackageSegment PackageSegment;
+	GetIOPackagePath(Context, PackagePath, PackageSegment);
+	SetIORequest(Context, PackagePath, PackageSegment);
 	PushTask(Context, TT_Async, SRA_UPDATE_CALLBACK(DoGetIORequestResults), TT_Async, SRA_UPDATE_CALLBACK(DoCancelIO));
 }
 
