PackageResourceManager SearchForPackageOnDisk bugfix: IFileManager::FindFilesRecursive requires a wildcard that matches the full basename+extension, and I incorrectly left off the .* portion of the wildcard.
#rnx

diff --git a/Engine/Source/Runtime/CoreUObject/Private/Misc/PackageName.cpp b/Engine/Source/Runtime/CoreUObject/Private/Misc/PackageName.cpp
--- a/Engine/Source/Runtime/CoreUObject/Private/Misc/PackageName.cpp
+++ b/Engine/Source/Runtime/CoreUObject/Private/Misc/PackageName.cpp
@@ -1430,7 +1430,8 @@ bool FPackageName::SearchForPackageOnDisk(const FString& PackageName, FString* O
 			UE_LOG(LogPackageName, Warning, TEXT("SearchForPackageOnDisk: Invalid extension in packagename %s. Searching for any header extension instead."), *PackageName);
 			RequiredExtension = EPackageExtension::Unspecified;
 		}
-		FStringView PackageWildCard(FStringView(PackageName).Left(ExtensionStart));
+		TStringBuilder<128> PackageWildCard;
+		PackageWildCard << FStringView(PackageName).Left(ExtensionStart) << TEXT(".*");
 
 		FPackagePath FirstResult;
 		TArray<FPackagePath> FoundResults;
