Move BeginCachePlatformCookedData to before package harvesting for Save2 using the same approach as PreSave
This prevent missing property name changed from default when using tagged property serialization for example

Fixed -DiffOnly not calling presave properly, since it shouldn't be called only on the DiffCallstack pass

diff --git a/Engine/Source/Runtime/CoreUObject/Private/UObject/SavePackage2.cpp b/Engine/Source/Runtime/CoreUObject/Private/UObject/SavePackage2.cpp
--- a/Engine/Source/Runtime/CoreUObject/Private/UObject/SavePackage2.cpp
+++ b/Engine/Source/Runtime/CoreUObject/Private/UObject/SavePackage2.cpp
@@ -1875,9 +1875,14 @@ void BeginCachePlatformCookedData(FSaveContext& SaveContext)
 	// Cache platform cooked data
 	if (SaveContext.IsCooking() && !SaveContext.IsConcurrent())
 	{
-		for (FTaggedExport& Export : SaveContext.GetExports())
+		TArray<UObject*> ObjectsInPackage;
+		GetObjectsWithPackage(SaveContext.GetPackage(), ObjectsInPackage);
+		for (UObject* Object : ObjectsInPackage)
 		{
-			Export.Obj->BeginCacheForCookedPlatformData(SaveContext.GetTargetPlatform());
+			if (!SaveContext.IsUnsaveable(Object))
+			{
+				Object->BeginCacheForCookedPlatformData(SaveContext.GetTargetPlatform());
+			}
 		}
 	}
 #endif
@@ -1888,10 +1893,14 @@ void ClearCachedPlatformCookedData(FSaveContext& SaveContext)
 #if WITH_EDITOR
 	if (SaveContext.IsCooking() && !SaveContext.IsConcurrent())
 	{
-		for (FTaggedExport& Export : SaveContext.GetExports())
+		TArray<UObject*> ObjectsInPackage;
+		GetObjectsWithPackage(SaveContext.GetPackage(), ObjectsInPackage);
+		for (UObject* Object : ObjectsInPackage)
 		{
-			Export.Obj->ClearCachedCookedPlatformData(SaveContext.GetTargetPlatform());
-
+			if (!SaveContext.IsUnsaveable(Object))
+			{
+				Object->ClearCachedCookedPlatformData(SaveContext.GetTargetPlatform());
+			}
 		}
 	}
 #endif
@@ -1939,9 +1948,6 @@ ESavePackageResult InnerSave(FSaveContext& SaveContext)
 		return SaveContext.Result;
 	}
 
-	// Trigger platform cooked data caching
-	BeginCachePlatformCookedData(SaveContext);
-
 	// Create Linker
 	SlowTask.EnterProgressFrame();
 	SaveContext.Result = CreateLinker(SaveContext);
@@ -2132,9 +2138,9 @@ FSavePackageResultStruct UPackage::Save2(UPackage* InPackage, UObject* InAsset,
 		SaveContext.SetPreSaveCleanup(InAsset->PreSaveRoot(InFilename));
 	}
 
-	// Route Presave only if not calling concurrently or diffing, in those case they should be handled separately
+	// Route Presave only if not calling concurrently or diffing callstack, in those case they should be handled separately already
 	SlowTask.EnterProgressFrame();
-	if (!SaveContext.IsConcurrent() && !SaveContext.IsDiffing())
+	if (!SaveContext.IsConcurrent() && !SaveContext.IsDiffCallstack())
 	{
 		SaveContext.Result = RoutePresave(SaveContext);
 		if (SaveContext.Result != ESavePackageResult::Success)
@@ -2143,6 +2149,10 @@ FSavePackageResultStruct UPackage::Save2(UPackage* InPackage, UObject* InAsset,
 		}
 	}
 
+	// Trigger platform cooked data caching before package harvesting 
+	// because it might modify some property and hence affect harvested property name for tagged property for example
+	BeginCachePlatformCookedData(SaveContext);
+
 	SlowTask.EnterProgressFrame();
 	{
 		// @todo: Once GIsSavingPackage is reworked we should reinstore the saving flag here for the gc lock
