- Added ability for UObject to override the section name it uses for config properties. This is useful for moving a config class to another module, but still allowing back compatibility. The new name cannot be used in the .ini, only the old names can be used.

diff --git a/Engine/Source/Runtime/CoreUObject/Private/UObject/Obj.cpp b/Engine/Source/Runtime/CoreUObject/Private/UObject/Obj.cpp
--- a/Engine/Source/Runtime/CoreUObject/Private/UObject/Obj.cpp
+++ b/Engine/Source/Runtime/CoreUObject/Private/UObject/Obj.cpp
@@ -2287,6 +2287,9 @@ void UObject::LoadConfig( UClass* ConfigClass/*=NULL*/, const TCHAR* InFilename/
 		{
 			ClassSection = BaseClass->GetPathName();
 			LongCommitName = BaseClass->GetOutermost()->GetFName();
+
+			// allow the class to override the expected section name
+			OverrideConfigSection(ClassSection);
 		}
 
 		// globalconfig properties should always use the owning class's config file
@@ -2511,6 +2514,8 @@ void UObject::SaveConfig( uint64 Flags, const TCHAR* InFilename, FConfigCacheIni
 			if ( !bPerObject )
 			{
 				Section = BaseClass->GetPathName();
+				// allow the class to override the expected section name
+				OverrideConfigSection(Section);
 			}
 
 			// globalconfig properties should always use the owning class's config file
diff --git a/Engine/Source/Runtime/CoreUObject/Public/UObject/Object.h b/Engine/Source/Runtime/CoreUObject/Public/UObject/Object.h
--- a/Engine/Source/Runtime/CoreUObject/Public/UObject/Object.h
+++ b/Engine/Source/Runtime/CoreUObject/Public/UObject/Object.h
@@ -1109,6 +1109,13 @@ public:
 	/** Returns the override config hierarchy platform (if NDAd platforms need defaults to not be in Base*.ini but still want editor to load them) */
 	virtual const TCHAR* GetConfigOverridePlatform() const { return nullptr; }
 
+	/**
+	 * Allows Non-PerObjectConfig classes, to override the ini section name used for loading config settings
+	 *
+	 * @param SectionName	Reference to the unmodified config section name, that can be altered/modified
+	 */
+	virtual void OverrideConfigSection(FString& SectionName) {}
+
 	/**
 	 * Allows PerObjectConfig classes, to override the ini section name used for the PerObjectConfig object.
 	 *
