Prototype Wrapped Object Pointers (ObjectHandle/ObjectPtr)

Ensure that the Save1 codepath doesn't attempt to lazy load/resolve wrapped object pointers mid-save (which is not permitted).

diff --git a/Engine/Source/Runtime/CoreUObject/Private/UObject/SavePackage.cpp b/Engine/Source/Runtime/CoreUObject/Private/UObject/SavePackage.cpp
--- a/Engine/Source/Runtime/CoreUObject/Private/UObject/SavePackage.cpp
+++ b/Engine/Source/Runtime/CoreUObject/Private/UObject/SavePackage.cpp
@@ -1387,6 +1387,18 @@ public:
 	virtual FArchive& operator<<(FSoftObjectPtr& Value) override { return *this; }
 	virtual FArchive& operator<<(FSoftObjectPath& Value) override { return *this; }
 
+	/** Object pointer references that are unresolved should not be traversed at this point as we are not allowed to load/resolve anything */
+	virtual FArchive& operator<<(FObjectPtr& Value) override
+	{
+		// Everything we want to traverse in this export reference sorter will have been resolved.  If we encounter something that is unresolved,
+		// it will be part of another package and not something we need to worry about for export sorting.
+		if (IsObjectHandleResolved(Value.GetHandle()))
+		{
+			return FArchiveUObject::operator<<(Value);
+		}
+		return *this;
+	}
+
 private:
 
 	/**
