Fix incorrect licensee flag set in FEngineVersion. Fallout from CL 12740027, fixed in 14301918.

diff --git a/Engine/Source/Runtime/Core/Public/UObject/ObjectVersion.h b/Engine/Source/Runtime/Core/Public/UObject/ObjectVersion.h
--- a/Engine/Source/Runtime/Core/Public/UObject/ObjectVersion.h
+++ b/Engine/Source/Runtime/Core/Public/UObject/ObjectVersion.h
@@ -639,6 +639,8 @@ enum EUnrealEngineObjectUE4Version
 	VER_UE4_NON_OUTER_PACKAGE_IMPORT,
 	// Added DependencyFlags to AssetRegistry
 	VER_UE4_ASSETREGISTRY_DEPENDENCYFLAGS,
+	// Fixed corrupt licensee flag in 4.26 assets
+	VER_UE4_CORRECT_LICENSEE_FLAG,
 
 	// -----<new versions can be added before this line>-------------------------------------------------
 	// - this needs to be the last line (see note below)
diff --git a/Engine/Source/Runtime/CoreUObject/Private/UObject/PackageFileSummary.cpp b/Engine/Source/Runtime/CoreUObject/Private/UObject/PackageFileSummary.cpp
--- a/Engine/Source/Runtime/CoreUObject/Private/UObject/PackageFileSummary.cpp
+++ b/Engine/Source/Runtime/CoreUObject/Private/UObject/PackageFileSummary.cpp
@@ -30,6 +30,21 @@ static ECustomVersionSerializationFormat::Type GetCustomVersionFormatForArchive(
 	return CustomVersionFormat;
 }
 
+static void FixCorruptEngineVersion(int ObjectVersion, FEngineVersion& Version)
+{
+	// The move of EpicInternal.txt in CL 12740027 broke checks for non-licensee builds in UGS. resulted in checks for Epic internal builds in UGS breaking, and assets being saved out with the licensee flag set.
+	// Detect such assets and clear the licensee bit.
+	if (ObjectVersion < VER_UE4_CORRECT_LICENSEE_FLAG
+		&& Version.GetMajor() == 4
+		&& Version.GetMinor() == 26
+		&& Version.GetPatch() == 0
+		&& Version.GetChangelist() >= 12740027
+		&& Version.IsLicenseeVersion())
+	{
+		Version.Set(4, 26, 0, Version.GetChangelist(), Version.GetBranch());
+	}
+}
+
 void operator<<(FStructuredArchive::FSlot Slot, FPackageFileSummary& Sum)
 {
 	FArchive& BaseArchive = Slot.GetUnderlyingArchive();
@@ -263,6 +278,7 @@ void operator<<(FStructuredArchive::FSlot Slot, FPackageFileSummary& Sum)
 			else
 			{
 				Record << SA_VALUE(TEXT("SavedByEngineVersion"), Sum.SavedByEngineVersion);
+				FixCorruptEngineVersion(Sum.GetFileVersionUE4(), Sum.SavedByEngineVersion);
 			}
 		}
 		else
@@ -286,6 +302,7 @@ void operator<<(FStructuredArchive::FSlot Slot, FPackageFileSummary& Sum)
 			else
 			{
 				Record << SA_VALUE(TEXT("CompatibleWithEngineVersion"), Sum.CompatibleWithEngineVersion);
+				FixCorruptEngineVersion(Sum.GetFileVersionUE4(), Sum.CompatibleWithEngineVersion);
 			}
 		}
 		else
