FPackageResourceManagerFile needs to handle IFileManager.IterateDirectoryRecursively returning unnormalized files.
#rnx

diff --git a/Engine/Source/Runtime/CoreUObject/Private/UObject/PackageResourceManagerFile.cpp b/Engine/Source/Runtime/CoreUObject/Private/UObject/PackageResourceManagerFile.cpp
--- a/Engine/Source/Runtime/CoreUObject/Private/UObject/PackageResourceManagerFile.cpp
+++ b/Engine/Source/Runtime/CoreUObject/Private/UObject/PackageResourceManagerFile.cpp
@@ -428,13 +428,23 @@ bool FPackageResourceManagerFile::FDirectoryVisitorBaseMounted::TryConvertToPack
 	{
 		return false;
 	}
+
 	FStringView Filename(FilenameOrDirectory);
+	FString BufferPath;
 	if (!Filename.StartsWith(FileMount, ESearchCase::IgnoreCase))
 	{
-		UE_LOG(LogPackageResourceManager, Warning,
-			TEXT("FDirectoryVisitorBaseMounted: FileManager IterateDirectoryRecursively(\"%s\") returned file \"%s\" that is not a subpath of the root \"%s\"."),
-			*RootDir, FilenameOrDirectory, *FileMount);
-		return false;
+		// Check whether the reason it doesn't start with the text of the FileMount is because it is not normalized and is e.g. a relative path
+		// ConvertRelativePathToFull will normalize it in addition to converting it to an absolute path
+		BufferPath = Filename;
+		BufferPath = FPaths::ConvertRelativePathToFull(MoveTemp(BufferPath));
+		Filename = BufferPath;
+		if (!Filename.StartsWith(FileMount, ESearchCase::IgnoreCase))
+		{
+			UE_LOG(LogPackageResourceManager, Warning,
+				TEXT("FDirectoryVisitorBaseMounted: FileManager IterateDirectoryRecursively(\"%s\") returned file \"%s\" that is not a subpath of the root \"%s\"."),
+				*RootDir, FilenameOrDirectory, *FileMount);
+			return false;
+		}
 	}
 	FStringView RelPath(Filename.RightChop(FileMount.Len()));
 	int32 ExtensionStart;
