Prototype Wrapped Object Pointers (ObjectHandle/ObjectPtr)

Temporarily disable lazy resolve/load in editor/cooker to reduce risk while preparing for check-in of wrapped object pointer upgrade

diff --git a/Engine/Source/Runtime/CoreUObject/Public/UObject/ObjectHandle.h b/Engine/Source/Runtime/CoreUObject/Public/UObject/ObjectHandle.h
--- a/Engine/Source/Runtime/CoreUObject/Public/UObject/ObjectHandle.h
+++ b/Engine/Source/Runtime/CoreUObject/Public/UObject/ObjectHandle.h
@@ -9,7 +9,7 @@
 #include "UObject/Object.h"
 #include "UObject/ObjectPathId.h"
 
-#define UE_WITH_OBJECT_HANDLE_LATE_RESOLVE WITH_EDITORONLY_DATA
+#define UE_WITH_OBJECT_HANDLE_LATE_RESOLVE 0 // @TODO: OBJPTR: Should be WITH_EDITORONLY_DATA when ready to tackle lazy load/resolve issues
 #define UE_WITH_OBJECT_HANDLE_TRACKING WITH_EDITORONLY_DATA
 
 /**
