Revised ICH template archetype search logic to remove its dependency on the outer class CDO.

This addresses a regression that could manifest in some cases where the CDO has not yet been created, and is a bit closer to the original implementation.

#rnx

diff --git a/Engine/Source/Runtime/CoreUObject/Private/UObject/UObjectArchetype.cpp b/Engine/Source/Runtime/CoreUObject/Private/UObject/UObjectArchetype.cpp
--- a/Engine/Source/Runtime/CoreUObject/Private/UObject/UObjectArchetype.cpp
+++ b/Engine/Source/Runtime/CoreUObject/Private/UObject/UObjectArchetype.cpp
@@ -68,18 +68,22 @@ UObject* GetArchetypeFromRequiredInfoImpl(const UClass* Class, const UObject* Ou
 			// Get a lock on the UObject hash tables for the duration of the GetArchetype operation
 			FScopedUObjectHashTablesLock HashTablesLock;
 
-			const UObject* ArchetypeToSearch;
-			UObject* FoundArchetype = nullptr;
 			if (!!(ObjectFlags & RF_InheritableComponentTemplate) && Outer->IsA<UClass>())
 			{
-				// In this case, we are looking for the archetype of a non-native component template, which is owned by an outer UClass object. This
-				// differs from the "normal" path in that we don't search in the outer class object's archetype (which is a UClass derivative). Instead,
-				// we set the search object to the outer class's own default object (CDO). This way we'll defer the search to the virtual class API below.
+				// In this case, we are looking for the archetype of a non-native component template, which is owned by an outer UClass object.
+				// This differs from the "normal" path in that we don't search in the outer class object's archetype (which is a UClass derivative).
+				// Instead, we redirect the search through the outer class object itself.
 				const UClass* OuterClass = static_cast<const UClass*>(Outer);
-				ArchetypeToSearch = bUseUpToDateClass ? OuterClass->GetAuthoritativeClass()->GetDefaultObject() : OuterClass->GetDefaultObject();
+				if(bUseUpToDateClass)
+				{
+					OuterClass = OuterClass->GetAuthoritativeClass();
+				}
+
+				Result = OuterClass->FindArchetype(Class, Name, ObjectFlags);
 			}
 			else
 			{
+				UObject* ArchetypeToSearch = nullptr;
 #if UE_CACHE_ARCHETYPE
 				ArchetypeToSearch = Outer->GetArchetype();
 #if UE_VERIFY_CACHED_ARCHETYPE
@@ -91,24 +95,23 @@ UObject* GetArchetypeFromRequiredInfoImpl(const UClass* Class, const UObject* Ou
 #else
 				ArchetypeToSearch = GetArchetypeFromRequiredInfoImpl(Outer->GetClass(), Outer->GetOuter(), Outer->GetFName(), Outer->GetFlags(), bUseUpToDateClass);
 #endif // UE_CACHE_ARCHETYPE
-				FoundArchetype = static_cast<UObject*>(FindObjectWithOuter(ArchetypeToSearch, Class, Name));
-			}
-
-			if (FoundArchetype)
-			{
-				Result = FoundArchetype; // found that my outers archetype had a matching component, that must be my archetype
-			}
-			else
-			{
-				if (GEventDrivenLoaderEnabled && EVENT_DRIVEN_ASYNC_LOAD_ACTIVE_AT_RUNTIME)
+				UObject* MyArchetype = static_cast<UObject*>(FindObjectWithOuter(ArchetypeToSearch, Class, Name));
+				if (MyArchetype)
+				{
+					Result = MyArchetype; // found that my outers archetype had a matching component, that must be my archetype
+				}
+				else
 				{
-					if (ArchetypeToSearch->HasAnyFlags(RF_NeedLoad))
+					if (GEventDrivenLoaderEnabled && EVENT_DRIVEN_ASYNC_LOAD_ACTIVE_AT_RUNTIME)
 					{
-						UE_LOG(LogClass, Fatal, TEXT("%s had RF_NeedLoad when searching for an archetype of %s in %s"), *GetFullNameSafe(ArchetypeToSearch), *GetFullNameSafe(Class), *GetFullNameSafe(Outer));
+						if (ArchetypeToSearch->HasAnyFlags(RF_NeedLoad))
+						{
+							UE_LOG(LogClass, Fatal, TEXT("%s had RF_NeedLoad when searching for an archetype of %s in %s"), *GetFullNameSafe(ArchetypeToSearch), *GetFullNameSafe(Class), *GetFullNameSafe(Outer));
+						}
 					}
-				}
 
-				Result = ArchetypeToSearch->GetClass()->FindArchetype(Class, Name, ObjectFlags);
+					Result = ArchetypeToSearch->GetClass()->FindArchetype(Class, Name, ObjectFlags);
+				}
 			}
 		}
 
