Prototype Wrapped Object Pointers (ObjectHandle/ObjectPtr)

Systematic changes in support of non-scalar wrapped object pointer upgrades for Engine + Shootergame

diff --git a/Engine/Source/Runtime/Core/Public/Containers/Array.h b/Engine/Source/Runtime/Core/Public/Containers/Array.h
--- a/Engine/Source/Runtime/Core/Public/Containers/Array.h
+++ b/Engine/Source/Runtime/Core/Public/Containers/Array.h
@@ -10,6 +10,7 @@
 #include "Templates/UnrealTypeTraits.h"
 #include "Templates/UnrealTemplate.h"
 #include "Containers/ContainerAllocationPolicies.h"
+#include "Containers/ContainerElementTypeCompatibility.h"
 #include "Serialization/Archive.h"
 #include "Serialization/MemoryImageWriter.h"
 
@@ -2109,6 +2110,141 @@ public:
 		return *Ptr;
 	}
 
+	/** Implicit conversion operator to container of compatible element type. */
+	template <
+		typename AliasElementType = ElementType,
+		typename TEnableIf<TIsContainerElementTypeReinterpretable<AliasElementType>::Value>::Type* = nullptr
+		>
+	operator TArray<typename TContainerElementTypeCompatibility<AliasElementType>::ReinterpretType, Allocator>& ()
+	{
+		using ElementCompat = TContainerElementTypeCompatibility<ElementType>;
+		ElementCompat::ReinterpretRange(begin(), end());
+		return *reinterpret_cast<TArray<typename ElementCompat::ReinterpretType>*>(this);
+	}
+
+	/** Implicit conversion operator to constant container of compatible element type. */
+	template <
+		typename AliasElementType = ElementType,
+		typename std::enable_if_t<TIsContainerElementTypeReinterpretable<AliasElementType>::Value>* = nullptr
+		>
+	operator const TArray<typename TContainerElementTypeCompatibility<AliasElementType>::ReinterpretType, Allocator>& () const
+	{
+		using ElementCompat = TContainerElementTypeCompatibility<ElementType>;
+		ElementCompat::ReinterpretRange(begin(), end());
+		return *reinterpret_cast<const TArray<typename ElementCompat::ReinterpretType>*>(this);
+	}
+
+	/**
+	 * Move assignment operator.
+	 * Compatible element type version.
+	 *
+	 * @param Other Array to assign and move from.
+	 */
+	template <
+		typename AliasElementType = ElementType,
+		typename std::enable_if_t<TIsContainerElementTypeCopyable<AliasElementType>::Value>* = nullptr
+	>
+	TArray& operator=(TArray<typename TContainerElementTypeCompatibility<ElementType>::CopyFromOtherType, Allocator>&& Other)
+	{
+		TContainerElementTypeCompatibility<ElementType>::CopyingFromOtherType();
+		DestructItems(GetData(), ArrayNum);
+		MoveOrCopy(*this, Other, ArrayMax);
+		return *this; 
+	}
+
+	/**
+	 * Assignment operator. First deletes all currently contained elements
+	 * and then copies from other array.
+	 * Compatible element type version.
+	 *
+	 * @param Other The source array to assign from.
+	 */
+	template <
+		typename OtherAllocator,
+		typename AliasElementType = ElementType,
+		typename std::enable_if_t<TIsContainerElementTypeCopyable<AliasElementType>::Value>* = nullptr
+	>
+	TArray& operator=(const TArray<typename TContainerElementTypeCompatibility<ElementType>::CopyFromOtherType, OtherAllocator>& Other)
+	{
+		TContainerElementTypeCompatibility<ElementType>::CopyingFromOtherType();
+		DestructItems(GetData(), ArrayNum);
+		CopyToEmpty(Other.GetData(), Other.Num(), ArrayMax, 0);
+		return *this;
+	}
+
+	/**
+	 * Inserts given elements into the array at given location.
+	 * Compatible element type version.
+	 *
+	 * @param Items Array of elements to insert.
+	 * @param InIndex Tells where to insert the new elements.
+	 * @returns Location at which the item was inserted.
+	 */
+	template <
+		typename OtherAllocator,
+		typename AliasElementType = ElementType,
+		typename std::enable_if_t<TIsContainerElementTypeCopyable<AliasElementType>::Value>* = nullptr
+	>
+	SizeType Insert(const TArray<typename TContainerElementTypeCompatibility<ElementType>::CopyFromOtherType, OtherAllocator>& Items, const SizeType InIndex)
+	{
+		TContainerElementTypeCompatibility<ElementType>::CopyingFromOtherType();
+
+		auto NumNewElements = Items.Num();
+
+		InsertUninitializedImpl(InIndex, NumNewElements);
+		ConstructItems<ElementType>(GetData() + InIndex, Items.GetData(), NumNewElements);
+
+		return InIndex;
+	}
+
+	/**
+	 * Inserts given elements into the array at given location.
+	 * Compatible element type version.
+	 *
+	 * @param Items Array of elements to insert.
+	 * @param InIndex Tells where to insert the new elements.
+	 * @returns Location at which the item was inserted.
+	 */
+	template <
+		typename OtherAllocator,
+		typename AliasElementType = ElementType,
+		typename std::enable_if_t<TIsContainerElementTypeCopyable<AliasElementType>::Value>* = nullptr
+	>
+	SizeType Insert(TArray<typename TContainerElementTypeCompatibility<ElementType>::CopyFromOtherType, OtherAllocator>&& Items, const SizeType InIndex)
+	{
+		check((const void*)this != (const void*)&Items);
+		TContainerElementTypeCompatibility<ElementType>::CopyingFromOtherType();
+
+		auto NumNewElements = Items.Num();
+
+		InsertUninitializedImpl(InIndex, NumNewElements);
+		RelocateConstructItems<ElementType>(GetData() + InIndex, Items.GetData(), NumNewElements);
+		Items.ArrayNum = 0;
+
+		return InIndex;
+	}
+
+	/**
+	 * Adds a raw array of elements to the end of the TArray.
+	 * Compatible element type version.
+	 *
+	 * @param Ptr   A pointer to an array of elements to add.
+	 * @param Count The number of elements to insert from Ptr.
+	 * @see Add, Insert
+	 */
+	template <
+		typename AliasElementType = ElementType,
+		typename std::enable_if_t<TIsContainerElementTypeCopyable<AliasElementType>::Value>* = nullptr
+	>
+	void Append(const typename TContainerElementTypeCompatibility<ElementType>::CopyFromOtherType* Ptr, SizeType Count)
+	{
+		check(Ptr != nullptr || Count == 0);
+		TContainerElementTypeCompatibility<ElementType>::CopyingFromOtherType();
+
+		SizeType Pos = AddUninitialized(Count);
+		ConstructItems<ElementType>(GetData() + Pos, Ptr, Count);
+	}
+
 private:
 
 	/**
diff --git a/Engine/Source/Runtime/Core/Public/Containers/ArrayView.h b/Engine/Source/Runtime/Core/Public/Containers/ArrayView.h
--- a/Engine/Source/Runtime/Core/Public/Containers/ArrayView.h
+++ b/Engine/Source/Runtime/Core/Public/Containers/ArrayView.h
@@ -3,10 +3,12 @@
 #pragma once
 
 #include "CoreTypes.h"
+#include "Templates/IsConst.h"
 #include "Templates/IsSigned.h"
 #include "Templates/PointerIsConvertibleFromTo.h"
 #include "Misc/AssertionMacros.h"
 #include "Templates/UnrealTypeTraits.h"
+#include "Traits/ElementType.h"
 #include "Containers/Array.h"
 
 namespace ArrayViewPrivate
@@ -33,6 +35,19 @@ namespace ArrayViewPrivate
 		return GetData(Forward<T>(Arg));
 	}
 
+	// Gets the data from the passed argument and proceeds to reinterpret the resulting elements
+	template <typename T>
+	FORCEINLINE decltype(auto) GetReinterpretedDataHelper(T&& Arg)
+	{
+		auto NaturalPtr = GetData(Forward<T>(Arg));
+		using NaturalElementType = typename TRemovePointer<decltype(NaturalPtr)>::Type;
+
+		auto EndPtr = NaturalPtr + GetNum(Arg);
+		TContainerElementTypeCompatibility<NaturalElementType>::ReinterpretRange(NaturalPtr, EndPtr);
+
+		return reinterpret_cast<typename TContainerElementTypeCompatibility<NaturalElementType>::ReinterpretType*>(NaturalPtr);
+	}
+
 	/**
 	 * Trait testing whether a type is compatible with the view type
 	 */
@@ -40,6 +55,34 @@ namespace ArrayViewPrivate
 	struct TIsCompatibleRangeType
 	{
 		static constexpr bool Value = TIsCompatibleElementType<typename TRemovePointer<decltype(GetData(DeclVal<RangeType&>()))>::Type, ElementType>::Value;
+
+		template <typename T>
+		static decltype(auto) GetData(T&& Arg)
+		{
+			return ArrayViewPrivate::GetDataHelper(Forward<T>(Arg));
+		}
+	};
+
+	/**
+	 * Trait testing whether a type is reinterpretable in a way that permits use with the view type
+	 */
+	template <typename RangeType, typename ElementType>
+	struct TIsReinterpretableRangeType
+	{
+	private:
+		using NaturalElementType = typename TRemovePointer<decltype(GetData(DeclVal<RangeType&>()))>::Type;
+
+	public:
+		static constexpr bool Value = 
+			!TIsSame<typename TContainerElementTypeCompatibility<NaturalElementType>::ReinterpretType, NaturalElementType>::Value
+			&&
+			TIsCompatibleElementType<typename TContainerElementTypeCompatibility<NaturalElementType>::ReinterpretType, ElementType>::Value;
+
+		template <typename T>
+		static decltype(auto) GetData(T&& Arg)
+		{
+			return ArrayViewPrivate::GetReinterpretedDataHelper(Forward<T>(Arg));
+		}
 	};
 }
 
@@ -98,6 +141,9 @@ private:
 	template <typename T>
 	using TIsCompatibleRangeType = ArrayViewPrivate::TIsCompatibleRangeType<T, ElementType>;
 
+	template <typename T>
+	using TIsReinterpretableRangeType = ArrayViewPrivate::TIsReinterpretableRangeType<T, ElementType>;
+
 public:
 	/**
 	 * Constructor from another range
@@ -110,12 +156,19 @@ public:
 		typename = typename TEnableIf<
 			TAnd<
 				TIsContiguousContainer<CVUnqualifiedOtherRangeType>,
-				TIsCompatibleRangeType<OtherRangeType>
+				TOr<
+					TIsCompatibleRangeType<OtherRangeType>,
+					TIsReinterpretableRangeType<OtherRangeType>
+				>
 			>::Value
 		>::Type
 	>
 	FORCEINLINE TArrayView(OtherRangeType&& Other)
-		: DataPtr(ArrayViewPrivate::GetDataHelper(Forward<OtherRangeType>(Other)))
+		: DataPtr(TChooseClass<
+						TIsCompatibleRangeType<OtherRangeType>::Value,
+						TIsCompatibleRangeType<OtherRangeType>,
+						TIsReinterpretableRangeType<OtherRangeType>
+					>::Result::GetData(Forward<OtherRangeType>(Other)))
 	{
 		const auto InCount = GetNum(Forward<OtherRangeType>(Other));
 		check((InCount >= 0) && ((sizeof(InCount) < sizeof(SizeType)) || (InCount <= static_cast<decltype(InCount)>(TNumericLimits<SizeType>::Max()))));
diff --git a/Engine/Source/Runtime/Core/Public/Containers/ContainerElementTypeCompatibility.h b/Engine/Source/Runtime/Core/Public/Containers/ContainerElementTypeCompatibility.h
new file mode 100644
--- /dev/null
+++ b/Engine/Source/Runtime/Core/Public/Containers/ContainerElementTypeCompatibility.h
@@ -0,0 +1,32 @@
+// Copyright Epic Games, Inc. All Rights Reserved.
+
+#pragma once
+
+#include "CoreTypes.h"
+
+template <typename InElementType>
+struct TContainerElementTypeCompatibility
+{
+	typedef InElementType ReinterpretType;
+
+	template <typename IterBeginType, typename IterEndType, typename OperatorType = InElementType&(*)(IterBeginType&)>
+	static void ReinterpretRange(IterBeginType Iter, IterEndType IterEnd, OperatorType Operator = [](IterBeginType& InIt) -> InElementType& { return *InIt; })
+	{
+	}
+
+	typedef InElementType CopyFromOtherType;
+
+	static constexpr void CopyingFromOtherType() {}
+};
+
+template <typename ElementType>
+struct TIsContainerElementTypeReinterpretable
+{
+	enum { Value = !TIsSame<typename TContainerElementTypeCompatibility<ElementType>::ReinterpretType, ElementType>::Value };
+};
+
+template <typename ElementType>
+struct TIsContainerElementTypeCopyable
+{
+	enum { Value = !TIsSame<typename TContainerElementTypeCompatibility<ElementType>::CopyFromOtherType, ElementType>::Value };
+};
diff --git a/Engine/Source/Runtime/Core/Public/Containers/Map.h b/Engine/Source/Runtime/Core/Public/Containers/Map.h
--- a/Engine/Source/Runtime/Core/Public/Containers/Map.h
+++ b/Engine/Source/Runtime/Core/Public/Containers/Map.h
@@ -5,6 +5,7 @@
 #include "CoreTypes.h"
 
 #include "Algo/Reverse.h"
+#include "Containers/ContainerElementTypeCompatibility.h"
 #include "Containers/Set.h"
 #include "Containers/UnrealString.h"
 #include "Misc/AssertionMacros.h"
@@ -1246,6 +1247,46 @@ public:
 
 	FORCEINLINE       ValueType& operator[](KeyConstPointerType Key) { return this->FindChecked(Key); }
 	FORCEINLINE const ValueType& operator[](KeyConstPointerType Key) const { return this->FindChecked(Key); }
+
+	/** Implicit conversion operator to container of compatible key and value element type. */
+	template <
+		typename KeyTypeAlias = KeyType,
+		typename ValueTypeAlias = ValueType,
+		typename std::enable_if_t<
+			TIsContainerElementTypeReinterpretable<KeyTypeAlias>::Value
+			||
+			TIsContainerElementTypeReinterpretable<ValueTypeAlias>::Value
+		>* = nullptr
+	>
+	operator TMap<typename TContainerElementTypeCompatibility<KeyTypeAlias>::ReinterpretType, typename TContainerElementTypeCompatibility<ValueTypeAlias>::ReinterpretType>& ()
+	{
+		using KeyCompat = TContainerElementTypeCompatibility<KeyType>;
+		using ValueCompat = TContainerElementTypeCompatibility<ValueType>;
+		KeyCompat::ReinterpretRange(Super::Pairs.begin(), Super::Pairs.end(), [](typename Super::ElementSetType::TRangedForIterator& It) -> KeyType& { return It->Key; });
+		ValueCompat::ReinterpretRange(Super::Pairs.begin(), Super::Pairs.end(), [](typename Super::ElementSetType::TRangedForIterator& It) -> ValueType& { return It->Value; });
+
+		return *reinterpret_cast<TMap<typename KeyCompat::ReinterpretType, typename ValueCompat::ReinterpretType>*>(this);
+	}
+
+	/** Implicit conversion operator to constant container of compatible key and value element type. */
+	template <
+		typename KeyTypeAlias = KeyType,
+		typename ValueTypeAlias = ValueType,
+		typename std::enable_if_t<
+			TIsContainerElementTypeReinterpretable<KeyTypeAlias>::Value
+			||
+			TIsContainerElementTypeReinterpretable<ValueTypeAlias>::Value
+		>* = nullptr
+	>
+	operator const TMap<typename TContainerElementTypeCompatibility<KeyTypeAlias>::ReinterpretType, typename TContainerElementTypeCompatibility<ValueTypeAlias>::ReinterpretType>& () const
+	{
+		using KeyCompat = TContainerElementTypeCompatibility<KeyType>;
+		using ValueCompat = TContainerElementTypeCompatibility<ValueType>;
+		KeyCompat::ReinterpretRange(Super::Pairs.begin(), Super::Pairs.end(), [](typename Super::ElementSetType::TRangedForConstIterator& It) -> const KeyType& { return It->Key; });
+		ValueCompat::ReinterpretRange(Super::Pairs.begin(), Super::Pairs.end(), [](typename Super::ElementSetType::TRangedForConstIterator& It) -> const ValueType& { return It->Value; });
+
+		return *reinterpret_cast<const TMap<typename KeyCompat::ReinterpretType, typename ValueCompat::ReinterpretType>*>(this);
+	}
 };
 
 namespace Freeze
diff --git a/Engine/Source/Runtime/Core/Public/Containers/Set.h b/Engine/Source/Runtime/Core/Public/Containers/Set.h
--- a/Engine/Source/Runtime/Core/Public/Containers/Set.h
+++ b/Engine/Source/Runtime/Core/Public/Containers/Set.h
@@ -7,6 +7,7 @@
 #include "Templates/UnrealTypeTraits.h"
 #include "Templates/UnrealTemplate.h"
 #include "Containers/ContainerAllocationPolicies.h"
+#include "Containers/ContainerElementTypeCompatibility.h"
 #include "Templates/Sorting.h"
 #include "Containers/Array.h"
 #include "Math/UnrealMathUtility.h"
@@ -1214,6 +1215,112 @@ public:
 		Elements.CheckAddress(Addr);
 	}
 
+	/** Implicit conversion operator to container of compatible element type. */
+	template <
+		typename AliasElementType = ElementType,
+		typename std::enable_if_t<TIsContainerElementTypeReinterpretable<AliasElementType>::Value>* = nullptr
+		>
+	operator TSet<typename TContainerElementTypeCompatibility<AliasElementType>::ReinterpretType>& ()
+	{
+		using ElementCompat = TContainerElementTypeCompatibility<ElementType>;
+		ElementCompat::ReinterpretRange(begin(), end());
+		return *reinterpret_cast<TSet<typename ElementCompat::ReinterpretType>*>(this);
+	}
+
+	/** Implicit conversion operator to constant container of compatible element type. */
+	template <
+		typename AliasElementType = ElementType,
+		typename std::enable_if_t<TIsContainerElementTypeReinterpretable<AliasElementType>::Value>* = nullptr
+		>
+	operator const TSet<typename TContainerElementTypeCompatibility<AliasElementType>::ReinterpretType>& () const
+	{
+		using ElementCompat = TContainerElementTypeCompatibility<ElementType>;
+		ElementCompat::ReinterpretRange(begin(), end());
+		return *reinterpret_cast<const TSet<typename ElementCompat::ReinterpretType>*>(this);
+	}
+
+	/**
+	 * Move assignment operator.
+	 * Compatible element type version.
+	 *
+	 * @param Other Set to assign and move from.
+	 */
+	template <
+		typename OtherKeyFuncs,
+		typename AliasElementType = ElementType,
+		typename std::enable_if_t<TIsContainerElementTypeCopyable<AliasElementType>::Value>* = nullptr
+	>
+	TSet& operator=(TSet<typename TContainerElementTypeCompatibility<ElementType>::CopyFromOtherType, OtherKeyFuncs, Allocator>&& Other)
+	{
+		TContainerElementTypeCompatibility<ElementType>::CopyingFromOtherType();
+		Reset();
+		Append(MoveTemp(Other));
+		return *this;
+	}
+
+	/**
+	 * Assignment operator. First deletes all currently contained elements
+	 * and then copies from other set.
+	 * Compatible element type version.
+	 *
+	 * @param Other The source set to assign from.
+	 */
+	template <
+		typename OtherKeyFuncs,
+		typename OtherAllocator,
+		typename AliasElementType = ElementType,
+		typename std::enable_if_t<TIsContainerElementTypeCopyable<AliasElementType>::Value>* = nullptr
+	>
+	TSet& operator=(const TSet<typename TContainerElementTypeCompatibility<ElementType>::CopyFromOtherType, OtherKeyFuncs, OtherAllocator>& Other)
+	{
+		TContainerElementTypeCompatibility<ElementType>::CopyingFromOtherType();
+		Reset();
+		Append(Other);
+		return *this;
+	}
+
+	/**
+	 * Add all items from another set to our set (union without creating a new set)
+	 * Compatible element type version.
+	 * @param OtherSet - The other set of items to add.
+	 */
+	template <
+		typename OtherKeyFuncs,
+		typename OtherAllocator,
+		typename AliasElementType = ElementType,
+		typename std::enable_if_t<TIsContainerElementTypeCopyable<AliasElementType>::Value>* = nullptr
+	>
+	void Append(const TSet<typename TContainerElementTypeCompatibility<ElementType>::CopyFromOtherType, OtherKeyFuncs, OtherAllocator>& OtherSet)
+	{
+		TContainerElementTypeCompatibility<ElementType>::CopyingFromOtherType();
+		Reserve(Elements.Num() + OtherSet.Num());
+		for (const ElementType& Element : OtherSet)
+		{
+			Add(Element);
+		}
+	}
+
+	/**
+	 * Add all items from another set to our set (union without creating a new set)
+	 * Compatible element type version.
+	 * @param OtherSet - The other set of items to add.
+	 */
+	template <
+		typename OtherKeyFuncs,
+		typename AliasElementType = ElementType,
+		typename std::enable_if_t<TIsContainerElementTypeCopyable<AliasElementType>::Value>* = nullptr
+	>
+	void Append(TSet<typename TContainerElementTypeCompatibility<ElementType>::CopyFromOtherType, OtherKeyFuncs, Allocator>&& OtherSet)
+	{
+		TContainerElementTypeCompatibility<ElementType>::CopyingFromOtherType();
+		Reserve(Elements.Num() + OtherSet.Num());
+		for (ElementType& Element : OtherSet)
+		{
+			Add(MoveTempIfPossible(Element));
+		}
+		OtherSet.Reset();
+	}
+
 private:
 	/** Extracts the element value from the set's element structure and passes it to the user provided comparison class. */
 	template <typename PREDICATE_CLASS>
diff --git a/Engine/Source/Runtime/CoreUObject/Public/UObject/ObjectPtr.h b/Engine/Source/Runtime/CoreUObject/Public/UObject/ObjectPtr.h
--- a/Engine/Source/Runtime/CoreUObject/Public/UObject/ObjectPtr.h
+++ b/Engine/Source/Runtime/CoreUObject/Public/UObject/ObjectPtr.h
@@ -164,6 +164,7 @@ template <typename T>
 struct TObjectPtr : private FObjectPtr
 {
 public:
+	using ElementType = T;
 
 	TObjectPtr() = default;
 	TObjectPtr(TObjectPtr<T>&&) = default;
@@ -278,7 +279,7 @@ public:
 	FORCEINLINE T& operator*() const { return *Get(); }
 
 	UE_OBJPTR_DEPRECATED(5.0, "Conversion to a mutable pointer is deprecated.  Please pass a TObjectPtr<T>& instead so that assignment can be tracked accurately.")
-	FORCEINLINE operator T*& () { return GetInternalRef(); }
+	explicit FORCEINLINE operator T*& () { return GetInternalRef(); }
 
 	using FObjectPtr::IsNull;
 	using FObjectPtr::operator bool;
@@ -310,11 +311,10 @@ private:
 	}
 };
 
-template <typename T>
-struct TIsTObjectPtr<TObjectPtr<T>>
-{
-	enum { Value = true };
-};
+template <typename T> struct TIsTObjectPtr<               TObjectPtr<T>> { enum { Value = true }; };
+template <typename T> struct TIsTObjectPtr<const          TObjectPtr<T>> { enum { Value = true }; };
+template <typename T> struct TIsTObjectPtr<      volatile TObjectPtr<T>> { enum { Value = true }; };
+template <typename T> struct TIsTObjectPtr<const volatile TObjectPtr<T>> { enum { Value = true }; };
 
 template <typename T>
 struct TRemoveObjectPointer
@@ -329,6 +329,12 @@ struct TRemoveObjectPointer<TObjectPtr<T>>
 
 namespace ObjectPtr_Private
 {
+	template <typename T> struct TRawPointerType                               { using Type = T;  };
+	template <typename T> struct TRawPointerType<               TObjectPtr<T>> { using Type = T*; };
+	template <typename T> struct TRawPointerType<const          TObjectPtr<T>> { using Type = T*; };
+	template <typename T> struct TRawPointerType<      volatile TObjectPtr<T>> { using Type = T*; };
+	template <typename T> struct TRawPointerType<const volatile TObjectPtr<T>> { using Type = T*; };
+
 	/** Coerce to pointer through implicit conversion to const T* (overload through less specific "const T*" parameter to avoid ambiguity with other coercion options that may also exist. */
 	template <typename T>
 	FORCEINLINE const T* CoerceToPointer(const T* Other)
@@ -476,6 +482,288 @@ FORCEINLINE T** ToRawPtrArrayUnsafe(T** ArrayOfPtr)
 	return ArrayOfPtr;
 }
 
+PRAGMA_DISABLE_DEPRECATION_WARNINGS
+template <
+	typename ArrayType,
+	typename ArrayTypeNoRef = std::remove_reference_t<ArrayType>,
+	std::enable_if_t<TIsTArray<ArrayTypeNoRef>::Value>* = nullptr
+>
+decltype(auto) ToRawPtrTArrayUnsafe(ArrayType&& Array)
+{
+	using ArrayElementType         = typename ArrayTypeNoRef::ElementType;
+	using ArrayAllocatorType       = typename ArrayTypeNoRef::Allocator;
+	using RawPointerType           = typename ObjectPtr_Private::TRawPointerType<ArrayElementType>::Type;
+	using QualifiedRawPointerType  = typename TCopyQualifiersFromTo<ArrayElementType, RawPointerType>::Type;
+	using NewArrayType             = TArray<QualifiedRawPointerType, ArrayAllocatorType>;
+	using RefQualifiedNewArrayType = typename TCopyQualifiersAndRefsFromTo<ArrayType, NewArrayType>::Type;
+
+	return (RefQualifiedNewArrayType&)Array;
+}
+
+template <
+	typename SetElementType,
+	typename Allocator,
+	typename RawPtrType = typename SetElementType::ElementType*,
+	std::enable_if_t<
+		TIsTObjectPtr<std::remove_cv_t<SetElementType>>::Value
+	>* = nullptr
+>
+FORCEINLINE TSet<RawPtrType, DefaultKeyFuncs<RawPtrType>, Allocator>& ToRawPtrTSetUnsafe(TSet<SetElementType, DefaultKeyFuncs<SetElementType>, Allocator>& InSet)
+{
+	return InSet;
+}
+
+template <
+	typename SetElementType,
+	typename Allocator,
+	typename RawPtrType = typename SetElementType::ElementType*,
+	std::enable_if_t<
+		TIsTObjectPtr<std::remove_cv_t<SetElementType>>::Value
+	>* = nullptr
+>
+FORCEINLINE const TSet<RawPtrType, DefaultKeyFuncs<RawPtrType>, Allocator>& ToRawPtrTSetUnsafe(const TSet<SetElementType, DefaultKeyFuncs<SetElementType>, Allocator>& InSet)
+{
+	return InSet;
+}
+
+template <
+	typename ElementType,
+	typename KeyFuncs,
+	typename Allocator,
+	std::enable_if_t<
+		!TIsTObjectPtr<std::remove_cv_t<ElementType>>::Value
+	>* = nullptr
+>
+FORCEINLINE TSet<ElementType, KeyFuncs, Allocator>& ToRawPtrTSetUnsafe(TSet<ElementType, KeyFuncs, Allocator>& InSet)
+{
+	return InSet;
+}
+
+template <
+	typename ElementType,
+	typename KeyFuncs,
+	typename Allocator,
+	std::enable_if_t<
+		!TIsTObjectPtr<std::remove_cv_t<ElementType>>::Value
+	>* = nullptr
+>
+FORCEINLINE const TSet<ElementType, KeyFuncs, Allocator>& ToRawPtrTSetUnsafe(const TSet<ElementType, KeyFuncs, Allocator>& InSet)
+{
+	return InSet;
+}
+
+template <
+	typename KeyType,
+	typename ValueType,
+	typename Allocator,
+	typename KeyRawPtrType = typename KeyType::ElementType*,
+	std::enable_if_t<
+		TIsTObjectPtr<std::remove_cv_t<KeyType>>::Value
+		&&
+		!TIsTObjectPtr<std::remove_cv_t<ValueType>>::Value
+	>* = nullptr
+>
+FORCEINLINE TMap<KeyRawPtrType, ValueType, Allocator>& ToRawPtrTMapUnsafe(TMap<KeyType, ValueType, Allocator>& InMap)
+{
+	return InMap;
+}
+
+template <
+	typename KeyType,
+	typename ValueType,
+	typename Allocator,
+	typename ValueRawPtrType = typename ValueType::ElementType*,
+	std::enable_if_t<
+		!TIsTObjectPtr<std::remove_cv_t<KeyType>>::Value
+		&&
+		TIsTObjectPtr<std::remove_cv_t<ValueType>>::Value
+	>* = nullptr
+>
+FORCEINLINE TMap<KeyType, ValueRawPtrType, Allocator>& ToRawPtrTMapUnsafe(TMap<KeyType, ValueType, Allocator>& InMap)
+{
+	return InMap;
+}
+
+template <
+	typename KeyType,
+	typename ValueType,
+	typename Allocator,
+	typename KeyRawPtrType = typename KeyType::ElementType*,
+	typename ValueRawPtrType = typename ValueType::ElementType*,
+	std::enable_if_t<
+		TIsTObjectPtr<std::remove_cv_t<KeyType>>::Value
+		&&
+		TIsTObjectPtr<std::remove_cv_t<ValueType>>::Value
+	>* = nullptr
+>
+FORCEINLINE TMap<KeyRawPtrType, ValueRawPtrType, Allocator>& ToRawPtrTMapUnsafe(TMap<KeyType, ValueType, Allocator>& InMap)
+{
+	return InMap;
+}
+
+template <
+	typename KeyType,
+	typename ValueType,
+	typename Allocator,
+	typename KeyRawPtrType = typename KeyType::ElementType*,
+	std::enable_if_t<
+		TIsTObjectPtr<std::remove_cv_t<KeyType>>::Value
+		&&
+		!TIsTObjectPtr<std::remove_cv_t<ValueType>>::Value
+	>* = nullptr
+>
+FORCEINLINE const TMap<KeyRawPtrType, ValueType, Allocator>& ToRawPtrTMapUnsafe(const TMap<KeyType, ValueType, Allocator>& InMap)
+{
+	return InMap;
+}
+
+template <
+	typename KeyType,
+	typename ValueType,
+	typename Allocator,
+	typename ValueRawPtrType = typename ValueType::ElementType*,
+	std::enable_if_t<
+		!TIsTObjectPtr<std::remove_cv_t<KeyType>>::Value
+		&&
+		TIsTObjectPtr<std::remove_cv_t<ValueType>>::Value
+	>* = nullptr
+>
+FORCEINLINE const TMap<KeyType, ValueRawPtrType, Allocator>& ToRawPtrTMapUnsafe(const TMap<KeyType, ValueType, Allocator>& InMap)
+{
+	return InMap;
+}
+
+template <
+	typename KeyType,
+	typename ValueType,
+	typename Allocator,
+	typename KeyRawPtrType = typename KeyType::ElementType*,
+	typename ValueRawPtrType = typename ValueType::ElementType*,
+	std::enable_if_t<
+		TIsTObjectPtr<std::remove_cv_t<KeyType>>::Value
+		&&
+		TIsTObjectPtr<std::remove_cv_t<ValueType>>::Value
+	>* = nullptr
+>
+FORCEINLINE const TMap<KeyRawPtrType, ValueRawPtrType, Allocator>& ToRawPtrTMapUnsafe(const TMap<KeyType, ValueType, Allocator>& InMap)
+{
+	return InMap;
+}
+
+template <
+	typename KeyType,
+	typename ValueType,
+	typename Allocator,
+	std::enable_if_t<
+		!TIsTObjectPtr<std::remove_cv_t<KeyType>>::Value
+		&&
+		!TIsTObjectPtr<std::remove_cv_t<ValueType>>::Value
+	>* = nullptr
+>
+FORCEINLINE TMap<KeyType, ValueType, Allocator>& ToRawPtrTMapUnsafe(TMap<KeyType, ValueType, Allocator>& InMap)
+{
+	return InMap;
+}
+
+template <
+	typename KeyType,
+	typename ValueType,
+	typename Allocator,
+	std::enable_if_t<
+		!TIsTObjectPtr<std::remove_cv_t<KeyType>>::Value
+		&&
+		!TIsTObjectPtr<std::remove_cv_t<ValueType>>::Value
+	>* = nullptr
+>
+FORCEINLINE const TMap<KeyType, ValueType, Allocator>& ToRawPtrTMapUnsafe(const TMap<KeyType, ValueType, Allocator>& InMap)
+{
+	return InMap;
+}
+PRAGMA_ENABLE_DEPRECATION_WARNINGS
+
+template <typename T>
+struct TContainerElementTypeCompatibility<TObjectPtr<T>>
+{
+	typedef T* ReinterpretType;
+	
+	template <typename IterBeginType, typename IterEndType, typename OperatorType = typename TRemoveReference<decltype(*DeclVal<IterBeginType>())>::Type& (*)(IterBeginType&)>
+	UE_OBJPTR_DEPRECATED(5.0, "Reinterpretation between ranges of one type to another type is deprecated.")
+	static void ReinterpretRange(IterBeginType Iter, IterEndType IterEnd, OperatorType Operator = [](IterBeginType& InIt) -> decltype(auto) { return *InIt; })
+	{
+#if UE_WITH_OBJECT_HANDLE_LATE_RESOLVE
+		while (Iter != IterEnd)
+		{
+			Operator(Iter).Get();
+			++Iter;
+		}
+#endif
+	}
+
+	typedef T* CopyFromOtherType;
+
+	UE_OBJPTR_DEPRECATED(5.0, "Copying ranges of one type to another type is deprecated.")
+	static constexpr void CopyingFromOtherType() {}
+};
+
+template <typename T>
+struct TContainerElementTypeCompatibility<const TObjectPtr<T>>
+{
+	typedef T* const ReinterpretType;
+
+	template <typename IterBeginType, typename IterEndType, typename OperatorType = const TObjectPtr<T>&(*)(IterBeginType&)>
+	UE_OBJPTR_DEPRECATED(5.0, "Reinterpretation between ranges of one type to another type is deprecated.")
+	static void ReinterpretRange(IterBeginType Iter, IterEndType IterEnd, OperatorType Operator = [](IterBeginType& InIt) -> const TObjectPtr<T>& { return *InIt; })
+	{
+#if UE_WITH_OBJECT_HANDLE_LATE_RESOLVE
+		while (Iter != IterEnd)
+		{
+			Operator(Iter).Get();
+			++Iter;
+		}
+#endif
+	}
+
+	typedef T* const CopyFromOtherType;
+
+	UE_OBJPTR_DEPRECATED(5.0, "Copying ranges of one type to another type is deprecated.")
+	static constexpr void CopyingFromOtherType() {}
+};
+
+
+// Trait which allows TObjectPtr to be default constructed by memsetting to zero.
+template <typename T>
+struct TIsZeroConstructType<TObjectPtr<T>>
+{
+	enum { Value = true };
+};
+
+// Trait which allows TObjectPtr to be memcpy'able from pointers.
+template <typename T>
+struct TIsBitwiseConstructible<TObjectPtr<T>, T*>
+{
+	enum { Value = true };
+};
+
+template <typename T, class PREDICATE_CLASS>
+struct TDereferenceWrapper<TObjectPtr<T>, PREDICATE_CLASS>
+{
+	const PREDICATE_CLASS& Predicate;
+
+	TDereferenceWrapper(const PREDICATE_CLASS& InPredicate)
+		: Predicate(InPredicate) {}
+
+	/** Dereference pointers */
+	FORCEINLINE bool operator()(const TObjectPtr<T>& A, const TObjectPtr<T>& B) const
+	{
+		return Predicate(*A, *B);
+	}
+};
+
+template <typename T>
+struct TCallTraits<TObjectPtr<T>> : public TCallTraitsBase<TObjectPtr<T>>
+{
+	using ConstPointerType = typename TCallTraitsParamTypeHelper<const TObjectPtr<const T>, true>::ConstParamType;
+};
 
 
 template <typename T>
@@ -505,3 +793,15 @@ inline void Exchange(T*& A, TObjectPtr<T>& B)
 {
 	Swap(A, (T*&)B);
 }
+
+/** Exchange variants between TArray<TObjectPtr<T>> and TArray<T*> */
+template <typename T>
+inline void Exchange(TArray<TObjectPtr<T>>& A, TArray<T*>& B)
+{
+	Swap(ToRawPtrTArrayUnsafe(A), B);
+}
+template <typename T>
+inline void Exchange(TArray<T*>& A, TArray<TObjectPtr<T>>& B)
+{
+	Swap(A, ToRawPtrTArrayUnsafe(B));
+}
