Fixed C++ vs UPROPERTY system size mismatch for FExpressionInnput and FMaterialInput

+ Added more debug info when asserting about mismatched struct sizes

diff --git a/Engine/Source/Runtime/CoreUObject/Private/UObject/Class.cpp b/Engine/Source/Runtime/CoreUObject/Private/UObject/Class.cpp
--- a/Engine/Source/Runtime/CoreUObject/Private/UObject/Class.cpp
+++ b/Engine/Source/Runtime/CoreUObject/Private/UObject/Class.cpp
@@ -3207,7 +3207,8 @@ void UScriptStruct::DestroyStruct(void* Dest, int32 ArrayDim) const
 		}
 		ClearedSize = TheCppStructOps->GetSize();
 		// here we want to make sure C++ and the property system agree on the size
-		check(Stride == ClearedSize && PropertiesSize == ClearedSize);
+		checkf(Stride == ClearedSize && PropertiesSize == ClearedSize, TEXT("C++ and the property system struct size mismatch for %s (C++ Size: %d, Property Size: %d)"),
+			*GetPathName(), ClearedSize, Stride);
 	}
 
 	if (PropertiesSize > ClearedSize)
