Also undef TEST_NAME_ROOT in Texture2D and BulkData tests.

diff --git a/Engine/Source/Runtime/CoreUObject/Private/Tests/Serialization/BulkDataTests.cpp b/Engine/Source/Runtime/CoreUObject/Private/Tests/Serialization/BulkDataTests.cpp
--- a/Engine/Source/Runtime/CoreUObject/Private/Tests/Serialization/BulkDataTests.cpp
+++ b/Engine/Source/Runtime/CoreUObject/Private/Tests/Serialization/BulkDataTests.cpp
@@ -80,7 +80,9 @@ namespace BulkDataTest
 		TestNull(TEXT("Getting a copy of BulkData object after calling ::RemoveBulkData should return nullptr!"), DataPtr);
 
 		return true;
-	}	
+	}
+
+	#undef TEST_NAME_ROOT
 }
 
 #endif // WITH_DEV_AUTOMATION_TESTS
diff --git a/Engine/Source/Runtime/Engine/Private/Tests/Texture2DTests.cpp b/Engine/Source/Runtime/Engine/Private/Tests/Texture2DTests.cpp
--- a/Engine/Source/Runtime/Engine/Private/Tests/Texture2DTests.cpp
+++ b/Engine/Source/Runtime/Engine/Private/Tests/Texture2DTests.cpp
@@ -34,6 +34,8 @@ namespace Texture2DTest
 
 		return true;
 	}
+
+	#undef TEST_NAME_ROOT
 }
 
 #endif 
