Fix non-unity

diff --git a/Engine/Source/Runtime/CoreUObject/Private/Serialization/AsyncLoading2.cpp b/Engine/Source/Runtime/CoreUObject/Private/Serialization/AsyncLoading2.cpp
--- a/Engine/Source/Runtime/CoreUObject/Private/Serialization/AsyncLoading2.cpp
+++ b/Engine/Source/Runtime/CoreUObject/Private/Serialization/AsyncLoading2.cpp
@@ -50,6 +50,7 @@
 #include "Serialization/UnversionedPropertySerialization.h"
 #include "Serialization/Zenaphore.h"
 #include "UObject/GCObject.h"
+#include "UObject/LinkerLoad.h"
 #include "UObject/ObjectRedirector.h"
 #include "Serialization/BulkData.h"
 #include "Serialization/LargeMemoryReader.h"
