UObject: made ParseObject() redirect aware for a Control Rig Jira

diff --git a/Engine/Source/Runtime/CoreUObject/Private/UObject/UObjectGlobals.cpp b/Engine/Source/Runtime/CoreUObject/Private/UObject/UObjectGlobals.cpp
--- a/Engine/Source/Runtime/CoreUObject/Private/UObject/UObjectGlobals.cpp
+++ b/Engine/Source/Runtime/CoreUObject/Private/UObject/UObjectGlobals.cpp
@@ -822,12 +822,25 @@ bool ParseObject( const TCHAR* Stream, const TCHAR* Match, UClass* Class, UObjec
 		UObject* Res = StaticFindObject( Class, InParent, TempStr );
 		if( !Res )
 		{
-			// Match found, object not found
-			if (bInvalidObject)
-			{
-				*bInvalidObject = true;
+			if (Class->IsChildOf<UClass>())
+			{ 
+				const FString ObjectName = FPackageName::ObjectPathToObjectName(FString(TempStr));
+				const FName RedirectedObjectName = FLinkerLoad::FindNewNameForClass(*ObjectName, false);
+				if (!RedirectedObjectName.IsNone())
+				{
+					Res = StaticFindObject(Class, InParent, *RedirectedObjectName.ToString());
+				}
+			}
+
+			if (!Res)
+			{ 
+				// Match found, object not found
+				if (bInvalidObject)
+				{
+					*bInvalidObject = true;
+				}
+				return 0;
 			}
-			return 0;
 		}
 
 		// Match found, object found
