Fix DirectoriesToAlwaysCook and DirectoriesToNeverCook not recognizing a path if it is the roots of a plugin content directory without a terminating / (e.g. /PluginRoot instead of /PluginRoot/).
Change all FPackageName functions that find the mount point for a filepath or packagename to find the mountpoint even if the input path is the root path of the mountpoint but without the terminating slash.
#rnx

diff --git a/Engine/Source/Runtime/CoreUObject/Private/Misc/PackageName.cpp b/Engine/Source/Runtime/CoreUObject/Private/Misc/PackageName.cpp
--- a/Engine/Source/Runtime/CoreUObject/Private/Misc/PackageName.cpp
+++ b/Engine/Source/Runtime/CoreUObject/Private/Misc/PackageName.cpp
@@ -139,6 +139,7 @@ bool FPackageName::TryConvertGameRelativePackagePathToLocalPath(FStringView Rela
 	if (RelativePackagePath.StartsWith(TEXT("/"), ESearchCase::CaseSensitive))
 	{
 		// If this starts with /, this includes a root like /engine
+		// MERGE Note: Remove the change to this function from UE4/Main; UE5 fixes it differently inside of TryConvertLongPackageNameToFilename
 		return FPackageName::TryConvertLongPackageNameToFilename(FString(RelativePackagePath), OutLocalPath);
 	}
 	else
@@ -401,7 +402,7 @@ void FPackageName::InternalFilenameToLongPackageName(FStringView InFilename, FSt
 		FReadScopeLock ScopeLock(ContentMountPointCriticalSection);
 		for (const auto& Pair : Paths.ContentRootToPath)
 		{
-			if (Filename.StartsWith(Pair.RootPath))
+			if (FPathViews::IsParentPathOf(Pair.RootPath, Filename))
 			{
 				bIsValidLongPackageName = true;
 				break;
@@ -440,9 +441,10 @@ void FPackageName::InternalFilenameToLongPackageName(FStringView InFilename, FSt
 		FReadScopeLock ScopeLock(ContentMountPointCriticalSection);
 		for (const auto& Pair : Paths.ContentPathToRoot)
 		{
-			if (Result.StartsWith(Pair.ContentPath))
+			FStringView RelPath;
+			if (FPathViews::TryMakeChildPathRelativeTo(Result, Pair.ContentPath, RelPath))
 			{
-				OutPackageName << Pair.RootPath << Result.RightChop(Pair.ContentPath.Len());
+				OutPackageName << Pair.RootPath << RelPath;
 				return;
 			}
 		}
@@ -521,9 +523,10 @@ bool FPackageName::TryConvertLongPackageNameToFilename(const FString& InLongPack
 	FReadScopeLock ScopeLock(ContentMountPointCriticalSection);
 	for (const auto& Pair : Paths.ContentRootToPath)
 	{
-		if (InLongPackageName.StartsWith(Pair.RootPath))
+		FStringView RelPath;
+		if (FPathViews::TryMakeChildPathRelativeTo(InLongPackageName, Pair.RootPath, RelPath))
 		{
-			OutFilename = Pair.ContentPath + InLongPackageName.Mid(Pair.RootPath.Len()) + InExtension;
+			OutFilename = Pair.ContentPath + RelPath + InExtension;
 			return true;
 		}
 	}
@@ -538,7 +541,7 @@ bool FPackageName::ConvertRootPathToContentPath( const FString& RootPath, FStrin
 	FReadScopeLock ScopeLock(ContentMountPointCriticalSection);
 	for (const auto& Pair : Paths.ContentRootToPath)
 	{
-		if (RootPath.StartsWith(Pair.RootPath))
+		if (FPathViews::IsParentPathOf(Pair.RootPath, RootPath))
 		{
 			OutContentPath = Pair.ContentPath;
 			return true;
@@ -651,10 +654,11 @@ bool FPackageName::SplitLongPackageName(const FString& InLongPackageName, FStrin
 
 	// Check to see whether our package came from a valid root
 	OutPackageRoot.Empty();
+	FStringView PackageRelPath;
 	for(auto RootIt = ValidRoots.CreateConstIterator(); RootIt; ++RootIt)
 	{
 		const FString& PackageRoot = *RootIt;
-		if(InLongPackageName.StartsWith(PackageRoot))
+		if (FPathViews::TryMakeChildPathRelativeTo(InLongPackageName, PackageRoot, PackageRelPath))
 		{
 			OutPackageRoot = PackageRoot / "";
 			break;
@@ -668,9 +672,9 @@ bool FPackageName::SplitLongPackageName(const FString& InLongPackageName, FStrin
 	}
 
 	// Use the standard path functions to get the rest
-	const FString RemainingPackageName = InLongPackageName.Mid(OutPackageRoot.Len());
-	OutPackagePath = FPaths::GetPath(RemainingPackageName) / "";
-	OutPackageName = FPaths::GetCleanFilename(RemainingPackageName);
+	FString PackageRelPathStr(PackageRelPath);
+	OutPackagePath = FPaths::GetPath(PackageRelPathStr) / "";
+	OutPackageName = FPaths::GetCleanFilename(PackageRelPathStr);
 
 	if(bStripRootLeadingSlash && OutPackageRoot.StartsWith(TEXT("/"), ESearchCase::CaseSensitive))
 	{
@@ -878,7 +882,7 @@ bool FPackageName::IsValidLongPackageName(FStringView InLongPackageName, bool bI
 	for (int32 RootIdx = 0; RootIdx < ValidRoots.Num(); ++RootIdx)
 	{
 		const FString& Root = ValidRoots[RootIdx];
-		if (InLongPackageName.StartsWith(Root))
+		if (FPathViews::IsParentPathOf(Root, InLongPackageName))
 		{
 			if (OutReason) *OutReason = EErrorCode::PackageNameUnknown;
 			return true;
@@ -958,7 +962,7 @@ bool FPackageName::IsValidPath(const FString& InPath)
 	FReadScopeLock ScopeLock(ContentMountPointCriticalSection);
 	for (const FPathPair& Pair : Paths.ContentRootToPath)
 	{
-		if (InPath.StartsWith(Pair.RootPath))
+		if (FPathViews::IsParentPathOf(Pair.RootPath, InPath))
 		{
 			return true;
 		}
@@ -993,7 +997,7 @@ FName FPackageName::GetPackageMountPoint(const FString& InPackagePath, bool InWi
 	int32 WithoutSlashes = InWithoutSlashes ? 1 : 0;
 	for (auto RootIt = MountPoints.CreateConstIterator(); RootIt; ++RootIt)
 	{
-		if (InPackagePath.StartsWith(*RootIt))
+		if (FPathViews::IsParentPathOf(*RootIt, InPackagePath))
 		{
 			return FName(*RootIt->Mid(WithoutSlashes, RootIt->Len() - (2 * WithoutSlashes)));
 		}
@@ -1111,11 +1115,11 @@ bool FPackageName::TryGetMountPointForPath(FStringView InFilePathOrPackageName,
 	for (const auto& Pair : Paths.ContentRootToPath)
 	{
 		FString RootFileAbsPath = FPaths::ConvertRelativePathToFull(Pair.ContentPath);
-		if (InFilePathOrPackageName.StartsWith(Pair.RootPath))
+		FStringView RelPath;
+		if (FPathViews::TryMakeChildPathRelativeTo(InFilePathOrPackageName, Pair.RootPath, RelPath))
 		{
 			OutMountPointPackageName << Pair.RootPath;
 			OutMountPointFilePath << Pair.ContentPath;
-			FStringView RelPath(InFilePathOrPackageName.RightChop(Pair.RootPath.Len()));
 			OutRelPath << RelPath;
 			if (OutFlexNameType)
 			{
@@ -1134,11 +1138,11 @@ bool FPackageName::TryGetMountPointForPath(FStringView InFilePathOrPackageName,
 			}
 			return true;
 		}
-		else if (PossibleAbsFilePath.StartsWith(RootFileAbsPath))
+		else if (FPathViews::TryMakeChildPathRelativeTo(PossibleAbsFilePath, RootFileAbsPath, RelPath))
 		{
 			OutMountPointPackageName << Pair.RootPath;
 			OutMountPointFilePath << Pair.ContentPath;
-			OutRelPath << FStringView(PossibleAbsFilePath).RightChop(RootFileAbsPath.Len());
+			OutRelPath << RelPath;
 			if (OutFlexNameType)
 			{
 				*OutFlexNameType = EFlexNameType::LocalPath;
@@ -2001,22 +2005,22 @@ FWideStringView FPackageName::ObjectPathToObjectName(FWideStringView InObjectPat
 
 bool FPackageName::IsExtraPackage(FStringView InPackageName)
 {
-	return InPackageName.StartsWith(FLongPackagePathsSingleton::Get().ExtraRootPath);
+	return FPathViews::IsParentPathOf(FLongPackagePathsSingleton::Get().ExtraRootPath, InPackageName);
 }
 
 bool FPackageName::IsScriptPackage(FStringView InPackageName)
 {
-	return InPackageName.StartsWith(FLongPackagePathsSingleton::Get().ScriptRootPath);
+	return FPathViews::IsParentPathOf(FLongPackagePathsSingleton::Get().ScriptRootPath, InPackageName);
 }
 
 bool FPackageName::IsMemoryPackage(FStringView InPackageName)
 {
-	return InPackageName.StartsWith(FLongPackagePathsSingleton::Get().MemoryRootPath);
+	return FPathViews::IsParentPathOf(FLongPackagePathsSingleton::Get().MemoryRootPath, InPackageName);
 }
 
 bool FPackageName::IsTempPackage(FStringView InPackageName)
 {
-	return InPackageName.StartsWith(FLongPackagePathsSingleton::Get().TempRootPath);
+	return FPathViews::IsParentPathOf(FLongPackagePathsSingleton::Get().TempRootPath, InPackageName);
 }
 
 bool FPackageName::IsLocalizedPackage(FStringView InPackageName)
