Fix iostore builds not working:
If FPackageName::DoesPackageExist uses DoesPackageExistOverrideDelegate to find the package, it needs to set the extension to a special marker EmptyExtension, so that IPackageResourceManagerFile sends the filename without an extension to IFileManager rather than searching for one of the regular extensions.
#rnx

diff --git a/Engine/Source/Runtime/CoreUObject/Private/Misc/PackageName.cpp b/Engine/Source/Runtime/CoreUObject/Private/Misc/PackageName.cpp
--- a/Engine/Source/Runtime/CoreUObject/Private/Misc/PackageName.cpp
+++ b/Engine/Source/Runtime/CoreUObject/Private/Misc/PackageName.cpp
@@ -1327,6 +1327,10 @@ bool FPackageName::DoesPackageExist(const FPackagePath& PackagePath, const FGuid
 			if (OutPackagePath)
 			{
 				*OutPackagePath = PackagePath;
+				if (OutPackagePath->GetHeaderExtension() == EPackageExtension::Unspecified)
+				{
+					OutPackagePath->SetHeaderExtension(EPackageExtension::EmptyString);
+				}
 			}
 			return true;
 		}
diff --git a/Engine/Source/Runtime/CoreUObject/Private/Misc/PackagePath.cpp b/Engine/Source/Runtime/CoreUObject/Private/Misc/PackagePath.cpp
--- a/Engine/Source/Runtime/CoreUObject/Private/Misc/PackagePath.cpp
+++ b/Engine/Source/Runtime/CoreUObject/Private/Misc/PackagePath.cpp
@@ -49,6 +49,8 @@ EPackageSegment ExtensionToSegment(EPackageExtension PackageExtension)
 		return EPackageSegment::Header;
 	case EPackageExtension::Custom:
 		return EPackageSegment::Header;
+	case EPackageExtension::EmptyString:
+		return EPackageSegment::Header;
 	case EPackageExtension::Exports:
 		return EPackageSegment::Exports;
 	case EPackageExtension::BulkDataDefault:
@@ -99,6 +101,8 @@ const TCHAR* LexToString(EPackageExtension PackageExtension)
 		return TEXT(".utxtmap");
 	case EPackageExtension::Custom:
 		return TEXT(".CustomExtension");
+	case EPackageExtension::EmptyString:
+		return TEXT("");
 	case EPackageExtension::Exports:
 		return TEXT(".uexp");
 	case EPackageExtension::BulkDataDefault:
@@ -749,6 +753,7 @@ EPackageExtension AllExtensions[] =
 	EPackageExtension::TextAsset,
 	EPackageExtension::TextMap,
 	EPackageExtension::Custom,
+	EPackageExtension::EmptyString,
 	EPackageExtension::Exports,
 	EPackageExtension::BulkDataDefault,
 	EPackageExtension::BulkDataOptional,
@@ -769,12 +774,12 @@ TConstArrayView<EPackageExtension> FPackagePath::GetPossibleExtensions(EPackageS
 	else
 	{
 #if WITH_TEXT_ARCHIVE_SUPPORT
-		constexpr int NumHeaderExtensions = 4;
+		constexpr int NumHeaderSearchExtensions = 4;
 #else
-		constexpr int NumHeaderExtensions = 2;
+		constexpr int NumHeaderSearchExtensions = 2;
 #endif
 		static_assert(static_cast<int>(EPackageExtension::TextMap) == static_cast<int>(EPackageExtension::Asset) + 3, "Need to update the list of header extensions");
-		return TConstArrayView<EPackageExtension>(&UE::PackagePathPrivate::AllExtensions[static_cast<int>(EPackageExtension::Asset)], NumHeaderExtensions);
+		return TConstArrayView<EPackageExtension>(&UE::PackagePathPrivate::AllExtensions[static_cast<int>(EPackageExtension::Asset)], NumHeaderSearchExtensions);
 	}
 }
 
diff --git a/Engine/Source/Runtime/CoreUObject/Private/Misc/PackagePathTest.cpp b/Engine/Source/Runtime/CoreUObject/Private/Misc/PackagePathTest.cpp
--- a/Engine/Source/Runtime/CoreUObject/Private/Misc/PackagePathTest.cpp
+++ b/Engine/Source/Runtime/CoreUObject/Private/Misc/PackagePathTest.cpp
@@ -48,13 +48,14 @@ public:
 			TestEqual(Errorf(TEXT("ParsedExtension returns correct extension for %s"), LexToString(Extension)), ParsedExtension, Extension);
 			TestEqual(Errorf(TEXT("ParsedExtension sets correct ExtensionStart for %s"), LexToString(Extension)), ExtensionStart, PathWithoutExtension.Len());
 		};
-		static_assert(EPackageExtensionCount == 10, "Need to add extensions to the test");
+		static_assert(EPackageExtensionCount == 11, "Need to add extensions to the test");
 		TestExtension(EPackageExtension::Unspecified);
 		TestExtension(EPackageExtension::Asset);
 		TestExtension(EPackageExtension::Map);
 		TestExtension(EPackageExtension::TextAsset);
 		TestExtension(EPackageExtension::TextMap);
-		//TestExtension(EPackageExtension::Custom);
+		//TestExtension(EPackageExtension::Custom); // Custom requires a separate test
+		//TestExtension(EPackageExtension::EmptyString); // Parse never returns EmptyString; it returns Unspecified instead
 		TestExtension(EPackageExtension::Exports);
 		TestExtension(EPackageExtension::BulkDataDefault);
 		TestExtension(EPackageExtension::BulkDataOptional);
diff --git a/Engine/Source/Runtime/CoreUObject/Public/Misc/PackagePath.h b/Engine/Source/Runtime/CoreUObject/Public/Misc/PackagePath.h
--- a/Engine/Source/Runtime/CoreUObject/Public/Misc/PackagePath.h
+++ b/Engine/Source/Runtime/CoreUObject/Public/Misc/PackagePath.h
@@ -50,6 +50,8 @@ enum class EPackageExtension : uint8
 	 * one of the enumerated possibilies, e.g. a custom extension on a temp file
 	 */
 	Custom,
+	/** Used by iostore to indicate that the package should be requested with no extension */
+	EmptyString,
 					
 	// Other Segments
 
