Fixing PVS 7.7 Issues:

 warning V1062: The 'XYZ' class defines a custom 'new' operator. The 'delete' operator must also be defined.

diff --git a/Engine/Source/Runtime/Core/Private/FileCache/FileCache.cpp b/Engine/Source/Runtime/Core/Private/FileCache/FileCache.cpp
--- a/Engine/Source/Runtime/Core/Private/FileCache/FileCache.cpp
+++ b/Engine/Source/Runtime/Core/Private/FileCache/FileCache.cpp
@@ -510,7 +510,7 @@ IMemoryReadStreamRef FFileCacheHandle::ReadDataUncached(FGraphEventArray& OutCom
 	return new FMemoryReadStreamAsyncRequest(AsyncRequest, BytesToRead);
 }
 
-class FMemoryReadStreamCache : public IMemoryReadStream
+class FMemoryReadStreamCache : public IMemoryReadStream //-V1062
 {
 public:
 	virtual const void* Read(int64& OutSize, int64 InOffset, int64 InSize) override
diff --git a/Engine/Source/Runtime/Core/Public/Delegates/DelegateBase.h b/Engine/Source/Runtime/Core/Public/Delegates/DelegateBase.h
--- a/Engine/Source/Runtime/Core/Public/Delegates/DelegateBase.h
+++ b/Engine/Source/Runtime/Core/Public/Delegates/DelegateBase.h
@@ -37,7 +37,7 @@ ALIAS_TEMPLATE_TYPE_LAYOUT(template<typename ElementType>, FDelegateAllocatorTyp
 /**
  * Base class for unicast delegates.
  */
-class FDelegateBase
+class FDelegateBase //-V1062
 {
 	template <typename>
 	friend class TMulticastDelegateBase;
diff --git a/Engine/Source/Runtime/Core/Public/Misc/MemStack.h b/Engine/Source/Runtime/Core/Public/Misc/MemStack.h
--- a/Engine/Source/Runtime/Core/Public/Misc/MemStack.h
+++ b/Engine/Source/Runtime/Core/Public/Misc/MemStack.h
@@ -66,7 +66,7 @@ private:
  * Items are allocated via PushBytes() or the specialized operator new()s.
  * Items are freed en masse by using FMemMark to Pop() them.
  **/
-class CORE_API FMemStackBase
+class CORE_API FMemStackBase //-V1062
 {
 public:
 #if ( PLATFORM_WINDOWS && defined(__clang__) )
diff --git a/Engine/Source/Runtime/CoreUObject/Public/UObject/ObjectMacros.h b/Engine/Source/Runtime/CoreUObject/Public/UObject/ObjectMacros.h
--- a/Engine/Source/Runtime/CoreUObject/Public/UObject/ObjectMacros.h
+++ b/Engine/Source/Runtime/CoreUObject/Public/UObject/ObjectMacros.h
@@ -1583,6 +1583,11 @@ public: \
 	inline void* operator new( const size_t InSize, EInternal* InMem ) \
 	{ \
 		return (void*)InMem; \
+	} \
+	/* Eliminate V1062 warning from PVS-Studio while keeping MSVC and Clang happy. */ \
+	inline void operator delete(void* InMem) \
+	{ \
+		::operator delete(InMem); \
 	}
 
 #define DEFINE_FORBIDDEN_DEFAULT_CONSTRUCTOR_CALL(TClass) \
diff --git a/Engine/Source/Runtime/SlateCore/Public/Widgets/SUserWidget.h b/Engine/Source/Runtime/SlateCore/Public/Widgets/SUserWidget.h
--- a/Engine/Source/Runtime/SlateCore/Public/Widgets/SUserWidget.h
+++ b/Engine/Source/Runtime/SlateCore/Public/Widgets/SUserWidget.h
@@ -91,4 +91,10 @@ class SUserWidget : public SCompoundWidget
 		return FMemory::Malloc(InSize);
 	}
 
+	public:
+	// This exists to silence V1062 and provide implementation balance with the new above
+	void operator delete (void* InMem)
+	{
+		FMemory::Free(InMem);
+	}
 };
