EngineTest System fixes when running with the new loader and iodispatcher.

diff --git a/Engine/Source/Runtime/CoreUObject/Private/Tests/ObjectHandleTest.cpp b/Engine/Source/Runtime/CoreUObject/Private/Tests/ObjectHandleTest.cpp
--- a/Engine/Source/Runtime/CoreUObject/Private/Tests/ObjectHandleTest.cpp
+++ b/Engine/Source/Runtime/CoreUObject/Private/Tests/ObjectHandleTest.cpp
@@ -4,6 +4,7 @@
 
 #include "HAL/PlatformProperties.h"
 #include "ObjectRefTrackingTestBase.h"
+#include "IO/IoDispatcher.h"
 
 static_assert(sizeof(FObjectHandle) == sizeof(void*), "FObjectHandle type must always compile to something equivalent to a pointer size.");
 
@@ -150,8 +151,15 @@ bool FObjectHandleTestFailResolveInvalidTarget::RunTest(const FString& Parameter
 {
 	if (FPlatformProperties::RequiresCookedData())
 	{
-		AddExpectedError(TEXT("Couldn't find file for package"));
-		AddExpectedError(TEXT("Found 0 dependent packages..."));
+		if (FIoDispatcher::IsInitialized())
+		{
+			AddExpectedError(TEXT("SkipPackage"));
+		}
+		else
+		{
+			AddExpectedError(TEXT("Couldn't find file for package"));
+			AddExpectedError(TEXT("Found 0 dependent packages..."));
+		}
 	}
 
 	// Confirm we don't successfully resolve an incorrect reference to engine content
