Add GetAllocatedSize function to FDiskCachedAssetData, FPackageDependencyData and FLinkerTables to support measuring how much memory is used by AssetRegistry and AssetDataGatherer.
[REDACTED_HANDLE] Steve.Robb
#rnx

diff --git a/Engine/Source/Runtime/AssetRegistry/Private/DiskCachedAssetData.h b/Engine/Source/Runtime/AssetRegistry/Private/DiskCachedAssetData.h
--- a/Engine/Source/Runtime/AssetRegistry/Private/DiskCachedAssetData.h
+++ b/Engine/Source/Runtime/AssetRegistry/Private/DiskCachedAssetData.h
@@ -47,4 +47,13 @@ public:
 
 		DependencyData.SerializeForCache(Ar);
 	}
+
+	/** Returns the amount of memory allocated by this container, not including sizeof(*this). */
+	uint32 GetAllocatedSize() const
+	{
+		uint32 Result = 0;
+		Result += AssetDataList.GetAllocatedSize();
+		Result += DependencyData.GetAllocatedSize();
+		return Result;
+	}
 };
diff --git a/Engine/Source/Runtime/AssetRegistry/Private/PackageDependencyData.h b/Engine/Source/Runtime/AssetRegistry/Private/PackageDependencyData.h
--- a/Engine/Source/Runtime/AssetRegistry/Private/PackageDependencyData.h
+++ b/Engine/Source/Runtime/AssetRegistry/Private/PackageDependencyData.h
@@ -57,4 +57,13 @@ public:
 		return ImportUsedInGame.Num() == ImportMap.Num() &&
 			SoftPackageUsedInGame.Num() == SoftPackageReferenceList.Num();
 	}
+
+	/** Returns the amount of memory allocated by this container, not including sizeof(*this). */
+	uint32 GetAllocatedSize() const
+	{
+		uint32 Result = FLinkerTables::GetAllocatedSize();
+		Result += ImportUsedInGame.GetAllocatedSize();
+		Result += SoftPackageUsedInGame.GetAllocatedSize();
+		return Result;
+	}
 };
diff --git a/Engine/Source/Runtime/CoreUObject/Private/UObject/Linker.cpp b/Engine/Source/Runtime/CoreUObject/Private/UObject/Linker.cpp
--- a/Engine/Source/Runtime/CoreUObject/Private/UObject/Linker.cpp
+++ b/Engine/Source/Runtime/CoreUObject/Private/UObject/Linker.cpp
@@ -160,6 +160,25 @@ FName FLinker::GetExportClassName( int32 i )
 	return NAME_Class;
 }
 
+uint32 FLinkerTables::GetAllocatedSize() const
+{
+	uint32 Result = 0;
+	Result += ImportMap.GetAllocatedSize();
+	Result += ExportMap.GetAllocatedSize();
+	Result += DependsMap.GetAllocatedSize();
+	for (const TArray<FPackageIndex>& DependencyList : DependsMap)
+	{
+		Result += DependencyList.GetAllocatedSize();
+	}
+	Result += SoftPackageReferenceList.GetAllocatedSize();
+	Result += SearchableNamesMap.GetAllocatedSize();
+	for (const TPair<FPackageIndex, TArray<FName>>& Pair : SearchableNamesMap)
+	{
+		Result += Pair.Value.GetAllocatedSize();
+	}
+	return Result;
+}
+
 /*----------------------------------------------------------------------------
 	FLinker.
 ----------------------------------------------------------------------------*/
diff --git a/Engine/Source/Runtime/CoreUObject/Public/UObject/Linker.h b/Engine/Source/Runtime/CoreUObject/Public/UObject/Linker.h
--- a/Engine/Source/Runtime/CoreUObject/Public/UObject/Linker.h
+++ b/Engine/Source/Runtime/CoreUObject/Public/UObject/Linker.h
@@ -166,6 +166,9 @@ public:
 	/** Serializes the searchable name map */
 	COREUOBJECT_API void SerializeSearchableNamesMap(FArchive &Ar);
 	COREUOBJECT_API void SerializeSearchableNamesMap(FStructuredArchive::FSlot Slot);
+
+	/** Returns the amount of memory allocated by this container, not including sizeof(*this). */
+	COREUOBJECT_API uint32 GetAllocatedSize() const;
 };
 
 
