Speculative fix for licensee reported compile error on Xbox

I'm unable to repro a reported compilation issue on Xbox but CL 14587450 addressed the same issue for Hololens by adding a platform exclusion to these static asserts.  I've changed it to exclude all Microsoft platforms rather than be piecemeal.

diff --git a/Engine/Source/Runtime/CoreUObject/Private/Tests/ObjectPtrTest.cpp b/Engine/Source/Runtime/CoreUObject/Private/Tests/ObjectPtrTest.cpp
--- a/Engine/Source/Runtime/CoreUObject/Private/Tests/ObjectPtrTest.cpp
+++ b/Engine/Source/Runtime/CoreUObject/Private/Tests/ObjectPtrTest.cpp
@@ -84,12 +84,12 @@ static_assert(TModels<CEqualityComparableWith, FMutableObjectPtr, FMutableInterf
 
 // Ensure that TObjectPtr<[const] UPackage> is not comparable with a TObjectPtr<[const] UInterface> regardless of constness
 // TODO: This only ensures that at least one of the A==B,B==A,A!=B,B!=A operations fail, not that they all fail.
-#if !(PLATFORM_WINDOWS || PLATFORM_HOLOLENS) || !defined(_MSC_EXTENSIONS) // MSVC static analyzer is run in non-conformance mode, and that causes these checks to fail.
+#if !(PLATFORM_MICROSOFT) || !defined(_MSC_EXTENSIONS) // MSVC static analyzer is run in non-conformance mode, and that causes these checks to fail.
 static_assert(!TModels<CEqualityComparableWith, FConstPackagePtr, FConstInterfacePtr>::Value, "Must not be able to compare equality and inequality bidirectionally between TObjectPtr<const UPackage> and TObjectPtr<const UInterface>");
 static_assert(!TModels<CEqualityComparableWith, FMutablePackagePtr, FConstInterfacePtr>::Value, "Must not be able to compare equality and inequality bidirectionally between TObjectPtr<UPackage> and TObjectPtr<const UInterface>");
 static_assert(!TModels<CEqualityComparableWith, FConstPackagePtr, FMutableInterfacePtr>::Value, "Must not be able to compare equality and inequality bidirectionally between TObjectPtr<const UPackage> and TObjectPtr<UInterface>");
 static_assert(!TModels<CEqualityComparableWith, FMutablePackagePtr, FMutableInterfacePtr>::Value, "Must not be able to compare equality and inequality bidirectionally between TObjectPtr<UPackage> and TObjectPtr<UInterface>");
-#endif // _MSC_EXTENSIONS
+#endif // #if !(PLATFORM_MICROSOFT) || !defined(_MSC_EXTENSIONS)
 
 // Ensure that TObjectPtr<[const] UObject> is comparable with a raw pointer of the same referenced type regardless of constness
 static_assert(TModels<CEqualityComparableWith, FConstObjectPtr, const UObject*>::Value, "Must be able to compare equality and inequality bidirectionally between TObjectPtr<const UObject> and const UObject*");
@@ -127,10 +127,10 @@ static_assert(!TModels<CEqualityComparableWith, FMutableObjectPtr, char*>::Value
 static_assert(TModels<CEqualityComparableWith, FConstObjectPtr, TYPE_OF_NULLPTR>::Value, "Must be able to compare equality and inequality bidirectionally between TObjectPtr<const UObject> and nullptr");
 static_assert(TModels<CEqualityComparableWith, FMutableObjectPtr, TYPE_OF_NULLPTR>::Value, "Must be able to compare equality and inequality bidirectionally between TObjectPtr<UObject> and nullptr");
 
-#if !(PLATFORM_WINDOWS || PLATFORM_HOLOLENS) || !defined(_MSC_EXTENSIONS) // MSVC static analyzer is run in non-conformance mode, and that causes these checks to fail.
+#if !(PLATFORM_MICROSOFT) || !defined(_MSC_EXTENSIONS) // MSVC static analyzer is run in non-conformance mode, and that causes these checks to fail.
 static_assert(!TModels<CEqualityComparableWith, FConstObjectPtr, long>::Value, "Should not be able to compare equality and inequality bidirectionally between TObjectPtr<const UObject> and long");
 static_assert(!TModels<CEqualityComparableWith, FMutableObjectPtr, long>::Value, "Should not be able to compare equality and inequality bidirectionally between TObjectPtr<UObject> and long");
-#endif // _MSC_EXTENSIONS
+#endif // #if !(PLATFORM_MICROSOFT) || !defined(_MSC_EXTENSIONS)
 
 // Ensure that the use of incomplete types doesn't provide a means to bypass type safety on TObjectPtr
 // NOTE: This is disabled because we're permitting this operation with a deprecation warning.
