Prototype Wrapped Object Pointers (ObjectHandle/ObjectPtr)

Change cast of wrapped pointer to return raw pointer instead of wrapped pointer.  This is because the cast may return an IInterface instead of a UObject, and TObjectPtr<IInterface> is invalid and can result in some code (eg: access tracking) interpreting the pointer as a UObject, when that is not valid.  Could be modified in the future to return a wrapped pointer if the type being cast to is a UObject, and a raw pointer if the type being cast to is an IInterface.

diff --git a/Engine/Source/Runtime/CoreUObject/Public/Templates/Casts.h b/Engine/Source/Runtime/CoreUObject/Public/Templates/Casts.h
--- a/Engine/Source/Runtime/CoreUObject/Public/Templates/Casts.h
+++ b/Engine/Source/Runtime/CoreUObject/Public/Templates/Casts.h
@@ -311,14 +311,14 @@ template< class T, class U > FORCEINLINE T* ExactCast  ( const TWeakObjectPtr<U>
 template< class T, class U > FORCEINLINE T* CastChecked( const TWeakObjectPtr<U>& Src, ECastCheckedType::Type CheckType = ECastCheckedType::NullChecked ) { return CastChecked<T>(Src.Get(), CheckType); }
 
 // object ptr versions
-template <class T, class U> FORCEINLINE TObjectPtr<T> Cast       ( const TObjectPtr<U>& Src                                                                   ) { return TObjectPtr<T>(Cast       <T>(Src.Get())); }
-template <class T, class U> FORCEINLINE TObjectPtr<T> ExactCast  ( const TObjectPtr<U>& Src                                                                   ) { return TObjectPtr<T>(ExactCast  <T>(Src.Get())); }
-template <class T, class U> FORCEINLINE TObjectPtr<T> CastChecked( const TObjectPtr<U>& Src, ECastCheckedType::Type CheckType = ECastCheckedType::NullChecked ) { return TObjectPtr<T>(CastChecked<T>(Src.Get(), CheckType)); }
+template <class T, class U> FORCEINLINE T* Cast       ( const TObjectPtr<U>& Src                                                                   ) { return Cast       <T>(Src.Get()); }
+template <class T, class U> FORCEINLINE T* ExactCast  ( const TObjectPtr<U>& Src                                                                   ) { return ExactCast  <T>(Src.Get()); }
+template <class T, class U> FORCEINLINE T* CastChecked( const TObjectPtr<U>& Src, ECastCheckedType::Type CheckType = ECastCheckedType::NullChecked ) { return CastChecked<T>(Src.Get(), CheckType); }
 
 // const object ptr versions
-template <class T, class U> FORCEINLINE TObjectPtr<const T> Cast       ( const TObjectPtr<const U>& Src                                                                   ) { return TObjectPtr<const T>(Cast       <T>(Src.Get())); }
-template <class T, class U> FORCEINLINE TObjectPtr<const T> ExactCast  ( const TObjectPtr<const U>& Src                                                                   ) { return TObjectPtr<const T>(ExactCast  <T>(Src.Get())); }
-template <class T, class U> FORCEINLINE TObjectPtr<const T> CastChecked( const TObjectPtr<const U>& Src, ECastCheckedType::Type CheckType = ECastCheckedType::NullChecked ) { return TObjectPtr<const T>(CastChecked<T>(Src.Get(), CheckType)); }
+template <class T, class U> FORCEINLINE const T* Cast       ( const TObjectPtr<const U>& Src                                                                   ) { return Cast       <T>(const_cast<U*>(Src.Get())); }
+template <class T, class U> FORCEINLINE const T* ExactCast  ( const TObjectPtr<const U>& Src                                                                   ) { return ExactCast  <T>(const_cast<U*>(Src.Get())); }
+template <class T, class U> FORCEINLINE const T* CastChecked( const TObjectPtr<const U>& Src, ECastCheckedType::Type CheckType = ECastCheckedType::NullChecked ) { return CastChecked<T>(const_cast<U*>(Src.Get()), CheckType); }
 
 // TSubclassOf versions
 template< class T, class U > FORCEINLINE T* Cast       ( const TSubclassOf<U>& Src                                                                   ) { return Cast       <T>(*Src); }
