Change the test object pointer value in the unit test named "AttemptToFindUninitializedScriptStructMembers" to avoid it being interpreted as a late resolved object reference when late resolve is enabled for wrapped object pointers.

diff --git a/Engine/Source/Runtime/CoreUObject/Private/UObject/Class.cpp b/Engine/Source/Runtime/CoreUObject/Private/UObject/Class.cpp
--- a/Engine/Source/Runtime/CoreUObject/Private/UObject/Class.cpp
+++ b/Engine/Source/Runtime/CoreUObject/Private/UObject/Class.cpp
@@ -3450,10 +3450,10 @@ int32 FStructUtils::AttemptToFindUninitializedScriptStructMembers()
 	check(TestUninitializedScriptStructMembersTestStruct != nullptr);
 	
 	{
-		const void* BadPointer = (void*)0xFFFFFFFFFFFFFFFFull;
+		const void* BadPointer = (void*)0xFEFEFEFEFEFEFEFEull;
 
 		// First test if the tests aren't broken
-		FScriptStructTestWrapper WrapperFF(TestUninitializedScriptStructMembersTestStruct, 0xFF);
+		FScriptStructTestWrapper WrapperFF(TestUninitializedScriptStructMembersTestStruct, 0xFE);
 		const FObjectPropertyBase* UninitializedProperty = CastFieldChecked<const FObjectPropertyBase>(TestUninitializedScriptStructMembersTestStruct->FindPropertyByName(TEXT("UninitializedObjectReference")));
 		const FObjectPropertyBase* InitializedProperty = CastFieldChecked<const FObjectPropertyBase>(TestUninitializedScriptStructMembersTestStruct->FindPropertyByName(TEXT("InitializedObjectReference")));
 		
