SavePackage fixes:
1) Save1 had reversed condition around SetLoadedPath: it _should_ be done if IsMountedPath.
2) Save2 needs the calls to SetPersistentGuid and SetLoadedPath that were added to Save1.
#rnx

diff --git a/Engine/Source/Runtime/CoreUObject/Private/UObject/SavePackage.cpp b/Engine/Source/Runtime/CoreUObject/Private/UObject/SavePackage.cpp
--- a/Engine/Source/Runtime/CoreUObject/Private/UObject/SavePackage.cpp
+++ b/Engine/Source/Runtime/CoreUObject/Private/UObject/SavePackage.cpp
@@ -4665,7 +4665,7 @@ FSavePackageResultStruct UPackage::Save(UPackage* InOuter, UObject* Base, EObjec
 		{
 			// if the save was successful, update the internal package filename path if we aren't currently cooking
 #if WITH_EDITOR
-			if (TargetPlatform == nullptr && !TargetPackagePath.IsMountedPath())
+			if (TargetPlatform == nullptr && TargetPackagePath.IsMountedPath())
 			{
 				InOuter->SetLoadedPath(TargetPackagePath);
 			}
diff --git a/Engine/Source/Runtime/CoreUObject/Private/UObject/SavePackage/SaveContext.h b/Engine/Source/Runtime/CoreUObject/Private/UObject/SavePackage/SaveContext.h
--- a/Engine/Source/Runtime/CoreUObject/Private/UObject/SavePackage/SaveContext.h
+++ b/Engine/Source/Runtime/CoreUObject/Private/UObject/SavePackage/SaveContext.h
@@ -86,6 +86,8 @@ public:
 				}, false/*bIncludeNestedObjects*/);
 		}
 
+		TargetPackagePath = FPackagePath::FromLocalPath(InFilename);
+
 		// if we aren't cooking and top level flags aren't empty, add RF_HasExternalPackage to them to catch external packages data
 		if (SaveArgs.TopLevelFlags != RF_NoFlags && !IsCooking())
 		{
@@ -145,6 +147,11 @@ public:
 		return Filename;
 	}
 
+	const FPackagePath& GetTargetPackagePath() const
+	{
+		return TargetPackagePath;
+	}
+
 	EObjectMark GetExcludedObjectMarks() const
 	{
 		return ExcludedObjectMarks;
@@ -548,6 +555,7 @@ private:
 	// Args
 	UPackage* Package;
 	UObject* Asset;
+	FPackagePath TargetPackagePath;
 	const TCHAR* Filename;
 	FSavePackageArgs SaveArgs;
 
diff --git a/Engine/Source/Runtime/CoreUObject/Private/UObject/SavePackage2.cpp b/Engine/Source/Runtime/CoreUObject/Private/UObject/SavePackage2.cpp
--- a/Engine/Source/Runtime/CoreUObject/Private/UObject/SavePackage2.cpp
+++ b/Engine/Source/Runtime/CoreUObject/Private/UObject/SavePackage2.cpp
@@ -210,6 +210,19 @@ FORCEINLINE void EnsurePackageLocalization(UPackage* InPackage)
 	FPackageLocalizationManager::Get().ConditionalUpdateCache();
 }
 
+void PreSavePackage(FSaveContext& SaveContext)
+{
+#if WITH_EDITOR
+	// if the in memory package filename is different the filename we are saving it to,
+	// regenerate a new persistent id for it.
+	UPackage* Package = SaveContext.GetPackage();
+	if (!SaveContext.IsCooking() && !SaveContext.IsFromAutoSave() && !Package->GetLoadedPath().IsEmpty() && Package->GetLoadedPath() != SaveContext.GetTargetPackagePath())
+	{
+		Package->SetPersistentGuid(FGuid::NewGuid());
+	}
+#endif //WITH_EDITOR
+}
+
 ESavePackageResult RoutePresave(FSaveContext& SaveContext)
 {
 	SCOPED_SAVETIMER(UPackage_RoutePresave);
@@ -1906,6 +1919,26 @@ void ClearCachedPlatformCookedData(FSaveContext& SaveContext)
 #endif
 }
 
+void PostSavePackage(FSaveContext& SaveContext)
+{
+	UPackage* Package = SaveContext.GetPackage();
+	// Package has been saved, so unmark the NewlyCreated flag.
+	Package->ClearPackageFlags(PKG_NewlyCreated);
+
+	// Send a message that the package was saved
+	UPackage::PackageSavedEvent.Broadcast(SaveContext.GetFilename(), Package);
+
+	// update the internal package filename path if we're saving to a valid mounted path and we aren't currently cooking
+#if WITH_EDITOR
+	const FPackagePath& PackagePath = SaveContext.GetTargetPackagePath();
+	if (!SaveContext.IsCooking() && PackagePath.IsMountedPath())
+	{
+		Package->SetLoadedPath(PackagePath);
+	}
+#endif
+}
+
+
 /**
  * InnerSave is the portion of Save that can be safely run concurrently
  */
@@ -2133,6 +2166,7 @@ FSavePackageResultStruct UPackage::Save2(UPackage* InPackage, UObject* InAsset,
 
 	// PreSave Asset
 	SlowTask.EnterProgressFrame();
+	PreSavePackage(SaveContext);
 	if (InAsset)
 	{
 		SaveContext.SetPreSaveCleanup(InAsset->PreSaveRoot(InFilename));
@@ -2174,14 +2208,10 @@ FSavePackageResultStruct UPackage::Save2(UPackage* InPackage, UObject* InAsset,
 
 	ClearCachedPlatformCookedData(SaveContext);
 
-	// Package Saved event
+	// PostSave Package - edit in memory package and send events
 	SlowTask.EnterProgressFrame();
 	{
-		// Package has been save, so unmark NewlyCreated flag.
-		InPackage->ClearPackageFlags(PKG_NewlyCreated);
-
-		// send a message that the package was saved
-		UPackage::PackageSavedEvent.Broadcast(InFilename, InPackage);
+		PostSavePackage(SaveContext);
 	}
 	return SaveContext.GetFinalResult();
 }
@@ -2214,6 +2244,7 @@ ESavePackageResult UPackage::SaveConcurrent(TArrayView<FPackageSaveInfo> InPacka
 			EnsureLoadingComplete(SaveContext.GetPackage()); //@todo: needed??
 
 			// PreSave Asset
+			PreSavePackage(SaveContext);
 			if (SaveContext.GetAsset())
 			{
 				SCOPED_SAVETIMER(UPackage_SaveConcurrent_PreSaveRoot);
@@ -2293,14 +2324,10 @@ ESavePackageResult UPackage::SaveConcurrent(TArrayView<FPackageSaveInfo> InPacka
 
 			ClearCachedPlatformCookedData(SaveContext);
 
-			// Package Saved event
+			// PostSave Package - edit in memory package and send events
 			if (SaveContext.Result == ESavePackageResult::Success)
 			{
-				// Package has been save, so unmark NewlyCreated flag.
-				SaveContext.GetPackage()->ClearPackageFlags(PKG_NewlyCreated);
-
-				// send a message that the package was saved
-				UPackage::PackageSavedEvent.Broadcast(SaveContext.GetFilename(), SaveContext.GetPackage());
+				PostSavePackage(SaveContext);
 			}
 			OutResults.Add(SaveContext.GetFinalResult());
 		}
