Prototype Wrapped Object Pointers (ObjectHandle/ObjectPtr)

Change the FObjectPtrProperty reflection type to derive from FObjectProperty type (instead of TFObjectPropertyBase) to avoid incompatibility with code that uses the reflection system to operate on specifically FObjectProperty fields.  Also removed the override of the GetID virtual method to allow ObjectPtr properties to have the same identifying name as Object properties.  This ensures there's no serialization incompatibility when moving between Object<->ObjectPtr or ObjectPtr->(Object compatible type) in the serialization and reflection code.

diff --git a/Engine/Source/Runtime/Core/Public/UObject/UnrealNames.inl b/Engine/Source/Runtime/Core/Public/UObject/UnrealNames.inl
--- a/Engine/Source/Runtime/Core/Public/UObject/UnrealNames.inl
+++ b/Engine/Source/Runtime/Core/Public/UObject/UnrealNames.inl
@@ -44,7 +44,6 @@ REGISTER_NAME(33,CoreUObject)
 
 // More class properties
 REGISTER_NAME(34,EnumProperty)
-REGISTER_NAME(35,ObjectPtrProperty)
 
 
 // Special types.
diff --git a/Engine/Source/Runtime/CoreUObject/Private/UObject/PropertyObjectPtr.cpp b/Engine/Source/Runtime/CoreUObject/Private/UObject/PropertyObjectPtr.cpp
--- a/Engine/Source/Runtime/CoreUObject/Private/UObject/PropertyObjectPtr.cpp
+++ b/Engine/Source/Runtime/CoreUObject/Private/UObject/PropertyObjectPtr.cpp
@@ -20,16 +20,11 @@ FString FObjectPtrProperty::GetCPPMacroType(FString& ExtendedTypeText) const
 	return TEXT("OBJECTPTR");
 }
 
-FName FObjectPtrProperty::GetID() const
-{
-	return NAME_ObjectPtrProperty;
-}
-
 void FObjectPtrProperty::SerializeItem(FStructuredArchive::FSlot Slot, void* Value, void const* Defaults) const
 {
 	FArchive& UnderlyingArchive = Slot.GetUnderlyingArchive();
 
-	FObjectPtr* ObjectPtr = GetPropertyValuePtr(Value);
+	FObjectPtr* ObjectPtr = (FObjectPtr*)GetPropertyValuePtr(Value);
 
 	Slot << *ObjectPtr;
 }
@@ -41,7 +36,7 @@ EConvertFromTypeResult FObjectPtrProperty::ConvertFromType(const FPropertyTag& T
 	{
 		uint8* DestAddress = ContainerPtrToValuePtr<uint8>(Data, Tag.ArrayIndex);
 		FArchive& UnderlyingArchive = Slot.GetUnderlyingArchive();
-		FObjectPtr* ObjectPtr = GetPropertyValuePtr(DestAddress);
+		FObjectPtr* ObjectPtr = (FObjectPtr*)GetPropertyValuePtr(DestAddress);
 		Slot << *ObjectPtr;
 
 		return UnderlyingArchive.IsCriticalError() ? EConvertFromTypeResult::CannotConvert : EConvertFromTypeResult::Converted;
@@ -100,7 +95,7 @@ bool FObjectPtrProperty::Identical(const void* A, const void* B, uint32 PortFlag
 
 UObject* FObjectPtrProperty::GetObjectPropertyValue(const void* PropertyValueAddress) const
 {
-	return GetPropertyValue(PropertyValueAddress).Get();
+	return ((FObjectPtr&)GetPropertyValue(PropertyValueAddress)).Get();
 }
 
 void FObjectPtrProperty::SetObjectPropertyValue(void* PropertyValueAddress, UObject* Value) const
@@ -115,6 +110,6 @@ bool FObjectPtrProperty::AllowCrossLevel() const
 
 uint32 FObjectPtrProperty::GetValueTypeHashInternal(const void* Src) const
 {
-	return GetTypeHash(GetPropertyValue(Src));
+	return GetTypeHash((FObjectPtr&)GetPropertyValue(Src));
 }
 
diff --git a/Engine/Source/Runtime/CoreUObject/Public/UObject/UnrealType.h b/Engine/Source/Runtime/CoreUObject/Public/UObject/UnrealType.h
--- a/Engine/Source/Runtime/CoreUObject/Public/UObject/UnrealType.h
+++ b/Engine/Source/Runtime/CoreUObject/Public/UObject/UnrealType.h
@@ -2286,15 +2286,15 @@ public:
 //
 // Describes a reference variable to another object which may be nil.
 //
-class COREUOBJECT_API FObjectPtrProperty : public TFObjectPropertyBase<FObjectPtr>
+class COREUOBJECT_API FObjectPtrProperty : public FObjectProperty
 {
-	DECLARE_FIELD(FObjectPtrProperty, TFObjectPropertyBase<FObjectPtr>, CASTCLASS_FObjectPtrProperty)
+	DECLARE_FIELD(FObjectPtrProperty, FObjectProperty, CASTCLASS_FObjectPtrProperty)
 
 	using Super::Super;
 
 #if WITH_EDITORONLY_DATA
 	explicit FObjectPtrProperty(UField* InField)
-		: TFObjectPropertyBase(InField)
+		: FObjectProperty(InField)
 	{
 	}
 #endif // WITH_EDITORONLY_DATA
@@ -2305,7 +2305,6 @@ class COREUOBJECT_API FObjectPtrProperty : public TFObjectPropertyBase<FObjectPt
 	// End of UHT interface
 
 	// FProperty interface
-	virtual FName GetID() const override;
 	virtual bool SameType(const FProperty* Other) const override;
 	virtual bool Identical(const void* A, const void* B, uint32 PortFlags) const override;
 	virtual void SerializeItem(FStructuredArchive::FSlot Slot, void* Value, void const* Defaults) const override;
