Fixed naming of previously modified WrapperFF variable used in unit tests to WrapperFE.  Also found second instance of the same FF pattern used for object references and changed it to FE.

diff --git a/Engine/Source/Runtime/CoreUObject/Private/UObject/Class.cpp b/Engine/Source/Runtime/CoreUObject/Private/UObject/Class.cpp
--- a/Engine/Source/Runtime/CoreUObject/Private/UObject/Class.cpp
+++ b/Engine/Source/Runtime/CoreUObject/Private/UObject/Class.cpp
@@ -3359,12 +3359,12 @@ private:
 
 static void FindUninitializedScriptStructMembers(UScriptStruct* ScriptStruct, EScriptStructTestCtorSyntax ConstructorSyntax, TSet<const FProperty*>& OutUninitializedProperties)
 {
-	FScriptStructTestWrapper WrapperFF(ScriptStruct, 0xFF, ConstructorSyntax);
+	FScriptStructTestWrapper WrapperFE(ScriptStruct, 0xFE, ConstructorSyntax);
 	FScriptStructTestWrapper Wrapper00(ScriptStruct, 0x00, ConstructorSyntax);
 	FScriptStructTestWrapper WrapperAA(ScriptStruct, 0xAA, ConstructorSyntax);
 	FScriptStructTestWrapper Wrapper55(ScriptStruct, 0x55, ConstructorSyntax);
 
-	const void* BadPointer = (void*)0xFFFFFFFFFFFFFFFFull;
+	const void* BadPointer = (void*)0xFEFEFEFEFEFEFEFEull;
 
 	for (const FProperty* Property : TFieldRange<FProperty>(ScriptStruct, EFieldIteratorFlags::ExcludeSuper))
 	{
@@ -3379,7 +3379,7 @@ static void FindUninitializedScriptStructMembers(UScriptStruct* ScriptStruct, ES
 		if (const FObjectPropertyBase* ObjectProperty = CastField<const FObjectPropertyBase>(Property))
 		{
 			// Check any reflected pointer properties to make sure they got initialized
-			const UObject* PropValue = ObjectProperty->GetObjectPropertyValue_InContainer(WrapperFF.GetData());
+			const UObject* PropValue = ObjectProperty->GetObjectPropertyValue_InContainer(WrapperFE.GetData());
 			if (PropValue == BadPointer)
 			{
 				OutUninitializedProperties.Add(Property);
@@ -3389,7 +3389,7 @@ static void FindUninitializedScriptStructMembers(UScriptStruct* ScriptStruct, ES
 		{
 			// Check for uninitialized boolean properties (done separately to deal with byte-wide booleans that would evaluate to true with either 0x55 or 0xAA)
 			const bool bValue0 = BoolProperty->GetPropertyValue_InContainer(Wrapper00.GetData());
-			const bool bValue1 = BoolProperty->GetPropertyValue_InContainer(WrapperFF.GetData());
+			const bool bValue1 = BoolProperty->GetPropertyValue_InContainer(WrapperFE.GetData());
 
 			if (bValue0 != bValue1)
 			{
@@ -3453,17 +3453,17 @@ int32 FStructUtils::AttemptToFindUninitializedScriptStructMembers()
 		const void* BadPointer = (void*)0xFEFEFEFEFEFEFEFEull;
 
 		// First test if the tests aren't broken
-		FScriptStructTestWrapper WrapperFF(TestUninitializedScriptStructMembersTestStruct, 0xFE);
+		FScriptStructTestWrapper WrapperFE(TestUninitializedScriptStructMembersTestStruct, 0xFE);
 		const FObjectPropertyBase* UninitializedProperty = CastFieldChecked<const FObjectPropertyBase>(TestUninitializedScriptStructMembersTestStruct->FindPropertyByName(TEXT("UninitializedObjectReference")));
 		const FObjectPropertyBase* InitializedProperty = CastFieldChecked<const FObjectPropertyBase>(TestUninitializedScriptStructMembersTestStruct->FindPropertyByName(TEXT("InitializedObjectReference")));
 		
-		const UObject* UninitializedPropValue = UninitializedProperty->GetObjectPropertyValue_InContainer(WrapperFF.GetData());
+		const UObject* UninitializedPropValue = UninitializedProperty->GetObjectPropertyValue_InContainer(WrapperFE.GetData());
 		if (UninitializedPropValue != BadPointer)
 		{
 			UE_LOG(LogClass, Warning, TEXT("ObjectProperty %s%s::%s seems to be initialized properly but it shouldn't be. Verify that AttemptToFindUninitializedScriptStructMembers() is working properly"), 
 				TestUninitializedScriptStructMembersTestStruct->GetPrefixCPP(), *TestUninitializedScriptStructMembersTestStruct->GetName(), *UninitializedProperty->GetNameCPP());
 		}
-		const UObject* InitializedPropValue = InitializedProperty->GetObjectPropertyValue_InContainer(WrapperFF.GetData());
+		const UObject* InitializedPropValue = InitializedProperty->GetObjectPropertyValue_InContainer(WrapperFE.GetData());
 		if (InitializedPropValue != nullptr)
 		{
 			UE_LOG(LogClass, Warning, TEXT("ObjectProperty %s%s::%s seems to be not initialized properly but it should be. Verify that AttemptToFindUninitializedScriptStructMembers() is working properly"),
