Extract package store interface from I/O store package loader

diff --git a/Engine/Source/Developer/IoStoreUtilities/Private/IoStoreUtilities.cpp b/Engine/Source/Developer/IoStoreUtilities/Private/IoStoreUtilities.cpp
--- a/Engine/Source/Developer/IoStoreUtilities/Private/IoStoreUtilities.cpp
+++ b/Engine/Source/Developer/IoStoreUtilities/Private/IoStoreUtilities.cpp
@@ -28,6 +28,7 @@
 #include "Serialization/MemoryReader.h"
 #include "Serialization/AsyncLoading2.h"
 #include "Serialization/ArrayReader.h"
+#include "Serialization/PackageStore.h"
 #include "UObject/Class.h"
 #include "UObject/NameBatchSerialization.h"
 #include "UObject/PackageFileSummary.h"
diff --git a/Engine/Source/Runtime/CoreUObject/Private/Serialization/AsyncLoading2.cpp b/Engine/Source/Runtime/CoreUObject/Private/Serialization/AsyncLoading2.cpp
--- a/Engine/Source/Runtime/CoreUObject/Private/Serialization/AsyncLoading2.cpp
+++ b/Engine/Source/Runtime/CoreUObject/Private/Serialization/AsyncLoading2.cpp
@@ -6,6 +6,7 @@
 
 #include "Serialization/AsyncLoading2.h"
 #include "Serialization/AsyncPackageLoader.h"
+#include "Serialization/PackageStore.h"
 #include "HAL/PlatformFileManager.h"
 #include "HAL/FileManager.h"
 #include "HAL/Event.h"
@@ -66,13 +67,6 @@
 PRAGMA_DISABLE_OPTIMIZATION
 #endif
 
-FArchive& operator<<(FArchive& Ar, FMappedName& MappedName)
-{
-	Ar << MappedName.Index << MappedName.Number;
-
-	return Ar;
-}
-
 FArchive& operator<<(FArchive& Ar, FContainerHeader& ContainerHeader)
 {
 	Ar << ContainerHeader.ContainerId;
@@ -471,81 +465,6 @@ struct FAsyncPackageDesc2
 #endif
 };
 
-class FNameMap
-{
-public:
-	void LoadGlobal(FIoDispatcher& IoDispatcher)
-	{
-		TRACE_CPUPROFILER_EVENT_SCOPE(LoadGlobalNameMap);
-
-		check(NameEntries.Num() == 0);
-
-		FIoChunkId NamesId = CreateIoChunkId(0, 0, EIoChunkType::LoaderGlobalNames);
-		FIoChunkId HashesId = CreateIoChunkId(0, 0, EIoChunkType::LoaderGlobalNameHashes);
-
-		FIoBatch Batch = IoDispatcher.NewBatch();
-		FIoRequest NameRequest = Batch.Read(NamesId, FIoReadOptions(), IoDispatcherPriority_High);
-		FIoRequest HashRequest = Batch.Read(HashesId, FIoReadOptions(), IoDispatcherPriority_High);
-		FEvent* BatchCompletedEvent = FPlatformProcess::GetSynchEventFromPool();
-		Batch.IssueAndTriggerEvent(BatchCompletedEvent);
-
-		ReserveNameBatch(	IoDispatcher.GetSizeForChunk(NamesId).ValueOrDie(),
-							IoDispatcher.GetSizeForChunk(HashesId).ValueOrDie());
-
-		BatchCompletedEvent->Wait();
-		FPlatformProcess::ReturnSynchEventToPool(BatchCompletedEvent);
-
-		FIoBuffer NameBuffer = NameRequest.GetResult().ConsumeValueOrDie();
-		FIoBuffer HashBuffer = HashRequest.GetResult().ConsumeValueOrDie();
-
-		Load(MakeArrayView(NameBuffer.Data(), NameBuffer.DataSize()), MakeArrayView(HashBuffer.Data(), HashBuffer.DataSize()), FMappedName::EType::Global);
-	}
-
-	int32 Num() const
-	{
-		return NameEntries.Num();
-	}
-
-	void Load(TArrayView<const uint8> NameBuffer, TArrayView<const uint8> HashBuffer, FMappedName::EType InNameMapType)
-	{
-		LoadNameBatch(NameEntries, NameBuffer, HashBuffer);
-		NameMapType = InNameMapType;
-	}
-
-	FName GetName(const FMappedName& MappedName) const
-	{
-		check(MappedName.GetType() == NameMapType);
-		check(MappedName.GetIndex() < uint32(NameEntries.Num()));
-		FNameEntryId NameEntry = NameEntries[MappedName.GetIndex()];
-		return FName::CreateFromDisplayId(NameEntry, MappedName.GetNumber());
-	}
-
-	bool TryGetName(const FMappedName& MappedName, FName& OutName) const
-	{
-		check(MappedName.GetType() == NameMapType);
-		uint32 Index = MappedName.GetIndex();
-		if (Index < uint32(NameEntries.Num()))
-		{
-			FNameEntryId NameEntry = NameEntries[MappedName.GetIndex()];
-			OutName = FName::CreateFromDisplayId(NameEntry, MappedName.GetNumber());
-			return true;
-		}
-		return false;
-	}
-
-	FMinimalName GetMinimalName(const FMappedName& MappedName) const
-	{
-		check(MappedName.GetType() == NameMapType);
-		check(MappedName.GetIndex() < uint32(NameEntries.Num()));
-		FNameEntryId NameEntry = NameEntries[MappedName.GetIndex()];
-		return FMinimalName(NameEntry, MappedName.GetNumber());
-	}
-
-private:
-	TArray<FNameEntryId> NameEntries;
-	FMappedName::EType NameMapType = FMappedName::EType::Global;
-};
-
 struct FPublicExport
 {
 	UObject* Object = nullptr;
@@ -561,6 +480,7 @@ using FUnreachablePublicExports = TArray<FUnreachablePublicExport>;
 
 struct FGlobalImportStore
 {
+	FNameMap GlobalNameMap;
 	TMap<FPackageObjectIndex, UObject*> ScriptObjects;
 	TMap<FPackageObjectIndex, FPublicExport> PublicExportObjects;
 	TMap<int32, FPackageObjectIndex> ObjectIndexToPublicExport;
@@ -574,6 +494,62 @@ struct FGlobalImportStore
 		ObjectIndexToPublicExport.Reserve(32768);
 	}
 
+	void Initialize(FIoDispatcher& IoDispatcher)
+	{
+		{
+			TRACE_CPUPROFILER_EVENT_SCOPE(LoadGlobalNameMap);
+
+			FIoChunkId NamesId = CreateIoChunkId(0, 0, EIoChunkType::LoaderGlobalNames);
+			FIoChunkId HashesId = CreateIoChunkId(0, 0, EIoChunkType::LoaderGlobalNameHashes);
+
+			FIoBatch Batch = IoDispatcher.NewBatch();
+			FIoRequest NameRequest = Batch.Read(NamesId, FIoReadOptions(), IoDispatcherPriority_High);
+			FIoRequest HashRequest = Batch.Read(HashesId, FIoReadOptions(), IoDispatcherPriority_High);
+			FEvent* BatchCompletedEvent = FPlatformProcess::GetSynchEventFromPool();
+			Batch.IssueAndTriggerEvent(BatchCompletedEvent);
+
+			ReserveNameBatch(	IoDispatcher.GetSizeForChunk(NamesId).ValueOrDie(),
+								IoDispatcher.GetSizeForChunk(HashesId).ValueOrDie());
+
+			BatchCompletedEvent->Wait();
+			FPlatformProcess::ReturnSynchEventToPool(BatchCompletedEvent);
+
+			FIoBuffer NameBuffer = NameRequest.GetResult().ConsumeValueOrDie();
+			FIoBuffer HashBuffer = HashRequest.GetResult().ConsumeValueOrDie();
+
+			GlobalNameMap.Load(MakeArrayView(NameBuffer.Data(), NameBuffer.DataSize()), MakeArrayView(HashBuffer.Data(), HashBuffer.DataSize()), FMappedName::EType::Global);
+		}
+
+		{
+			TRACE_CPUPROFILER_EVENT_SCOPE(SetupInitialLoadData);
+
+			FEvent* InitialLoadEvent = FPlatformProcess::GetSynchEventFromPool();
+
+			FIoBatch IoBatch = IoDispatcher.NewBatch();
+			FIoRequest IoRequest = IoBatch.Read(CreateIoChunkId(0, 0, EIoChunkType::LoaderInitialLoadMeta), FIoReadOptions(), IoDispatcherPriority_High);
+			IoBatch.IssueAndTriggerEvent(InitialLoadEvent);
+
+			InitialLoadEvent->Wait();
+			FPlatformProcess::ReturnSynchEventToPool(InitialLoadEvent);
+
+			FIoBuffer InitialLoadIoBuffer = IoRequest.GetResult().ConsumeValueOrDie();
+			FLargeMemoryReader InitialLoadArchive(InitialLoadIoBuffer.Data(), InitialLoadIoBuffer.DataSize());
+			int32 NumScriptObjects = 0;
+			InitialLoadArchive << NumScriptObjects;
+			ScriptObjectEntries = MakeArrayView(reinterpret_cast<const FScriptObjectEntry*>(InitialLoadIoBuffer.Data() + InitialLoadArchive.Tell()), NumScriptObjects);
+
+			ScriptObjectEntriesMap.Reserve(ScriptObjectEntries.Num());
+			for (FScriptObjectEntry& ScriptObjectEntry : ScriptObjectEntries)
+			{
+				const FMappedName& MappedName = FMappedName::FromMinimalName(ScriptObjectEntry.ObjectName);
+				check(MappedName.IsGlobal());
+				ScriptObjectEntry.ObjectName = GlobalNameMap.GetMinimalName(MappedName);
+
+				ScriptObjectEntriesMap.Add(ScriptObjectEntry.GlobalIndex, &ScriptObjectEntry);
+			}
+		}
+	}
+
 	TArray<FPackageId> RemovePublicExports(const FUnreachablePublicExports& PublicExports)
 	{
 		TArray<FPackageId> PackageIds;
@@ -786,6 +762,7 @@ private:
 	// Packages in active loading or completely loaded packages, with Desc.DiskPackageName as key.
 	// Does not track temp packages with custom UPackage names, since they are never imorted by other packages.
 	TMap<FPackageId, FLoadedPackageRef> Packages;
+	IPackageStore* PackageStore = nullptr;
 
 public:
 	FLoadedPackageStore()
@@ -793,6 +770,11 @@ public:
 		Packages.Reserve(32768);
 	}
 
+	void Initialize(IPackageStore& InPackageStore)
+	{
+		PackageStore = &InPackageStore;
+	}
+
 	int32 NumTracked() const
 	{
 		return Packages.Num();
@@ -830,303 +812,6 @@ public:
 		}
 	}
 #endif
-};
-
-class FPackageStore
-{
-public:
-	FPackageStore(FIoDispatcher& InIoDispatcher, FNameMap& InGlobalNameMap)
-		: IoDispatcher(InIoDispatcher)
-		, GlobalNameMap(InGlobalNameMap) { }
-	
-	struct FLoadedContainer
-	{
-		TUniquePtr<FNameMap> ContainerNameMap;
-		TArray<uint8> StoreEntries; //FPackageStoreEntry[PackageCount];
-		uint32 PackageCount;
-		bool bValid = false;
-	};
-
-	FIoDispatcher& IoDispatcher;
-	FNameMap& GlobalNameMap;
-	TMap<FIoContainerId, TUniquePtr<FLoadedContainer>> LoadedContainers;
-
-	TArray<FString> CurrentCultureNames;
-
-	FCriticalSection PackageNameMapsCritical;
-
-	TMap<FPackageId, FPackageStoreEntry*> StoreEntriesMap;
-	TMap<FPackageId, FPackageId> RedirectsPackageMap;
-	TSet<FPackageId> TargetRedirectIds;
-	int32 NextCustomPackageIndex = 0;
-
-	FGlobalImportStore ImportStore;
-	FLoadedPackageStore LoadedPackageStore;
-	int32 ScriptArcsCount = 0;
-
-public:
-	void Initialize()
-	{
-		TRACE_CPUPROFILER_EVENT_SCOPE(InitializePackageStore);
-
-		FPackageName::DoesPackageExistOverride().BindLambda([this](FName PackageName)
-		{
-			FPackageId PackageId = FPackageId::FromName(PackageName);
-			FScopeLock Lock(&PackageNameMapsCritical);
-			return StoreEntriesMap.Contains(PackageId);
-		});
-	}
-
-	void SetupCulture()
-	{
-		FInternationalization& Internationalization = FInternationalization::Get();
-		FString CurrentCulture = Internationalization.GetCurrentCulture()->GetName();
-		FParse::Value(FCommandLine::Get(), TEXT("CULTURE="), CurrentCulture);
-		CurrentCultureNames = Internationalization.GetPrioritizedCultureNames(CurrentCulture);
-	}
-
-	void SetupInitialLoadData()
-	{
-		TRACE_CPUPROFILER_EVENT_SCOPE(SetupInitialLoadData);
-
-		FEvent* InitialLoadEvent = FPlatformProcess::GetSynchEventFromPool();
-
-		FIoBatch IoBatch = IoDispatcher.NewBatch();
-		FIoRequest IoRequest = IoBatch.Read(CreateIoChunkId(0, 0, EIoChunkType::LoaderInitialLoadMeta), FIoReadOptions(), IoDispatcherPriority_High);
-		IoBatch.IssueAndTriggerEvent(InitialLoadEvent);
-
-		InitialLoadEvent->Wait();
-		FPlatformProcess::ReturnSynchEventToPool(InitialLoadEvent);
-
-		FIoBuffer InitialLoadIoBuffer = IoRequest.GetResult().ConsumeValueOrDie();
-		FLargeMemoryReader InitialLoadArchive(InitialLoadIoBuffer.Data(), InitialLoadIoBuffer.DataSize());
-		int32 NumScriptObjects = 0;
-		InitialLoadArchive << NumScriptObjects;
-		ImportStore.ScriptObjectEntries = MakeArrayView(reinterpret_cast<const FScriptObjectEntry*>(InitialLoadIoBuffer.Data() + InitialLoadArchive.Tell()), NumScriptObjects);
-
-		ImportStore.ScriptObjectEntriesMap.Reserve(ImportStore.ScriptObjectEntries.Num());
-		for (FScriptObjectEntry& ScriptObjectEntry : ImportStore.ScriptObjectEntries)
-		{
-			const FMappedName& MappedName = FMappedName::FromMinimalName(ScriptObjectEntry.ObjectName);
-			check(MappedName.IsGlobal());
-			ScriptObjectEntry.ObjectName = GlobalNameMap.GetMinimalName(MappedName);
-
-			ImportStore.ScriptObjectEntriesMap.Add(ScriptObjectEntry.GlobalIndex, &ScriptObjectEntry);
-		}
-	}
-
-	void LoadContainers(TArrayView<const FIoContainerId> Containers)
-	{
-		TRACE_CPUPROFILER_EVENT_SCOPE(LoadContainers);
-
-		int32 ContainersToLoad = 0;
-
-		for (const FIoContainerId& ContainerId : Containers)
-		{
-			if (ContainerId.IsValid())
-			{
-				++ContainersToLoad;
-			}
-		}
-
-		if (!ContainersToLoad)
-		{
-			return;
-		}
-
-		TAtomic<int32> Remaining(ContainersToLoad);
-
-		FEvent* Event = FPlatformProcess::GetSynchEventFromPool();
-		FIoBatch IoBatch = IoDispatcher.NewBatch();
-
-		for (const FIoContainerId& ContainerId : Containers)
-		{
-			if (!ContainerId.IsValid())
-			{
-				continue;
-			}
-
-			TUniquePtr<FLoadedContainer>& LoadedContainerPtr = LoadedContainers.FindOrAdd(ContainerId);
-			if (!LoadedContainerPtr)
-			{
-				LoadedContainerPtr.Reset(new FLoadedContainer);
-			}
-			FLoadedContainer& LoadedContainer = *LoadedContainerPtr;
-
-			UE_LOG(LogStreaming, Log, TEXT("Loading mounted container ID '0x%llX'"), ContainerId.Value());
-			LoadedContainer.bValid = true;
-
-			FIoChunkId HeaderChunkId = CreateIoChunkId(ContainerId.Value(), 0, EIoChunkType::ContainerHeader);
-			IoBatch.ReadWithCallback(HeaderChunkId, FIoReadOptions(), IoDispatcherPriority_High, [this, &Remaining, Event, &LoadedContainer, ContainerId](TIoStatusOr<FIoBuffer> Result)
-			{
-				// Execution method Thread will run the async block synchronously when multithreading is NOT supported
-				const EAsyncExecution ExecutionMethod = FPlatformProcess::SupportsMultithreading() ? EAsyncExecution::TaskGraph : EAsyncExecution::Thread;
-
-				if (!Result.IsOk())
-				{
-					if (EIoErrorCode::NotFound == Result.Status().GetErrorCode())
-					{
-						UE_LOG(LogStreaming, Warning, TEXT("Header for container '0x%llX' not found."), ContainerId.Value());
-					}
-					else
-					{
-						UE_LOG(LogStreaming, Fatal, TEXT("Failed reading header for container '0x%llX' (%s)"), ContainerId.Value(), *Result.Status().ToString());
-					}
-
-					if (--Remaining == 0)
-					{
-						Event->Trigger();
-					}
-					return;
-				}
-
-				Async(ExecutionMethod, [this, &Remaining, Event, IoBuffer = Result.ConsumeValueOrDie(), &LoadedContainer]()
-				{
-					LLM_SCOPE(ELLMTag::AsyncLoading);
-
-					FMemoryReaderView Ar(MakeArrayView(IoBuffer.Data(), IoBuffer.DataSize()));
-
-					FContainerHeader ContainerHeader;
-					Ar << ContainerHeader;
-
-					const bool bHasContainerLocalNameMap = ContainerHeader.Names.Num() > 0;
-					if (bHasContainerLocalNameMap)
-					{
-						TRACE_CPUPROFILER_EVENT_SCOPE(LoadContainerNameMap);
-						LoadedContainer.ContainerNameMap.Reset(new FNameMap());
-						LoadedContainer.ContainerNameMap->Load(ContainerHeader.Names, ContainerHeader.NameHashes, FMappedName::EType::Container);
-					}
-
-					LoadedContainer.PackageCount = ContainerHeader.PackageCount;
-					LoadedContainer.StoreEntries = MoveTemp(ContainerHeader.StoreEntries);
-					{
-						TRACE_CPUPROFILER_EVENT_SCOPE(AddPackages);
-						FScopeLock Lock(&PackageNameMapsCritical);
-
-						TArrayView<FPackageStoreEntry> StoreEntries(reinterpret_cast<FPackageStoreEntry*>(LoadedContainer.StoreEntries.GetData()), LoadedContainer.PackageCount);
-
-						int32 Index = 0;
-						StoreEntriesMap.Reserve(StoreEntriesMap.Num() + LoadedContainer.PackageCount);
-						for (FPackageStoreEntry& ContainerEntry : StoreEntries)
-						{
-							const FPackageId& PackageId = ContainerHeader.PackageIds[Index];
-
-							FPackageStoreEntry*& GlobalEntry = StoreEntriesMap.FindOrAdd(PackageId);
-							if (!GlobalEntry)
-							{
-								GlobalEntry = &ContainerEntry;
-							}
-							++Index;
-						}
-
-						{
-							TRACE_CPUPROFILER_EVENT_SCOPE(LoadPackageStoreLocalization);
-							const FSourceToLocalizedPackageIdMap* LocalizedPackages = nullptr;
-							for (const FString& CultureName : CurrentCultureNames)
-							{
-								LocalizedPackages = ContainerHeader.CulturePackageMap.Find(CultureName);
-								if (LocalizedPackages)
-								{
-									break;
-								}
-							}
-
-							if (LocalizedPackages)
-							{
-								for (auto& Pair : *LocalizedPackages)
-								{
-									const FPackageId& SourceId = Pair.Key;
-									const FPackageId& LocalizedId = Pair.Value;
-									RedirectsPackageMap.Emplace(SourceId, LocalizedId);
-									TargetRedirectIds.Add(LocalizedId);
-								}
-							}
-						}
-
-						{
-							TRACE_CPUPROFILER_EVENT_SCOPE(LoadPackageStoreRedirects);
-							for (const TPair<FPackageId, FPackageId>& Redirect : ContainerHeader.PackageRedirects)
-							{
-								const FPackageId& SourceId = Redirect.Key;
-								const FPackageId& RedirectedId = Redirect.Value;
-								RedirectsPackageMap.Emplace(SourceId, RedirectedId);
-								TargetRedirectIds.Add(RedirectedId);
-							}
-						}
-					}
-
-					if (--Remaining == 0)
-					{
-						Event->Trigger();
-					}
-				});
-			});
-		}
-
-		IoBatch.Issue();
-		Event->Wait();
-		FPlatformProcess::ReturnSynchEventToPool(Event);
-
-		ApplyRedirects(RedirectsPackageMap);
-	}
-
-	void OnContainerMounted(const FIoContainerId& ContainerId)
-	{
-		LLM_SCOPE(ELLMTag::AsyncLoading);
-		LoadContainers(MakeArrayView(&ContainerId, 1));
-	}
-
-	void ApplyRedirects(const TMap<FPackageId, FPackageId>& Redirects)
-	{
-		TRACE_CPUPROFILER_EVENT_SCOPE(ApplyRedirects);
-
-		FScopeLock Lock(&PackageNameMapsCritical);
-
-		if (Redirects.Num() == 0)
-		{
-			return;
-		}
-
-		for (auto It = Redirects.CreateConstIterator(); It; ++It)
-		{
-			const FPackageId& SourceId = It.Key();
-			const FPackageId& RedirectId = It.Value();
-			check(RedirectId.IsValid());
-			FPackageStoreEntry* RedirectEntry = StoreEntriesMap.FindRef(RedirectId);
-			check(RedirectEntry);
-			FPackageStoreEntry*& PackageEntry = StoreEntriesMap.FindOrAdd(SourceId);
-			if (RedirectEntry)
-			{
-				PackageEntry = RedirectEntry;
-			}
-		}
-
-		for (auto It = StoreEntriesMap.CreateIterator(); It; ++It)
-		{
-			FPackageStoreEntry* StoreEntry = It.Value();
-
-			for (FPackageId& ImportedPackageId : StoreEntry->ImportedPackages)
-			{
-				if (const FPackageId* RedirectId = Redirects.Find(ImportedPackageId))
-				{
-					ImportedPackageId = *RedirectId;
-				}
-			}
-		}
-	}
-
-	void FinalizeInitialLoad()
-	{
-		ImportStore.FindAllScriptObjects();
-
-		UE_LOG(LogStreaming, Display, TEXT("AsyncLoading2 - InitialLoad Finalized: %d script object entries in %.2f KB"),
-			ImportStore.ScriptObjects.Num(), (float)ImportStore.ScriptObjects.GetAllocatedSize() / 1024.f);
-	}
-
-	inline FGlobalImportStore& GetGlobalImportStore()
-	{
-		return ImportStore;
-	}
 
 	void RemovePackage(FPackageId PackageId, UPackage* Package)
 	{
@@ -1138,13 +823,13 @@ public:
 			return;
 		}
 
-		int32 RefCount = LoadedPackageStore.RemovePackage(PackageId);
+		int32 RefCount = RemovePackage(PackageId);
 		if (RefCount < 0) // not found
 		{
-			FPackageId* RedirectedId = RedirectsPackageMap.Find(PackageId);
-			if (RedirectedId)
+			FPackageId RedirectedId = PackageStore->GetRedirectedPackageId(PackageId);
+			if (RedirectedId.IsValid())
 			{
-				RefCount = LoadedPackageStore.RemovePackage(*RedirectedId);
+				RefCount = RemovePackage(RedirectedId);
 			}
 		}
 		if (RefCount > 0)
@@ -1169,68 +854,51 @@ public:
 		}
 	}
 
-	void RemovePackages(const FUnreachablePackages& Packages)
+	void RemovePackages(const FUnreachablePackages& UnreachablePackages)
 	{
-		int32 PackageCount = Packages.Num();
+		const int32 PackageCount = UnreachablePackages.Num();
 		TArray<FPackageId> PackageIds;
 		PackageIds.AddUninitialized(PackageCount);
 		bool bForceSingleThreaded = PackageCount < 64;
 		{
 			// TRACE_CPUPROFILER_EVENT_SCOPE(FPackageId::FromName);
-			ParallelFor(PackageCount, [&Packages, &PackageIds](int32 Index)
+			ParallelFor(PackageCount, [&UnreachablePackages, &PackageIds](int32 Index)
 			{
-				PackageIds[Index] = FPackageId::FromName(Packages[Index].Key);
+				PackageIds[Index] = FPackageId::FromName(UnreachablePackages[Index].Key);
 			}, bForceSingleThreaded);
 		}
 		for (int32 Index = 0; Index < PackageCount; ++Index)
 		{
-			RemovePackage(PackageIds[Index], Packages[Index].Value);
+			RemovePackage(PackageIds[Index], UnreachablePackages[Index].Value);
 		}
 	}
 
 	void ClearAllPublicExportsLoaded(const TArray<FPackageId>& PackageIds)
 	{
-		int32 PackageCount = PackageIds.Num();
-		bool bForceSingleThreaded = PackageCount < 1024;
+		const int32 PackageCount = PackageIds.Num();
+		const bool bForceSingleThreaded = PackageCount < 1024;
 		ParallelFor(PackageCount, [this, &PackageIds](int32 Index)
 		{
-			if (FLoadedPackageRef* PackageRef = LoadedPackageStore.FindPackageRef(PackageIds[Index]))
+			if (FLoadedPackageRef* PackageRef = FindPackageRef(PackageIds[Index]))
 			{
 				PackageRef->ClearAllPublicExportsLoaded();
 			}
 		}, bForceSingleThreaded);
 	}
-
-	inline const FPackageStoreEntry* FindStoreEntry(FPackageId PackageId)
-	{
-		FScopeLock Lock(&PackageNameMapsCritical);
-		FPackageStoreEntry* Entry = StoreEntriesMap.FindRef(PackageId);
-		return Entry;
-	}
-
-	inline FPackageId GetRedirectedPackageId(FPackageId PackageId)
-	{
-		FScopeLock Lock(&PackageNameMapsCritical);
-		FPackageId RedirectedId = RedirectsPackageMap.FindRef(PackageId);
-		return RedirectedId;
-	}
-
-	bool IsRedirect(FPackageId PackageId) const
-	{
-		return TargetRedirectIds.Contains(PackageId);
-	}
 };
 
 struct FPackageImportStore
 {
-	FPackageStore& GlobalPackageStore;
+	IPackageStore& PackageStore;
 	FGlobalImportStore& GlobalImportStore;
+	FLoadedPackageStore& LoadedPackageStore;
 	const FAsyncPackageDesc2& Desc;
 	TArrayView<const FPackageObjectIndex> ImportMap;
 
-	FPackageImportStore(FPackageStore& InGlobalPackageStore, const FAsyncPackageDesc2& InDesc)
-		: GlobalPackageStore(InGlobalPackageStore)
-		, GlobalImportStore(GlobalPackageStore.ImportStore)
+	FPackageImportStore(IPackageStore& InPackageStore, FGlobalImportStore& InGlobalImportStore, FLoadedPackageStore& InLoadedPackageStore, const FAsyncPackageDesc2& InDesc)
+		: PackageStore(InPackageStore)
+		, GlobalImportStore(InGlobalImportStore)
+		, LoadedPackageStore(InLoadedPackageStore)
 		, Desc(InDesc)
 	{
 		TRACE_CPUPROFILER_EVENT_SCOPE(NewPackageImportStore);
@@ -1356,7 +1024,7 @@ private:
 	{
 		for (const FPackageId& ImportedPackageId : Desc.StoreEntry->ImportedPackages)
 		{
-			FLoadedPackageRef& PackageRef = GlobalPackageStore.LoadedPackageStore.GetPackageRef(ImportedPackageId);
+			FLoadedPackageRef& PackageRef = LoadedPackageStore.GetPackageRef(ImportedPackageId);
 			if (PackageRef.AddRef())
 			{
 				AddAsyncFlags(PackageRef.GetPackage());
@@ -1364,7 +1032,7 @@ private:
 		}
 		if (Desc.CanBeImported())
 		{
-			FLoadedPackageRef& PackageRef = GlobalPackageStore.LoadedPackageStore.GetPackageRef(Desc.DiskPackageId);
+			FLoadedPackageRef& PackageRef = LoadedPackageStore.GetPackageRef(Desc.DiskPackageId);
 			PackageRef.ClearErrorFlags();
 			if (PackageRef.AddRef())
 			{
@@ -1377,7 +1045,7 @@ private:
 	{
 		for (const FPackageId& ImportedPackageId : Desc.StoreEntry->ImportedPackages)
 		{
-			FLoadedPackageRef& PackageRef = GlobalPackageStore.LoadedPackageStore.GetPackageRef(ImportedPackageId);
+			FLoadedPackageRef& PackageRef = LoadedPackageStore.GetPackageRef(ImportedPackageId);
 			if (PackageRef.ReleaseRef(Desc.DiskPackageId, ImportedPackageId))
 			{
 				ClearAsyncFlags(PackageRef.GetPackage());
@@ -1386,7 +1054,7 @@ private:
 		if (Desc.CanBeImported())
 		{
 			// clear own reference, and possible all async flags if no remaining ref count
-			FLoadedPackageRef& PackageRef =	GlobalPackageStore.LoadedPackageStore.GetPackageRef(Desc.DiskPackageId);
+			FLoadedPackageRef& PackageRef =	LoadedPackageStore.GetPackageRef(Desc.DiskPackageId);
 			if (PackageRef.ReleaseRef(Desc.DiskPackageId, Desc.DiskPackageId))
 			{
 				ClearAsyncFlags(PackageRef.GetPackage());
@@ -2250,8 +1918,9 @@ private:
 	/** I/O Dispatcher */
 	FIoDispatcher& IoDispatcher;
 
-	FNameMap GlobalNameMap;
-	FPackageStore GlobalPackageStore;
+	TUniquePtr<IPackageStore> PackageStore;
+	FLoadedPackageStore LoadedPackageStore;
+	FGlobalImportStore GlobalImportStore;
 
 	/** Initial load pending CDOs */
 	TMap<UClass*, TArray<FEventLoadNode2*>> PendingCDOs;
@@ -2525,11 +2194,17 @@ public:
 		}
 	}
 
+	IPackageStore& GetPackageStore()
+	{
+		return *PackageStore.Get();
+	}
+
 private:
 
 	void SuspendWorkers();
 	void ResumeWorkers();
 
+	void LoadGlobalNameMap();
 	void LazyInitializeFromLoadPackage();
 	void FinalizeInitialLoad();
 
@@ -2704,7 +2379,13 @@ void FAsyncLoadingThread2::InitializeLoading()
 	FBulkDataBase::SetIoDispatcher(&IoDispatcher);
 #endif
 
-	GlobalPackageStore.Initialize();
+	//TODO: Move construction of package store from the loader
+	PackageStore = MakeFilePackageStore(IoDispatcher);
+
+	FPackageName::DoesPackageExistOverride().BindLambda([this](FName PackageName)
+	{
+		return PackageStore->DoesPackageExist(FPackageId::FromName(PackageName));
+	});
 
 	AsyncThreadReady.Increment();
 
@@ -3330,6 +3011,9 @@ void FGlobalImportStore::FindAllScriptObjects()
 	ScriptObjectEntriesMap.Empty();
 	ScriptObjectEntries.Empty();
 	ScriptObjects.Shrink();
+
+	UE_LOG(LogStreaming, Display, TEXT("AsyncLoading2 - InitialLoad Finalized: %d script object entries in %.2f KB"),
+		ScriptObjects.Num(), (float)ScriptObjects.GetAllocatedSize() / 1024.f);
 }
 
 void FAsyncPackage2::ImportPackagesRecursive()
@@ -3353,16 +3037,16 @@ void FAsyncPackage2::ImportPackagesRecursive()
 
 	int32 ImportedPackageIndex = 0;
 
-	FPackageStore& GlobalPackageStore = AsyncLoadingThread.GlobalPackageStore;
+	IPackageStore& PackageStore = AsyncLoadingThread.GetPackageStore();
 	for (const FPackageId& ImportedPackageId : Desc.StoreEntry->ImportedPackages)
 	{
-		FLoadedPackageRef& PackageRef = GlobalPackageStore.LoadedPackageStore.GetPackageRef(ImportedPackageId);
+		FLoadedPackageRef& PackageRef = ImportStore.LoadedPackageStore.GetPackageRef(ImportedPackageId);
 		if (PackageRef.AreAllPublicExportsLoaded())
 		{
 			continue;
 		}
 
-		const FPackageStoreEntry* ImportedPackageEntry = GlobalPackageStore.FindStoreEntry(ImportedPackageId);
+		const FPackageStoreEntry* ImportedPackageEntry = PackageStore.GetPackageEntry(ImportedPackageId);
 
 		if (!ImportedPackageEntry)
 		{
@@ -3431,7 +3115,7 @@ EAsyncPackageState::Type FAsyncPackage2::Event_ProcessPackageSummary(FAsyncLoadi
 	{
 		if (Package->Desc.CanBeImported())
 		{
-			FLoadedPackageRef* PackageRef = Package->ImportStore.GlobalPackageStore.LoadedPackageStore.FindPackageRef(Package->Desc.DiskPackageId);
+			FLoadedPackageRef* PackageRef = Package->ImportStore.LoadedPackageStore.FindPackageRef(Package->Desc.DiskPackageId);
 			check(PackageRef);
 			PackageRef->SetHasFailed();
 		}
@@ -3994,7 +3678,7 @@ EAsyncPackageState::Type FAsyncPackage2::Event_ExportsDone(FAsyncLoadingThreadSt
 	if (!Package->bLoadHasFailed && Package->Desc.CanBeImported())
 	{
 		FLoadedPackageRef& PackageRef =
-			Package->AsyncLoadingThread.GlobalPackageStore.LoadedPackageStore.GetPackageRef((Package->Desc.DiskPackageId));
+			Package->AsyncLoadingThread.LoadedPackageStore.GetPackageRef(Package->Desc.DiskPackageId);
 		PackageRef.SetAllPublicExportsLoaded();
 	}
 
@@ -4675,7 +4359,6 @@ EAsyncPackageState::Type FAsyncLoadingThread2::TickAsyncLoadingFromGameThread(FA
 FAsyncLoadingThread2::FAsyncLoadingThread2(FIoDispatcher& InIoDispatcher)
 	: Thread(nullptr)
 	, IoDispatcher(InIoDispatcher)
-	, GlobalPackageStore(InIoDispatcher, GlobalNameMap)
 {
 	GEventDrivenLoaderEnabled = true;
 
@@ -4802,21 +4485,19 @@ void FAsyncLoadingThread2::ResumeWorkers()
 void FAsyncLoadingThread2::LazyInitializeFromLoadPackage()
 {
 	TRACE_CPUPROFILER_EVENT_SCOPE(LazyInitializeFromLoadPackage);
-	GlobalNameMap.LoadGlobal(IoDispatcher);
 	if (GIsInitialLoad)
 	{
-		GlobalPackageStore.SetupInitialLoadData();
+		GlobalImportStore.Initialize(IoDispatcher);
 	}
-	GlobalPackageStore.SetupCulture();
-	GlobalPackageStore.LoadContainers(IoDispatcher.GetMountedContainers().Array());
-	IoDispatcher.OnContainerMounted().AddRaw(&GlobalPackageStore, &FPackageStore::OnContainerMounted);
-}
 
+	PackageStore->Initialize();
+	LoadedPackageStore.Initialize(*PackageStore.Get());
+}
 
 void FAsyncLoadingThread2::FinalizeInitialLoad()
 {
 	TRACE_CPUPROFILER_EVENT_SCOPE(FinalizeInitialLoad);
-	GlobalPackageStore.FinalizeInitialLoad();
+	GlobalImportStore.FindAllScriptObjects();
 	check(PendingCDOs.Num() == 0);
 	PendingCDOs.Empty();
 }
@@ -5152,17 +4833,17 @@ void FAsyncLoadingThread2::RemoveUnreachableObjects(const FUnreachablePublicExpo
 	if (PublicExports.Num() > 0)
 	{
 		// TRACE_CPUPROFILER_EVENT_SCOPE(RemovePublicExports);
-		PublicExportPackages = GlobalPackageStore.ImportStore.RemovePublicExports(PublicExports);
+		PublicExportPackages = GlobalImportStore.RemovePublicExports(PublicExports);
 	}
 	if (Packages.Num() > 0)
 	{
 		// TRACE_CPUPROFILER_EVENT_SCOPE(RemovePackages);
-		GlobalPackageStore.RemovePackages(Packages);
+		LoadedPackageStore.RemovePackages(Packages);
 	}
 	if (PublicExportPackages.Num() > 0)
 	{
 		// TRACE_CPUPROFILER_EVENT_SCOPE(ClearAllPublicExportsLoaded);
-		GlobalPackageStore.ClearAllPublicExportsLoaded(PublicExportPackages);
+		LoadedPackageStore.ClearAllPublicExportsLoaded(PublicExportPackages);
 	}
 }
 
@@ -5185,14 +4866,14 @@ void FAsyncLoadingThread2::NotifyUnreachableObjects(const TArrayView<FUObjectIte
 	const int32 PublicExportCount = PublicExports.Num();
 	if (PackageCount > 0 || PublicExportCount > 0)
 	{
-		const int32 OldLoadedPackageCount = GlobalPackageStore.LoadedPackageStore.NumTracked();
-		const int32 OldPublicExportCount = GlobalPackageStore.GetGlobalImportStore().PublicExportObjects.Num();
+		const int32 OldLoadedPackageCount = LoadedPackageStore.NumTracked();
+		const int32 OldPublicExportCount = GlobalImportStore.PublicExportObjects.Num();
 
 		const double RemoveStartTime = FPlatformTime::Seconds();
 		RemoveUnreachableObjects(PublicExports, Packages);
 
-		const int32 NewLoadedPackageCount = GlobalPackageStore.LoadedPackageStore.NumTracked();
-		const int32 NewPublicExportCount = GlobalPackageStore.GetGlobalImportStore().PublicExportObjects.Num();
+		const int32 NewLoadedPackageCount = LoadedPackageStore.NumTracked();
+		const int32 NewPublicExportCount = GlobalImportStore.PublicExportObjects.Num();
 		const int32 RemovedLoadedPackageCount = OldLoadedPackageCount - NewLoadedPackageCount;
 		const int32 RemovedPublicExportCount = OldPublicExportCount - NewPublicExportCount;
 
@@ -5219,7 +4900,7 @@ void FAsyncLoadingThread2::NotifyUnreachableObjects(const TArrayView<FUObjectIte
 #if ALT2_VERIFY_ASYNC_FLAGS
 	if (!IsAsyncLoadingPackages())
 	{
-		GlobalPackageStore.LoadedPackageStore.VerifyLoadedPackages();
+		LoadedPackageStore.VerifyLoadedPackages();
 		VerifyLoadFlagsWhenFinishedLoading();
 	}
 #endif
@@ -5268,7 +4949,7 @@ FAsyncPackage2::FAsyncPackage2(
 , Data(InData)
 , AsyncLoadingThread(InAsyncLoadingThread)
 , GraphAllocator(InGraphAllocator)
-, ImportStore(AsyncLoadingThread.GlobalPackageStore, Desc)
+, ImportStore(AsyncLoadingThread.GetPackageStore(), AsyncLoadingThread.GlobalImportStore, AsyncLoadingThread.LoadedPackageStore, Desc)
 {
 	TRACE_LOADTIME_NEW_ASYNC_PACKAGE(this, Desc.DiskPackageName);
 	AddRequestID(Desc.RequestID);
@@ -5474,7 +5155,7 @@ void FAsyncPackage2::CreateUPackage(const FPackageSummary* PackageSummary)
 		TRACE_CPUPROFILER_EVENT_SCOPE(UPackageFind);
 		if (Desc.CanBeImported())
 		{
-			PackageRef = ImportStore.GlobalPackageStore.LoadedPackageStore.FindPackageRef(Desc.DiskPackageId);
+			PackageRef = ImportStore.LoadedPackageStore.FindPackageRef(Desc.DiskPackageId);
 			UE_ASYNC_PACKAGE_CLOG(!PackageRef, Fatal, Desc, TEXT("CreateUPackage"), TEXT("Package has been destroyed by GC."));
 			LinkerRoot = PackageRef->GetPackage();
 #if DO_CHECK
@@ -5671,7 +5352,7 @@ int32 FAsyncLoadingThread2::LoadPackage(const FPackagePath& InPackagePath, FName
 
 	// Verify PackageToLoadName, or fixup to handle any input string that can be converted to a long package name.
 	FPackageId DiskPackageId = FPackageId::FromName(DiskPackageName);
-	const FPackageStoreEntry* StoreEntry = GlobalPackageStore.FindStoreEntry(DiskPackageId);
+	const FPackageStoreEntry* StoreEntry = PackageStore->GetPackageEntry(DiskPackageId);
 	if (!StoreEntry)
 	{
 		FString PackageNameStr = DiskPackageName.ToString();
@@ -5682,7 +5363,7 @@ int32 FAsyncLoadingThread2::LoadPackage(const FPackagePath& InPackagePath, FName
 			{
 				DiskPackageName = *NewPackageNameStr;
 				DiskPackageId = FPackageId::FromName(DiskPackageName);
-				StoreEntry = GlobalPackageStore.FindStoreEntry(DiskPackageId);
+				StoreEntry = PackageStore->GetPackageEntry(DiskPackageId);
 				bHasCustomPackageName &= Name != DiskPackageName;
 			}
 		}
@@ -5696,7 +5377,7 @@ int32 FAsyncLoadingThread2::LoadPackage(const FPackagePath& InPackagePath, FName
 	if (bHasCustomPackageName)
 	{
 		FPackageId PackageId = FPackageId::FromName(Name);
-		if (!GlobalPackageStore.FindStoreEntry(PackageId))
+		if (!PackageStore->GetPackageEntry(PackageId))
 		{
 			FString PackageNameStr = Name.ToString();
 			if (FPackageName::IsValidLongPackageName(PackageNameStr))
@@ -5710,7 +5391,7 @@ int32 FAsyncLoadingThread2::LoadPackage(const FPackagePath& InPackagePath, FName
 				if (FPackageName::TryConvertFilenameToLongPackageName(PackageNameStr, NewPackageNameStr))
 				{
 					PackageId = FPackageId::FromName(FName(*NewPackageNameStr));
-					if (!GlobalPackageStore.FindStoreEntry(PackageId))
+					if (!PackageStore->GetPackageEntry(PackageId))
 					{
 						CustomPackageName = *NewPackageNameStr;
 						CustomPackageId = PackageId;
@@ -5721,7 +5402,7 @@ int32 FAsyncLoadingThread2::LoadPackage(const FPackagePath& InPackagePath, FName
 	}
 	// When explicitly requesting a redirected package then set CustomName to
 	// the redirected name, otherwise the UPackage name will be set to the base game name.
-	else if (GlobalPackageStore.IsRedirect(DiskPackageId))
+	else if (PackageStore->IsRedirect(DiskPackageId))
 	{
 		bHasCustomPackageName = true;
 		CustomPackageName = DiskPackageName;
@@ -5783,7 +5464,7 @@ int32 FAsyncLoadingThread2::LoadPackage(const FPackagePath& InPackagePath, FName
 		FAsyncPackageDesc2 PackageDesc(RequestID, InPackagePriority, DiskPackageId, StoreEntry, DiskPackageName, CustomPackageId, CustomPackageName, MoveTemp(CompletionDelegatePtr));
 
 		// Fixup for redirected packages since the slim StoreEntry itself has been stripped from both package names and package ids
-		FPackageId RedirectedDiskPackageId = GlobalPackageStore.GetRedirectedPackageId(DiskPackageId);
+		FPackageId RedirectedDiskPackageId = PackageStore->GetRedirectedPackageId(DiskPackageId);
 		if (RedirectedDiskPackageId.IsValid())
 		{
 			PackageDesc.DiskPackageId = RedirectedDiskPackageId;
diff --git a/Engine/Source/Runtime/CoreUObject/Private/Serialization/FilePackageStore.cpp b/Engine/Source/Runtime/CoreUObject/Private/Serialization/FilePackageStore.cpp
new file mode 100644
--- /dev/null
+++ b/Engine/Source/Runtime/CoreUObject/Private/Serialization/FilePackageStore.cpp
@@ -0,0 +1,290 @@
+// Copyright Epic Games, Inc. All Rights Reserved.
+
+#include "Serialization/PackageStore.h"
+#include "Serialization/MappedName.h"
+#include "Serialization/AsyncLoading2.h"
+#include "Serialization/LargeMemoryReader.h"
+#include "Serialization/MemoryReader.h"
+#include "IO/IoDispatcher.h"
+#include "IO/IoContainerId.h"
+#include "Internationalization/Culture.h"
+#include "Internationalization/Internationalization.h"
+#include "Async/Async.h"
+#include "Misc/CommandLine.h"
+
+/*
+ * File/container based package store.
+ */
+class FFilePackageStore final
+	: public IPackageStore
+{
+public:
+	FFilePackageStore(FIoDispatcher& InIoDispatcher)
+		: IoDispatcher(InIoDispatcher) { }
+
+	virtual ~FFilePackageStore() { }
+
+	virtual void Initialize() override
+	{
+		// Setup culture
+		{
+			FInternationalization& Internationalization = FInternationalization::Get();
+			FString CurrentCulture = Internationalization.GetCurrentCulture()->GetName();
+			FParse::Value(FCommandLine::Get(), TEXT("CULTURE="), CurrentCulture);
+			CurrentCultureNames = Internationalization.GetPrioritizedCultureNames(CurrentCulture);
+		}
+
+		LoadContainers(IoDispatcher.GetMountedContainers().Array());
+		IoDispatcher.OnContainerMounted().AddRaw(this, &FFilePackageStore::OnContainerMounted);
+	}
+
+	virtual bool DoesPackageExist(FPackageId PackageId) override
+	{
+		return nullptr != GetPackageEntry(PackageId);
+	}
+
+	virtual const FPackageStoreEntry* GetPackageEntry(FPackageId PackageId) override
+	{
+		FScopeLock Lock(&PackageNameMapsCritical);
+		FPackageStoreEntry* Entry = StoreEntriesMap.FindRef(PackageId);
+		return Entry;
+	}
+
+	virtual FPackageId GetRedirectedPackageId(FPackageId PackageId) override
+	{
+		FScopeLock Lock(&PackageNameMapsCritical);
+		FPackageId RedirectedId = RedirectsPackageMap.FindRef(PackageId);
+		return RedirectedId;
+	}
+
+	virtual bool IsRedirect(FPackageId PackageId) override
+	{
+		return TargetRedirectIds.Contains(PackageId);
+	}
+
+private:
+	void LoadContainers(TArrayView<const FIoContainerId> Containers)
+	{
+		TRACE_CPUPROFILER_EVENT_SCOPE(LoadContainers);
+
+		int32 ContainersToLoad = 0;
+
+		for (const FIoContainerId& ContainerId : Containers)
+		{
+			if (ContainerId.IsValid())
+			{
+				++ContainersToLoad;
+			}
+		}
+
+		if (!ContainersToLoad)
+		{
+			return;
+		}
+
+		TAtomic<int32> Remaining(ContainersToLoad);
+
+		FEvent* Event = FPlatformProcess::GetSynchEventFromPool();
+		FIoBatch IoBatch = IoDispatcher.NewBatch();
+
+		for (const FIoContainerId& ContainerId : Containers)
+		{
+			if (!ContainerId.IsValid())
+			{
+				continue;
+			}
+
+			TUniquePtr<FLoadedContainer>& LoadedContainerPtr = LoadedContainers.FindOrAdd(ContainerId);
+			if (!LoadedContainerPtr)
+			{
+				LoadedContainerPtr.Reset(new FLoadedContainer);
+			}
+			FLoadedContainer& LoadedContainer = *LoadedContainerPtr;
+
+			UE_LOG(LogStreaming, Log, TEXT("Loading mounted container ID '0x%llX'"), ContainerId.Value());
+			LoadedContainer.bValid = true;
+
+			FIoChunkId HeaderChunkId = CreateIoChunkId(ContainerId.Value(), 0, EIoChunkType::ContainerHeader);
+			IoBatch.ReadWithCallback(HeaderChunkId, FIoReadOptions(), IoDispatcherPriority_High, [this, &Remaining, Event, &LoadedContainer, ContainerId](TIoStatusOr<FIoBuffer> Result)
+			{
+				// Execution method Thread will run the async block synchronously when multithreading is NOT supported
+				const EAsyncExecution ExecutionMethod = FPlatformProcess::SupportsMultithreading() ? EAsyncExecution::TaskGraph : EAsyncExecution::Thread;
+
+				if (!Result.IsOk())
+				{
+					if (EIoErrorCode::NotFound == Result.Status().GetErrorCode())
+					{
+						UE_LOG(LogStreaming, Warning, TEXT("Header for container '0x%llX' not found."), ContainerId.Value());
+					}
+					else
+					{
+						UE_LOG(LogStreaming, Fatal, TEXT("Failed reading header for container '0x%llX' (%s)"), ContainerId.Value(), *Result.Status().ToString());
+					}
+
+					if (--Remaining == 0)
+					{
+						Event->Trigger();
+					}
+					return;
+				}
+
+				Async(ExecutionMethod, [this, &Remaining, Event, IoBuffer = Result.ConsumeValueOrDie(), &LoadedContainer]()
+				{
+					LLM_SCOPE(ELLMTag::AsyncLoading);
+
+					FMemoryReaderView Ar(MakeArrayView(IoBuffer.Data(), IoBuffer.DataSize()));
+
+					FContainerHeader ContainerHeader;
+					Ar << ContainerHeader;
+
+					const bool bHasContainerLocalNameMap = ContainerHeader.Names.Num() > 0;
+					if (bHasContainerLocalNameMap)
+					{
+						TRACE_CPUPROFILER_EVENT_SCOPE(LoadContainerNameMap);
+						LoadedContainer.ContainerNameMap.Reset(new FNameMap());
+						LoadedContainer.ContainerNameMap->Load(ContainerHeader.Names, ContainerHeader.NameHashes, FMappedName::EType::Container);
+					}
+
+					LoadedContainer.PackageCount = ContainerHeader.PackageCount;
+					LoadedContainer.StoreEntries = MoveTemp(ContainerHeader.StoreEntries);
+					{
+						TRACE_CPUPROFILER_EVENT_SCOPE(AddPackages);
+						FScopeLock Lock(&PackageNameMapsCritical);
+
+						TArrayView<FPackageStoreEntry> StoreEntries(reinterpret_cast<FPackageStoreEntry*>(LoadedContainer.StoreEntries.GetData()), LoadedContainer.PackageCount);
+
+						int32 Index = 0;
+						StoreEntriesMap.Reserve(StoreEntriesMap.Num() + LoadedContainer.PackageCount);
+						for (FPackageStoreEntry& ContainerEntry : StoreEntries)
+						{
+							const FPackageId& PackageId = ContainerHeader.PackageIds[Index];
+
+							FPackageStoreEntry*& GlobalEntry = StoreEntriesMap.FindOrAdd(PackageId);
+							if (!GlobalEntry)
+							{
+								GlobalEntry = &ContainerEntry;
+							}
+							++Index;
+						}
+
+						{
+							TRACE_CPUPROFILER_EVENT_SCOPE(LoadPackageStoreLocalization);
+							const FSourceToLocalizedPackageIdMap* LocalizedPackages = nullptr;
+							for (const FString& CultureName : CurrentCultureNames)
+							{
+								LocalizedPackages = ContainerHeader.CulturePackageMap.Find(CultureName);
+								if (LocalizedPackages)
+								{
+									break;
+								}
+							}
+
+							if (LocalizedPackages)
+							{
+								for (auto& Pair : *LocalizedPackages)
+								{
+									const FPackageId& SourceId = Pair.Key;
+									const FPackageId& LocalizedId = Pair.Value;
+									RedirectsPackageMap.Emplace(SourceId, LocalizedId);
+									TargetRedirectIds.Add(LocalizedId);
+								}
+							}
+						}
+
+						{
+							TRACE_CPUPROFILER_EVENT_SCOPE(LoadPackageStoreRedirects);
+							for (const TPair<FPackageId, FPackageId>& Redirect : ContainerHeader.PackageRedirects)
+							{
+								const FPackageId& SourceId = Redirect.Key;
+								const FPackageId& RedirectedId = Redirect.Value;
+								RedirectsPackageMap.Emplace(SourceId, RedirectedId);
+								TargetRedirectIds.Add(RedirectedId);
+							}
+						}
+					}
+
+					if (--Remaining == 0)
+					{
+						Event->Trigger();
+					}
+				});
+			});
+		}
+
+		IoBatch.Issue();
+		Event->Wait();
+		FPlatformProcess::ReturnSynchEventToPool(Event);
+
+		ApplyRedirects(RedirectsPackageMap);
+	}
+
+	void OnContainerMounted(const FIoContainerId& ContainerId)
+	{
+		LLM_SCOPE(ELLMTag::AsyncLoading);
+		LoadContainers(MakeArrayView(&ContainerId, 1));
+	}
+
+	void ApplyRedirects(const TMap<FPackageId, FPackageId>& Redirects)
+	{
+		TRACE_CPUPROFILER_EVENT_SCOPE(ApplyRedirects);
+
+		FScopeLock Lock(&PackageNameMapsCritical);
+
+		if (Redirects.Num() == 0)
+		{
+			return;
+		}
+
+		for (auto It = Redirects.CreateConstIterator(); It; ++It)
+		{
+			const FPackageId& SourceId = It.Key();
+			const FPackageId& RedirectId = It.Value();
+			check(RedirectId.IsValid());
+			FPackageStoreEntry* RedirectEntry = StoreEntriesMap.FindRef(RedirectId);
+			check(RedirectEntry);
+			FPackageStoreEntry*& PackageEntry = StoreEntriesMap.FindOrAdd(SourceId);
+			if (RedirectEntry)
+			{
+				PackageEntry = RedirectEntry;
+			}
+		}
+
+		for (auto It = StoreEntriesMap.CreateIterator(); It; ++It)
+		{
+			FPackageStoreEntry* StoreEntry = It.Value();
+
+			for (FPackageId& ImportedPackageId : StoreEntry->ImportedPackages)
+			{
+				if (const FPackageId* RedirectId = Redirects.Find(ImportedPackageId))
+				{
+					ImportedPackageId = *RedirectId;
+				}
+			}
+		}
+	}
+
+	struct FLoadedContainer
+	{
+		TUniquePtr<FNameMap> ContainerNameMap;
+		TArray<uint8> StoreEntries; //FPackageStoreEntry[PackageCount];
+		uint32 PackageCount = 0;
+		bool bValid = false;
+	};
+
+	FIoDispatcher& IoDispatcher;
+	TMap<FIoContainerId, TUniquePtr<FLoadedContainer>> LoadedContainers;
+
+	TArray<FString> CurrentCultureNames;
+
+	FCriticalSection PackageNameMapsCritical;
+
+	TMap<FPackageId, FPackageStoreEntry*> StoreEntriesMap;
+	TMap<FPackageId, FPackageId> RedirectsPackageMap;
+	TSet<FPackageId> TargetRedirectIds;
+	int32 NextCustomPackageIndex = 0;
+};
+
+TUniquePtr<IPackageStore> MakeFilePackageStore(FIoDispatcher& IoDispatcher)
+{
+	return MakeUnique<FFilePackageStore>(IoDispatcher);
+}
diff --git a/Engine/Source/Runtime/CoreUObject/Private/Serialization/MappedName.cpp b/Engine/Source/Runtime/CoreUObject/Private/Serialization/MappedName.cpp
new file mode 100644
--- /dev/null
+++ b/Engine/Source/Runtime/CoreUObject/Private/Serialization/MappedName.cpp
@@ -0,0 +1,18 @@
+// Copyright Epic Games, Inc. All Rights Reserved.
+
+#include "Serialization/MappedName.h"
+#include "Serialization/Archive.h"
+#include "UObject/NameBatchSerialization.h"
+
+FArchive& operator<<(FArchive& Ar, FMappedName& MappedName)
+{
+	Ar << MappedName.Index << MappedName.Number;
+
+	return Ar;
+}
+
+void FNameMap::Load(TArrayView<const uint8> NameBuffer, TArrayView<const uint8> HashBuffer, FMappedName::EType InNameMapType)
+{
+	LoadNameBatch(NameEntries, NameBuffer, HashBuffer);
+	NameMapType = InNameMapType;
+}
diff --git a/Engine/Source/Runtime/CoreUObject/Public/Serialization/AsyncLoading2.h b/Engine/Source/Runtime/CoreUObject/Public/Serialization/AsyncLoading2.h
--- a/Engine/Source/Runtime/CoreUObject/Public/Serialization/AsyncLoading2.h
+++ b/Engine/Source/Runtime/CoreUObject/Public/Serialization/AsyncLoading2.h
@@ -10,6 +10,7 @@
 #include "UObject/ObjectResource.h"
 #include "UObject/PackageId.h"
 #include "Serialization/Archive.h"
+#include "Serialization/MappedName.h"
 #include "IO/IoContainerId.h"
 
 class FArchive;
@@ -20,93 +21,6 @@ class IEDLBootNotificationManager;
 using FSourceToLocalizedPackageIdMap = TArray<TPair<FPackageId, FPackageId>>;
 using FCulturePackageMap = TMap<FString, FSourceToLocalizedPackageIdMap>;
 
-class FMappedName
-{
-	static constexpr uint32 InvalidIndex = ~uint32(0);
-	static constexpr uint32 IndexBits = 30u;
-	static constexpr uint32 IndexMask = (1u << IndexBits) - 1u;
-	static constexpr uint32 TypeMask = ~IndexMask;
-	static constexpr uint32 TypeShift = IndexBits;
-
-public:
-	enum class EType
-	{
-		Package,
-		Container,
-		Global
-	};
-
-	inline FMappedName() = default;
-
-	static inline FMappedName Create(const uint32 InIndex, const uint32 InNumber, EType InType)
-	{
-		check(InIndex <= MAX_int32);
-		return FMappedName((uint32(InType) << TypeShift) | InIndex, InNumber);
-	}
-
-	static inline FMappedName FromMinimalName(const FMinimalName& MinimalName)
-	{
-		return *reinterpret_cast<const FMappedName*>(&MinimalName);
-	}
-
-	static inline bool IsResolvedToMinimalName(const FMinimalName& MinimalName)
-	{
-		// Not completely safe, relies on that no FName will have its Index and Number equal to Max_uint32
-		const FMappedName MappedName = FromMinimalName(MinimalName);
-		return MappedName.IsValid();
-	}
-
-	static inline FName SafeMinimalNameToName(const FMinimalName& MinimalName)
-	{
-		return IsResolvedToMinimalName(MinimalName) ? MinimalNameToName(MinimalName) : NAME_None;
-	}
-
-	inline FMinimalName ToUnresolvedMinimalName() const
-	{
-		return *reinterpret_cast<const FMinimalName*>(this);
-	}
-
-	inline bool IsValid() const
-	{
-		return Index != InvalidIndex && Number != InvalidIndex;
-	}
-
-	inline EType GetType() const
-	{
-		return static_cast<EType>(uint32((Index & TypeMask) >> TypeShift));
-	}
-
-	inline bool IsGlobal() const
-	{
-		return ((Index & TypeMask) >> TypeShift) != 0;
-	}
-
-	inline uint32 GetIndex() const
-	{
-		return Index & IndexMask;
-	}
-
-	inline uint32 GetNumber() const
-	{
-		return Number;
-	}
-
-	inline bool operator!=(FMappedName Other) const
-	{
-		return Index != Other.Index || Number != Other.Number;
-	}
-
-	COREUOBJECT_API friend FArchive& operator<<(FArchive& Ar, FMappedName& MappedName);
-
-private:
-	inline FMappedName(const uint32 InIndex, const uint32 InNumber)
-		: Index(InIndex)
-		, Number(InNumber) { }
-
-	uint32 Index = InvalidIndex;
-	uint32 Number = InvalidIndex;
-};
-
 struct FContainerHeader
 {
 	FIoContainerId ContainerId;
@@ -269,38 +183,6 @@ struct FExportBundleEntry
 	COREUOBJECT_API friend FArchive& operator<<(FArchive& Ar, FExportBundleEntry& ExportBundleEntry);
 };
 
-template<typename T>
-class TPackageStoreEntryCArrayView
-{
-	const uint32 ArrayNum = 0;
-	const uint32 OffsetToDataFromThis = 0;
-
-public:
-	inline uint32 Num() const						{ return ArrayNum; }
-
-	inline const T* Data() const					{ return (T*)((char*)this + OffsetToDataFromThis); }
-	inline T* Data()								{ return (T*)((char*)this + OffsetToDataFromThis); }
-
-	inline const T* begin() const					{ return Data(); }
-	inline T* begin()								{ return Data(); }
-
-	inline const T* end() const						{ return Data() + ArrayNum; }
-	inline T* end()									{ return Data() + ArrayNum; }
-
-	inline const T& operator[](uint32 Index) const	{ return Data()[Index]; }
-	inline T& operator[](uint32 Index)				{ return Data()[Index]; }
-};
-
-struct FPackageStoreEntry
-{
-	uint64 ExportBundlesSize;
-	int32 ExportCount;
-	int32 ExportBundleCount;
-	uint32 LoadOrder;
-	uint32 Pad;
-	TPackageStoreEntryCArrayView<FPackageId> ImportedPackages;
-};
-
 /**
  * Export bundle header
  */
diff --git a/Engine/Source/Runtime/CoreUObject/Public/Serialization/MappedName.h b/Engine/Source/Runtime/CoreUObject/Public/Serialization/MappedName.h
new file mode 100644
--- /dev/null
+++ b/Engine/Source/Runtime/CoreUObject/Public/Serialization/MappedName.h
@@ -0,0 +1,147 @@
+// Copyright Epic Games, Inc. All Rights Reserved.
+
+#pragma once
+
+#include "CoreMinimal.h"
+
+/**
+ * Index and name number into a name map.
+ */
+class FMappedName
+{
+	static constexpr uint32 InvalidIndex = ~uint32(0);
+	static constexpr uint32 IndexBits = 30u;
+	static constexpr uint32 IndexMask = (1u << IndexBits) - 1u;
+	static constexpr uint32 TypeMask = ~IndexMask;
+	static constexpr uint32 TypeShift = IndexBits;
+
+public:
+	enum class EType
+	{
+		Package,
+		Container,
+		Global
+	};
+
+	inline FMappedName() = default;
+
+	static inline FMappedName Create(const uint32 InIndex, const uint32 InNumber, EType InType)
+	{
+		check(InIndex <= MAX_int32);
+		return FMappedName((uint32(InType) << TypeShift) | InIndex, InNumber);
+	}
+
+	static inline FMappedName FromMinimalName(const FMinimalName& MinimalName)
+	{
+		return *reinterpret_cast<const FMappedName*>(&MinimalName);
+	}
+
+	static inline bool IsResolvedToMinimalName(const FMinimalName& MinimalName)
+	{
+		// Not completely safe, relies on that no FName will have its Index and Number equal to Max_uint32
+		const FMappedName MappedName = FromMinimalName(MinimalName);
+		return MappedName.IsValid();
+	}
+
+	static inline FName SafeMinimalNameToName(const FMinimalName& MinimalName)
+	{
+		return IsResolvedToMinimalName(MinimalName) ? MinimalNameToName(MinimalName) : NAME_None;
+	}
+
+	inline FMinimalName ToUnresolvedMinimalName() const
+	{
+		return *reinterpret_cast<const FMinimalName*>(this);
+	}
+
+	inline bool IsValid() const
+	{
+		return Index != InvalidIndex && Number != InvalidIndex;
+	}
+
+	inline EType GetType() const
+	{
+		return static_cast<EType>(uint32((Index & TypeMask) >> TypeShift));
+	}
+
+	inline bool IsGlobal() const
+	{
+		return ((Index & TypeMask) >> TypeShift) != 0;
+	}
+
+	inline uint32 GetIndex() const
+	{
+		return Index & IndexMask;
+	}
+
+	inline uint32 GetNumber() const
+	{
+		return Number;
+	}
+
+	inline bool operator!=(FMappedName Other) const
+	{
+		return Index != Other.Index || Number != Other.Number;
+	}
+
+	COREUOBJECT_API friend FArchive& operator<<(FArchive& Ar, FMappedName& MappedName);
+
+private:
+	inline FMappedName(const uint32 InIndex, const uint32 InNumber)
+		: Index(InIndex)
+		, Number(InNumber) { }
+
+	uint32 Index = InvalidIndex;
+	uint32 Number = InvalidIndex;
+};
+
+/*
+ * Maps serialized name entries to names.
+ */
+class FNameMap
+{
+public:
+	inline int32 Num() const
+	{
+		return NameEntries.Num();
+	}
+
+	void Load(TArrayView<const uint8> NameBuffer, TArrayView<const uint8> HashBuffer, FMappedName::EType NameMapType);
+
+	FName GetName(const FMappedName& MappedName) const
+	{
+		check(MappedName.GetType() == NameMapType);
+		check(MappedName.GetIndex() < uint32(NameEntries.Num()));
+
+		FNameEntryId NameEntry = NameEntries[MappedName.GetIndex()];
+		return FName::CreateFromDisplayId(NameEntry, MappedName.GetNumber());
+	}
+
+	bool TryGetName(const FMappedName& MappedName, FName& OutName) const
+	{
+		check(MappedName.GetType() == NameMapType);
+
+		uint32 Index = MappedName.GetIndex();
+		if (Index < uint32(NameEntries.Num()))
+		{
+			FNameEntryId NameEntry = NameEntries[MappedName.GetIndex()];
+			OutName = FName::CreateFromDisplayId(NameEntry, MappedName.GetNumber());
+
+			return true;
+		}
+
+		return false;
+	}
+
+	FMinimalName GetMinimalName(const FMappedName& MappedName) const
+	{
+		check(MappedName.GetType() == NameMapType);
+		check(MappedName.GetIndex() < uint32(NameEntries.Num()));
+
+		FNameEntryId NameEntry = NameEntries[MappedName.GetIndex()];
+		return FMinimalName(NameEntry, MappedName.GetNumber());
+	}
+
+private:
+	TArray<FNameEntryId> NameEntries;
+	FMappedName::EType NameMapType = FMappedName::EType::Global;
+};
diff --git a/Engine/Source/Runtime/CoreUObject/Public/Serialization/PackageStore.h b/Engine/Source/Runtime/CoreUObject/Public/Serialization/PackageStore.h
new file mode 100644
--- /dev/null
+++ b/Engine/Source/Runtime/CoreUObject/Public/Serialization/PackageStore.h
@@ -0,0 +1,74 @@
+// Copyright Epic Games, Inc. All Rights Reserved.
+
+#pragma once
+
+#include "CoreMinimal.h"
+#include "UObject/PackageId.h"
+
+class FIoDispatcher;
+struct FPackageStoreEntry;
+
+/**
+ * Package store entry array view.
+ */
+template<typename T>
+class TPackageStoreEntryCArrayView
+{
+	const uint32 ArrayNum = 0;
+	const uint32 OffsetToDataFromThis = 0;
+
+public:
+	inline uint32 Num() const						{ return ArrayNum; }
+
+	inline const T* Data() const					{ return (T*)((char*)this + OffsetToDataFromThis); }
+	inline T* Data()								{ return (T*)((char*)this + OffsetToDataFromThis); }
+
+	inline const T* begin() const					{ return Data(); }
+	inline T* begin()								{ return Data(); }
+
+	inline const T* end() const						{ return Data() + ArrayNum; }
+	inline T* end()									{ return Data() + ArrayNum; }
+
+	inline const T& operator[](uint32 Index) const	{ return Data()[Index]; }
+	inline T& operator[](uint32 Index)				{ return Data()[Index]; }
+};
+
+/**
+ * Package store entry.
+ */
+struct FPackageStoreEntry
+{
+	uint64 ExportBundlesSize;
+	int32 ExportCount;
+	int32 ExportBundleCount;
+	uint32 LoadOrder;
+	uint32 Pad;
+	TPackageStoreEntryCArrayView<FPackageId> ImportedPackages;
+};
+
+/**
+ * Stores information about available packages that can be loaded.
+ */
+class IPackageStore
+{
+public:
+	/* Destructor. */
+	virtual ~IPackageStore() { }
+
+	/* Initialize the package store. */
+	virtual void Initialize() = 0;
+
+	/* Returns whether the package exists. */
+	virtual bool DoesPackageExist(FPackageId PackageId) = 0;
+
+	/* Get the package information for the specified package ID. */
+	virtual const FPackageStoreEntry* GetPackageEntry(FPackageId PackageId) = 0;
+
+	/* Returns the redirected package ID for the specified package ID. */
+	virtual FPackageId GetRedirectedPackageId(FPackageId PackageId) = 0;
+
+	/* Returns whether the package ID is a redirect. */
+	virtual bool IsRedirect(FPackageId PackageId) = 0;
+};
+
+TUniquePtr<IPackageStore> MakeFilePackageStore(FIoDispatcher& IoDispatcher);
