Prototype Wrapped Object Pointers (ObjectHandle/ObjectPtr)

Avoid an assert experienced when wrapped object pointers are used in named values via the StructuredArchive's FArchiveFromStructuredArchive adapter.  This change ensures that the wrapped pointers are serialized directly instead of through the adapter.  The groundwork for avoiding the adapter was already there, but what was missing was the operator necessary to convert from the templated TObjectPtr class to the private non-templated FObjectPtr base class used by serialization.  I suspect this same issue may also be happening for other types like TLazyObjectPtr which don't implement the necessary operator to convert to their private base type used in serialization.

diff --git a/Engine/Source/Runtime/CoreUObject/Public/UObject/ObjectPtr.h b/Engine/Source/Runtime/CoreUObject/Public/UObject/ObjectPtr.h
--- a/Engine/Source/Runtime/CoreUObject/Public/UObject/ObjectPtr.h
+++ b/Engine/Source/Runtime/CoreUObject/Public/UObject/ObjectPtr.h
@@ -3,6 +3,7 @@
 #pragma once
 
 #include "HAL/Platform.h"
+#include "Serialization/StructuredArchive.h"
 #include "UObject/Class.h"
 #include "UObject/Object.h"
 #include "UObject/ObjectHandle.h"
@@ -298,6 +299,11 @@ public:
 		return Ar;
 	}
 
+	friend FORCEINLINE void operator<<(FStructuredArchiveSlot&& Slot, TObjectPtr<T>& ObjectPtr)
+	{
+		Slot << static_cast<FObjectPtr&>(ObjectPtr);
+	}
+
 	friend struct FObjectPtr;
 	template <typename U> friend struct TObjectPtr;
 private:
