Prototype Wrapped Object Pointers (ObjectHandle/ObjectPtr)

Fix issue where constructor for FSoftObjectPath was using FObjectHandle variant in unintended cases in configurations with lazy load disabled.  This led to errors when using the "package override" feature.

diff --git a/Engine/Source/Runtime/CoreUObject/Private/UObject/ObjectHandle.cpp b/Engine/Source/Runtime/CoreUObject/Private/UObject/ObjectHandle.cpp
--- a/Engine/Source/Runtime/CoreUObject/Private/UObject/ObjectHandle.cpp
+++ b/Engine/Source/Runtime/CoreUObject/Private/UObject/ObjectHandle.cpp
@@ -301,6 +301,9 @@ static inline UPackage* FindOrLoadPackage(FName PackageName, int32 LoadFlags)
 	UPackage* TargetPackage = (UPackage*)StaticFindObjectFast(UPackage::StaticClass(), nullptr, PackageName);
 	if (!ScriptPackageName && !TargetPackage)
 	{
+		// @TODO: OBJPTR: When using the "external package" feature, we will have objects that have a differing package path vs "outer hierarchy" path
+		//				  The package path should be used when loading.  The "outer hierarchy" path may need to be used when finding existing objects in memory.
+		//				  This will need further evaluation and testing before lazy load can be enabled.
 		// @TODO: OBJPTR: This seems like it will do a search on disk every time instead of leveraging known missing packages
 		// @TODO: OBJPTR: Instancing context may be important to consider when loading the package.
 		LoadFlags |= LOAD_NoWarn | LOAD_NoVerify; //This does nothing? | LOAD_DisableDependencyPreloading;
diff --git a/Engine/Source/Runtime/CoreUObject/Private/UObject/SoftObjectPath.cpp b/Engine/Source/Runtime/CoreUObject/Private/UObject/SoftObjectPath.cpp
--- a/Engine/Source/Runtime/CoreUObject/Private/UObject/SoftObjectPath.cpp
+++ b/Engine/Source/Runtime/CoreUObject/Private/UObject/SoftObjectPath.cpp
@@ -19,21 +19,30 @@ FSoftObjectPath::FSoftObjectPath(const UObject* InObject)
 	}
 }
 
-FSoftObjectPath::FSoftObjectPath(FObjectHandle InObjectHandle)
+void FSoftObjectPath::InitializeFromObjectHandle(FObjectHandle InObjectHandle)
 {
 	if (!IsObjectHandleNull(InObjectHandle))
 	{
-		//Try to compose the path without forcing the referenced data to be loaded
-		FObjectRef Ref = MakeObjectRef(InObjectHandle);
-		FObjectPathId::ResolvedNameContainerType ResolvedNames;
-		Ref.ObjectPath.Resolve(ResolvedNames);
-		TStringBuilder<128> CompletePath;
-		CompletePath.Append(Ref.PackageName.ToString());
-		CompletePath.Append(TEXT('.'));
-		for (int32 ResolvedNameIndex = 0; ResolvedNameIndex < ResolvedNames.Num(); ++ResolvedNameIndex)
+		if (IsObjectHandleResolved(InObjectHandle))
 		{
-			switch (ResolvedNameIndex)
+			if (UObject* ResolvedObject = ResolveObjectHandle(InObjectHandle))
 			{
+				SetPath(ResolvedObject->GetPathName());
+			}
+		}
+		else
+		{
+			//Try to compose the path without forcing the referenced data to be loaded
+			FObjectRef Ref = MakeObjectRef(InObjectHandle);
+			FObjectPathId::ResolvedNameContainerType ResolvedNames;
+			Ref.ObjectPath.Resolve(ResolvedNames);
+			TStringBuilder<128> CompletePath;
+			CompletePath.Append(Ref.PackageName.ToString());
+			CompletePath.Append(TEXT('.'));
+			for (int32 ResolvedNameIndex = 0; ResolvedNameIndex < ResolvedNames.Num(); ++ResolvedNameIndex)
+			{
+				switch (ResolvedNameIndex)
+				{
 				case 0:
 				{
 					break;
@@ -48,10 +57,11 @@ FSoftObjectPath::FSoftObjectPath(FObjectHandle InObjectHandle)
 					CompletePath << '.';
 					break;
 				}
+				}
+				CompletePath << ResolvedNames[ResolvedNameIndex];
 			}
-			CompletePath << ResolvedNames[ResolvedNameIndex];
+			SetPath(CompletePath.ToString());
 		}
-		SetPath(CompletePath.ToString());
 	}
 }
 
diff --git a/Engine/Source/Runtime/CoreUObject/Public/UObject/SoftObjectPath.h b/Engine/Source/Runtime/CoreUObject/Public/UObject/SoftObjectPath.h
--- a/Engine/Source/Runtime/CoreUObject/Public/UObject/SoftObjectPath.h
+++ b/Engine/Source/Runtime/CoreUObject/Public/UObject/SoftObjectPath.h
@@ -51,18 +51,15 @@ struct COREUOBJECT_API FSoftObjectPath
 	/** Construct from an existing object in memory */
 	FSoftObjectPath(const UObject* InObject);
 
-	/** Construct from a FObjectHandle/FObjectPtr/TObjectPtr which may or may not have the referenced object in memory */
-	explicit FSoftObjectPath(FObjectHandle InObjectHandle);
-
 	FSoftObjectPath(FObjectPtr InObject)
-		: FSoftObjectPath(InObject.GetHandle())
 	{
+		InitializeFromObjectHandle(InObject.GetHandle());
 	}
 
 	template <typename T>
 	FSoftObjectPath(const TObjectPtr<T>& InObject)
-		: FSoftObjectPath(InObject.GetHandle())
 	{
+		InitializeFromObjectHandle(InObject.GetHandle());
 	}
 
 	~FSoftObjectPath() {}
@@ -248,6 +245,9 @@ private:
 	/** Package names currently being duplicated, needed by FixupForPIE */
 	static TSet<FName> PIEPackageNames;
 
+	/** Initialize from a FObjectHandle/FObjectPtr/TObjectPtr which may or may not have the referenced object in memory */
+	void InitializeFromObjectHandle(FObjectHandle InObjectHandle);
+
 	UObject* ResolveObjectInternal() const;
 	UObject* ResolveObjectInternal(const TCHAR* PathString) const;
 
