Prototype Wrapped Object Pointers (ObjectHandle/ObjectPtr)

Address the following serialization issues:

1) Ensure ObjectPtr Property can convert from SoftObjectProperty and anything else that Object Property can convert from by removing its ConvertFromType override.
2) Fix serialization of wrapped pointers through archives that don't derive from FArchiveUObject, but still expect to be able to serialize both raw and wrapped object references (FNetBitWriter, FExportArchive, and FObjectAndNameAsStringProxyArchive).
3) Restore assert in base FArchive serialization of wrapped object pointer which I had erroneously made in CL15130804 now that it is clear that it is not safe to be relying on silently not writing/reading a wrapped object pointer (even if that may be the default behavior for a raw pointer).
4) Bump Niagara custom version guid to avoid issues with incorrectly serialized data in DDC.

diff --git a/Engine/Plugins/FX/Niagara/Source/NiagaraCore/Private/NiagaraCustomVersion.cpp b/Engine/Plugins/FX/Niagara/Source/NiagaraCore/Private/NiagaraCustomVersion.cpp
--- a/Engine/Plugins/FX/Niagara/Source/NiagaraCore/Private/NiagaraCustomVersion.cpp
+++ b/Engine/Plugins/FX/Niagara/Source/NiagaraCore/Private/NiagaraCustomVersion.cpp
@@ -8,4 +8,4 @@ const FGuid FNiagaraCustomVersion::GUID(0xFCF57AFA, 0x50764283, 0xB9A9E658, 0xFF
 // Register the custom version with core
 FCustomVersionRegistration GRegisterNiagaraCustomVersion(FNiagaraCustomVersion::GUID, FNiagaraCustomVersion::LatestVersion, TEXT("NiagaraVer"));
 
-const FGuid FNiagaraCustomVersion::LatestScriptCompileVersion(0x2C2A181E, 0xB29145D6, 0x96EC9584, 0xDCAA75AD);
+const FGuid FNiagaraCustomVersion::LatestScriptCompileVersion(0x26D0C2AD, 0x739E42CD, 0x9D3A3514, 0x5E853816);
diff --git a/Engine/Source/Runtime/Core/Private/Serialization/Archive.cpp b/Engine/Source/Runtime/Core/Private/Serialization/Archive.cpp
--- a/Engine/Source/Runtime/Core/Private/Serialization/Archive.cpp
+++ b/Engine/Source/Runtime/Core/Private/Serialization/Archive.cpp
@@ -467,6 +467,8 @@ FArchive& FArchive::operator<<(struct FLazyObjectPtr& Value)
 
 FArchive& FArchive::operator<<(struct FObjectPtr& Value)
 {
+	// The base FArchive does not implement this method. Use FArchiveUObject instead.
+	UE_LOG(LogSerialization, Fatal, TEXT("FArchive does not support FObjectPtr serialization. Use FArchiveUObject instead."));
 	return *this;
 }
 
diff --git a/Engine/Source/Runtime/CoreUObject/Private/Serialization/ArchiveUObject.cpp b/Engine/Source/Runtime/CoreUObject/Private/Serialization/ArchiveUObject.cpp
--- a/Engine/Source/Runtime/CoreUObject/Private/Serialization/ArchiveUObject.cpp
+++ b/Engine/Source/Runtime/CoreUObject/Private/Serialization/ArchiveUObject.cpp
@@ -182,6 +182,11 @@ FArchive& FObjectAndNameAsStringProxyArchive::operator<<(FSoftObjectPath& Value)
 	return *this;
 }
 
+FArchive& FObjectAndNameAsStringProxyArchive::operator<<(FObjectPtr& Obj)
+{
+	return FArchiveUObject::SerializeObjectPtr(*this, Obj);
+}
+
 void FSerializedPropertyScope::PushProperty()
 {
 	if (Property)
diff --git a/Engine/Source/Runtime/CoreUObject/Private/Serialization/AsyncLoading2.cpp b/Engine/Source/Runtime/CoreUObject/Private/Serialization/AsyncLoading2.cpp
--- a/Engine/Source/Runtime/CoreUObject/Private/Serialization/AsyncLoading2.cpp
+++ b/Engine/Source/Runtime/CoreUObject/Private/Serialization/AsyncLoading2.cpp
@@ -1166,6 +1166,7 @@ public:
 	}
 
 	//~ Begin FArchive::FArchiveUObject Interface
+	virtual FArchive& operator<<(FObjectPtr& Value) override { return FArchiveUObject::SerializeObjectPtr(*this, Value); }
 	virtual FArchive& operator<<(FSoftObjectPath& Value) override { return FArchiveUObject::SerializeSoftObjectPath(*this, Value); }
 	virtual FArchive& operator<<(FWeakObjectPtr& Value) override { return FArchiveUObject::SerializeWeakObjectPtr(*this, Value); }
 	//~ End FArchive::FArchiveUObject Interface
diff --git a/Engine/Source/Runtime/CoreUObject/Private/UObject/CoreNet.cpp b/Engine/Source/Runtime/CoreUObject/Private/UObject/CoreNet.cpp
--- a/Engine/Source/Runtime/CoreUObject/Private/UObject/CoreNet.cpp
+++ b/Engine/Source/Runtime/CoreUObject/Private/UObject/CoreNet.cpp
@@ -443,6 +443,11 @@ FArchive& FNetBitWriter::operator<<(FSoftObjectPtr& Value)
 	return FArchiveUObject::SerializeSoftObjectPtr(*this, Value);
 }
 
+FArchive& FNetBitWriter::operator<<(FObjectPtr& Value)
+{
+	return FArchiveUObject::SerializeObjectPtr(*this, Value);
+}
+
 FArchive& FNetBitWriter::operator<<(struct FWeakObjectPtr& WeakObjectPtr)
 {
 	return FArchiveUObject::SerializeWeakObjectPtr(*this, WeakObjectPtr);
diff --git a/Engine/Source/Runtime/CoreUObject/Private/UObject/PropertyObjectPtr.cpp b/Engine/Source/Runtime/CoreUObject/Private/UObject/PropertyObjectPtr.cpp
--- a/Engine/Source/Runtime/CoreUObject/Private/UObject/PropertyObjectPtr.cpp
+++ b/Engine/Source/Runtime/CoreUObject/Private/UObject/PropertyObjectPtr.cpp
@@ -29,22 +29,6 @@ void FObjectPtrProperty::SerializeItem(FStructuredArchive::FSlot Slot, void* Val
 	Slot << *ObjectPtr;
 }
 
-EConvertFromTypeResult FObjectPtrProperty::ConvertFromType(const FPropertyTag& Tag, FStructuredArchive::FSlot Slot, uint8* Data, UStruct* DefaultsStruct)
-{
-	// Convert raw object reference to TObjectPtr
-	if (Tag.Type == NAME_ObjectProperty)
-	{
-		uint8* DestAddress = ContainerPtrToValuePtr<uint8>(Data, Tag.ArrayIndex);
-		FArchive& UnderlyingArchive = Slot.GetUnderlyingArchive();
-		FObjectPtr* ObjectPtr = (FObjectPtr*)GetPropertyValuePtr(DestAddress);
-		Slot << *ObjectPtr;
-
-		return UnderlyingArchive.IsCriticalError() ? EConvertFromTypeResult::CannotConvert : EConvertFromTypeResult::Converted;
-	}
-
-	return EConvertFromTypeResult::UseSerializeItem;
-}
-
 bool FObjectPtrProperty::SameType(const FProperty* Other) const
 {
 	// @TODO: OBJPTR: Should this be done through a new, separate API on FProperty (eg: ImplicitConv)
diff --git a/Engine/Source/Runtime/CoreUObject/Public/Serialization/ObjectAndNameAsStringProxyArchive.h b/Engine/Source/Runtime/CoreUObject/Public/Serialization/ObjectAndNameAsStringProxyArchive.h
--- a/Engine/Source/Runtime/CoreUObject/Public/Serialization/ObjectAndNameAsStringProxyArchive.h
+++ b/Engine/Source/Runtime/CoreUObject/Public/Serialization/ObjectAndNameAsStringProxyArchive.h
@@ -38,5 +38,6 @@ struct FObjectAndNameAsStringProxyArchive : public FNameAsStringProxyArchive
 	COREUOBJECT_API virtual FArchive& operator<<(FWeakObjectPtr& Obj) override;
 	COREUOBJECT_API virtual FArchive& operator<<(FSoftObjectPtr& Value) override;
 	COREUOBJECT_API virtual FArchive& operator<<(FSoftObjectPath& Value) override;
+	COREUOBJECT_API virtual FArchive& operator<<(FObjectPtr& Obj) override;
 };
 
diff --git a/Engine/Source/Runtime/CoreUObject/Public/UObject/CoreNet.h b/Engine/Source/Runtime/CoreUObject/Public/UObject/CoreNet.h
--- a/Engine/Source/Runtime/CoreUObject/Public/UObject/CoreNet.h
+++ b/Engine/Source/Runtime/CoreUObject/Public/UObject/CoreNet.h
@@ -346,6 +346,7 @@ public:
 	virtual FArchive& operator<<(UObject*& Object) override;
 	virtual FArchive& operator<<(FSoftObjectPath& Value) override;
 	virtual FArchive& operator<<(FSoftObjectPtr& Value) override;
+	virtual FArchive& operator<<(FObjectPtr& Value) override;
 	virtual FArchive& operator<<(struct FWeakObjectPtr& Value) override;
 
 	virtual void CountMemory(FArchive& Ar) const override;
diff --git a/Engine/Source/Runtime/CoreUObject/Public/UObject/UnrealType.h b/Engine/Source/Runtime/CoreUObject/Public/UObject/UnrealType.h
--- a/Engine/Source/Runtime/CoreUObject/Public/UObject/UnrealType.h
+++ b/Engine/Source/Runtime/CoreUObject/Public/UObject/UnrealType.h
@@ -2309,7 +2309,6 @@ class COREUOBJECT_API FObjectPtrProperty : public FObjectProperty
 	virtual bool Identical(const void* A, const void* B, uint32 PortFlags) const override;
 	virtual void SerializeItem(FStructuredArchive::FSlot Slot, void* Value, void const* Defaults) const override;
 	virtual void EmitReferenceInfo(UClass& OwnerClass, int32 BaseOffset, TArray<const FStructProperty*>& EncounteredStructProps) override;
-	virtual EConvertFromTypeResult ConvertFromType(const FPropertyTag& Tag, FStructuredArchive::FSlot Slot, uint8* Data, UStruct* DefaultsStruct) override;
 	// End of FProperty interface
 
 	// FObjectProperty interface
