Prototype Wrapped Object Pointers (ObjectHandle/ObjectPtr)

Fix another issue of an FArchive that serializes object references without deriving from FArchiveUObject.  This one affects serialization to/from text assets.  Confirmed I can export an asset containing object references to text format when the fields are of wrapped pointer type.  Comparing the resulting output to a text export of the same asset when the fields are raw pointers show no diffs aside from Guid and Localization which are shown to vary from export-to-export or editor session to editor session.

NOTE: Unlike lazy/soft/weak pointers, I want serialization of wrapped pointers to match serialization of raw pointers, so I don't create a separate indexed array for it in the text output, it goes into the same "Objects" array as raw object pointers.

diff --git a/Engine/Source/Runtime/CoreUObject/Private/Serialization/ArchiveUObjectFromStructuredArchive.cpp b/Engine/Source/Runtime/CoreUObject/Private/Serialization/ArchiveUObjectFromStructuredArchive.cpp
--- a/Engine/Source/Runtime/CoreUObject/Private/Serialization/ArchiveUObjectFromStructuredArchive.cpp
+++ b/Engine/Source/Runtime/CoreUObject/Private/Serialization/ArchiveUObjectFromStructuredArchive.cpp
@@ -148,6 +148,12 @@ FArchive& FArchiveUObjectFromStructuredArchiveImpl::operator<<(FWeakObjectPtr& V
 	return *this;
 }
 
+FArchive& FArchiveUObjectFromStructuredArchiveImpl::operator<<(FObjectPtr& Value)
+{
+	// @TODO: OBJPTR: If we want lazy loading from text archives, this will need to be modified to not route through the default raw pointer serialization path.
+	return FArchiveUObject::SerializeObjectPtr(*this, Value);
+}
+
 void FArchiveUObjectFromStructuredArchiveImpl::PushFileRegionType(EFileRegionType Type)
 {
 	check(CurrentFileRegionType == EFileRegionType::None);
diff --git a/Engine/Source/Runtime/CoreUObject/Public/Serialization/ArchiveUObjectFromStructuredArchive.h b/Engine/Source/Runtime/CoreUObject/Public/Serialization/ArchiveUObjectFromStructuredArchive.h
--- a/Engine/Source/Runtime/CoreUObject/Public/Serialization/ArchiveUObjectFromStructuredArchive.h
+++ b/Engine/Source/Runtime/CoreUObject/Public/Serialization/ArchiveUObjectFromStructuredArchive.h
@@ -7,6 +7,7 @@
 #include "Serialization/FileRegions.h"
 #include "UObject/ObjectResource.h"
 
+#include "UObject/ObjectPtr.h"
 #include "UObject/SoftObjectPath.h"
 #include "UObject/SoftObjectPtr.h"
 #include "UObject/LazyObjectPtr.h"
@@ -29,6 +30,7 @@ public:
 	virtual FArchive& operator<<(FSoftObjectPtr& Value) override;
 	virtual FArchive& operator<<(FSoftObjectPath& Value) override;
 	virtual FArchive& operator<<(FWeakObjectPtr& Value) override;
+	virtual FArchive& operator<<(FObjectPtr& Value) override;
 	//~ End FArchive Interface
 
 	virtual void PushFileRegionType(EFileRegionType Type) override;
