Don't clear BP references on reinstantiation if the object reference shares an outer with the container, even if they aren't in the same package (in the case of objects with a package override).

diff --git a/Engine/Source/Runtime/CoreUObject/Public/Serialization/ArchiveReplaceOrClearExternalReferences.h b/Engine/Source/Runtime/CoreUObject/Public/Serialization/ArchiveReplaceOrClearExternalReferences.h
--- a/Engine/Source/Runtime/CoreUObject/Public/Serialization/ArchiveReplaceOrClearExternalReferences.h
+++ b/Engine/Source/Runtime/CoreUObject/Public/Serialization/ArchiveReplaceOrClearExternalReferences.h
@@ -41,19 +41,16 @@ public:
 		// if Resolved is a private object in another package just clear the reference:
 		if (Resolved)
 		{
-			UObject* Outermost = Resolved->GetOutermost();
-			if (Outermost)
+			UPackage* ObjPackage = Resolved->GetPackage();
+			if (ObjPackage)
 			{
-				UPackage* ObjPackage = dynamic_cast<UPackage*>(Outermost);
-				if (ObjPackage)
+				CA_SUPPRESS(6011);
+				if (ObjPackage != Obj && 
+					DestPackage != ObjPackage && 
+					Obj->GetOutermostObject() != TSuper::GetSearchObject()->GetOutermostObject() &&
+					!Obj->HasAnyFlags(RF_Public))
 				{
-					CA_SUPPRESS(6011);
-					if (ObjPackage != Obj && 
-						DestPackage != ObjPackage && 
-						!Obj->HasAnyFlags(RF_Public))
-					{
-						Resolved = nullptr;
-					}
+					Resolved = nullptr;
 				}
 			}
 		}
