Fix compile errors about operator ambiguity when comparing a TObjectPtr and NULL (not nullptr) specifically in MSVC non-conformance mode.  The fix introduces an issue where people could now compare TObjectPtr with an arbitrary integer when compiling in non-conformance mode.  The alternatives considered:

1) Use safe bool idiom - fails due to ambiguity with implicit cast operator to T*
2) Use constexpr comparison operators with validation that the TYPE_OF_NULL parameter is always zero - can't find a way to accomplish this with function parameters at compile time.

The issue arises only when a type has conversion operator to a bool (explicit or implicit), and removal of that operator just leads to other problems where a TObjectPtr must be resolved before it can be checked against null, which is something that I've tried to avoid in anticipation of lazy load.  If we end up abandoning the idea of null checking a TObjectPtr without resolving it, we can revisit this code and strip out this workaround.

diff --git a/Engine/Source/Runtime/CoreUObject/Private/Tests/ObjectPtrTest.cpp b/Engine/Source/Runtime/CoreUObject/Private/Tests/ObjectPtrTest.cpp
--- a/Engine/Source/Runtime/CoreUObject/Private/Tests/ObjectPtrTest.cpp
+++ b/Engine/Source/Runtime/CoreUObject/Private/Tests/ObjectPtrTest.cpp
@@ -127,10 +127,10 @@ static_assert(!TModels<CEqualityComparableWith, FMutableObjectPtr, char*>::Value
 static_assert(TModels<CEqualityComparableWith, FConstObjectPtr, TYPE_OF_NULLPTR>::Value, "Must be able to compare equality and inequality bidirectionally between TObjectPtr<const UObject> and nullptr");
 static_assert(TModels<CEqualityComparableWith, FMutableObjectPtr, TYPE_OF_NULLPTR>::Value, "Must be able to compare equality and inequality bidirectionally between TObjectPtr<UObject> and nullptr");
 
-#if !(PLATFORM_MICROSOFT) || !defined(_MSC_EXTENSIONS) // MSVC static analyzer is run in non-conformance mode, and that causes these checks to fail.
+#if !UE_OBJECT_PTR_NONCONFORMANCE_SUPPORT // Specialized NULL support causes these checks to fail.
 static_assert(!TModels<CEqualityComparableWith, FConstObjectPtr, long>::Value, "Should not be able to compare equality and inequality bidirectionally between TObjectPtr<const UObject> and long");
 static_assert(!TModels<CEqualityComparableWith, FMutableObjectPtr, long>::Value, "Should not be able to compare equality and inequality bidirectionally between TObjectPtr<UObject> and long");
-#endif // #if !(PLATFORM_MICROSOFT) || !defined(_MSC_EXTENSIONS)
+#endif // #if !UE_OBJECT_PTR_NONCONFORMANCE_SUPPORT
 
 // Ensure that the use of incomplete types doesn't provide a means to bypass type safety on TObjectPtr
 // NOTE: This is disabled because we're permitting this operation with a deprecation warning.
diff --git a/Engine/Source/Runtime/CoreUObject/Public/UObject/ObjectPtr.h b/Engine/Source/Runtime/CoreUObject/Public/UObject/ObjectPtr.h
--- a/Engine/Source/Runtime/CoreUObject/Public/UObject/ObjectPtr.h
+++ b/Engine/Source/Runtime/CoreUObject/Public/UObject/ObjectPtr.h
@@ -25,6 +25,17 @@
  */
 #define UE_TRANSITIONAL_OBJECT_PTR(Type) auto 
 
+#if PLATFORM_MICROSOFT && defined(_MSC_EXTENSIONS)
+	/**
+	 * Non-conformance mode in MSVC has issues where the presence of a conversion operator to bool (even an explicit one)
+	 * leads to operator ambiguity for operator== and operator!= with a NULL (not nullptr).  This macro controls the presence
+	 * of code to deal with this issue.
+	 */
+	#define UE_OBJECT_PTR_NONCONFORMANCE_SUPPORT 1
+#else
+	#define UE_OBJECT_PTR_NONCONFORMANCE_SUPPORT 0
+#endif
+
 template <typename T>
 struct TObjectPtr;
 
@@ -123,6 +134,7 @@ public:
 	FORCEINLINE UObject& operator*() const { return *Get(); }
 
 	FORCEINLINE bool IsNull() const { return IsObjectHandleNull(Handle); }
+	FORCEINLINE bool operator!() const { return IsNull(); }
 	explicit FORCEINLINE operator bool() const { return !IsNull(); }
 	FORCEINLINE FObjectHandle GetHandle() const { return Handle; }
 	FORCEINLINE FObjectHandle& GetHandleRef() const { return Handle; }
@@ -283,6 +295,7 @@ public:
 	explicit FORCEINLINE operator T*& () { return GetInternalRef(); }
 
 	using FObjectPtr::IsNull;
+	using FObjectPtr::operator!;
 	using FObjectPtr::operator bool;
 	using FObjectPtr::GetHandle;
 
@@ -348,6 +361,20 @@ namespace ObjectPtr_Private
 		return Other;
 	}
 
+#if UE_OBJECT_PTR_NONCONFORMANCE_SUPPORT
+	/**
+	 * Force acceptance of explicitly TYPE_OF_NULL to avoid ambiguity issues triggered by the presence of a type conversion operator to bool.
+	 * This has the negative consequence of making it possible to coerce an arbitrary integer for the purpose of use in the comparison operators.
+	 */
+	template <typename T>
+	UE_OBJPTR_DEPRECATED(5.0, "Coercing a NULL for operations with a TObjectPtr is deprecated when running in a non-standards conforming compiler mode.")
+	constexpr const T* CoerceToPointer(TYPE_OF_NULL Other)
+	{
+		checkfSlow(Other == 0, TEXT("TObjectPtr cannot be compared to a non-zero NULL type value."));
+		return nullptr;
+	}
+#endif
+
 	/** Coerce to pointer through implicit conversion to CommonPointerType where CommonPointerType is deduced, and must be a C++ pointer, not a wrapper type. */
 	template <
 		typename T,
