Fix circular dependency load deferral for wrapped object pointer fields by ensuring they match behavior for standard raw pointers if they are not lazy loaded.

diff --git a/Engine/Source/Runtime/CoreUObject/Private/UObject/PropertyObjectPtr.cpp b/Engine/Source/Runtime/CoreUObject/Private/UObject/PropertyObjectPtr.cpp
--- a/Engine/Source/Runtime/CoreUObject/Private/UObject/PropertyObjectPtr.cpp
+++ b/Engine/Source/Runtime/CoreUObject/Private/UObject/PropertyObjectPtr.cpp
@@ -4,6 +4,8 @@
 #include "UObject/ObjectMacros.h"
 #include "UObject/PropertyHelper.h"
 #include "UObject/UnrealType.h"
+#include "UObject/LinkerPlaceholderExportObject.h"
+#include "UObject/LinkerPlaceholderClass.h"
 
 /*-----------------------------------------------------------------------------
 	FObjectPtrProperty.
@@ -23,10 +25,51 @@ FString FObjectPtrProperty::GetCPPMacroType(FString& ExtendedTypeText) const
 void FObjectPtrProperty::SerializeItem(FStructuredArchive::FSlot Slot, void* Value, void const* Defaults) const
 {
 	FArchive& UnderlyingArchive = Slot.GetUnderlyingArchive();
-
 	FObjectPtr* ObjectPtr = (FObjectPtr*)GetPropertyValuePtr(Value);
 
-	Slot << *ObjectPtr;
+	if (UnderlyingArchive.IsObjectReferenceCollector())
+	{
+		Slot << *ObjectPtr;
+
+		if(!UnderlyingArchive.IsSaving() && IsObjectHandleResolved(ObjectPtr->GetHandle()))
+		{
+			CheckValidObject(ObjectPtr);
+		}
+	}
+	else
+	{
+		FObjectHandle OriginalHandle = ObjectPtr->GetHandle();
+		Slot << *ObjectPtr;
+
+		FObjectHandle CurrentHandle = ObjectPtr->GetHandle();
+		if ((OriginalHandle != CurrentHandle) && IsObjectHandleResolved(CurrentHandle))
+		{
+			UObject* ResolvedObject = ObjectPtr->Get();
+	#if USE_CIRCULAR_DEPENDENCY_LOAD_DEFERRING
+			if (ULinkerPlaceholderExportObject* PlaceholderVal = Cast<ULinkerPlaceholderExportObject>(ResolvedObject))
+			{
+				PlaceholderVal->AddReferencingPropertyValue(this, Value);
+			}
+			else if (ULinkerPlaceholderClass* PlaceholderClass = Cast<ULinkerPlaceholderClass>(ResolvedObject))
+			{
+				PlaceholderClass->AddReferencingPropertyValue(this, Value);
+			}
+			// NOTE: we don't remove this from CurrentValue if it is a 
+			//       ULinkerPlaceholderExportObject; this is because this property 
+			//       could be an array inner, and another member of that array (also 
+			//       referenced through this property)... if this becomes a problem,
+			//       then we could inc/decrement a ref count per referencing property 
+			//
+			// @TODO: if this becomes problematic (because ObjectValue doesn't match 
+			//        this property's PropertyClass), then we could spawn another
+			//        placeholder object (of PropertyClass's type), or use null; but
+			//        we'd have to modify ULinkerPlaceholderExportObject::ReplaceReferencingObjectValues()
+			//        to accommodate this (as it depends on finding itself as the set value)
+	#endif // USE_CIRCULAR_DEPENDENCY_LOAD_DEFERRING
+
+			CheckValidObject(Value);
+		}
+	}
 }
 
 bool FObjectPtrProperty::SameType(const FProperty* Other) const
