Remove the ability to implicitly convert or reinterpret-in-place a TSet or TMap of wrapped pointers as raw pointers, and fix-up all code that was depending on this compatibility API.  This ability can't be permitted when we enable lazy load (UE_WITH_OBJECT_HANDLE_LATE_RESOLVE is non-zero) because the hash order of a wrapped pointer differs from a raw pointer, this means that reinterpreting a container of one type to a container of another type means the reinterpreted container isn't in hash order which is invalid.

diff --git a/Engine/Plugins/Enterprise/DatasmithContent/Source/DatasmithContent/Private/ObjectTemplates/DatasmithObjectTemplate.cpp b/Engine/Plugins/Enterprise/DatasmithContent/Source/DatasmithContent/Private/ObjectTemplates/DatasmithObjectTemplate.cpp
--- a/Engine/Plugins/Enterprise/DatasmithContent/Source/DatasmithContent/Private/ObjectTemplates/DatasmithObjectTemplate.cpp
+++ b/Engine/Plugins/Enterprise/DatasmithContent/Source/DatasmithContent/Private/ObjectTemplates/DatasmithObjectTemplate.cpp
@@ -71,12 +71,12 @@ UDatasmithAssetUserData* FDatasmithObjectTemplateUtils::FindOrCreateDatasmithUse
 	return UserData;
 }
 
-TMap< TSubclassOf< UDatasmithObjectTemplate >, UDatasmithObjectTemplate* >* FDatasmithObjectTemplateUtils::FindOrCreateObjectTemplates(UObject* Outer)
+TMap< TSubclassOf< UDatasmithObjectTemplate >, TObjectPtr<UDatasmithObjectTemplate> >* FDatasmithObjectTemplateUtils::FindOrCreateObjectTemplates(UObject* Outer)
 {
 #if WITH_EDITORONLY_DATA
 	if (UDatasmithAssetUserData* UserAssetData = FindOrCreateDatasmithUserData(Outer))
 	{
-		return &ToRawPtrTMapUnsafe(UserAssetData->ObjectTemplates);
+		return &UserAssetData->ObjectTemplates;
 	}
 #endif // #if WITH_EDITORONLY_DATA
 return nullptr;
@@ -95,14 +95,14 @@ void UDatasmithObjectTemplate::Apply(UObject* Destination, bool bForce)
 UDatasmithObjectTemplate* FDatasmithObjectTemplateUtils::GetObjectTemplate(UObject* Outer, TSubclassOf< UDatasmithObjectTemplate > Subclass)
 {
 #if WITH_EDITORONLY_DATA
-	TMap< TSubclassOf< UDatasmithObjectTemplate >, UDatasmithObjectTemplate* >* ObjectTemplatesMap = FindOrCreateObjectTemplates(Outer);
+	TMap< TSubclassOf< UDatasmithObjectTemplate >, TObjectPtr<UDatasmithObjectTemplate> >* ObjectTemplatesMap = FindOrCreateObjectTemplates(Outer);
 
 	if (!ObjectTemplatesMap)
 	{
 		return nullptr;
 	}
 
-	UDatasmithObjectTemplate** ObjectTemplatePtr = ObjectTemplatesMap->Find(Subclass);
+	TObjectPtr<UDatasmithObjectTemplate>* ObjectTemplatePtr = ObjectTemplatesMap->Find(Subclass);
 
 	return ObjectTemplatePtr ? *ObjectTemplatePtr : nullptr;
 #else
@@ -115,7 +115,7 @@ void FDatasmithObjectTemplateUtils::SetObjectTemplate(UObject* Outer, UDatasmith
 #if WITH_EDITORONLY_DATA
 	if (UDatasmithAssetUserData* UserData = FindOrCreateDatasmithUserData(Outer))
 	{
-		TMap< TSubclassOf< UDatasmithObjectTemplate >, UDatasmithObjectTemplate* >& ObjectTemplatesMap = UserData->ObjectTemplates;
+		TMap< TSubclassOf< UDatasmithObjectTemplate >, TObjectPtr<UDatasmithObjectTemplate> >& ObjectTemplatesMap = UserData->ObjectTemplates;
 		ObjectTemplatesMap.FindOrAdd(ObjectTemplate->GetClass()) = ObjectTemplate;
 		ObjectTemplate->SetFlags(RF_Public);
 
diff --git a/Engine/Plugins/Enterprise/DatasmithContent/Source/DatasmithContent/Public/ObjectTemplates/DatasmithObjectTemplate.h b/Engine/Plugins/Enterprise/DatasmithContent/Source/DatasmithContent/Public/ObjectTemplates/DatasmithObjectTemplate.h
--- a/Engine/Plugins/Enterprise/DatasmithContent/Source/DatasmithContent/Public/ObjectTemplates/DatasmithObjectTemplate.h
+++ b/Engine/Plugins/Enterprise/DatasmithContent/Source/DatasmithContent/Public/ObjectTemplates/DatasmithObjectTemplate.h
@@ -107,7 +107,7 @@ struct DATASMITHCONTENT_API FDatasmithObjectTemplateUtils
 
 	static class UDatasmithAssetUserData* FindOrCreateDatasmithUserData( UObject* Outer );
 
-	static TMap< TSubclassOf< UDatasmithObjectTemplate >, UDatasmithObjectTemplate* >* FindOrCreateObjectTemplates( UObject* Outer );
+	static TMap< TSubclassOf< UDatasmithObjectTemplate >, TObjectPtr<UDatasmithObjectTemplate> >* FindOrCreateObjectTemplates( UObject* Outer );
 
 	static UDatasmithObjectTemplate* GetObjectTemplate( UObject* Outer, TSubclassOf< UDatasmithObjectTemplate > Subclass );
 
diff --git a/Engine/Plugins/Enterprise/DatasmithImporter/Source/DatasmithImporter/Private/DatasmithImporterModule.cpp b/Engine/Plugins/Enterprise/DatasmithImporter/Source/DatasmithImporter/Private/DatasmithImporterModule.cpp
--- a/Engine/Plugins/Enterprise/DatasmithImporter/Source/DatasmithImporter/Private/DatasmithImporterModule.cpp
+++ b/Engine/Plugins/Enterprise/DatasmithImporter/Source/DatasmithImporter/Private/DatasmithImporterModule.cpp
@@ -549,7 +549,7 @@ void FDatasmithImporterModule::ExecuteReimportDatasmithMaterials( TArray<FAssetD
 
 void FDatasmithImporterModule::DiffAgainstTemplates( UObject* Outer )
 {
-	TMap< TSubclassOf< UDatasmithObjectTemplate >, UDatasmithObjectTemplate* >* ObjectTemplates = FDatasmithObjectTemplateUtils::FindOrCreateObjectTemplates( Outer );
+	TMap< TSubclassOf< UDatasmithObjectTemplate >, TObjectPtr<UDatasmithObjectTemplate> >* ObjectTemplates = FDatasmithObjectTemplateUtils::FindOrCreateObjectTemplates( Outer );
 
 	if ( !ObjectTemplates )
 	{
@@ -579,7 +579,7 @@ void FDatasmithImporterModule::DiffAgainstTemplates( UObject* Outer )
 
 void FDatasmithImporterModule::ResetFromTemplates( UObject* Outer )
 {
-	if ( TMap< TSubclassOf< UDatasmithObjectTemplate >, UDatasmithObjectTemplate* >* ObjectTemplates = FDatasmithObjectTemplateUtils::FindOrCreateObjectTemplates( Outer ) )
+	if ( TMap< TSubclassOf< UDatasmithObjectTemplate >, TObjectPtr<UDatasmithObjectTemplate> >* ObjectTemplates = FDatasmithObjectTemplateUtils::FindOrCreateObjectTemplates( Outer ) )
 	{
 		for ( auto It = ObjectTemplates->CreateConstIterator(); It; ++It )
 		{
diff --git a/Engine/Plugins/Enterprise/DatasmithImporter/Source/DatasmithImporter/Private/Utility/DatasmithImporterImpl.cpp b/Engine/Plugins/Enterprise/DatasmithImporter/Source/DatasmithImporter/Private/Utility/DatasmithImporterImpl.cpp
--- a/Engine/Plugins/Enterprise/DatasmithImporter/Source/DatasmithImporter/Private/Utility/DatasmithImporterImpl.cpp
+++ b/Engine/Plugins/Enterprise/DatasmithImporter/Source/DatasmithImporter/Private/Utility/DatasmithImporterImpl.cpp
@@ -383,7 +383,7 @@ TArray< FDatasmithImporterImpl::FMigratedTemplatePairType > FDatasmithImporterIm
 		return Results;
 	}
 
-	TMap< TSubclassOf< UDatasmithObjectTemplate >, UDatasmithObjectTemplate* >* SourceTemplates = FDatasmithObjectTemplateUtils::FindOrCreateObjectTemplates( SourceObject );
+	TMap< TSubclassOf< UDatasmithObjectTemplate >, TObjectPtr<UDatasmithObjectTemplate> >* SourceTemplates = FDatasmithObjectTemplateUtils::FindOrCreateObjectTemplates( SourceObject );
 
 	if ( !SourceTemplates )
 	{
diff --git a/Engine/Plugins/FX/Niagara/Source/Niagara/Classes/NiagaraEmitter.h b/Engine/Plugins/FX/Niagara/Source/Niagara/Classes/NiagaraEmitter.h
--- a/Engine/Plugins/FX/Niagara/Source/Niagara/Classes/NiagaraEmitter.h
+++ b/Engine/Plugins/FX/Niagara/Source/Niagara/Classes/NiagaraEmitter.h
@@ -552,7 +552,7 @@ public:
 	void OnScalabilityCVarChanged();
 
 #if WITH_EDITORONLY_DATA
-	NIAGARA_API const TMap<FGuid, UNiagaraMessageDataBase*>& GetMessages() const { return MessageKeyToMessageMap; };
+	NIAGARA_API const TMap<FGuid, TObjectPtr<UNiagaraMessageDataBase>>& GetMessages() const { return MessageKeyToMessageMap; };
 	NIAGARA_API void AddMessage(const FGuid& MessageKey, UNiagaraMessageDataBase* NewMessage) { MessageKeyToMessageMap.Add(MessageKey, NewMessage); };
 	NIAGARA_API void RemoveMessage(const FGuid& MessageKey) { MessageKeyToMessageMap.Remove(MessageKey); };
 	void RemoveMessageDelegateable(const FGuid MessageKey) { MessageKeyToMessageMap.Remove(MessageKey); };
diff --git a/Engine/Plugins/FX/Niagara/Source/Niagara/Classes/NiagaraSystem.h b/Engine/Plugins/FX/Niagara/Source/Niagara/Classes/NiagaraSystem.h
--- a/Engine/Plugins/FX/Niagara/Source/Niagara/Classes/NiagaraSystem.h
+++ b/Engine/Plugins/FX/Niagara/Source/Niagara/Classes/NiagaraSystem.h
@@ -498,7 +498,7 @@ public:
 	bool CanObtainUserVariable(const FNiagaraVariableBase& InVar) const;
 
 #if WITH_EDITORONLY_DATA
-	const TMap<FGuid, UNiagaraMessageDataBase*>& GetMessages() const { return MessageKeyToMessageMap; };
+	const TMap<FGuid, TObjectPtr<UNiagaraMessageDataBase>>& GetMessages() const { return MessageKeyToMessageMap; };
 	void AddMessage(const FGuid& MessageKey, UNiagaraMessageDataBase* NewMessage) { MessageKeyToMessageMap.Add(MessageKey, NewMessage); };
 	void RemoveMessage(const FGuid& MessageKey) { MessageKeyToMessageMap.Remove(MessageKey); };
 	void RemoveMessageDelegateable(const FGuid MessageKey) { MessageKeyToMessageMap.Remove(MessageKey); };
diff --git a/Engine/Plugins/FX/Niagara/Source/NiagaraEditor/Private/NiagaraGraph.cpp b/Engine/Plugins/FX/Niagara/Source/NiagaraEditor/Private/NiagaraGraph.cpp
--- a/Engine/Plugins/FX/Niagara/Source/NiagaraEditor/Private/NiagaraGraph.cpp
+++ b/Engine/Plugins/FX/Niagara/Source/NiagaraEditor/Private/NiagaraGraph.cpp
@@ -767,7 +767,7 @@ void UNiagaraGraph::StandardizeParameterNames()
 
 	// Since we'll be modifying the keys, make a copy of the map and then clear the original so it cal be 
 	// repopulated.
-	TMap<FNiagaraVariable, UNiagaraScriptVariable*> OldVariableToScriptVariable = VariableToScriptVariable;
+	TMap<FNiagaraVariable, TObjectPtr<UNiagaraScriptVariable>> OldVariableToScriptVariable = VariableToScriptVariable;
 	VariableToScriptVariable.Empty();
 	for (TPair<FNiagaraVariable, UNiagaraScriptVariable*> VariableScriptVariablePair : OldVariableToScriptVariable)
 	{
@@ -1297,12 +1297,12 @@ void UNiagaraGraph::GetParameters(TArray<FNiagaraVariable>& Inputs, TArray<FNiag
 // 	Outputs.Sort(SortVars);
 }
 
-const TMap<FNiagaraVariable, UNiagaraScriptVariable*>& UNiagaraGraph::GetAllMetaData() const
+const TMap<FNiagaraVariable, TObjectPtr<UNiagaraScriptVariable>>& UNiagaraGraph::GetAllMetaData() const
 {
 	return VariableToScriptVariable;
 }
 
-TMap<FNiagaraVariable, UNiagaraScriptVariable*>& UNiagaraGraph::GetAllMetaData()
+TMap<FNiagaraVariable, TObjectPtr<UNiagaraScriptVariable>>& UNiagaraGraph::GetAllMetaData()
 {
 	return VariableToScriptVariable;
 }
@@ -1767,7 +1767,7 @@ bool UNiagaraGraph::RenameParameter(const FNiagaraVariable& Parameter, FName New
 
 void UNiagaraGraph::ScriptVariableChanged(FNiagaraVariable Variable)
 {
-	UNiagaraScriptVariable** ScriptVariable = GetAllMetaData().Find(Variable);
+	TObjectPtr<UNiagaraScriptVariable>* ScriptVariable = GetAllMetaData().Find(Variable);
 	if (!ScriptVariable || !*ScriptVariable || (*ScriptVariable)->Metadata.GetIsStaticSwitch()) {
 		return;
 	}
@@ -1943,8 +1943,8 @@ bool UNiagaraGraph::AppendCompileHash(FNiagaraCompileHashVisitor* InVisitor, con
 	}
 
 	// We need to sort the variables in a stable manner.
-	TArray<UNiagaraScriptVariable*> Values;
-	ToRawPtrTMapUnsafe(VariableToScriptVariable).GenerateValueArray(Values);
+	TArray<TObjectPtr<UNiagaraScriptVariable>> Values;
+	VariableToScriptVariable.GenerateValueArray(Values);
 	Values.Remove(nullptr);
 	Values.Sort([&](const UNiagaraScriptVariable& A, const UNiagaraScriptVariable& B)
 	{
diff --git a/Engine/Plugins/FX/Niagara/Source/NiagaraEditor/Private/NiagaraNodeParameterMapBase.cpp b/Engine/Plugins/FX/Niagara/Source/NiagaraEditor/Private/NiagaraNodeParameterMapBase.cpp
--- a/Engine/Plugins/FX/Niagara/Source/NiagaraEditor/Private/NiagaraNodeParameterMapBase.cpp
+++ b/Engine/Plugins/FX/Niagara/Source/NiagaraEditor/Private/NiagaraNodeParameterMapBase.cpp
@@ -402,7 +402,7 @@ void UNiagaraNodeParameterMapBase::SelectParameterFromPin(const UEdGraphPin* InP
 
 			FNiagaraTypeDefinition TypeDef = Schema->PinToTypeDefinition(InPin);
 			FNiagaraVariable PinVariable = FNiagaraVariable(TypeDef, InPin->PinName);
-			UNiagaraScriptVariable** PinAssociatedScriptVariable = NiagaraGraph->GetAllMetaData().Find(PinVariable);
+			TObjectPtr<UNiagaraScriptVariable>* PinAssociatedScriptVariable = NiagaraGraph->GetAllMetaData().Find(PinVariable);
 			if (PinAssociatedScriptVariable != nullptr)
 			{
 				NiagaraGraph->OnSubObjectSelectionChanged().Broadcast(*PinAssociatedScriptVariable);
diff --git a/Engine/Plugins/FX/Niagara/Source/NiagaraEditor/Private/NiagaraNodeWithDynamicPins.cpp b/Engine/Plugins/FX/Niagara/Source/NiagaraEditor/Private/NiagaraNodeWithDynamicPins.cpp
--- a/Engine/Plugins/FX/Niagara/Source/NiagaraEditor/Private/NiagaraNodeWithDynamicPins.cpp
+++ b/Engine/Plugins/FX/Niagara/Source/NiagaraEditor/Private/NiagaraNodeWithDynamicPins.cpp
@@ -125,7 +125,7 @@ void UNiagaraNodeWithDynamicPins::UpdateAddedPinMetaData(const UEdGraphPin* Adde
 		const UEdGraphSchema_Niagara* Schema = GetDefault<UEdGraphSchema_Niagara>();
 		FNiagaraVariable PinVariable = Schema->PinToNiagaraVariable(AddedPin, false);
 		
-		if (UNiagaraScriptVariable** ScriptVariable = Graph->GetAllMetaData().Find(PinVariable))
+		if (TObjectPtr<UNiagaraScriptVariable>* ScriptVariable = Graph->GetAllMetaData().Find(PinVariable))
 		{
 			Graph->UpdateUsageForScriptVariable(*ScriptVariable);
 		}
@@ -324,7 +324,7 @@ void UNiagaraNodeWithDynamicPins::AddParameter(FNiagaraVariable Parameter, const
 		// Resolve the unique parameter name before adding to the graph as the pin needs to be created first to resolve the parameter metadata usage.
 		if (FNiagaraConstants::FindEngineConstant(Parameter) == nullptr)
 		{
-			UNiagaraScriptVariable** FoundScriptVariable = Graph->GetAllMetaData().Find(Parameter);
+			TObjectPtr<UNiagaraScriptVariable>* FoundScriptVariable = Graph->GetAllMetaData().Find(Parameter);
 			if (!FoundScriptVariable)
 			{
 				Parameter.SetName(Graph->MakeUniqueParameterName(Parameter.GetName()));
@@ -387,7 +387,7 @@ void UNiagaraNodeWithDynamicPins::RemoveDynamicPin(UEdGraphPin* Pin)
 					return;
 				}
 
-				UNiagaraScriptVariable** PinAssociatedScriptVariable = Graph->GetAllMetaData().Find(PinVariable);
+				TObjectPtr<UNiagaraScriptVariable>* PinAssociatedScriptVariable = Graph->GetAllMetaData().Find(PinVariable);
 				if (PinAssociatedScriptVariable != nullptr)
 				{
 					bool bShouldRemoveScriptVariable = !Graph->UpdateUsageForScriptVariable(*PinAssociatedScriptVariable);
diff --git a/Engine/Plugins/FX/Niagara/Source/NiagaraEditor/Private/ViewModels/NiagaraParameterPanelViewModel.cpp b/Engine/Plugins/FX/Niagara/Source/NiagaraEditor/Private/ViewModels/NiagaraParameterPanelViewModel.cpp
--- a/Engine/Plugins/FX/Niagara/Source/NiagaraEditor/Private/ViewModels/NiagaraParameterPanelViewModel.cpp
+++ b/Engine/Plugins/FX/Niagara/Source/NiagaraEditor/Private/ViewModels/NiagaraParameterPanelViewModel.cpp
@@ -515,7 +515,7 @@ const TArray<FNiagaraScriptVariableAndViewInfo> FNiagaraSystemToolkitParameterPa
 	{
 		if (ensureMsgf(Graph.IsValid(), TEXT("Invalid Graph visited when trying to get viewed parameters for system toolkit parameter panel!")))
 		{
-			const TMap<FNiagaraVariable, UNiagaraScriptVariable*>& GraphVarToScriptVarMap = Graph->GetAllMetaData();
+			const TMap<FNiagaraVariable, TObjectPtr<UNiagaraScriptVariable>>& GraphVarToScriptVarMap = Graph->GetAllMetaData();
 			for (auto Iter = GraphVarToScriptVarMap.CreateConstIterator(); Iter; ++Iter)
 			{
 				const FNiagaraVariableMetaData& MetaData = Iter.Value()->Metadata;
@@ -831,7 +831,7 @@ bool FNiagaraScriptToolkitParameterPanelViewModel::GetCanRenameParameterAndToolT
 void FNiagaraScriptToolkitParameterPanelViewModel::HandleActionSelected(const TSharedPtr<FEdGraphSchemaAction>& InAction, ESelectInfo::Type InSelectionType)
 {
 	const FNiagaraScriptVarAndViewInfoAction* Action = static_cast<FNiagaraScriptVarAndViewInfoAction*>(InAction.Get());
-	UNiagaraScriptVariable** ScriptVarPtr = ScriptViewModel->GetGraphViewModel()->GetGraph()->GetAllMetaData().Find(Action->ScriptVariableAndViewInfo.ScriptVariable);
+	TObjectPtr<UNiagaraScriptVariable>* ScriptVarPtr = ScriptViewModel->GetGraphViewModel()->GetGraph()->GetAllMetaData().Find(Action->ScriptVariableAndViewInfo.ScriptVariable);
 	if (ensureMsgf(ScriptVarPtr != nullptr, TEXT("Failed to get UNiagaraScriptVariable from selected action!")))
 	{
 		VariableObjectSelection->SetSelectedObject(*ScriptVarPtr);
@@ -965,7 +965,7 @@ TSharedRef<SWidget> FNiagaraScriptToolkitParameterPanelViewModel::GetScriptParam
 	else
 	{
 		// Failed to find the parameter name in the cache, try to find the variable in the graph script variables and generate view info.
-		const UNiagaraScriptVariable* const* ScriptVarPtr = ScriptViewModel->GetGraphViewModel()->GetGraph()->GetAllMetaData().Find(PinVar);
+		TObjectPtr<UNiagaraScriptVariable> const* ScriptVarPtr = ScriptViewModel->GetGraphViewModel()->GetGraph()->GetAllMetaData().Find(PinVar);
 		if (ScriptVarPtr != nullptr)
 		{
 			const UNiagaraScriptVariable* ScriptVar = *ScriptVarPtr;
@@ -1058,7 +1058,7 @@ const TArray<FNiagaraScriptVariableAndViewInfo> FNiagaraScriptToolkitParameterPa
 	UNiagaraGraph* ViewedGraph = ScriptViewModel->GetGraphViewModel()->GetGraph();
 	if (ensureMsgf(ViewedGraph != nullptr, TEXT("Invalid Graph found when trying to get viewed parameters for script toolkit parameter panel!")))
 	{
-		const TMap<FNiagaraVariable, UNiagaraScriptVariable*>& GraphVarToScriptVarMap = ViewedGraph->GetAllMetaData();
+		const TMap<FNiagaraVariable, TObjectPtr<UNiagaraScriptVariable>>& GraphVarToScriptVarMap = ViewedGraph->GetAllMetaData();
 		for (auto Iter = GraphVarToScriptVarMap.CreateConstIterator(); Iter; ++Iter)
 		{
 			const FNiagaraVariable& Variable = Iter.Value()->Variable;
diff --git a/Engine/Plugins/FX/Niagara/Source/NiagaraEditor/Private/ViewModels/Stack/NiagaraStackModuleItem.cpp b/Engine/Plugins/FX/Niagara/Source/NiagaraEditor/Private/ViewModels/Stack/NiagaraStackModuleItem.cpp
--- a/Engine/Plugins/FX/Niagara/Source/NiagaraEditor/Private/ViewModels/Stack/NiagaraStackModuleItem.cpp
+++ b/Engine/Plugins/FX/Niagara/Source/NiagaraEditor/Private/ViewModels/Stack/NiagaraStackModuleItem.cpp
@@ -145,7 +145,7 @@ void UNiagaraStackModuleItem::RefreshChildrenInternal(const TArray<UNiagaraStack
 		{
 			UNiagaraScriptSource* Source = CastChecked<UNiagaraScriptSource>(FunctionCallNode->FunctionScript->GetSource());
 			UNiagaraGraph* Graph = CastChecked<UNiagaraGraph>(Source->NodeGraph);
-			const TMap<FNiagaraVariable, UNiagaraScriptVariable*>& MetaDataMap = Graph->GetAllMetaData();
+			const TMap<FNiagaraVariable, TObjectPtr<UNiagaraScriptVariable>>& MetaDataMap = Graph->GetAllMetaData();
 			auto Iter = MetaDataMap.CreateConstIterator();
 			while (Iter)
 			{
diff --git a/Engine/Plugins/FX/Niagara/Source/NiagaraEditor/Private/Widgets/SNiagaraGraphParameterMapGetNode.cpp b/Engine/Plugins/FX/Niagara/Source/NiagaraEditor/Private/Widgets/SNiagaraGraphParameterMapGetNode.cpp
--- a/Engine/Plugins/FX/Niagara/Source/NiagaraEditor/Private/Widgets/SNiagaraGraphParameterMapGetNode.cpp
+++ b/Engine/Plugins/FX/Niagara/Source/NiagaraEditor/Private/Widgets/SNiagaraGraphParameterMapGetNode.cpp
@@ -185,7 +185,7 @@ FReply SNiagaraGraphParameterMapGetNode::OnBorderMouseButtonDown(const FGeometry
 				if (Schema)
 				{
 					FNiagaraVariable Var = Schema->PinToNiagaraVariable(OutputPins[InWhichPin-1]->GetPinObj());
-					UNiagaraScriptVariable** PinAssociatedScriptVariable = Graph->GetAllMetaData().Find(Var);
+					TObjectPtr<UNiagaraScriptVariable>* PinAssociatedScriptVariable = Graph->GetAllMetaData().Find(Var);
 					if (PinAssociatedScriptVariable != nullptr)
 					{
 						Graph->OnSubObjectSelectionChanged().Broadcast(*PinAssociatedScriptVariable);
diff --git a/Engine/Plugins/FX/Niagara/Source/NiagaraEditor/Private/Widgets/SNiagaraParameterPanel.cpp b/Engine/Plugins/FX/Niagara/Source/NiagaraEditor/Private/Widgets/SNiagaraParameterPanel.cpp
--- a/Engine/Plugins/FX/Niagara/Source/NiagaraEditor/Private/Widgets/SNiagaraParameterPanel.cpp
+++ b/Engine/Plugins/FX/Niagara/Source/NiagaraEditor/Private/Widgets/SNiagaraParameterPanel.cpp
@@ -627,10 +627,10 @@ void SNiagaraAddParameterMenu2::CollectAllActions(FGraphActionListBuilderBase& O
 	{
 		for (TWeakObjectPtr<UNiagaraGraph>& Graph : Graphs)
 		{
-			TMap<FNiagaraVariable, UNiagaraScriptVariable*> GraphParameters = Graph.Get()->GetAllMetaData(); //@todo(ng) try an array view so we can sort it without copying
+			TMap<FNiagaraVariable, TObjectPtr<UNiagaraScriptVariable>> GraphParameters = Graph.Get()->GetAllMetaData(); //@todo(ng) try an array view so we can sort it without copying
 			GraphParameters.KeySort([](const FNiagaraVariable& A, const FNiagaraVariable& B) { return A.GetName().LexicalLess(B.GetName()); }); //@todo(ng) sort by scope
 
-			for (TTuple<FNiagaraVariable, UNiagaraScriptVariable*>& GraphParameter : GraphParameters)
+			for (TTuple<FNiagaraVariable, TObjectPtr<UNiagaraScriptVariable>>& GraphParameter : GraphParameters)
 			{
 				FNiagaraVariable Variable = GraphParameter.Key;
 				const FText Category = LOCTEXT("NiagaraAddExistingParameterMenu", "Add Existing Parameter");
diff --git a/Engine/Plugins/FX/Niagara/Source/NiagaraEditor/Public/NiagaraGraph.h b/Engine/Plugins/FX/Niagara/Source/NiagaraEditor/Public/NiagaraGraph.h
--- a/Engine/Plugins/FX/Niagara/Source/NiagaraEditor/Public/NiagaraGraph.h
+++ b/Engine/Plugins/FX/Niagara/Source/NiagaraEditor/Public/NiagaraGraph.h
@@ -297,8 +297,8 @@ class UNiagaraGraph : public UEdGraph
 	void SetPerScriptMetaData(const FNiagaraVariable& InVar, const FNiagaraVariableMetaData& InMetaData);
 
 
-	const TMap<FNiagaraVariable, UNiagaraScriptVariable*>& GetAllMetaData() const;
-	TMap<FNiagaraVariable, UNiagaraScriptVariable*>& GetAllMetaData();
+	const TMap<FNiagaraVariable, TObjectPtr<UNiagaraScriptVariable>>& GetAllMetaData() const;
+	TMap<FNiagaraVariable, TObjectPtr<UNiagaraScriptVariable>>& GetAllMetaData();
 
 	const TMap<FNiagaraVariable, FNiagaraGraphParameterReferenceCollection>& GetParameterReferenceMap() const; // NOTE: The const is a lie! (This indirectly calls RefreshParameterReferences, which can recreate the entire map)
 
diff --git a/Engine/Plugins/FX/Niagara/Source/NiagaraEditor/Public/NiagaraNodeFunctionCall.h b/Engine/Plugins/FX/Niagara/Source/NiagaraEditor/Public/NiagaraNodeFunctionCall.h
--- a/Engine/Plugins/FX/Niagara/Source/NiagaraEditor/Public/NiagaraNodeFunctionCall.h
+++ b/Engine/Plugins/FX/Niagara/Source/NiagaraEditor/Public/NiagaraNodeFunctionCall.h
@@ -136,7 +136,7 @@ public:
 
 	virtual TSharedPtr<SGraphNode> CreateVisualWidget() override;
 
-	NIAGARAEDITOR_API const TMap<FGuid, UNiagaraMessageData*>& GetMessages() const { return MessageKeyToMessageMap; };
+	NIAGARAEDITOR_API const TMap<FGuid, TObjectPtr<UNiagaraMessageData>>& GetMessages() const { return MessageKeyToMessageMap; };
 	NIAGARAEDITOR_API void AddMessage(const FGuid& MessageKey, UNiagaraMessageData* NewMessage) { MessageKeyToMessageMap.Add(MessageKey, NewMessage); };
 	NIAGARAEDITOR_API void RemoveMessage(const FGuid& MessageKey) { MessageKeyToMessageMap.Remove(MessageKey); };
 	void RemoveMessageDelegateable(const FGuid MessageKey) { MessageKeyToMessageMap.Remove(MessageKey); };
diff --git a/Engine/Source/Runtime/ClothingSystemRuntimeCommon/Private/ClothPhysicalMeshData.cpp b/Engine/Source/Runtime/ClothingSystemRuntimeCommon/Private/ClothPhysicalMeshData.cpp
--- a/Engine/Source/Runtime/ClothingSystemRuntimeCommon/Private/ClothPhysicalMeshData.cpp
+++ b/Engine/Source/Runtime/ClothingSystemRuntimeCommon/Private/ClothPhysicalMeshData.cpp
@@ -83,11 +83,11 @@ void FClothPhysicalMeshData::ClearWeightMaps()
 	AddWeightMap(EWeightMapTargetCommon::AnimDriveMultiplier);
 }
 
-void FClothPhysicalMeshData::BuildSelfCollisionData(const TMap<FName, UClothConfigBase*>& ClothConfigs)
+void FClothPhysicalMeshData::BuildSelfCollisionData(const TMap<FName, TObjectPtr<UClothConfigBase>>& ClothConfigs)
 {
 #if WITH_APEX_CLOTHING  // Only apex clothing needs to build the SelfCollisionIndices
 	float SCRadius = 0.f;
-	for (const TPair<FName, UClothConfigBase*>& ClothConfig : ClothConfigs)
+	for (const TPair<FName, TObjectPtr<UClothConfigBase>>& ClothConfig : ClothConfigs)
 	{
 		SCRadius = ClothConfig.Value->NeedsSelfCollisionIndices();
 		if (SCRadius > 0.f)
diff --git a/Engine/Source/Runtime/ClothingSystemRuntimeCommon/Public/ClothPhysicalMeshData.h b/Engine/Source/Runtime/ClothingSystemRuntimeCommon/Public/ClothPhysicalMeshData.h
--- a/Engine/Source/Runtime/ClothingSystemRuntimeCommon/Public/ClothPhysicalMeshData.h
+++ b/Engine/Source/Runtime/ClothingSystemRuntimeCommon/Public/ClothPhysicalMeshData.h
@@ -33,7 +33,7 @@ struct CLOTHINGSYSTEMRUNTIMECOMMON_API FClothPhysicalMeshData
 	void ClearWeightMaps();
 
 	/** Build the self collision indices for the relevant config. */
-	void BuildSelfCollisionData(const TMap<FName, UClothConfigBase*>& ClothConfigs);
+	void BuildSelfCollisionData(const TMap<FName, TObjectPtr<UClothConfigBase>>& ClothConfigs);
 
 	/** Retrieve whether a vertex weight array has already been registered. */
 	template<typename T>
diff --git a/Engine/Source/Runtime/Core/Public/Containers/Map.h b/Engine/Source/Runtime/Core/Public/Containers/Map.h
--- a/Engine/Source/Runtime/Core/Public/Containers/Map.h
+++ b/Engine/Source/Runtime/Core/Public/Containers/Map.h
@@ -1253,46 +1253,6 @@ public:
 
 	FORCEINLINE       ValueType& operator[](KeyConstPointerType Key) { return this->FindChecked(Key); }
 	FORCEINLINE const ValueType& operator[](KeyConstPointerType Key) const { return this->FindChecked(Key); }
-
-	/** Implicit conversion operator to container of compatible key and value element type. */
-	template <
-		typename KeyTypeAlias = KeyType,
-		typename ValueTypeAlias = ValueType,
-		typename std::enable_if_t<
-			TIsContainerElementTypeReinterpretable<KeyTypeAlias>::Value
-			||
-			TIsContainerElementTypeReinterpretable<ValueTypeAlias>::Value
-		>* = nullptr
-	>
-	operator TMap<typename TContainerElementTypeCompatibility<KeyTypeAlias>::ReinterpretType, typename TContainerElementTypeCompatibility<ValueTypeAlias>::ReinterpretType>& ()
-	{
-		using KeyCompat = TContainerElementTypeCompatibility<KeyType>;
-		using ValueCompat = TContainerElementTypeCompatibility<ValueType>;
-		KeyCompat::ReinterpretRange(Super::Pairs.begin(), Super::Pairs.end(), [](typename Super::ElementSetType::TRangedForIterator& It) -> KeyType& { return It->Key; });
-		ValueCompat::ReinterpretRange(Super::Pairs.begin(), Super::Pairs.end(), [](typename Super::ElementSetType::TRangedForIterator& It) -> ValueType& { return It->Value; });
-
-		return *reinterpret_cast<TMap<typename KeyCompat::ReinterpretType, typename ValueCompat::ReinterpretType>*>(this);
-	}
-
-	/** Implicit conversion operator to constant container of compatible key and value element type. */
-	template <
-		typename KeyTypeAlias = KeyType,
-		typename ValueTypeAlias = ValueType,
-		typename std::enable_if_t<
-			TIsContainerElementTypeReinterpretable<KeyTypeAlias>::Value
-			||
-			TIsContainerElementTypeReinterpretable<ValueTypeAlias>::Value
-		>* = nullptr
-	>
-	operator const TMap<typename TContainerElementTypeCompatibility<KeyTypeAlias>::ReinterpretType, typename TContainerElementTypeCompatibility<ValueTypeAlias>::ReinterpretType>& () const
-	{
-		using KeyCompat = TContainerElementTypeCompatibility<KeyType>;
-		using ValueCompat = TContainerElementTypeCompatibility<ValueType>;
-		KeyCompat::ReinterpretRange(Super::Pairs.begin(), Super::Pairs.end(), [](typename Super::ElementSetType::TRangedForConstIterator& It) -> const KeyType& { return It->Key; });
-		ValueCompat::ReinterpretRange(Super::Pairs.begin(), Super::Pairs.end(), [](typename Super::ElementSetType::TRangedForConstIterator& It) -> const ValueType& { return It->Value; });
-
-		return *reinterpret_cast<const TMap<typename KeyCompat::ReinterpretType, typename ValueCompat::ReinterpretType>*>(this);
-	}
 };
 
 namespace Freeze
diff --git a/Engine/Source/Runtime/Core/Public/Containers/Set.h b/Engine/Source/Runtime/Core/Public/Containers/Set.h
--- a/Engine/Source/Runtime/Core/Public/Containers/Set.h
+++ b/Engine/Source/Runtime/Core/Public/Containers/Set.h
@@ -1218,30 +1218,6 @@ public:
 		Elements.CheckAddress(Addr);
 	}
 
-	/** Implicit conversion operator to container of compatible element type. */
-	template <
-		typename AliasElementType = ElementType,
-		typename std::enable_if_t<TIsContainerElementTypeReinterpretable<AliasElementType>::Value>* = nullptr
-		>
-	operator TSet<typename TContainerElementTypeCompatibility<AliasElementType>::ReinterpretType>& ()
-	{
-		using ElementCompat = TContainerElementTypeCompatibility<ElementType>;
-		ElementCompat::ReinterpretRange(begin(), end());
-		return *reinterpret_cast<TSet<typename ElementCompat::ReinterpretType>*>(this);
-	}
-
-	/** Implicit conversion operator to constant container of compatible element type. */
-	template <
-		typename AliasElementType = ElementType,
-		typename std::enable_if_t<TIsContainerElementTypeReinterpretable<AliasElementType>::Value>* = nullptr
-		>
-	operator const TSet<typename TContainerElementTypeCompatibility<AliasElementType>::ReinterpretType>& () const
-	{
-		using ElementCompat = TContainerElementTypeCompatibility<ElementType>;
-		ElementCompat::ReinterpretRange(begin(), end());
-		return *reinterpret_cast<const TSet<typename ElementCompat::ReinterpretType>*>(this);
-	}
-
 	/**
 	 * Move assignment operator.
 	 * Compatible element type version.
diff --git a/Engine/Source/Runtime/CoreUObject/Public/UObject/ObjectPtr.h b/Engine/Source/Runtime/CoreUObject/Public/UObject/ObjectPtr.h
--- a/Engine/Source/Runtime/CoreUObject/Public/UObject/ObjectPtr.h
+++ b/Engine/Source/Runtime/CoreUObject/Public/UObject/ObjectPtr.h
@@ -532,186 +532,6 @@ decltype(auto) ToRawPtrTArrayUnsafe(ArrayType&& Array)
 
 	return (RefQualifiedNewArrayType&)Array;
 }
-
-template <
-	typename SetElementType,
-	typename Allocator,
-	typename RawPtrType = typename SetElementType::ElementType*,
-	std::enable_if_t<
-		TIsTObjectPtr<std::remove_cv_t<SetElementType>>::Value
-	>* = nullptr
->
-FORCEINLINE TSet<RawPtrType, DefaultKeyFuncs<RawPtrType>, Allocator>& ToRawPtrTSetUnsafe(TSet<SetElementType, DefaultKeyFuncs<SetElementType>, Allocator>& InSet)
-{
-	return InSet;
-}
-
-template <
-	typename SetElementType,
-	typename Allocator,
-	typename RawPtrType = typename SetElementType::ElementType*,
-	std::enable_if_t<
-		TIsTObjectPtr<std::remove_cv_t<SetElementType>>::Value
-	>* = nullptr
->
-FORCEINLINE const TSet<RawPtrType, DefaultKeyFuncs<RawPtrType>, Allocator>& ToRawPtrTSetUnsafe(const TSet<SetElementType, DefaultKeyFuncs<SetElementType>, Allocator>& InSet)
-{
-	return InSet;
-}
-
-template <
-	typename ElementType,
-	typename KeyFuncs,
-	typename Allocator,
-	std::enable_if_t<
-		!TIsTObjectPtr<std::remove_cv_t<ElementType>>::Value
-	>* = nullptr
->
-FORCEINLINE TSet<ElementType, KeyFuncs, Allocator>& ToRawPtrTSetUnsafe(TSet<ElementType, KeyFuncs, Allocator>& InSet)
-{
-	return InSet;
-}
-
-template <
-	typename ElementType,
-	typename KeyFuncs,
-	typename Allocator,
-	std::enable_if_t<
-		!TIsTObjectPtr<std::remove_cv_t<ElementType>>::Value
-	>* = nullptr
->
-FORCEINLINE const TSet<ElementType, KeyFuncs, Allocator>& ToRawPtrTSetUnsafe(const TSet<ElementType, KeyFuncs, Allocator>& InSet)
-{
-	return InSet;
-}
-
-template <
-	typename KeyType,
-	typename ValueType,
-	typename Allocator,
-	typename KeyRawPtrType = typename KeyType::ElementType*,
-	std::enable_if_t<
-		TIsTObjectPtr<std::remove_cv_t<KeyType>>::Value
-		&&
-		!TIsTObjectPtr<std::remove_cv_t<ValueType>>::Value
-	>* = nullptr
->
-FORCEINLINE TMap<KeyRawPtrType, ValueType, Allocator>& ToRawPtrTMapUnsafe(TMap<KeyType, ValueType, Allocator>& InMap)
-{
-	return InMap;
-}
-
-template <
-	typename KeyType,
-	typename ValueType,
-	typename Allocator,
-	typename ValueRawPtrType = typename ValueType::ElementType*,
-	std::enable_if_t<
-		!TIsTObjectPtr<std::remove_cv_t<KeyType>>::Value
-		&&
-		TIsTObjectPtr<std::remove_cv_t<ValueType>>::Value
-	>* = nullptr
->
-FORCEINLINE TMap<KeyType, ValueRawPtrType, Allocator>& ToRawPtrTMapUnsafe(TMap<KeyType, ValueType, Allocator>& InMap)
-{
-	return InMap;
-}
-
-template <
-	typename KeyType,
-	typename ValueType,
-	typename Allocator,
-	typename KeyRawPtrType = typename KeyType::ElementType*,
-	typename ValueRawPtrType = typename ValueType::ElementType*,
-	std::enable_if_t<
-		TIsTObjectPtr<std::remove_cv_t<KeyType>>::Value
-		&&
-		TIsTObjectPtr<std::remove_cv_t<ValueType>>::Value
-	>* = nullptr
->
-FORCEINLINE TMap<KeyRawPtrType, ValueRawPtrType, Allocator>& ToRawPtrTMapUnsafe(TMap<KeyType, ValueType, Allocator>& InMap)
-{
-	return InMap;
-}
-
-template <
-	typename KeyType,
-	typename ValueType,
-	typename Allocator,
-	typename KeyRawPtrType = typename KeyType::ElementType*,
-	std::enable_if_t<
-		TIsTObjectPtr<std::remove_cv_t<KeyType>>::Value
-		&&
-		!TIsTObjectPtr<std::remove_cv_t<ValueType>>::Value
-	>* = nullptr
->
-FORCEINLINE const TMap<KeyRawPtrType, ValueType, Allocator>& ToRawPtrTMapUnsafe(const TMap<KeyType, ValueType, Allocator>& InMap)
-{
-	return InMap;
-}
-
-template <
-	typename KeyType,
-	typename ValueType,
-	typename Allocator,
-	typename ValueRawPtrType = typename ValueType::ElementType*,
-	std::enable_if_t<
-		!TIsTObjectPtr<std::remove_cv_t<KeyType>>::Value
-		&&
-		TIsTObjectPtr<std::remove_cv_t<ValueType>>::Value
-	>* = nullptr
->
-FORCEINLINE const TMap<KeyType, ValueRawPtrType, Allocator>& ToRawPtrTMapUnsafe(const TMap<KeyType, ValueType, Allocator>& InMap)
-{
-	return InMap;
-}
-
-template <
-	typename KeyType,
-	typename ValueType,
-	typename Allocator,
-	typename KeyRawPtrType = typename KeyType::ElementType*,
-	typename ValueRawPtrType = typename ValueType::ElementType*,
-	std::enable_if_t<
-		TIsTObjectPtr<std::remove_cv_t<KeyType>>::Value
-		&&
-		TIsTObjectPtr<std::remove_cv_t<ValueType>>::Value
-	>* = nullptr
->
-FORCEINLINE const TMap<KeyRawPtrType, ValueRawPtrType, Allocator>& ToRawPtrTMapUnsafe(const TMap<KeyType, ValueType, Allocator>& InMap)
-{
-	return InMap;
-}
-
-template <
-	typename KeyType,
-	typename ValueType,
-	typename Allocator,
-	std::enable_if_t<
-		!TIsTObjectPtr<std::remove_cv_t<KeyType>>::Value
-		&&
-		!TIsTObjectPtr<std::remove_cv_t<ValueType>>::Value
-	>* = nullptr
->
-FORCEINLINE TMap<KeyType, ValueType, Allocator>& ToRawPtrTMapUnsafe(TMap<KeyType, ValueType, Allocator>& InMap)
-{
-	return InMap;
-}
-
-template <
-	typename KeyType,
-	typename ValueType,
-	typename Allocator,
-	std::enable_if_t<
-		!TIsTObjectPtr<std::remove_cv_t<KeyType>>::Value
-		&&
-		!TIsTObjectPtr<std::remove_cv_t<ValueType>>::Value
-	>* = nullptr
->
-FORCEINLINE const TMap<KeyType, ValueType, Allocator>& ToRawPtrTMapUnsafe(const TMap<KeyType, ValueType, Allocator>& InMap)
-{
-	return InMap;
-}
 PRAGMA_ENABLE_DEPRECATION_WARNINGS
 
 template <typename T>
diff --git a/Engine/Source/Runtime/Engine/Classes/Engine/World.h b/Engine/Source/Runtime/Engine/Classes/Engine/World.h
--- a/Engine/Source/Runtime/Engine/Classes/Engine/World.h
+++ b/Engine/Source/Runtime/Engine/Classes/Engine/World.h
@@ -731,7 +731,7 @@ struct ENGINE_API FLevelCollection
 	void SetDemoNetDriver(UDemoNetDriver* const InDemoNetDriver) { DemoNetDriver = InDemoNetDriver; }
 
 	/** Returns the set of levels in this collection. */
-	const TSet<ULevel*>& GetLevels() const { return Levels; }
+	const TSet<TObjectPtr<ULevel>>& GetLevels() const { return Levels; }
 
 	/** Adds a level to this collection and caches the collection pointer on the level for fast access. */
 	void AddLevel(ULevel* const Level);
diff --git a/Engine/Source/Runtime/Engine/Private/Components/SceneComponent.cpp b/Engine/Source/Runtime/Engine/Private/Components/SceneComponent.cpp
--- a/Engine/Source/Runtime/Engine/Private/Components/SceneComponent.cpp
+++ b/Engine/Source/Runtime/Engine/Private/Components/SceneComponent.cpp
@@ -2321,8 +2321,8 @@ void FSceneComponentInstanceData::AddReferencedObjects(FReferenceCollector& Coll
 
 void FSceneComponentInstanceData::FindAndReplaceInstances(const TMap<UObject*, UObject*>& OldToNewInstanceMap)
 {
-	TArray<USceneComponent*> SceneComponents;
-	ToRawPtrTMapUnsafe(AttachedInstanceComponents).GenerateKeyArray(SceneComponents);
+	TArray<TObjectPtr<USceneComponent>> SceneComponents;
+	AttachedInstanceComponents.GenerateKeyArray(SceneComponents);
 
 	for (USceneComponent* SceneComponent : SceneComponents)
 	{
diff --git a/Engine/Source/Runtime/Engine/Private/DataTable.cpp b/Engine/Source/Runtime/Engine/Private/DataTable.cpp
--- a/Engine/Source/Runtime/Engine/Private/DataTable.cpp
+++ b/Engine/Source/Runtime/Engine/Private/DataTable.cpp
@@ -461,9 +461,9 @@ void UDataTable::CleanBeforeStructChange()
 		{
 			class FRawStructWriter : public FObjectWriter
 			{
-				TSet<UObject*>& TemporarilyReferencedObjects;
+				TSet<TObjectPtr<UObject>>& TemporarilyReferencedObjects;
 			public:
-				FRawStructWriter(TArray<uint8>& InBytes, TSet<UObject*>& InTemporarilyReferencedObjects)
+				FRawStructWriter(TArray<uint8>& InBytes, TSet<TObjectPtr<UObject>>& InTemporarilyReferencedObjects)
 					: FObjectWriter(InBytes), TemporarilyReferencedObjects(InTemporarilyReferencedObjects) {}
 				virtual FArchive& operator<<(class UObject*& Res) override
 				{
diff --git a/Engine/Source/Runtime/Engine/Private/HLOD/HLODProxyDesc.cpp b/Engine/Source/Runtime/Engine/Private/HLOD/HLODProxyDesc.cpp
--- a/Engine/Source/Runtime/Engine/Private/HLOD/HLODProxyDesc.cpp
+++ b/Engine/Source/Runtime/Engine/Private/HLOD/HLODProxyDesc.cpp
@@ -97,7 +97,7 @@ bool UHLODProxyDesc::UpdateFromLODActor(const ALODActor* InLODActor)
 
 	StaticMesh = InLODActor->StaticMeshComponent ? InLODActor->StaticMeshComponent->GetStaticMesh() : nullptr;
 
-	const TMap<FHLODInstancingKey, UInstancedStaticMeshComponent*>& ISMComponents = InLODActor->InstancedStaticMeshComponents;
+	const TMap<FHLODInstancingKey, TObjectPtr<UInstancedStaticMeshComponent>>& ISMComponents = InLODActor->InstancedStaticMeshComponents;
 	ISMComponentsDesc.Reset(ISMComponents.Num());
 	for (auto const& Pair : ISMComponents)
 	{
@@ -161,7 +161,7 @@ bool UHLODProxyDesc::ShouldUpdateDesc(const ALODActor* InLODActor) const
 	}
 
 	TArray<FHLODISMComponentDesc> LocalISMComponentsDesc;
-	const TMap<FHLODInstancingKey, UInstancedStaticMeshComponent*>& ISMComponents = InLODActor->InstancedStaticMeshComponents;
+	const TMap<FHLODInstancingKey, TObjectPtr<UInstancedStaticMeshComponent>>& ISMComponents = InLODActor->InstancedStaticMeshComponents;
 	LocalISMComponentsDesc.Reset(ISMComponents.Num());
 	for (auto const& Pair : ISMComponents)
 	{
diff --git a/Engine/Source/Runtime/Engine/Private/World.cpp b/Engine/Source/Runtime/Engine/Private/World.cpp
--- a/Engine/Source/Runtime/Engine/Private/World.cpp
+++ b/Engine/Source/Runtime/Engine/Private/World.cpp
@@ -933,7 +933,7 @@ void UWorld::BeginDestroy()
 
 	for (FLevelCollection& LevelCollection : LevelCollections)
 	{
-		TSet<ULevel*> CollectionLevels = LevelCollection.GetLevels();
+		TSet<TObjectPtr<ULevel>> CollectionLevels = LevelCollection.GetLevels();
 		for (ULevel* CollectionLevel : CollectionLevels)
 		{
 			LevelCollection.RemoveLevel(CollectionLevel);
diff --git a/Engine/Source/Runtime/Engine/Public/ActiveSound.h b/Engine/Source/Runtime/Engine/Public/ActiveSound.h
--- a/Engine/Source/Runtime/Engine/Public/ActiveSound.h
+++ b/Engine/Source/Runtime/Engine/Public/ActiveSound.h
@@ -340,7 +340,7 @@ public:
 	TMap<FConcurrencyGroupID, FConcurrencySoundData> ConcurrencyGroupData;
 
 	/** Optional USoundConcurrency to override for the sound. */
-	TSet<USoundConcurrency*> ConcurrencySet;
+	TSet<TObjectPtr<USoundConcurrency>> ConcurrencySet;
 
 private:
 	/** Optional SoundClass to override for the sound. */
diff --git a/Engine/Source/Runtime/Experimental/ChaosSolverEngine/Private/Chaos/ChaosGameplayEventDispatcher.cpp b/Engine/Source/Runtime/Experimental/ChaosSolverEngine/Private/Chaos/ChaosGameplayEventDispatcher.cpp
--- a/Engine/Source/Runtime/Experimental/ChaosSolverEngine/Private/Chaos/ChaosGameplayEventDispatcher.cpp
+++ b/Engine/Source/Runtime/Experimental/ChaosSolverEngine/Private/Chaos/ChaosGameplayEventDispatcher.cpp
@@ -30,7 +30,7 @@ void UChaosGameplayEventDispatcher::OnUnregister()
 }
 
 // internal
-static void DispatchPendingBreakEvents(TArray<FChaosBreakEvent> const& Events, TMap<UPrimitiveComponent*, FBreakEventCallbackWrapper> const& Registrations)
+static void DispatchPendingBreakEvents(TArray<FChaosBreakEvent> const& Events, TMap<TObjectPtr<UPrimitiveComponent>, FBreakEventCallbackWrapper> const& Registrations)
 {
 	for (FChaosBreakEvent const& E : Events)
 	{
diff --git a/Engine/Source/Runtime/Experimental/ChaosSolverEngine/Public/Chaos/ChaosGameplayEventDispatcher.h b/Engine/Source/Runtime/Experimental/ChaosSolverEngine/Public/Chaos/ChaosGameplayEventDispatcher.h
--- a/Engine/Source/Runtime/Experimental/ChaosSolverEngine/Public/Chaos/ChaosGameplayEventDispatcher.h
+++ b/Engine/Source/Runtime/Experimental/ChaosSolverEngine/Public/Chaos/ChaosGameplayEventDispatcher.h
@@ -67,7 +67,7 @@ struct FChaosHandlerSet
 struct FChaosPendingCollisionNotify
 {
 	FChaosPhysicsCollisionInfo CollisionInfo;
-	TSet<UObject*> NotifyRecipients;
+	TSet<TObjectPtr<UObject>> NotifyRecipients;
 };
 
 
diff --git a/Engine/Source/Runtime/Landscape/Classes/Landscape.h b/Engine/Source/Runtime/Landscape/Classes/Landscape.h
--- a/Engine/Source/Runtime/Landscape/Classes/Landscape.h
+++ b/Engine/Source/Runtime/Landscape/Classes/Landscape.h
@@ -273,8 +273,8 @@ public:
 	LANDSCAPE_API void GetUsedPaintLayers(const FGuid& InLayerGuid, TArray<ULandscapeLayerInfoObject*>& OutUsedLayerInfos) const;
 	LANDSCAPE_API void ClearPaintLayer(int32 InLayerIndex, ULandscapeLayerInfoObject* InLayerInfo);
 	LANDSCAPE_API void ClearPaintLayer(const FGuid& InLayerGuid, ULandscapeLayerInfoObject* InLayerInfo);
-	LANDSCAPE_API void ClearLayer(int32 InLayerIndex, TSet<ULandscapeComponent*>* InComponents = nullptr, ELandscapeClearMode InClearMode = ELandscapeClearMode::Clear_All);
-	LANDSCAPE_API void ClearLayer(const FGuid& InLayerGuid, TSet<ULandscapeComponent*>* InComponents = nullptr, ELandscapeClearMode InClearMode = ELandscapeClearMode::Clear_All, bool bMarkPackageDirty = true);
+	LANDSCAPE_API void ClearLayer(int32 InLayerIndex, TSet<TObjectPtr<ULandscapeComponent>>* InComponents = nullptr, ELandscapeClearMode InClearMode = ELandscapeClearMode::Clear_All);
+	LANDSCAPE_API void ClearLayer(const FGuid& InLayerGuid, TSet<TObjectPtr<ULandscapeComponent>>* InComponents = nullptr, ELandscapeClearMode InClearMode = ELandscapeClearMode::Clear_All, bool bMarkPackageDirty = true);
 	LANDSCAPE_API void DeleteLayer(int32 InLayerIndex);
 	LANDSCAPE_API void CollapseLayer(int32 InLayerIndex);
 	LANDSCAPE_API void DeleteLayers();
diff --git a/Engine/Source/Runtime/Landscape/Classes/LandscapeInfo.h b/Engine/Source/Runtime/Landscape/Classes/LandscapeInfo.h
--- a/Engine/Source/Runtime/Landscape/Classes/LandscapeInfo.h
+++ b/Engine/Source/Runtime/Landscape/Classes/LandscapeInfo.h
@@ -184,7 +184,7 @@ public:
 	LANDSCAPE_API void ExportHeightmap(const FString& Filename, const FIntRect& ExportRegion);
 	LANDSCAPE_API void ExportLayer(ULandscapeLayerInfoObject* LayerInfo, const FString& Filename);
 	LANDSCAPE_API void ExportLayer(ULandscapeLayerInfoObject* LayerInfo, const FString& Filename, const FIntRect& ExportRegion);
-	LANDSCAPE_API bool ApplySplines(bool bOnlySelected, TSet<ULandscapeComponent*>* OutModifiedComponents = nullptr, bool bMarkPackageDirty = true);
+	LANDSCAPE_API bool ApplySplines(bool bOnlySelected, TSet<TObjectPtr<ULandscapeComponent>>* OutModifiedComponents = nullptr, bool bMarkPackageDirty = true);
 
 	LANDSCAPE_API bool GetSelectedExtent(int32& MinX, int32& MinY, int32& MaxX, int32& MaxY) const;
 	FVector GetLandscapeCenterPos(float& LengthZ, int32 MinX = MAX_int32, int32 MinY = MAX_int32, int32 MaxX = MIN_int32, int32 MaxY = MIN_int32);
@@ -350,7 +350,7 @@ public:
 
 #if WITH_EDITOR
 private:
-	bool ApplySplinesInternal(bool bOnlySelected, TScriptInterface<ILandscapeSplineInterface> SplineOwner, TSet<ULandscapeComponent*>* OutModifiedComponents, bool bMarkPackageDirty, int32 LandscapeMinX, int32 LandscapeMinY, int32 LandscapeMaxX, int32 LandscapeMaxY, TFunctionRef<TSharedPtr<FModulateAlpha>(ULandscapeLayerInfoObject*)> GetOrCreateModulate);
+	bool ApplySplinesInternal(bool bOnlySelected, TScriptInterface<ILandscapeSplineInterface> SplineOwner, TSet<TObjectPtr<ULandscapeComponent>>* OutModifiedComponents, bool bMarkPackageDirty, int32 LandscapeMinX, int32 LandscapeMinY, int32 LandscapeMaxX, int32 LandscapeMaxY, TFunctionRef<TSharedPtr<FModulateAlpha>(ULandscapeLayerInfoObject*)> GetOrCreateModulate);
 	void MoveSegment(ULandscapeSplineSegment* InSegment, TScriptInterface<ILandscapeSplineInterface> From, TScriptInterface<ILandscapeSplineInterface> To);
 	void MoveControlPoint(ULandscapeSplineControlPoint* InControlPoint, TScriptInterface<ILandscapeSplineInterface> From, TScriptInterface<ILandscapeSplineInterface> To);
 	bool UpdateLayerInfoMapInternal(ALandscapeProxy* Proxy, bool bInvalidate);
diff --git a/Engine/Source/Runtime/Landscape/Private/LandscapeEditLayers.cpp b/Engine/Source/Runtime/Landscape/Private/LandscapeEditLayers.cpp
--- a/Engine/Source/Runtime/Landscape/Private/LandscapeEditLayers.cpp
+++ b/Engine/Source/Runtime/Landscape/Private/LandscapeEditLayers.cpp
@@ -6038,7 +6038,7 @@ void ALandscape::ClearPaintLayer(const FGuid& InLayerGuid, ULandscapeLayerInfoOb
 	});
 }
 
-void ALandscape::ClearLayer(int32 InLayerIndex, TSet<ULandscapeComponent*>* InComponents, ELandscapeClearMode InClearMode)
+void ALandscape::ClearLayer(int32 InLayerIndex, TSet<TObjectPtr<ULandscapeComponent>>* InComponents, ELandscapeClearMode InClearMode)
 {
 	const FLandscapeLayer* Layer = GetLayer(InLayerIndex);
 	if (Layer)
@@ -6047,7 +6047,7 @@ void ALandscape::ClearLayer(int32 InLayerIndex, TSet<ULandscapeComponent*>* InCo
 	}
 }
 
-void ALandscape::ClearLayer(const FGuid& InLayerGuid, TSet<ULandscapeComponent*>* InComponents, ELandscapeClearMode InClearMode, bool bMarkPackageDirty)
+void ALandscape::ClearLayer(const FGuid& InLayerGuid, TSet<TObjectPtr<ULandscapeComponent>>* InComponents, ELandscapeClearMode InClearMode, bool bMarkPackageDirty)
 {
 	ensure(HasLayersContent());
 
@@ -6269,7 +6269,7 @@ void ALandscape::UpdateLandscapeSplines(const FGuid& InTargetLayer, bool bInUpda
 		FScopedSetLandscapeEditingLayer Scope(this, TargetLayerGuid, [=] { this->RequestLayersContentUpdate(ELandscapeLayerUpdateMode::Update_All); });
 		// Temporarily disable material instance updates since it will be done once at the end (requested by RequestLayersContentUpdateForceAll)
 		GDisableUpdateLandscapeMaterialInstances = true;
-		TSet<ULandscapeComponent*>* ModifiedComponent = nullptr;
+		TSet<TObjectPtr<ULandscapeComponent>>* ModifiedComponent = nullptr;
 		if (LandscapeSplinesTargetLayerGuid.IsValid())
 		{
 			// Check that we can modify data
@@ -6299,9 +6299,9 @@ void ALandscape::UpdateLandscapeSplines(const FGuid& InTargetLayer, bool bInUpda
 
 			// Clear layers without affecting weightmap allocations
 			const bool bMarkPackageDirty = false;
-			ClearLayer(LandscapeSplinesTargetLayerGuid, (!bInForceUpdateAllCompoments && LandscapeSplinesAffectedComponents.Num()) ? &ToRawPtrTSetUnsafe(LandscapeSplinesAffectedComponents) : nullptr, Clear_All, bMarkPackageDirty);
+			ClearLayer(LandscapeSplinesTargetLayerGuid, (!bInForceUpdateAllCompoments && LandscapeSplinesAffectedComponents.Num()) ? &LandscapeSplinesAffectedComponents : nullptr, Clear_All, bMarkPackageDirty);
 			LandscapeSplinesAffectedComponents.Empty();
-			ModifiedComponent = &ToRawPtrTSetUnsafe(LandscapeSplinesAffectedComponents);
+			ModifiedComponent = &LandscapeSplinesAffectedComponents;
 			// For now, in Landscape Layer System Mode with a reserved layer for splines, we always update all the splines since we clear the whole layer first
 			bInUpdateOnlySelected = false;
 
diff --git a/Engine/Source/Runtime/Landscape/Private/LandscapeSplineRaster.cpp b/Engine/Source/Runtime/Landscape/Private/LandscapeSplineRaster.cpp
--- a/Engine/Source/Runtime/Landscape/Private/LandscapeSplineRaster.cpp
+++ b/Engine/Source/Runtime/Landscape/Private/LandscapeSplineRaster.cpp
@@ -548,7 +548,7 @@ void RasterizeSegmentAlpha(int32& MinX, int32& MinY, int32& MaxX, int32& MaxY, F
 	LandscapeEdit.GetComponentsInRegion(MinX, MinY, MaxX, MaxY, &ModifiedComponents);
 }
 
-bool ULandscapeInfo::ApplySplines(bool bOnlySelected, TSet<ULandscapeComponent*>* OutModifiedComponents, bool bMarkPackageDirty)
+bool ULandscapeInfo::ApplySplines(bool bOnlySelected, TSet<TObjectPtr<ULandscapeComponent>>* OutModifiedComponents, bool bMarkPackageDirty)
 {
 	bool bResult = false;
 
@@ -585,7 +585,7 @@ bool ULandscapeInfo::ApplySplines(bool bOnlySelected, TSet<ULandscapeComponent*>
 	return bResult;
 }
 
-bool ULandscapeInfo::ApplySplinesInternal(bool bOnlySelected, TScriptInterface<ILandscapeSplineInterface> SplineOwner, TSet<ULandscapeComponent*>* OutModifiedComponents, bool bMarkPackageDirty, int32 LandscapeMinX, int32 LandscapeMinY, int32 LandscapeMaxX, int32 LandscapeMaxY, TFunctionRef<TSharedPtr<FModulateAlpha>(ULandscapeLayerInfoObject*)> GetOrCreateModulate)
+bool ULandscapeInfo::ApplySplinesInternal(bool bOnlySelected, TScriptInterface<ILandscapeSplineInterface> SplineOwner, TSet<TObjectPtr<ULandscapeComponent>>* OutModifiedComponents, bool bMarkPackageDirty, int32 LandscapeMinX, int32 LandscapeMinY, int32 LandscapeMaxX, int32 LandscapeMaxY, TFunctionRef<TSharedPtr<FModulateAlpha>(ULandscapeLayerInfoObject*)> GetOrCreateModulate)
 {
 	if (!SplineOwner)
 	{
