Add a way to be notified when properties are overwritten for reasons that were not trackable before
  - PostReinitProperties is called when a subobject has it's properties overwritten a second time during a direct copy from a CDO
  - PostApplyToComponent is called when a component has it's properties overwritten by a user construction script

diff --git a/Engine/Source/Runtime/CoreUObject/Private/UObject/UObjectGlobals.cpp b/Engine/Source/Runtime/CoreUObject/Private/UObject/UObjectGlobals.cpp
--- a/Engine/Source/Runtime/CoreUObject/Private/UObject/UObjectGlobals.cpp
+++ b/Engine/Source/Runtime/CoreUObject/Private/UObject/UObjectGlobals.cpp
@@ -100,6 +100,7 @@ DEFINE_STAT(STAT_DestroyObject);
 
 DECLARE_CYCLE_STAT(TEXT("InstanceSubobjects"), STAT_InstanceSubobjects, STATGROUP_Object);
 DECLARE_CYCLE_STAT(TEXT("PostInitProperties"), STAT_PostInitProperties, STATGROUP_Object);
+DECLARE_CYCLE_STAT(TEXT("PostReinitProperties"), STAT_PostReinitProperties, STATGROUP_Object);
 
 CSV_DEFINE_CATEGORY(UObject, false);
 
@@ -2588,7 +2589,9 @@ UObject* StaticAllocateObject
 
 //@todo UE4 - move this stuff to UnObj.cpp or something
 
-
+void UObject::PostReinitProperties()
+{
+}
 
 void UObject::PostInitProperties()
 {
@@ -2892,6 +2895,14 @@ void FObjectInitializer::PostConstructInit()
 		InstanceSubobjects(Class, bNeedInstancing, bNeedSubobjectInstancing);
 	}
 
+	// Make sure subobjects knows that they had their properties overwritten
+	for (int32 Index = 0; Index < ComponentInits.SubobjectInits.Num(); Index++)
+	{
+		SCOPE_CYCLE_COUNTER(STAT_PostReinitProperties);
+		UObject* Subobject = ComponentInits.SubobjectInits[Index].Subobject;
+		Subobject->PostReinitProperties();
+	}
+
 	{
 		SCOPE_CYCLE_COUNTER(STAT_PostInitProperties);
 		Obj->PostInitProperties();
diff --git a/Engine/Source/Runtime/CoreUObject/Public/UObject/Object.h b/Engine/Source/Runtime/CoreUObject/Public/UObject/Object.h
--- a/Engine/Source/Runtime/CoreUObject/Public/UObject/Object.h
+++ b/Engine/Source/Runtime/CoreUObject/Public/UObject/Object.h
@@ -192,6 +192,12 @@ public:
 	 * This is called before any serialization or other setup has happened.
 	 */
 	virtual void PostInitProperties();
+	
+	/**
+	 * Called after properties are overwritten, including after subobjects initialization from a CDO.
+	 * This could be called multiple times during an object lifetime, which is not the case for PostInitProperties which is expected to be called only once.
+	 */
+	virtual void PostReinitProperties();
 
 	/**
 	* Called after the C++ constructor has run on the CDO for a class. This is an obscure routine used to deal with the recursion 
diff --git a/Engine/Source/Runtime/Engine/Classes/Components/ActorComponent.h b/Engine/Source/Runtime/Engine/Classes/Components/ActorComponent.h
--- a/Engine/Source/Runtime/Engine/Classes/Components/ActorComponent.h
+++ b/Engine/Source/Runtime/Engine/Classes/Components/ActorComponent.h
@@ -860,6 +860,9 @@ public:
 	/** Called before we throw away components during RerunConstructionScripts, to cache any data we wish to persist across that operation */
 	virtual TStructOnScope<FActorComponentInstanceData> GetComponentInstanceData() const;
 
+	/** Called after ApplyToComponent has run. */
+	virtual void PostApplyToComponent();
+
 	//~ Begin UObject Interface.
 	virtual void BeginDestroy() override;
 	virtual bool NeedsLoadForClient() const override;
diff --git a/Engine/Source/Runtime/Engine/Private/ComponentInstanceDataCache.cpp b/Engine/Source/Runtime/Engine/Private/ComponentInstanceDataCache.cpp
--- a/Engine/Source/Runtime/Engine/Private/ComponentInstanceDataCache.cpp
+++ b/Engine/Source/Runtime/Engine/Private/ComponentInstanceDataCache.cpp
@@ -385,10 +385,7 @@ void FActorComponentInstanceData::ApplyToComponent(UActorComponent* Component, c
 
 		FComponentPropertyReader ComponentPropertyReader(Component, *this);
 
-		if (Component->IsRegistered())
-		{
-			Component->ReregisterComponent();
-		}
+		Component->PostApplyToComponent();
 	}
 }
 
diff --git a/Engine/Source/Runtime/Engine/Private/Components/ActorComponent.cpp b/Engine/Source/Runtime/Engine/Private/Components/ActorComponent.cpp
--- a/Engine/Source/Runtime/Engine/Private/Components/ActorComponent.cpp
+++ b/Engine/Source/Runtime/Engine/Private/Components/ActorComponent.cpp
@@ -2175,6 +2175,14 @@ void UActorComponent::Serialize(FArchive& Ar)
 	}
 }
 
+void UActorComponent::PostApplyToComponent()
+{
+	if (IsRegistered())
+	{
+		ReregisterComponent();
+	}
+}
+
 AActor* UActorComponent::GetActorOwnerNoninline() const
 {
 	// This is defined out-of-line because AActor isn't defined where the inlined function is.
