Fix ObjectHandle test failing after the wrapped object pointer upgrade as there are now more trackable sub-references in the content being loaded during the test, leading to an excess when counting the number of handle reads.

diff --git a/Engine/Source/Runtime/CoreUObject/Private/Tests/ObjectRefTrackingTestBase.h b/Engine/Source/Runtime/CoreUObject/Private/Tests/ObjectRefTrackingTestBase.h
--- a/Engine/Source/Runtime/CoreUObject/Private/Tests/ObjectRefTrackingTestBase.h
+++ b/Engine/Source/Runtime/CoreUObject/Private/Tests/ObjectRefTrackingTestBase.h
@@ -49,10 +49,17 @@ public:
 #endif
 		}
 
-		bool TestNumReads(const TCHAR* What, uint32 ExpectedDelta)
+		bool TestNumReads(const TCHAR* What, uint32 ExpectedDelta, bool bAllowAdditionalReads = false)
 		{
 #if UE_WITH_OBJECT_HANDLE_TRACKING
-			return Test.TestEqual(What, OriginalNumReads + ExpectedDelta, Test.GetNumReads());
+			if (bAllowAdditionalReads)
+			{
+				return Test.TestTrue(What, Test.GetNumReads() >= OriginalNumReads + ExpectedDelta);
+			}
+			else
+			{
+				return Test.TestEqual(What, OriginalNumReads + ExpectedDelta, Test.GetNumReads());
+			}
 #else
 			return true;
 #endif
diff --git a/Engine/Source/Runtime/CoreUObject/Private/UObject/ObjectHandle.cpp b/Engine/Source/Runtime/CoreUObject/Private/UObject/ObjectHandle.cpp
--- a/Engine/Source/Runtime/CoreUObject/Private/UObject/ObjectHandle.cpp
+++ b/Engine/Source/Runtime/CoreUObject/Private/UObject/ObjectHandle.cpp
@@ -304,8 +304,11 @@ static inline UPackage* FindOrLoadPackage(FName PackageName, int32 LoadFlags)
 		// @TODO: OBJPTR: When using the "external package" feature, we will have objects that have a differing package path vs "outer hierarchy" path
 		//				  The package path should be used when loading.  The "outer hierarchy" path may need to be used when finding existing objects in memory.
 		//				  This will need further evaluation and testing before lazy load can be enabled.
-		// @TODO: OBJPTR: This seems like it will do a search on disk every time instead of leveraging known missing packages
 		// @TODO: OBJPTR: Instancing context may be important to consider when loading the package.
+		if (FLinkerLoad::IsKnownMissingPackage(PackageName))
+		{
+			return nullptr;
+		}
 		LoadFlags |= LOAD_NoWarn | LOAD_NoVerify; //This does nothing? | LOAD_DisableDependencyPreloading;
 		TargetPackage = LoadPackage(nullptr, *PackageName.ToString(), LoadFlags);
 	}
