(added missing file from last submit) Fix ObjectHandle test failing after the wrapped object pointer upgrade as there are now more trackable sub-references in the content being loaded during the test, leading to an excess when counting the number of handle reads.

diff --git a/Engine/Source/Runtime/CoreUObject/Private/Tests/ObjectHandleTest.cpp b/Engine/Source/Runtime/CoreUObject/Private/Tests/ObjectHandleTest.cpp
--- a/Engine/Source/Runtime/CoreUObject/Private/Tests/ObjectHandleTest.cpp
+++ b/Engine/Source/Runtime/CoreUObject/Private/Tests/ObjectHandleTest.cpp
@@ -50,12 +50,12 @@ protected:
 		return ResolveObjectHandle(TargetHandle);
 	}
 
-	bool TestResolvableNonNull(const ANSICHAR* PackageName, const ANSICHAR* ObjectName, const ANSICHAR* ClassPackageName = nullptr, const ANSICHAR* ClassName = nullptr)
+	bool TestResolvableNonNull(const ANSICHAR* PackageName, const ANSICHAR* ObjectName, const ANSICHAR* ClassPackageName = nullptr, const ANSICHAR* ClassName = nullptr, bool bExpectSubRefReads = false)
 	{
 		FSnapshotObjectRefMetrics ObjectRefMetrics(*this);
 		UObject* ResolvedObject = ConstructAndResolveHandle(PackageName, ObjectName, ClassPackageName, ClassName);
 		ObjectRefMetrics.TestNumResolves(TEXT("NumResolves should be incremented by one after a resolve attempt"), 1);
-		ObjectRefMetrics.TestNumReads(TEXT("NumReads should be incremented by one after a resolve attempt"), 1);
+		ObjectRefMetrics.TestNumReads(TEXT("NumReads should be incremented by one after a resolve attempt"), 1, bExpectSubRefReads /*bAllowAdditionalReads*/);
 
 		if (!ResolvedObject)
 		{
@@ -137,10 +137,10 @@ bool FObjectHandleTestResolveEngineContentTarget::RunTest(const FString& Paramet
 	if (!FPlatformProperties::RequiresCookedData())
 	{
 		// Confirm we successfully resolve a correct reference to a subobject in engine content
-		TestResolvableNonNull("/Engine/FunctionalTesting/Blueprints/AITesting_MoveGoal", "AITesting_MoveGoal.EventGraph.K2Node_VariableGet_142");
+		TestResolvableNonNull("/Engine/FunctionalTesting/Blueprints/AITesting_MoveGoal", "AITesting_MoveGoal.EventGraph.K2Node_VariableGet_142", nullptr, nullptr, true);
 
 		// Attempt to load something that uses a User Defined Enum
-		TestResolvableNonNull("/Engine/ArtTools/RenderToTexture/Macros/RenderToTextureMacros", "RenderToTextureMacros:Array to HLSL Float Array.K2Node_Select_1");
+		TestResolvableNonNull("/Engine/ArtTools/RenderToTexture/Macros/RenderToTextureMacros", "RenderToTextureMacros:Array to HLSL Float Array.K2Node_Select_1", nullptr, nullptr, true);
 	}
 
 	return true;
