Add support in FSoftObjectPath to resolve unloaded subobjects through top level object.

diff --git a/Engine/Source/Runtime/CoreUObject/Private/UObject/SoftObjectPath.cpp b/Engine/Source/Runtime/CoreUObject/Private/UObject/SoftObjectPath.cpp
--- a/Engine/Source/Runtime/CoreUObject/Private/UObject/SoftObjectPath.cpp
+++ b/Engine/Source/Runtime/CoreUObject/Private/UObject/SoftObjectPath.cpp
@@ -470,39 +470,50 @@ UObject* FSoftObjectPath::TryLoad(FUObjectSerializeContext* InLoadContext) const
 			UObject* TopLevelObject = TopLevelPath.TryLoad(InLoadContext);
 
 			// This probably loaded the top-level object, so re-resolve ourselves
-			return ResolveObject();
-		}
+			LoadedObject = ResolveObject();
 
-		FString PathString = ToString();
 #if WITH_EDITOR
-		if (GPlayInEditorID != INDEX_NONE)
-		{
-			// If we are in PIE and this hasn't already been fixed up, we need to fixup at resolution time. We cannot modify the path as it may be somewhere like a blueprint CDO
-			FSoftObjectPath FixupObjectPath = *this;
-			if (FixupObjectPath.FixupForPIE())
+			// If the the top-level object exists but we can't find the object, defer the loading to the top-level container object in case
+			// it knows how to load that specific object.
+			if (!LoadedObject && TopLevelObject)
 			{
-				PathString = FixupObjectPath.ToString();
+				LoadedObject = TopLevelObject->LoadSubobject(*SubPathString);
 			}
+#endif
 		}
+		else
+		{
+			FString PathString = ToString();
+#if WITH_EDITOR
+			if (GPlayInEditorID != INDEX_NONE)
+			{
+				// If we are in PIE and this hasn't already been fixed up, we need to fixup at resolution time. We cannot modify the path as it may be somewhere like a blueprint CDO
+				FSoftObjectPath FixupObjectPath = *this;
+				if (FixupObjectPath.FixupForPIE())
+				{
+					PathString = FixupObjectPath.ToString();
+				}
+			}
 #endif
 
-		LoadedObject = StaticLoadObject(UObject::StaticClass(), nullptr, *PathString, nullptr, LOAD_None, nullptr, true);
+			LoadedObject = StaticLoadObject(UObject::StaticClass(), nullptr, *PathString, nullptr, LOAD_None, nullptr, true);
 
 #if WITH_EDITOR
-		// Look at core redirects if we didn't find the object
-		if (!LoadedObject)
-		{
-			FSoftObjectPath FixupObjectPath = *this;
-			if (FixupObjectPath.FixupCoreRedirects())
+			// Look at core redirects if we didn't find the object
+			if (!LoadedObject)
 			{
-				LoadedObject = LoadObject<UObject>(nullptr, *FixupObjectPath.ToString());
+				FSoftObjectPath FixupObjectPath = *this;
+				if (FixupObjectPath.FixupCoreRedirects())
+				{
+					LoadedObject = LoadObject<UObject>(nullptr, *FixupObjectPath.ToString());
+				}
 			}
-		}
 #endif
 
-		while (UObjectRedirector* Redirector = Cast<UObjectRedirector>(LoadedObject))
-		{
-			LoadedObject = Redirector->DestinationObject;
+			while (UObjectRedirector* Redirector = Cast<UObjectRedirector>(LoadedObject))
+			{
+				LoadedObject = Redirector->DestinationObject;
+			}
 		}
 	}
 
diff --git a/Engine/Source/Runtime/CoreUObject/Public/UObject/Object.h b/Engine/Source/Runtime/CoreUObject/Public/UObject/Object.h
--- a/Engine/Source/Runtime/CoreUObject/Public/UObject/Object.h
+++ b/Engine/Source/Runtime/CoreUObject/Public/UObject/Object.h
@@ -258,6 +258,13 @@ public:
 	 * Called when the object was loaded from another class via active class redirects.
 	 */
 	virtual void LoadedFromAnotherClass(const FName& OldClassName) {}
+
+	/**
+	 * Called to defer loading a subobject to its top-level container object. Usefull with the package override
+	 * feature, where an object might not be saved in the same package as its outer, especially if the the top-level
+	 * object does lazy loading of these objects.
+	 */
+	virtual UObject* LoadSubobject(const TCHAR* SubObjectPath) { return nullptr; }
 #endif
 
 	/**
