Split the representation in the reflection system of TObjectPtr fields into FObjectPtrProperty (derived from FObjectProperty and already existed since the introduction of wrapped object pointers) and a new FClassPtrProperty (derived from FClassProperty, and newly introduced in this change).  As with raw pointers, UHT will generate code for one vs the other depending on whether the inner type is derived from UClass or not.  The new reflection type is meant to be functionally equivalent to the existing one, except it changes how the PROPERTY (not the instance of the field on a containing object) is serialized to ensure it includes the MetaClass member as inherited from FClassProperty.

diff --git a/Engine/Source/Programs/UnrealHeaderTool/Private/CodeGenerator.cpp b/Engine/Source/Programs/UnrealHeaderTool/Private/CodeGenerator.cpp
--- a/Engine/Source/Programs/UnrealHeaderTool/Private/CodeGenerator.cpp
+++ b/Engine/Source/Programs/UnrealHeaderTool/Private/CodeGenerator.cpp
@@ -1445,6 +1445,29 @@ void FNativeClassHeaderGenerator::PropertyNew(FOutputDevice& DeclOut, FOutputDev
 		return;
 	}
 
+	if (FClassPtrProperty* TypedProp = CastField<FClassPtrProperty>(Prop))
+	{
+		DeclOut.Logf(TEXT("%sstatic const UE4CodeGen_Private::FClassPtrPropertyParams %s;\r\n"), DeclSpaces, *NameWithoutScope);
+
+		Out.Logf(
+			TEXT("%sconst UE4CodeGen_Private::FClassPtrPropertyParams %s = { %s, %s, (EPropertyFlags)0x%016llx, UE4CodeGen_Private::EPropertyGenFlags::Class | UE4CodeGen_Private::EPropertyGenFlags::ObjectPtr, %s, %s, %s, %s, %s, %s };%s\r\n"),
+			Spaces,
+			Name,
+			*PropName,
+			*PropNotifyFunc,
+			PropFlags,
+			FPropertyObjectFlags,
+			*ArrayDim,
+			OffsetStr,
+			*GetSingletonNameFuncAddr(TypedProp->MetaClass, OutReferenceGatherers.UniqueCrossModuleReferences, false),
+			*GetSingletonNameFuncAddr(TypedProp->PropertyClass, OutReferenceGatherers.UniqueCrossModuleReferences, false),
+			*MetaDataParams,
+			*PropTag
+		);
+
+		return;
+	}
+
 	if (FSoftObjectProperty* TypedProp = CastField<FSoftObjectProperty>(Prop))
 	{
 		DeclOut.Logf(TEXT("%sstatic const UE4CodeGen_Private::FSoftObjectPropertyParams %s;\r\n"), DeclSpaces, *NameWithoutScope);
diff --git a/Engine/Source/Programs/UnrealHeaderTool/Private/HeaderParser.cpp b/Engine/Source/Programs/UnrealHeaderTool/Private/HeaderParser.cpp
--- a/Engine/Source/Programs/UnrealHeaderTool/Private/HeaderParser.cpp
+++ b/Engine/Source/Programs/UnrealHeaderTool/Private/HeaderParser.cpp
@@ -874,16 +874,25 @@ namespace
 			{
 				check(VarProperty.PropertyClass);
 
-				// @TODO: OBJPTR: Further work is needed to handle UClass derivatives in a manner similar to CPT_ObjectReference or CPT_SoftObjectReference
-				if (DoesAnythingInHierarchyHaveDefaultToInstanced(VarProperty.PropertyClass))
+				if (VarProperty.PropertyClass->IsChildOf(UClass::StaticClass()))
 				{
-					VarProperty.PropertyFlags |= CPF_InstancedReference;
-					AddEditInlineMetaData(VarProperty.MetaData);
+					FClassPtrProperty* Result = new FClassPtrProperty(Scope, Name, ObjectFlags);
+					Result->MetaClass     = VarProperty.MetaClass;
+					Result->PropertyClass = VarProperty.PropertyClass;
+					return Result;
 				}
+				else
+				{
+					if (DoesAnythingInHierarchyHaveDefaultToInstanced(VarProperty.PropertyClass))
+					{
+						VarProperty.PropertyFlags |= CPF_InstancedReference;
+						AddEditInlineMetaData(VarProperty.MetaData);
+					}
 
-				FObjectPtrProperty* Result = new FObjectPtrProperty(Scope, Name, ObjectFlags);
-				Result->PropertyClass = VarProperty.PropertyClass;
-				return Result;
+					FObjectPtrProperty* Result = new FObjectPtrProperty(Scope, Name, ObjectFlags);
+					Result->PropertyClass = VarProperty.PropertyClass;
+					return Result;
+				}
 			}
 
 			case CPT_SoftObjectReference:
@@ -1241,6 +1250,7 @@ namespace
 			|| Property->IsA<FSoftObjectProperty>()
 			|| Property->IsA<FObjectProperty>()
 			|| Property->IsA<FObjectPtrProperty>()
+			|| Property->IsA<FClassPtrProperty>()
 			|| Property->IsA<FFloatProperty>()
 			|| Property->IsA<FDoubleProperty>()
 			|| Property->IsA<FIntProperty>()
diff --git a/Engine/Source/Runtime/CoreUObject/Private/UObject/GarbageCollection.cpp b/Engine/Source/Runtime/CoreUObject/Private/UObject/GarbageCollection.cpp
--- a/Engine/Source/Runtime/CoreUObject/Private/UObject/GarbageCollection.cpp
+++ b/Engine/Source/Runtime/CoreUObject/Private/UObject/GarbageCollection.cpp
@@ -2425,15 +2425,6 @@ void FMulticastDelegateProperty::EmitReferenceInfo(UClass& OwnerClass, int32 Bas
 	FGCReferenceFixedArrayTokenHelper FixedArrayHelper(OwnerClass, BaseOffset + GetOffset_ForGC(), ArrayDim, this->ElementSize, *this);
 	OwnerClass.EmitObjectReference(BaseOffset + GetOffset_ForGC(), GetFName(), GCRT_MulticastDelegate);
 }
-/**
- * Emits tokens used by realtime garbage collection code to passed in OwnerClass' ReferenceTokenStream. The offset emitted is relative
- * to the passed in BaseOffset which is used by e.g. arrays of structs.
- */
-void FObjectPtrProperty::EmitReferenceInfo(UClass& OwnerClass, int32 BaseOffset, TArray<const FStructProperty*>& EncounteredStructProps)
-{
-	FGCReferenceFixedArrayTokenHelper FixedArrayHelper(OwnerClass, BaseOffset + GetOffset_ForGC(), ArrayDim, sizeof(FObjectPtr), *this);
-	OwnerClass.EmitObjectReference(BaseOffset + GetOffset_ForGC(), GetFName(), GCRT_Object);
-}
 
 /**
  * Emits tokens used by realtime garbage collection code to passed in OwnerClass' ReferenceTokenStream. The offset emitted is relative
diff --git a/Engine/Source/Runtime/CoreUObject/Private/UObject/PropertyClassPtr.cpp b/Engine/Source/Runtime/CoreUObject/Private/UObject/PropertyClassPtr.cpp
new file mode 100644
--- /dev/null
+++ b/Engine/Source/Runtime/CoreUObject/Private/UObject/PropertyClassPtr.cpp
@@ -0,0 +1,59 @@
+// Copyright Epic Games, Inc. All Rights Reserved.
+
+#include "CoreMinimal.h"
+#include "UObject/ObjectMacros.h"
+#include "UObject/PropertyHelper.h"
+#include "UObject/UnrealType.h"
+#include "UObject/LinkerPlaceholderExportObject.h"
+#include "UObject/LinkerPlaceholderClass.h"
+
+/*-----------------------------------------------------------------------------
+	FClassPtrProperty.
+-----------------------------------------------------------------------------*/
+IMPLEMENT_FIELD(FClassPtrProperty)
+
+FString FClassPtrProperty::GetCPPType(FString* ExtendedTypeText/*=NULL*/, uint32 CPPExportFlags/*=0*/) const
+{
+	return FString::Printf(TEXT("TObjectPtr<%s%s>"), PropertyClass->GetPrefixCPP(), *PropertyClass->GetName());
+}
+FString FClassPtrProperty::GetCPPMacroType(FString& ExtendedTypeText) const
+{
+	ExtendedTypeText = FString::Printf(TEXT("TObjectPtr<%s%s>"), PropertyClass->GetPrefixCPP(), *PropertyClass->GetName());
+	return TEXT("OBJECTPTR");
+}
+
+void FClassPtrProperty::SerializeItem(FStructuredArchive::FSlot Slot, void* Value, void const* Defaults) const
+{
+	// Share serialization code with FObjectPtrProperty
+	FObjectPtrProperty::StaticSerializeItem(this, Slot, Value, Defaults);
+}
+
+bool FClassPtrProperty::Identical(const void* A, const void* B, uint32 PortFlags) const
+{
+	// Share comparison code with FObjectPtrProperty
+	return FObjectPtrProperty::StaticIdentical(A, B, PortFlags);
+}
+
+bool FClassPtrProperty::SameType(const FProperty* Other) const
+{
+	bool bPropertyTypeMatch = FObjectProperty::SameType(Other) ||
+								(Other && Other->IsA<FClassProperty>() &&
+									(PropertyClass == ((FObjectPropertyBase*)Other)->PropertyClass) );
+	return bPropertyTypeMatch && (MetaClass == ((FClassProperty*)Other)->MetaClass);
+}
+
+UObject* FClassPtrProperty::GetObjectPropertyValue(const void* PropertyValueAddress) const
+{
+	return ((FObjectPtr&)GetPropertyValue(PropertyValueAddress)).Get();
+}
+
+void FClassPtrProperty::SetObjectPropertyValue(void* PropertyValueAddress, UObject* Value) const
+{
+	SetPropertyValue(PropertyValueAddress, TCppType(Value));
+}
+
+uint32 FClassPtrProperty::GetValueTypeHashInternal(const void* Src) const
+{
+	return GetTypeHash((FObjectPtr&)GetPropertyValue(Src));
+}
+
diff --git a/Engine/Source/Runtime/CoreUObject/Private/UObject/PropertyObjectPtr.cpp b/Engine/Source/Runtime/CoreUObject/Private/UObject/PropertyObjectPtr.cpp
--- a/Engine/Source/Runtime/CoreUObject/Private/UObject/PropertyObjectPtr.cpp
+++ b/Engine/Source/Runtime/CoreUObject/Private/UObject/PropertyObjectPtr.cpp
@@ -23,6 +23,11 @@ FString FObjectPtrProperty::GetCPPMacroType(FString& ExtendedTypeText) const
 }
 
 void FObjectPtrProperty::SerializeItem(FStructuredArchive::FSlot Slot, void* Value, void const* Defaults) const
+{
+	StaticSerializeItem(this, Slot, Value, Defaults);
+}
+
+void FObjectPtrProperty::StaticSerializeItem(const FObjectPropertyBase* ObjectProperty, FStructuredArchive::FSlot Slot, void* Value, void const* Defaults)
 {
 	FArchive& UnderlyingArchive = Slot.GetUnderlyingArchive();
 	FObjectPtr* ObjectPtr = (FObjectPtr*)GetPropertyValuePtr(Value);
@@ -33,7 +38,7 @@ void FObjectPtrProperty::SerializeItem(FStructuredArchive::FSlot Slot, void* Val
 
 		if(!UnderlyingArchive.IsSaving() && IsObjectHandleResolved(ObjectPtr->GetHandle()))
 		{
-			CheckValidObject(ObjectPtr);
+			ObjectProperty->CheckValidObject(ObjectPtr);
 		}
 	}
 	else
@@ -48,11 +53,11 @@ void FObjectPtrProperty::SerializeItem(FStructuredArchive::FSlot Slot, void* Val
 	#if USE_CIRCULAR_DEPENDENCY_LOAD_DEFERRING
 			if (ULinkerPlaceholderExportObject* PlaceholderVal = Cast<ULinkerPlaceholderExportObject>(ResolvedObject))
 			{
-				PlaceholderVal->AddReferencingPropertyValue(this, Value);
+				PlaceholderVal->AddReferencingPropertyValue(ObjectProperty, Value);
 			}
 			else if (ULinkerPlaceholderClass* PlaceholderClass = Cast<ULinkerPlaceholderClass>(ResolvedObject))
 			{
-				PlaceholderClass->AddReferencingPropertyValue(this, Value);
+				PlaceholderClass->AddReferencingPropertyValue(ObjectProperty, Value);
 			}
 			// NOTE: we don't remove this from CurrentValue if it is a 
 			//       ULinkerPlaceholderExportObject; this is because this property 
@@ -67,7 +72,7 @@ void FObjectPtrProperty::SerializeItem(FStructuredArchive::FSlot Slot, void* Val
 			//        to accommodate this (as it depends on finding itself as the set value)
 	#endif // USE_CIRCULAR_DEPENDENCY_LOAD_DEFERRING
 
-			CheckValidObject(Value);
+			ObjectProperty->CheckValidObject(Value);
 		}
 	}
 }
@@ -81,6 +86,11 @@ bool FObjectPtrProperty::SameType(const FProperty* Other) const
 }
 
 bool FObjectPtrProperty::Identical(const void* A, const void* B, uint32 PortFlags) const
+{
+	return StaticIdentical(A, B, PortFlags);
+}
+
+bool FObjectPtrProperty::StaticIdentical(const void* A, const void* B, uint32 PortFlags)
 {
 	FObjectPtr ObjectA = A ? *((FObjectPtr*)A) : FObjectPtr();
 	FObjectPtr ObjectB = B ? *((FObjectPtr*)B) : FObjectPtr();
diff --git a/Engine/Source/Runtime/CoreUObject/Private/UObject/UObjectGlobals.cpp b/Engine/Source/Runtime/CoreUObject/Private/UObject/UObjectGlobals.cpp
--- a/Engine/Source/Runtime/CoreUObject/Private/UObject/UObjectGlobals.cpp
+++ b/Engine/Source/Runtime/CoreUObject/Private/UObject/UObjectGlobals.cpp
@@ -4424,13 +4424,26 @@ namespace UE4CodeGen_Private
 
 			case EPropertyGenFlags::Class:
 			{
-				const FClassPropertyParams* Prop = (const FClassPropertyParams*)PropBase;
-				NewProp = new FClassProperty(Outer, UTF8_TO_TCHAR(Prop->NameUTF8), Prop->ObjectFlags, Prop->Offset, Prop->PropertyFlags, Prop->MetaClassFunc ? Prop->MetaClassFunc() : nullptr, Prop->ClassFunc ? Prop->ClassFunc() : nullptr);
+				if (EnumHasAllFlags(PropBase->Flags, EPropertyGenFlags::ObjectPtr))
+				{
+					const FClassPtrPropertyParams* Prop = (const FClassPtrPropertyParams*)PropBase;
+					NewProp = new FClassPtrProperty(Outer, UTF8_TO_TCHAR(Prop->NameUTF8), Prop->ObjectFlags, Prop->Offset, Prop->PropertyFlags, Prop->MetaClassFunc ? Prop->MetaClassFunc() : nullptr, Prop->ClassFunc ? Prop->ClassFunc() : nullptr);
 
 #if WITH_METADATA
-				MetaDataArray = Prop->MetaDataArray;
-				NumMetaData   = Prop->NumMetaData;
+					MetaDataArray = Prop->MetaDataArray;
+					NumMetaData   = Prop->NumMetaData;
 #endif
+				}
+				else
+				{
+					const FClassPropertyParams* Prop = (const FClassPropertyParams*)PropBase;
+					NewProp = new FClassProperty(Outer, UTF8_TO_TCHAR(Prop->NameUTF8), Prop->ObjectFlags, Prop->Offset, Prop->PropertyFlags, Prop->MetaClassFunc ? Prop->MetaClassFunc() : nullptr, Prop->ClassFunc ? Prop->ClassFunc() : nullptr);
+
+#if WITH_METADATA
+					MetaDataArray = Prop->MetaDataArray;
+					NumMetaData   = Prop->NumMetaData;
+#endif
+				}
 			}
 			break;
 
diff --git a/Engine/Source/Runtime/CoreUObject/Public/UObject/ObjectMacros.h b/Engine/Source/Runtime/CoreUObject/Public/UObject/ObjectMacros.h
--- a/Engine/Source/Runtime/CoreUObject/Public/UObject/ObjectMacros.h
+++ b/Engine/Source/Runtime/CoreUObject/Public/UObject/ObjectMacros.h
@@ -348,6 +348,7 @@ enum EClassCastFlags : uint64
 	CASTCLASS_FMulticastSparseDelegateProperty	= 0x0008000000000000,
 	CASTCLASS_FFieldPathProperty			= 0x0010000000000000,
 	CASTCLASS_FObjectPtrProperty			= 0x0020000000000000,
+	CASTCLASS_FClassPtrProperty				= 0x0040000000000000,
 };
 
 #define CASTCLASS_AllFlags ((EClassCastFlags)0xFFFFFFFFFFFFFFFF)
diff --git a/Engine/Source/Runtime/CoreUObject/Public/UObject/UObjectGlobals.h b/Engine/Source/Runtime/CoreUObject/Public/UObject/UObjectGlobals.h
--- a/Engine/Source/Runtime/CoreUObject/Public/UObject/UObjectGlobals.h
+++ b/Engine/Source/Runtime/CoreUObject/Public/UObject/UObjectGlobals.h
@@ -2854,6 +2854,7 @@ namespace UE4CodeGen_Private
 	typedef FObjectPropertyParams  FWeakObjectPropertyParams;
 	typedef FObjectPropertyParams  FLazyObjectPropertyParams;
 	typedef FObjectPropertyParams  FObjectPtrPropertyParams;
+	typedef FClassPropertyParams   FClassPtrPropertyParams;
 	typedef FObjectPropertyParams  FSoftObjectPropertyParams;
 
 	struct FFunctionParams
diff --git a/Engine/Source/Runtime/CoreUObject/Public/UObject/UnrealType.h b/Engine/Source/Runtime/CoreUObject/Public/UObject/UnrealType.h
--- a/Engine/Source/Runtime/CoreUObject/Public/UObject/UnrealType.h
+++ b/Engine/Source/Runtime/CoreUObject/Public/UObject/UnrealType.h
@@ -2178,10 +2178,10 @@ public:
 	FORCEINLINE void SetPropertyClass(UClass* NewPropertyClass) { PropertyClass = NewPropertyClass; }
 #endif // USE_CIRCULAR_DEPENDENCY_LOAD_DEFERRING
 
+	virtual void CheckValidObject(void* Value) const;
+
 protected:
 	virtual bool AllowCrossLevel() const;
-
-	virtual void CheckValidObject(void* Value) const;
 	// End of FObjectPropertyBase interface
 };
 
@@ -2308,9 +2308,12 @@ class COREUOBJECT_API FObjectPtrProperty : public FObjectProperty
 	virtual bool SameType(const FProperty* Other) const override;
 	virtual bool Identical(const void* A, const void* B, uint32 PortFlags) const override;
 	virtual void SerializeItem(FStructuredArchive::FSlot Slot, void* Value, void const* Defaults) const override;
-	virtual void EmitReferenceInfo(UClass& OwnerClass, int32 BaseOffset, TArray<const FStructProperty*>& EncounteredStructProps) override;
 	// End of FProperty interface
 
+	// Helper methods for sharing code with FClassPtrProperty even though one doesn't inherit from the other
+	static void StaticSerializeItem(const FObjectPropertyBase* ObjectProperty, FStructuredArchive::FSlot Slot, void* Value, void const* Defaults);
+	static bool StaticIdentical(const void* A, const void* B, uint32 PortFlags);
+
 	// FObjectProperty interface
 	virtual UObject* GetObjectPropertyValue(const void* PropertyValueAddress) const override;
 	virtual void SetObjectPropertyValue(void* PropertyValueAddress, UObject* Value) const override;
@@ -2555,6 +2558,42 @@ public:
 #endif // USE_CIRCULAR_DEPENDENCY_LOAD_DEFERRING
 };
 
+//
+// Describes a reference variable to another object which may be nil.
+//
+class COREUOBJECT_API FClassPtrProperty : public FClassProperty
+{
+	DECLARE_FIELD(FClassPtrProperty, FClassProperty, CASTCLASS_FClassPtrProperty)
+
+	using Super::Super;
+
+#if WITH_EDITORONLY_DATA
+	explicit FClassPtrProperty(UField* InField)
+		: FClassProperty(InField)
+	{
+	}
+#endif // WITH_EDITORONLY_DATA
+
+	// UHT interface
+	virtual FString GetCPPMacroType( FString& ExtendedTypeText ) const  override;
+	virtual FString GetCPPType( FString* ExtendedTypeText, uint32 CPPExportFlags ) const override;
+	// End of UHT interface
+
+	// FProperty interface
+	virtual bool SameType(const FProperty* Other) const override;
+	virtual bool Identical(const void* A, const void* B, uint32 PortFlags) const override;
+	virtual void SerializeItem(FStructuredArchive::FSlot Slot, void* Value, void const* Defaults) const override;
+	// End of FProperty interface
+
+	// FObjectProperty interface
+	virtual UObject* GetObjectPropertyValue(const void* PropertyValueAddress) const override;
+	virtual void SetObjectPropertyValue(void* PropertyValueAddress, UObject* Value) const override;
+private:
+	virtual uint32 GetValueTypeHashInternal(const void* Src) const override;
+public:
+	// End of FObjectProperty interface
+};
+
 /*-----------------------------------------------------------------------------
 	FSoftClassProperty.
 -----------------------------------------------------------------------------*/
