U5 - Don't do latency critical work on Backgroundworkers.

diff --git a/Engine/Source/Runtime/CoreUObject/Public/UObject/FastReferenceCollector.h b/Engine/Source/Runtime/CoreUObject/Public/UObject/FastReferenceCollector.h
--- a/Engine/Source/Runtime/CoreUObject/Public/UObject/FastReferenceCollector.h
+++ b/Engine/Source/Runtime/CoreUObject/Public/UObject/FastReferenceCollector.h
@@ -14,6 +14,7 @@
 #include "UObject/FieldPath.h"
 #include "UObject/UObjectArray.h"
 #include "UObject/FastReferenceCollectorOptions.h"
+#include "Async/ParallelFor.h"
 
 struct FStackEntry;
 
@@ -555,9 +556,19 @@ public:
 					ChunkTasks.Add(TGraphTask< FCollectorTaskProcessorTask >::CreateTask().ConstructAndDispatchWhenReady(TaskQueue, Chunk >= NumThreads ? BackgroundThreadName : NormalThreadName));
 				}
 
-				QUICK_SCOPE_CYCLE_COUNTER(STAT_GC_Subtask_Wait);
-				FTaskGraphInterface::Get().WaitUntilTasksComplete(ChunkTasks, ENamedThreads::GameThread_Local);
-				TaskQueue.CheckDone();
+				extern CORE_API int32 GUseNewTaskBackend;
+				if(GUseNewTaskBackend)
+				{
+					QUICK_SCOPE_CYCLE_COUNTER(STAT_GC_Subtask_Wait);
+					ParallelForTemplate(NumTasks, [this](int32){ TaskQueue.DoTask(); });
+					TaskQueue.CheckDone();
+				}
+				else
+				{
+					QUICK_SCOPE_CYCLE_COUNTER(STAT_GC_Subtask_Wait);
+					FTaskGraphInterface::Get().WaitUntilTasksComplete(ChunkTasks, ENamedThreads::GameThread_Local);
+					TaskQueue.CheckDone();
+				}
 			}
 		}
 	}
