U5 - for some reason the collector tasks don't like to run on the Gamethread

diff --git a/Engine/Source/Runtime/CoreUObject/Public/UObject/FastReferenceCollector.h b/Engine/Source/Runtime/CoreUObject/Public/UObject/FastReferenceCollector.h
--- a/Engine/Source/Runtime/CoreUObject/Public/UObject/FastReferenceCollector.h
+++ b/Engine/Source/Runtime/CoreUObject/Public/UObject/FastReferenceCollector.h
@@ -553,22 +553,14 @@ public:
 				}
 				for (int32 Chunk = 0; Chunk < NumTasks; Chunk++)
 				{
-					ChunkTasks.Add(TGraphTask< FCollectorTaskProcessorTask >::CreateTask().ConstructAndDispatchWhenReady(TaskQueue, Chunk >= NumThreads ? BackgroundThreadName : NormalThreadName));
+					extern CORE_API int32 GUseNewTaskBackend;
+					ENamedThreads::Type ThreadPriority = GUseNewTaskBackend ? ENamedThreads::AnyHiPriThreadHiPriTask : ((Chunk >= NumThreads) ? BackgroundThreadName : NormalThreadName);
+					ChunkTasks.Add(TGraphTask< FCollectorTaskProcessorTask >::CreateTask().ConstructAndDispatchWhenReady(TaskQueue, ThreadPriority));
 				}
 
-				extern CORE_API int32 GUseNewTaskBackend;
-				if(GUseNewTaskBackend)
-				{
-					QUICK_SCOPE_CYCLE_COUNTER(STAT_GC_Subtask_Wait);
-					ParallelForTemplate(NumTasks, [this](int32){ TaskQueue.DoTask(); });
-					TaskQueue.CheckDone();
-				}
-				else
-				{
-					QUICK_SCOPE_CYCLE_COUNTER(STAT_GC_Subtask_Wait);
-					FTaskGraphInterface::Get().WaitUntilTasksComplete(ChunkTasks, ENamedThreads::GameThread_Local);
-					TaskQueue.CheckDone();
-				}
+				QUICK_SCOPE_CYCLE_COUNTER(STAT_GC_Subtask_Wait);
+				FTaskGraphInterface::Get().WaitUntilTasksComplete(ChunkTasks, ENamedThreads::GameThread_Local);
+				TaskQueue.CheckDone();
 			}
 		}
 	}
