Fix quadratic search into an array when there are many dependencies between packages
  - 4m49s to 1m25s to load Chase_Frogger_P

diff --git a/Engine/Source/Runtime/CoreUObject/Private/Serialization/AsyncLoading.cpp b/Engine/Source/Runtime/CoreUObject/Private/Serialization/AsyncLoading.cpp
--- a/Engine/Source/Runtime/CoreUObject/Private/Serialization/AsyncLoading.cpp
+++ b/Engine/Source/Runtime/CoreUObject/Private/Serialization/AsyncLoading.cpp
@@ -404,9 +404,9 @@ void FAsyncPackage::PopulateFlushTree(struct FFlushTree* FlushTree)
 {
 	if (FlushTree->AddPackage(GetPackageName()))
 	{
-		for (FAsyncPackage* PendingImport : PendingImportedPackages)
+		for (TTuple<FName, FAsyncPackage*>& PendingImport : PendingImportedPackages)
 		{
-			PendingImport->PopulateFlushTree(FlushTree);
+			PendingImport.Value->PopulateFlushTree(FlushTree);
 		}
 	}
 }
@@ -5924,25 +5924,6 @@ EAsyncPackageState::Type FAsyncPackage::FinishLinker()
 	return Result;
 }
 
-/**
- * Find a package by name.
- * 
- * @param Dependencies package list.
- * @param PackageName long package name.
- * @return Index into the array if the package was found, otherwise INDEX_NONE
- */
-FORCEINLINE int32 ContainsDependencyPackage(TArray<FAsyncPackage*>& Dependencies, const FName& PackageName)
-{
-	for (int32 Index = 0; Index < Dependencies.Num(); ++Index)
-	{
-		if (Dependencies[Index]->GetPackageName() == PackageName)
-		{
-			return Index;
-		}
-	}
-	return INDEX_NONE;
-}
-
 /**
  * Adds a package to the list of pending import packages.
  *
@@ -5992,7 +5973,7 @@ void FAsyncPackage::AddImportDependency(const FName& PendingImport, const FName&
 		TUniquePtr<FLoadPackageAsyncDelegate> InternalDelegate = MakeUnique<FLoadPackageAsyncDelegate>(FLoadPackageAsyncDelegate::CreateRaw(this, &FAsyncPackage::ImportFullyLoadedCallback));
 		PackageToStream->AddCompletionCallback(MoveTemp(InternalDelegate), bInternalCallback);
 		PackageToStream->DependencyRefCount.Increment();
-		PendingImportedPackages.Add(PackageToStream);
+		PendingImportedPackages.Emplace(PackageToStream->GetPackageName(), PackageToStream);
 		if (FlushTree)
 		{
 			PackageToStream->PopulateFlushTree(FlushTree);
@@ -6013,7 +5994,7 @@ void FAsyncPackage::AddImportDependency(const FName& PendingImport, const FName&
  */
 bool FAsyncPackage::AddUniqueLinkerDependencyPackage(FAsyncPackage& PendingImport, FFlushTree* FlushTree)
 {
-	if (ContainsDependencyPackage(PendingImportedPackages, PendingImport.GetPackageName()) == INDEX_NONE)
+	if (!PendingImportedPackages.Contains(PendingImport.GetPackageName()))
 	{
 		FLinkerLoad* PendingImportLinker = PendingImport.Linker;
 		if (PendingImportLinker == nullptr || !PendingImportLinker->HasFinishedInitialization())
@@ -6041,9 +6022,9 @@ void FAsyncPackage::AddDependencyTree(FAsyncPackage& ImportedPackage, TSet<FAsyn
 		// we've already searched this package
 		return;
 	}
-	for (int32 Index = 0; Index < ImportedPackage.PendingImportedPackages.Num(); ++Index)
+	for (TTuple<FName, FAsyncPackage*>& Kvp : ImportedPackage.PendingImportedPackages)
 	{
-		FAsyncPackage& PendingImport = *ImportedPackage.PendingImportedPackages[Index];
+		FAsyncPackage& PendingImport = *Kvp.Value;
 		if (!AddUniqueLinkerDependencyPackage(PendingImport, FlushTree))
 		{
 			AddDependencyTree(PendingImport, SearchedPackages, FlushTree);
@@ -6060,6 +6041,8 @@ void FAsyncPackage::AddDependencyTree(FAsyncPackage& ImportedPackage, TSet<FAsyn
  */
 EAsyncPackageState::Type FAsyncPackage::LoadImports(FFlushTree* FlushTree)
 {
+	TRACE_CPUPROFILER_EVENT_SCOPE(FAsyncPackage::LoadImports);
+
 	SCOPE_CYCLE_COUNTER(STAT_FAsyncPackage_LoadImports);
 	LastObjectWorkWasPerformedOn	= LinkerRoot;
 	LastTypeOfWorkPerformed			= TEXT("loading imports");
@@ -6131,7 +6114,7 @@ EAsyncPackageState::Type FAsyncPackage::LoadImports(FFlushTree* FlushTree)
 			}
 		}
 
-		if (!ExistingPackage && ContainsDependencyPackage(PendingImportedPackages, ImportPackageFName) == INDEX_NONE)
+		if (!ExistingPackage && !PendingImportedPackages.Contains(ImportPackageFName))
 		{
 			const FString ImportPackageName(ImportPackageFName.ToString());
 			// The package doesn't exist and this import is not in the dependency list so add it now.
@@ -6166,13 +6149,13 @@ void FAsyncPackage::ImportFullyLoadedCallback(const FName& InPackageName, UPacka
 	if (Result != EAsyncLoadingResult::Canceled)
 	{
 		UE_LOG(LogStreaming, Verbose, TEXT("FAsyncPackage::LoadImports for %s: Loaded %s"), *Desc.PackagePath.GetDebugName(), *InPackageName.ToString());
-		int32 Index = ContainsDependencyPackage(PendingImportedPackages, InPackageName);
-		if (Index != INDEX_NONE)
+		FAsyncPackage** AsyncPackage = PendingImportedPackages.Find(InPackageName);
+		if (AsyncPackage != nullptr)
 		{
-		// Keep a reference to this package so that its linker doesn't go away too soon
-		ReferencedImports.Add(PendingImportedPackages[Index]);
-		PendingImportedPackages.RemoveAt(Index);
-	}
+			// Keep a reference to this package so that its linker doesn't go away too soon
+			ReferencedImports.Add(*AsyncPackage);
+			PendingImportedPackages.Remove(InPackageName);
+		}
 	}
 }
 
diff --git a/Engine/Source/Runtime/CoreUObject/Private/Serialization/AsyncLoading.h b/Engine/Source/Runtime/CoreUObject/Private/Serialization/AsyncLoading.h
--- a/Engine/Source/Runtime/CoreUObject/Private/Serialization/AsyncLoading.h
+++ b/Engine/Source/Runtime/CoreUObject/Private/Serialization/AsyncLoading.h
@@ -545,7 +545,7 @@ private:
 	/** Call backs called when we finished loading this package											*/
 	TArray<FCompletionCallback>	CompletionCallbacks;
 	/** Pending Import packages - we wait until all of them have been fully loaded. */
-	TArray<FAsyncPackage*> PendingImportedPackages;
+	TMap<FName, FAsyncPackage*> PendingImportedPackages;
 	/** Referenced imports - list of packages we need until we finish loading this package. */
 	TArray<FAsyncPackage*> ReferencedImports;
 	/** Root package if this package was loaded as a dependency of another. NULL otherwise */
diff --git a/Engine/Source/Runtime/CoreUObject/Private/UObject/LinkerLoad.cpp b/Engine/Source/Runtime/CoreUObject/Private/UObject/LinkerLoad.cpp
--- a/Engine/Source/Runtime/CoreUObject/Private/UObject/LinkerLoad.cpp
+++ b/Engine/Source/Runtime/CoreUObject/Private/UObject/LinkerLoad.cpp
@@ -4070,6 +4070,8 @@ UObject* FLinkerLoad::Create( UClass* ObjectClass, FName ObjectName, UObject* Ou
 
 void FLinkerLoad::Preload( UObject* Object )
 {
+	TRACE_CPUPROFILER_EVENT_SCOPE(FLinkerLoad::Preload);
+
 	//check(IsValidLowLevel());
 	check(Object);
 
