Narrow the GIsSavingPackage check that prevent object creation and FindObject calls to be restricted to the actual package being saved by using a temporary package flag

diff --git a/Engine/Source/Developer/IoStoreUtilities/Private/PackageStoreOptimizer.cpp b/Engine/Source/Developer/IoStoreUtilities/Private/PackageStoreOptimizer.cpp
--- a/Engine/Source/Developer/IoStoreUtilities/Private/PackageStoreOptimizer.cpp
+++ b/Engine/Source/Developer/IoStoreUtilities/Private/PackageStoreOptimizer.cpp
@@ -117,7 +117,7 @@ void FPackageStoreOptimizer::LoadCookedHeader(FPackageStorePackage* Package, con
 		Ar << Package->Summary;
 	}
 
-	Package->PackageFlags = Package->Summary.PackageFlags;
+	Package->PackageFlags = Package->Summary.GetPackageFlags();
 	Package->CookedHeaderSize = Package->Summary.TotalHeaderSize;
 
 	Ar.SetFilterEditorOnly((Package->PackageFlags & EPackageFlags::PKG_FilterEditorOnly) != 0);
diff --git a/Engine/Source/Editor/UnrealEd/Private/Commandlets/GatherTextFromAssetsCommandlet.cpp b/Engine/Source/Editor/UnrealEd/Private/Commandlets/GatherTextFromAssetsCommandlet.cpp
--- a/Engine/Source/Editor/UnrealEd/Private/Commandlets/GatherTextFromAssetsCommandlet.cpp
+++ b/Engine/Source/Editor/UnrealEd/Private/Commandlets/GatherTextFromAssetsCommandlet.cpp
@@ -660,7 +660,7 @@ int32 UGatherTextFromAssetsCommandlet::Main(const FString& Params)
 		if (bSkipGatherCache && PackageFileSummary.GetFileVersionUE4() >= VER_UE4_SERIALIZE_TEXT_IN_PACKAGES)
 		{
 			// Fallback on the old package flag check.
-			if (PackageFileSummary.PackageFlags & PKG_RequiresLocalizationGather)
+			if (PackageFileSummary.GetPackageFlags() & PKG_RequiresLocalizationGather)
 			{
 				PackagePendingGather.PackageLocCacheState = EPackageLocCacheState::Uncached_NoCache;
 			}
@@ -677,7 +677,7 @@ int32 UGatherTextFromAssetsCommandlet::Main(const FString& Params)
 		else if (PackageFileSummary.GetFileVersionUE4() < VER_UE4_SERIALIZE_TEXT_IN_PACKAGES || (!EditorVersion || EditorVersion->Version < FEditorObjectVersion::GatheredTextEditorOnlyPackageLocId))
 		{
 			// Fallback on the old package flag check.
-			if (PackageFileSummary.PackageFlags & PKG_RequiresLocalizationGather)
+			if (PackageFileSummary.GetPackageFlags() & PKG_RequiresLocalizationGather)
 			{
 				PackagePendingGather.PackageLocCacheState = EPackageLocCacheState::Uncached_TooOld;
 			}
@@ -696,7 +696,7 @@ int32 UGatherTextFromAssetsCommandlet::Main(const FString& Params)
 		}
 
 		// If this package doesn't have any cached data, then we have to load it for gather
-		if (PackageFileSummary.GetFileVersionUE4() >= VER_UE4_SERIALIZE_TEXT_IN_PACKAGES && PackageFileSummary.GatherableTextDataOffset == 0 && (PackageFileSummary.PackageFlags & PKG_RequiresLocalizationGather))
+		if (PackageFileSummary.GetFileVersionUE4() >= VER_UE4_SERIALIZE_TEXT_IN_PACKAGES && PackageFileSummary.GatherableTextDataOffset == 0 && (PackageFileSummary.GetPackageFlags() & PKG_RequiresLocalizationGather))
 		{
 			PackagePendingGather.PackageLocCacheState = EPackageLocCacheState::Uncached_NoCache;
 		}
diff --git a/Engine/Source/Editor/UnrealEd/Private/Commandlets/PackageUtilities.cpp b/Engine/Source/Editor/UnrealEd/Private/Commandlets/PackageUtilities.cpp
--- a/Engine/Source/Editor/UnrealEd/Private/Commandlets/PackageUtilities.cpp
+++ b/Engine/Source/Editor/UnrealEd/Private/Commandlets/PackageUtilities.cpp
@@ -947,7 +947,7 @@ void FPkgInfoReporter_Log::GeneratePackageReport( FLinkerLoad* InLinker /*=nullp
 	Out.Logf(ELogVerbosity::Display, TEXT("\t     File Version: %i"), Linker->UE4Ver() );
 	Out.Logf(ELogVerbosity::Display, TEXT("\t   Engine Version: %s"), *Linker->Summary.SavedByEngineVersion.ToString());
 	Out.Logf(ELogVerbosity::Display, TEXT("\t   Compat Version: %s"), *Linker->Summary.CompatibleWithEngineVersion.ToString());
-	Out.Logf(ELogVerbosity::Display, TEXT("\t     PackageFlags: %X"), Linker->Summary.PackageFlags );
+	Out.Logf(ELogVerbosity::Display, TEXT("\t     PackageFlags: %X"), Linker->Summary.GetPackageFlags() );
 	Out.Logf(ELogVerbosity::Display, TEXT("\t        NameCount: %d"), Linker->Summary.NameCount );
 	Out.Logf(ELogVerbosity::Display, TEXT("\t       NameOffset: %d"), Linker->Summary.NameOffset );
 	Out.Logf(ELogVerbosity::Display, TEXT("\t      ImportCount: %d"), Linker->Summary.ImportCount );
diff --git a/Engine/Source/Editor/UnrealEd/Private/Commandlets/StabilizeLocalizationKeys.cpp b/Engine/Source/Editor/UnrealEd/Private/Commandlets/StabilizeLocalizationKeys.cpp
--- a/Engine/Source/Editor/UnrealEd/Private/Commandlets/StabilizeLocalizationKeys.cpp
+++ b/Engine/Source/Editor/UnrealEd/Private/Commandlets/StabilizeLocalizationKeys.cpp
@@ -254,7 +254,7 @@ int32 UStabilizeLocalizationKeysCommandlet::Main(const FString& Params)
 			FPackageFileSummary PackageFileSummary;
 			(*FileReader) << PackageFileSummary;
 
-			const bool bRequiresKeyStabilization = !!(PackageFileSummary.PackageFlags & PKG_RequiresLocalizationGather);
+			const bool bRequiresKeyStabilization = !!(PackageFileSummary.GetPackageFlags() & PKG_RequiresLocalizationGather);
 			if (bRequiresKeyStabilization)
 			{
 				UnstablePackages.Add(PackageFilename);
diff --git a/Engine/Source/Editor/UnrealEd/Private/EditorServer.cpp b/Engine/Source/Editor/UnrealEd/Private/EditorServer.cpp
--- a/Engine/Source/Editor/UnrealEd/Private/EditorServer.cpp
+++ b/Engine/Source/Editor/UnrealEd/Private/EditorServer.cpp
@@ -2401,7 +2401,7 @@ bool UEditorEngine::PackageIsAMapFile( const TCHAR* PackageFilename, FText& OutN
 		delete CheckMapPackageFile;
 
 		// Check flag.
-		if( ( Summary.PackageFlags & PKG_ContainsMap ) == 0 )
+		if( ( Summary.GetPackageFlags() & PKG_ContainsMap ) == 0 )
 		{
 			FFormatNamedArguments Arguments;
 			Arguments.Add(TEXT("File"), FText::FromString( FString( PackageFilename ) ));
diff --git a/Engine/Source/Runtime/AssetRegistry/Private/PackageReader.cpp b/Engine/Source/Runtime/AssetRegistry/Private/PackageReader.cpp
--- a/Engine/Source/Runtime/AssetRegistry/Private/PackageReader.cpp
+++ b/Engine/Source/Runtime/AssetRegistry/Private/PackageReader.cpp
@@ -266,7 +266,7 @@ bool FPackageReader::ReadAssetDataFromThumbnailCache(TArray<FAssetData*>& AssetD
 		}
 
 		// Create a new FAssetData for this asset and update it with the gathered data
-		AssetDataList.Add(new FAssetData(FName(*PackageName), FName(*PackagePath), FName(*ObjectPathWithoutPackageName), FName(*AssetClassName), FAssetDataTagMap(), PackageFileSummary.ChunkIDs, PackageFileSummary.PackageFlags));
+		AssetDataList.Add(new FAssetData(FName(*PackageName), FName(*PackagePath), FName(*ObjectPathWithoutPackageName), FName(*AssetClassName), FAssetDataTagMap(), PackageFileSummary.ChunkIDs, PackageFileSummary.GetPackageFlags()));
 	}
 
 	return true;
@@ -593,7 +593,7 @@ int64 FPackageReader::TotalSize()
 
 uint32 FPackageReader::GetPackageFlags() const
 {
-	return PackageFileSummary.PackageFlags;
+	return PackageFileSummary.GetPackageFlags();
 }
 
 FArchive& FPackageReader::operator<<( FName& Name )
@@ -638,10 +638,10 @@ namespace AssetRegistry
 
 		const FString PackagePath = FPackageName::GetLongPackagePath(PackageName);
 		const int64 PackageFileSize = BinaryArchive.TotalSize();
-		const bool bIsMapPackage = (PackageFileSummary.PackageFlags & PKG_ContainsMap) != 0;
+		const bool bIsMapPackage = (PackageFileSummary.GetPackageFlags() & PKG_ContainsMap) != 0;
 
 		// To avoid large patch sizes, we have frozen cooked package format at the format before VER_UE4_ASSETREGISTRY_DEPENDENCYFLAGS
-		bool bPreDependencyFormat = PackageFileSummary.GetFileVersionUE4() < VER_UE4_ASSETREGISTRY_DEPENDENCYFLAGS || !!(PackageFileSummary.PackageFlags & PKG_FilterEditorOnly);
+		bool bPreDependencyFormat = PackageFileSummary.GetFileVersionUE4() < VER_UE4_ASSETREGISTRY_DEPENDENCYFLAGS || !!(PackageFileSummary.GetPackageFlags() & PKG_FilterEditorOnly);
 
 		// Load offsets to optionally-read data
 		if (bPreDependencyFormat)
@@ -672,7 +672,7 @@ namespace AssetRegistry
 			if (bLegacyPackage || bNoMapAsset)
 			{
 				FString AssetName = FPackageName::GetLongPackageAssetName(PackageName);
-				OutAssetDataList.Add(new FAssetData(FName(*PackageName), FName(*PackagePath), FName(*AssetName), FName(TEXT("World")), FAssetDataTagMap(), PackageFileSummary.ChunkIDs, PackageFileSummary.PackageFlags));
+				OutAssetDataList.Add(new FAssetData(FName(*PackageName), FName(*PackagePath), FName(*AssetName), FName(TEXT("World")), FAssetDataTagMap(), PackageFileSummary.ChunkIDs, PackageFileSummary.GetPackageFlags()));
 			}
 		}
 
@@ -745,7 +745,7 @@ namespace AssetRegistry
 			}
 
 			// Create a new FAssetData for this asset and update it with the gathered data
-			OutAssetDataList.Add(new FAssetData(PackageName, ObjectPath, FName(*ObjectClassName), MoveTemp(TagsAndValues), PackageFileSummary.ChunkIDs, PackageFileSummary.PackageFlags));
+			OutAssetDataList.Add(new FAssetData(PackageName, ObjectPath, FName(*ObjectClassName), MoveTemp(TagsAndValues), PackageFileSummary.ChunkIDs, PackageFileSummary.GetPackageFlags()));
 		}
 
 		return true;
diff --git a/Engine/Source/Runtime/CoreUObject/Private/Serialization/AsyncLoading.cpp b/Engine/Source/Runtime/CoreUObject/Private/Serialization/AsyncLoading.cpp
--- a/Engine/Source/Runtime/CoreUObject/Private/Serialization/AsyncLoading.cpp
+++ b/Engine/Source/Runtime/CoreUObject/Private/Serialization/AsyncLoading.cpp
@@ -4370,6 +4370,7 @@ EAsyncPackageState::Type FAsyncLoadingThread::ProcessLoadedPackages(bool bUseTim
 
 		
 	bDidSomething = LoadedPackagesToProcess.Num() > 0;
+
 	for (int32 PackageIndex = 0; PackageIndex < LoadedPackagesToProcess.Num() && !IsAsyncLoadingSuspendedInternal(); ++PackageIndex)
 	{
 		FAsyncPackage* Package = LoadedPackagesToProcess[PackageIndex];
@@ -7270,7 +7271,7 @@ void FAsyncArchive::ReadCallback(bool bWasCancelled, IAsyncReadRequest* Request)
 					FMaxPackageSummarySize::Value, *PackagePath.GetDebugName(), Ar.Tell());
 				
 				// Support for cooked EDL packages in the editor
-				bCookedForEDLInEditor = !FPlatformProperties::RequiresCookedData() && (Sum.PackageFlags & PKG_FilterEditorOnly) && Sum.PreloadDependencyCount > 0 && Sum.PreloadDependencyOffset > 0;
+				bCookedForEDLInEditor = !FPlatformProperties::RequiresCookedData() && (Sum.GetPackageFlags() & PKG_FilterEditorOnly) && Sum.PreloadDependencyCount > 0 && Sum.PreloadDependencyOffset > 0;
 
 				HeaderSize = Sum.TotalHeaderSize;
 				LogItem(TEXT("Starting Header"), 0, HeaderSize);
diff --git a/Engine/Source/Runtime/CoreUObject/Private/UObject/LinkerDiff.cpp b/Engine/Source/Runtime/CoreUObject/Private/UObject/LinkerDiff.cpp
--- a/Engine/Source/Runtime/CoreUObject/Private/UObject/LinkerDiff.cpp
+++ b/Engine/Source/Runtime/CoreUObject/Private/UObject/LinkerDiff.cpp
@@ -95,7 +95,7 @@ void FLinkerDiff::GenerateSummaryDiff(FLinker* LHSLinker, FLinker* RHSLinker)
 	//CustomVersionContainer
 
 	//COMPARE_MEMBER(LHSSummary, RHSSummary, TotalHeaderSize);
-	COMPARE_MEMBER(LHSSummary, RHSSummary, PackageFlags);
+	COMPARE_MEMBER(LHSSummary, RHSSummary, GetPackageFlags());
 	//FolderName
 	COMPARE_MEMBER(LHSSummary, RHSSummary, NameCount);
 	//COMPARE_MEMBER(LHSSummary, RHSSummary, NameOffset);
diff --git a/Engine/Source/Runtime/CoreUObject/Private/UObject/LinkerLoad.cpp b/Engine/Source/Runtime/CoreUObject/Private/UObject/LinkerLoad.cpp
--- a/Engine/Source/Runtime/CoreUObject/Private/UObject/LinkerLoad.cpp
+++ b/Engine/Source/Runtime/CoreUObject/Private/UObject/LinkerLoad.cpp
@@ -1266,7 +1266,7 @@ FLinkerLoad::ELinkerStatus FLinkerLoad::SerializePackageFileSummaryInternal()
 	}
 
 	// Set desired property tag format
-	bool bUseUnversionedProperties = (Summary.PackageFlags & PKG_UnversionedProperties) != 0;
+	bool bUseUnversionedProperties = (Summary.GetPackageFlags() & PKG_UnversionedProperties) != 0;
 	SetUseUnversionedPropertySerialization(bUseUnversionedProperties);
 	Loader->SetUseUnversionedPropertySerialization(bUseUnversionedProperties);
 
@@ -1297,14 +1297,14 @@ FLinkerLoad::ELinkerStatus FLinkerLoad::SerializePackageFileSummaryInternal()
 	}
 
 	// don't load packages that contain editor only data in builds that don't support that and vise versa
-	if (!FPlatformProperties::HasEditorOnlyData() && !(Summary.PackageFlags & PKG_FilterEditorOnly))
+	if (!FPlatformProperties::HasEditorOnlyData() && !(Summary.GetPackageFlags() & PKG_FilterEditorOnly))
 	{
 		UE_LOG(LogLinker, Warning, TEXT("Unable to load package (%s). Package contains EditorOnly data which is not supported by the current build."), *GetDebugName());
 		return LINKER_Failed;
 	}
 
 	// don't load packages that contain editor only data in builds that don't support that and vise versa
-	if (FPlatformProperties::HasEditorOnlyData() && !!(Summary.PackageFlags & PKG_FilterEditorOnly))
+	if (FPlatformProperties::HasEditorOnlyData() && !!(Summary.GetPackageFlags() & PKG_FilterEditorOnly))
 	{
 		// This warning can be disabled in ini or project settings
 		if (!GAllowCookedDataInEditorBuilds)
@@ -1327,7 +1327,7 @@ FLinkerLoad::ELinkerStatus FLinkerLoad::SerializePackageFileSummaryInternal()
 #if PLATFORM_WINDOWS && DO_GUARD_SLOW
 	if (!FPlatformProperties::RequiresCookedData() &&
 		// We can't check the post tag if the file is an EDL cooked package
-		!((Summary.PackageFlags & PKG_FilterEditorOnly) && Summary.PreloadDependencyCount > 0 && Summary.PreloadDependencyOffset > 0)
+		!((Summary.GetPackageFlags() & PKG_FilterEditorOnly) && Summary.PreloadDependencyCount > 0 && Summary.PreloadDependencyOffset > 0)
 		&& !IsTextFormat())
 	{
 		// check if this package version stored the 4-byte magic post tag
@@ -1443,7 +1443,7 @@ FLinkerLoad::ELinkerStatus FLinkerLoad::UpdateFromPackageFileSummary()
 	SetEngineVer(Summary.SavedByEngineVersion);
 	SetCustomVersions(SummaryVersions);
 
-	if (Summary.PackageFlags & PKG_FilterEditorOnly)
+	if (Summary.GetPackageFlags() & PKG_FilterEditorOnly)
 	{
 		SetFilterEditorOnly(true);
 	}
@@ -1470,7 +1470,7 @@ FLinkerLoad::ELinkerStatus FLinkerLoad::UpdateFromPackageFileSummary()
 	if (UPackage* LinkerRootPackage = LinkerRoot)
 	{
 		// Preserve PIE package flag
-		uint32 NewPackageFlags = Summary.PackageFlags;
+		uint32 NewPackageFlags = Summary.GetPackageFlags();
 		if (LinkerRootPackage->HasAnyPackageFlags(PKG_PlayInEditor))
 		{
 			NewPackageFlags |= PKG_PlayInEditor;
@@ -1509,7 +1509,7 @@ FLinkerLoad::ELinkerStatus FLinkerLoad::UpdateFromPackageFileSummary()
 		}
 
 #if WITH_EDITORONLY_DATA
-		LinkerRootPackage->bIsCookedForEditor = !!(Summary.PackageFlags & PKG_FilterEditorOnly);
+		LinkerRootPackage->bIsCookedForEditor = !!(Summary.GetPackageFlags() & PKG_FilterEditorOnly);
 #endif
 	}
 
diff --git a/Engine/Source/Runtime/CoreUObject/Private/UObject/LinkerSave.cpp b/Engine/Source/Runtime/CoreUObject/Private/UObject/LinkerSave.cpp
--- a/Engine/Source/Runtime/CoreUObject/Private/UObject/LinkerSave.cpp
+++ b/Engine/Source/Runtime/CoreUObject/Private/UObject/LinkerSave.cpp
@@ -43,7 +43,7 @@ FLinkerSave::FLinkerSave(UPackage* InParent, const TCHAR* InFilename, bool bForc
 		Summary.SetFileVersions( GPackageFileUE4Version, GPackageFileLicenseeUE4Version, bInSaveUnversioned );
 		Summary.SavedByEngineVersion = FEngineVersion::Current();
 		Summary.CompatibleWithEngineVersion = FEngineVersion::CompatibleWith();
-		Summary.PackageFlags = Package ? (Package->GetPackageFlags() & ~PKG_NewlyCreated) : 0;
+		Summary.SetPackageFlags(Package ? Package->GetPackageFlags() : 0);
 
 #if USE_STABLE_LOCALIZATION_KEYS
 		if (GIsEditor)
@@ -97,7 +97,7 @@ FLinkerSave::FLinkerSave(UPackage* InParent, FArchive *InSaver, bool bForceByteS
 		Summary.SetFileVersions(GPackageFileUE4Version, GPackageFileLicenseeUE4Version, bInSaveUnversioned);
 		Summary.SavedByEngineVersion = FEngineVersion::Current();
 		Summary.CompatibleWithEngineVersion = FEngineVersion::CompatibleWith();
-		Summary.PackageFlags = Package ? (Package->GetPackageFlags() & ~PKG_NewlyCreated) : 0;
+		Summary.SetPackageFlags(Package ? Package->GetPackageFlags() : 0);
 
 #if USE_STABLE_LOCALIZATION_KEYS
 		if (GIsEditor)
@@ -146,7 +146,7 @@ FLinkerSave::FLinkerSave(UPackage* InParent, bool bForceByteSwapping, bool bInSa
 		Summary.SetFileVersions( GPackageFileUE4Version, GPackageFileLicenseeUE4Version, bInSaveUnversioned );
 		Summary.SavedByEngineVersion = FEngineVersion::Current();
 		Summary.CompatibleWithEngineVersion = FEngineVersion::CompatibleWith();
-		Summary.PackageFlags = Package ? (Package->GetPackageFlags() & ~PKG_NewlyCreated) : 0;
+		Summary.SetPackageFlags(Package ? Package->GetPackageFlags() : 0);
 
 #if USE_STABLE_LOCALIZATION_KEYS
 		if (GIsEditor)
@@ -393,7 +393,7 @@ void FLinkerSave::SetUseUnversionedPropertySerialization(bool bInUseUnversioned)
 	}
 	if (bInUseUnversioned)
 	{
-		Summary.PackageFlags |= PKG_UnversionedProperties;
+		Summary.SetPackageFlags(Summary.GetPackageFlags() | PKG_UnversionedProperties);
 		if (LinkerRoot)
 		{
 			LinkerRoot->SetPackageFlags(PKG_UnversionedProperties);
@@ -401,7 +401,7 @@ void FLinkerSave::SetUseUnversionedPropertySerialization(bool bInUseUnversioned)
 	}
 	else
 	{
-		Summary.PackageFlags &= ~PKG_UnversionedProperties;
+		Summary.SetPackageFlags(Summary.GetPackageFlags() & ~PKG_UnversionedProperties);
 		if (LinkerRoot)
 		{
 			LinkerRoot->ClearPackageFlags(PKG_UnversionedProperties);
diff --git a/Engine/Source/Runtime/CoreUObject/Private/UObject/PackageFileSummary.cpp b/Engine/Source/Runtime/CoreUObject/Private/UObject/PackageFileSummary.cpp
--- a/Engine/Source/Runtime/CoreUObject/Private/UObject/PackageFileSummary.cpp
+++ b/Engine/Source/Runtime/CoreUObject/Private/UObject/PackageFileSummary.cpp
@@ -397,3 +397,9 @@ void FPackageFileSummary::SetCustomVersionContainer(const FCustomVersionContaine
 	CustomVersionContainer = InContainer;
 	CustomVersionContainer.SortByKey();
 }
+
+void FPackageFileSummary::SetPackageFlags(uint32 InPackageFlags)
+{
+	// Package summary flags cannot contain the 'NewlyCreated' nor the 'IsSaving' flag
+	PackageFlags = InPackageFlags & ~(PKG_NewlyCreated | PKG_IsSaving);
+}
diff --git a/Engine/Source/Runtime/CoreUObject/Private/UObject/SavePackage.cpp b/Engine/Source/Runtime/CoreUObject/Private/UObject/SavePackage.cpp
--- a/Engine/Source/Runtime/CoreUObject/Private/UObject/SavePackage.cpp
+++ b/Engine/Source/Runtime/CoreUObject/Private/UObject/SavePackage.cpp
@@ -2316,9 +2316,10 @@ FSavePackageResultStruct UPackage::Save(UPackage* InOuter, UObject* Base, EObjec
 			}
 
 			{
-				// set GIsSavingPackage here as it is now illegal to create any new object references; they potentially wouldn't be saved correctly								
-				FScopedSavingFlag IsSavingFlag(bSavingConcurrent);
-			
+				// set GIsSavingPackage here and mark the package as being saved
+				// as it is now illegal to create new object inside the package, they wouldn't be saved correctly
+				// it is now also illegal to find any new object references; they potentially wouldn't be saved correctly								
+				FScopedSavingFlag IsSavingFlag(bSavingConcurrent, InOuter);
 				{
 					SCOPED_SAVETIMER(UPackage_Save_TagExports);
 					COOK_STAT(FScopedDurationTimer SaveTimer(FSavePackageStats::TagPackageExportsTimeSec));
@@ -2483,7 +2484,7 @@ FSavePackageResultStruct UPackage::Save(UPackage* InOuter, UObject* Base, EObjec
 				}
 #endif
 
-				if (!(Linker->Summary.PackageFlags & PKG_FilterEditorOnly))
+				if (!(Linker->Summary.GetPackageFlags() & PKG_FilterEditorOnly))
 				{
 					// The Editor version is used as part of the check to see if a package is too old to use the gather cache, so we always have to add it if we have gathered loc for this asset
 					// We need to set the editor custom version before we copy the version container to the summary, otherwise we may end up with corrupt assets
@@ -2500,17 +2501,6 @@ FSavePackageResultStruct UPackage::Save(UPackage* InOuter, UObject* Base, EObjec
 
 				bool bUseUnversionedProperties = bSaveUnversioned && CanUseUnversionedPropertySerialization(TargetPlatform);
 				Linker->SetUseUnversionedPropertySerialization(bUseUnversionedProperties);
-				Linker->Saver->SetUseUnversionedPropertySerialization(bUseUnversionedProperties);
-				if (bUseUnversionedProperties)
-				{
-					Linker->Summary.PackageFlags |= PKG_UnversionedProperties;
-					Linker->LinkerRoot->SetPackageFlags(PKG_UnversionedProperties);
-				}
-				else
-				{
-					Linker->Summary.PackageFlags &= ~PKG_UnversionedProperties;
-					Linker->LinkerRoot->ClearPackageFlags(PKG_UnversionedProperties);
-				}
 
 				// Make sure the package has the same version as the linker
 				InOuter->LinkerPackageVersion = Linker->UE4Ver();
@@ -3120,7 +3110,7 @@ FSavePackageResultStruct UPackage::Save(UPackage* InOuter, UObject* Base, EObjec
 				FStructuredArchive::FStream Stream = StructuredArchiveRoot.EnterStream(SA_FIELD_NAME(TEXT("GatherableTextData")));
 				Linker->Summary.GatherableTextDataOffset = 0;
 				Linker->Summary.GatherableTextDataCount = 0;
-				if (!(Linker->Summary.PackageFlags & PKG_FilterEditorOnly))
+				if (!(Linker->Summary.GetPackageFlags() & PKG_FilterEditorOnly))
 				{
 					SCOPED_SAVETIMER(UPackage_Save_WriteGatherableTextData);
 
@@ -3605,7 +3595,7 @@ FSavePackageResultStruct UPackage::Save(UPackage* InOuter, UObject* Base, EObjec
 				SlowTask.EnterProgressFrame();
 
 				// Only save string asset and searchable name map if saving for editor
-				if (!(Linker->Summary.PackageFlags & PKG_FilterEditorOnly))
+				if (!(Linker->Summary.GetPackageFlags() & PKG_FilterEditorOnly))
 				{
 					SCOPED_SAVETIMER(UPackage_Save_SaveSoftPackagesAndSearchableNames);
 
@@ -3784,7 +3774,7 @@ FSavePackageResultStruct UPackage::Save(UPackage* InOuter, UObject* Base, EObjec
 							{
 								UE_LOG(LogSavePackage, Warning, TEXT("A dependency '%s' of '%s' is in the linker table, but is pending kill. We will keep the dependency anyway (%d)."), *ToTest->GetFullName(), *ForObj->GetFullName(), CallSite);
 							}
-							bool bNotFiltered = (ExcludedObjectMarks == OBJECTMARK_NOMARKS || !ToTest->HasAnyMarks(ExcludedObjectMarks)) && (!(Linker->Summary.PackageFlags & PKG_FilterEditorOnly) || !IsEditorOnlyObject(ToTest, false, true));
+							bool bNotFiltered = (ExcludedObjectMarks == OBJECTMARK_NOMARKS || !ToTest->HasAnyMarks(ExcludedObjectMarks)) && (!(Linker->Summary.GetPackageFlags() & PKG_FilterEditorOnly) || !IsEditorOnlyObject(ToTest, false, true));
 							if (bMandatory && !bNotFiltered)
 							{
 								UE_LOG(LogSavePackage, Warning, TEXT("A dependency '%s' of '%s' was filtered, but is mandatory. This indicates a problem with editor only stripping. We will keep the dependency anyway (%d)."), *ToTest->GetFullName(), *ForObj->GetFullName(), CallSite);
@@ -4328,7 +4318,7 @@ FSavePackageResultStruct UPackage::Save(UPackage* InOuter, UObject* Base, EObjec
 				Linker->LinkerRoot->ThisRequiresLocalizationGather(Linker->RequiresLocalizationGather());
 				
 				// Update package flags from package, in case serialization has modified package flags.
-				Linker->Summary.PackageFlags = Linker->LinkerRoot->GetPackageFlags() & ~PKG_NewlyCreated;
+				Linker->Summary.SetPackageFlags(Linker->LinkerRoot->GetPackageFlags());
 				
 				{
 					// Verify that the final serialization pass hasn't added any new custom versions. Otherwise this will result in crashes when loading the package.
diff --git a/Engine/Source/Runtime/CoreUObject/Private/UObject/SavePackage/SavePackageUtilities.cpp b/Engine/Source/Runtime/CoreUObject/Private/UObject/SavePackage/SavePackageUtilities.cpp
--- a/Engine/Source/Runtime/CoreUObject/Private/UObject/SavePackage/SavePackageUtilities.cpp
+++ b/Engine/Source/Runtime/CoreUObject/Private/UObject/SavePackage/SavePackageUtilities.cpp
@@ -168,6 +168,22 @@ void FArchiveObjectCrc32NonEditorProperties::Serialize(void* Data, int64 Length)
 
 static FThreadSafeCounter OutstandingAsyncWrites;
 
+
+// Initialization for GIsSavingPackage
+bool GIsSavingPackage = false;
+
+namespace UE
+{
+	bool IsSavingPackage(UObject* InOuter)
+	{
+		if (InOuter == nullptr)
+		{
+			return GIsSavingPackage;
+		}
+		return InOuter->GetPackage()->HasAnyPackageFlags(PKG_IsSaving);
+	}
+}
+
 namespace SavePackageUtilities
 {
 const FName NAME_World("World");
@@ -1559,8 +1575,9 @@ void VerifyEDLCookInfo(bool bFullReferencesExpected)
 	FEDLCookChecker::Verify(bFullReferencesExpected);
 }
 
-FScopedSavingFlag::FScopedSavingFlag(bool InSavingConcurrent)
+FScopedSavingFlag::FScopedSavingFlag(bool InSavingConcurrent, UPackage* InSavedPackage)
 	: bSavingConcurrent(InSavingConcurrent)
+	, SavedPackage(InSavedPackage)
 {
 	check(!IsGarbageCollecting());
 
@@ -1579,6 +1596,12 @@ FScopedSavingFlag::FScopedSavingFlag(bool InSavingConcurrent)
 	{
 		GIsSavingPackage = true;
 	}
+
+	// Mark the package as being saved 
+	if (SavedPackage)
+	{
+		SavedPackage->SetPackageFlags(PKG_IsSaving);
+	}
 }
 
 FScopedSavingFlag::~FScopedSavingFlag()
@@ -1595,6 +1618,12 @@ FScopedSavingFlag::~FScopedSavingFlag()
 	{
 		FGCCSyncObject::Get().UnlockAsync();
 	}
+
+	if (SavedPackage)
+	{
+		SavedPackage->ClearPackageFlags(PKG_IsSaving);
+	}
+
 }
 
 FSavePackageDiffSettings::FSavePackageDiffSettings(bool bDiffing)
@@ -1646,7 +1675,7 @@ void SaveThumbnails(UPackage* InOuter, FLinkerSave* Linker, FStructuredArchive::
 
 #if WITH_EDITORONLY_DATA
 	// Do we have any thumbnails to save?
-	if( !(Linker->Summary.PackageFlags & PKG_FilterEditorOnly) && InOuter->HasThumbnailMap() )
+	if( !(Linker->Summary.GetPackageFlags() & PKG_FilterEditorOnly) && InOuter->HasThumbnailMap() )
 	{
 		const FThumbnailMap& PackageThumbnailMap = InOuter->GetThumbnailMap();
 
@@ -2124,7 +2153,7 @@ namespace AssetRegistry
 
 		// Collect the tag map
 		TArray<UObject*> AssetObjects;
-		if (!(Linker->Summary.PackageFlags & PKG_FilterEditorOnly))
+		if (!(Linker->Summary.GetPackageFlags() & PKG_FilterEditorOnly))
 		{
 			// Find any exports which are not in the tag map
 			for (int32 i = 0; i < Linker->ExportMap.Num(); i++)
diff --git a/Engine/Source/Runtime/CoreUObject/Private/UObject/SavePackage/SavePackageUtilities.h b/Engine/Source/Runtime/CoreUObject/Private/UObject/SavePackage/SavePackageUtilities.h
--- a/Engine/Source/Runtime/CoreUObject/Private/UObject/SavePackage/SavePackageUtilities.h
+++ b/Engine/Source/Runtime/CoreUObject/Private/UObject/SavePackage/SavePackageUtilities.h
@@ -52,10 +52,14 @@ ENUM_CLASS_FLAGS(EAsyncWriteOptions)
 
 struct FScopedSavingFlag
 {
-	FScopedSavingFlag(bool InSavingConcurrent);
+	FScopedSavingFlag(bool InSavingConcurrent, UPackage* InSavedPackage);
 	~FScopedSavingFlag();
 
 	bool bSavingConcurrent;
+
+private:
+	// The package being saved
+	UPackage* SavedPackage = nullptr;
 };
 
 struct FSavePackageDiffSettings
diff --git a/Engine/Source/Runtime/CoreUObject/Private/UObject/SavePackage2.cpp b/Engine/Source/Runtime/CoreUObject/Private/UObject/SavePackage2.cpp
--- a/Engine/Source/Runtime/CoreUObject/Private/UObject/SavePackage2.cpp
+++ b/Engine/Source/Runtime/CoreUObject/Private/UObject/SavePackage2.cpp
@@ -1673,7 +1673,7 @@ ESavePackageResult UpdatePackageHeader(FStructuredArchive::FRecord& StructuredAr
 		Linker->LinkerRoot->ThisRequiresLocalizationGather(Linker->RequiresLocalizationGather());
 
 		// Update package flags from package, in case serialization has modified package flags.
-		Linker->Summary.PackageFlags = Linker->LinkerRoot->GetPackageFlags() & ~PKG_NewlyCreated;
+		Linker->Summary.SetPackageFlags(Linker->LinkerRoot->GetPackageFlags());
 
 		// @todo: custom versions: when can this be checked?
 		{
@@ -1963,7 +1963,7 @@ ESavePackageResult InnerSave(FSaveContext& SaveContext)
 
 	// @todo: Need to adjust GIsSavingPackage to properly prevent generating reference once, package harvesting is done
 	// GIsSavingPackage is too harsh however, since it should be scope only to the current package
-	FScopedSavingFlag IsSavingFlag(SaveContext.IsConcurrent());
+	FScopedSavingFlag IsSavingFlag(SaveContext.IsConcurrent(), SaveContext.GetPackage());
 
 	// Validate Exports
 	SlowTask.EnterProgressFrame();
diff --git a/Engine/Source/Runtime/CoreUObject/Private/UObject/UObjectGlobals.cpp b/Engine/Source/Runtime/CoreUObject/Private/UObject/UObjectGlobals.cpp
--- a/Engine/Source/Runtime/CoreUObject/Private/UObject/UObjectGlobals.cpp
+++ b/Engine/Source/Runtime/CoreUObject/Private/UObject/UObjectGlobals.cpp
@@ -70,8 +70,6 @@ DEFINE_LOG_CATEGORY(LogUObjectGlobals);
 #include "ProfilingDebugging/MallocProfiler.h"
 #endif
 
-bool GIsSavingPackage = false;
-
 int32 GAllowUnversionedContentInEditor = 0;
 
 static FAutoConsoleVariableRef CVarAllowUnversionedContentInEditor(
@@ -289,37 +287,37 @@ int32 UpdateSuffixForNextNewObject(UObject* Parent, const UClass* Class, TFuncti
 //
 // Find an object, path must unqualified
 //
-UObject* StaticFindObjectFast(UClass* ObjectClass, UObject* ObjectPackage, FName ObjectName, bool ExactClass, bool AnyPackage, EObjectFlags ExclusiveFlags, EInternalObjectFlags ExclusiveInternalFlags)
+UObject* StaticFindObjectFast(UClass* ObjectClass, UObject* ObjectPackage, FName ObjectName, bool bExactClass, bool bAnyPackage, EObjectFlags ExclusiveFlags, EInternalObjectFlags ExclusiveInternalFlags)
 {
-	if (GIsSavingPackage || IsGarbageCollectingOnGameThread())
+	if (UE::IsSavingPackage(nullptr) || IsGarbageCollectingOnGameThread())
 	{
 		UE_LOG(LogUObjectGlobals, Fatal,TEXT("Illegal call to StaticFindObjectFast() while serializing object data or garbage collecting!"));
 	}
 
 	// We don't want to return any objects that are currently being background loaded unless we're using FindObject during async loading.
 	ExclusiveInternalFlags |= IsInAsyncLoadingThread() ? EInternalObjectFlags::None : EInternalObjectFlags::AsyncLoading;	
-	UObject* FoundObject = StaticFindObjectFastInternal(ObjectClass, ObjectPackage, ObjectName, ExactClass, AnyPackage, ExclusiveFlags, ExclusiveInternalFlags);
+	UObject* FoundObject = StaticFindObjectFastInternal(ObjectClass, ObjectPackage, ObjectName, bExactClass, bAnyPackage, ExclusiveFlags, ExclusiveInternalFlags);
 
 	if (!FoundObject)
 	{
-		FoundObject = StaticFindObjectWithChangedLegacyPath(ObjectClass, ObjectPackage, ObjectName, ExactClass);
+		FoundObject = StaticFindObjectWithChangedLegacyPath(ObjectClass, ObjectPackage, ObjectName, bExactClass);
 	}
 
 	return FoundObject;
 }
 
-UObject* StaticFindObjectFastSafe(UClass* ObjectClass, UObject* ObjectPackage, FName ObjectName, bool ExactClass, bool AnyPackage, EObjectFlags ExclusiveFlags, EInternalObjectFlags ExclusiveInternalFlags)
+UObject* StaticFindObjectFastSafe(UClass* ObjectClass, UObject* ObjectPackage, FName ObjectName, bool bExactClass, bool bAnyPackage, EObjectFlags ExclusiveFlags, EInternalObjectFlags ExclusiveInternalFlags)
 {
 	UObject* FoundObject = nullptr;
 	
-	if (!GIsSavingPackage && !IsGarbageCollectingOnGameThread())
+	if (!UE::IsSavingPackage(nullptr) && !IsGarbageCollectingOnGameThread())
 	{
 		// We don't want to return any objects that are currently being background loaded unless we're using FindObject during async loading.
 		ExclusiveInternalFlags |= IsInAsyncLoadingThread() ? EInternalObjectFlags::None : EInternalObjectFlags::AsyncLoading;
-		FoundObject = StaticFindObjectFastInternal(ObjectClass, ObjectPackage, ObjectName, ExactClass, AnyPackage, ExclusiveFlags, ExclusiveInternalFlags);
+		FoundObject = StaticFindObjectFastInternal(ObjectClass, ObjectPackage, ObjectName, bExactClass, bAnyPackage, ExclusiveFlags, ExclusiveInternalFlags);
 		if (!FoundObject)
 		{
-			FoundObject = StaticFindObjectWithChangedLegacyPath(ObjectClass, ObjectPackage, ObjectName, ExactClass);
+			FoundObject = StaticFindObjectWithChangedLegacyPath(ObjectClass, ObjectPackage, ObjectName, bExactClass);
 		}
 	}
 
@@ -329,26 +327,15 @@ UObject* StaticFindObjectFastSafe(UClass* ObjectClass, UObject* ObjectPackage, F
 //
 // Find an optional object.
 //
-UObject* StaticFindObject( UClass* ObjectClass, UObject* InObjectPackage, const TCHAR* OrigInName, bool ExactClass )
+UObject* StaticFindObject( UClass* ObjectClass, UObject* InObjectPackage, const TCHAR* OrigInName, bool bExactClass )
 {
 	INC_DWORD_STAT(STAT_FindObject);
 
-	if (GIsSavingPackage)
-	{
-		UE_LOG(LogUObjectGlobals, Fatal,TEXT("Illegal call to StaticFindObject() while serializing object data!"));
-	}
-
-	if (IsGarbageCollectingOnGameThread())
-	{
-		UE_LOG(LogUObjectGlobals, Fatal,TEXT("Illegal call to StaticFindObject() while collecting garbage!"));
-	}
-
 	// Resolve the object and package name.
 	const bool bAnyPackage = InObjectPackage==ANY_PACKAGE;
 	UObject* ObjectPackage = bAnyPackage ? nullptr : InObjectPackage;
 
 	UObject* MatchingObject = nullptr;
-
 #if WITH_EDITOR
 	// If the editor is running, and T3D is being imported, ensure any packages referenced are fully loaded.
 	if ((GIsEditor == true) && (GIsImportingT3D == true))// && (ObjectPackage != ANY_PACKAGE) && (ObjectPackage != NULL))
@@ -394,7 +381,7 @@ UObject* StaticFindObject( UClass* ObjectClass, UObject* InObjectPackage, const
 		ObjectName = FName(*InName, FNAME_Add);
 	}
 
-	return StaticFindObjectFast(ObjectClass, ObjectPackage, ObjectName, ExactClass, bAnyPackage);
+	return StaticFindObjectFast(ObjectClass, ObjectPackage, ObjectName, bExactClass, bAnyPackage);
 }
 
 //
@@ -413,18 +400,18 @@ UObject* StaticFindObjectChecked( UClass* ObjectClass, UObject* ObjectParent, co
 }
 
 //
-// Find an object; won't assert on GIsSavingPackage or IsGarbageCollecting()
+// Find an object; won't assert on UE::IsSavingPackage() or IsGarbageCollecting()
 //
-UObject* StaticFindObjectSafe( UClass* ObjectClass, UObject* ObjectParent, const TCHAR* InName, bool ExactClass )
+UObject* StaticFindObjectSafe( UClass* ObjectClass, UObject* ObjectParent, const TCHAR* InName, bool bExactClass )
 {
-	if (!GIsSavingPackage && !IsGarbageCollectingOnGameThread())
+	if (!UE::IsSavingPackage(nullptr) && !IsGarbageCollectingOnGameThread())
 	{
-		FGCScopeGuard GCAndSavepackageGuard;
-		return StaticFindObject( ObjectClass, ObjectParent, InName, ExactClass );
+		FGCScopeGuard GCGuard;
+		return StaticFindObject( ObjectClass, ObjectParent, InName, bExactClass );
 	}
 	else
 	{
-		return NULL;
+		return nullptr;
 	}
 }
 
@@ -2421,15 +2408,15 @@ UObject* StaticAllocateObject
 		}
 	}
 
-	FLinkerLoad*	Linker						= NULL;
-	int32				LinkerIndex					= INDEX_NONE;
+	FLinkerLoad*	Linker						= nullptr;
+	int32			LinkerIndex					= INDEX_NONE;
 	bool			bWasConstructedOnOldObject	= false;
 	// True when the object to be allocated already exists and is a subobject.
 	bool bSubObject = false;
 	int32 TotalSize = InClass->GetPropertiesSize();
 	checkSlow(TotalSize);
 
-	if( Obj == NULL )
+	if( Obj == nullptr )
 	{	
 		int32 Alignment	= FMath::Max( 4, InClass->GetMinAlignment() );
 		Obj = (UObject *)GUObjectAllocator.AllocateUObject(TotalSize,Alignment,GIsInitialLoad);
@@ -2536,7 +2523,7 @@ UObject* StaticAllocateObject
 
 	if (!bSubObject)
 	{
-			FMemory::Memzero((void *)Obj, TotalSize);
+		FMemory::Memzero((void *)Obj, TotalSize);
 		new ((void *)Obj) UObjectBase(const_cast<UClass*>(InClass), InFlags|RF_NeedInitialization, InternalSetFlags, InOuter, InName);
 	}
 	else
@@ -3357,7 +3344,9 @@ UObject* StaticConstructObject_Internal(const FStaticConstructObjectParameters&
 	UObject* Result = NULL;
 
 #if WITH_EDITORONLY_DATA
-	UE_CLOG(GIsSavingPackage && InOuter != GetTransientPackage(), LogUObjectGlobals, Fatal, TEXT("Illegal call to StaticConstructObject() while serializing object data! (Object will not be saved!)"));
+	// Check if we can construct the object: you can construct the object if its a package (InOuter is null) or the package the object is created in is not currently saving
+	bool bCanConstruct = InOuter == nullptr || !UE::IsSavingPackage(Params.ExternalPackage ? Params.ExternalPackage : InOuter->GetPackage()); 
+	UE_CLOG(!bCanConstruct, LogUObjectGlobals, Fatal, TEXT("Illegal call to StaticConstructObject() while serializing object data! (Object will not be saved!)"));
 #endif
 
 	checkf(!InTemplate || InTemplate->IsA(InClass) || (InFlags & RF_ClassDefaultObject), TEXT("StaticConstructObject %s is not an instance of class %s and it is not a CDO."), *GetFullNameSafe(InTemplate), *GetFullNameSafe(InClass)); // template must be an instance of the class we are creating, except CDOs
@@ -3379,7 +3368,7 @@ UObject* StaticConstructObject_Internal(const FStaticConstructObjectParameters&
 
 	bool bRecycledSubobject = false;	
 	Result = StaticAllocateObject(InClass, InOuter, InName, InFlags, Params.InternalSetFlags, bCanRecycleSubobjects, &bRecycledSubobject, Params.ExternalPackage);
-	check(Result != NULL);
+	check(Result != nullptr);
 	// Don't call the constructor on recycled subobjects, they haven't been destroyed.
 	if (!bRecycledSubobject)
 	{		
diff --git a/Engine/Source/Runtime/CoreUObject/Public/UObject/ObjectMacros.h b/Engine/Source/Runtime/CoreUObject/Public/UObject/ObjectMacros.h
--- a/Engine/Source/Runtime/CoreUObject/Public/UObject/ObjectMacros.h
+++ b/Engine/Source/Runtime/CoreUObject/Public/UObject/ObjectMacros.h
@@ -117,7 +117,7 @@ enum EPackageFlags
 //	PKG_Unused						= 0x00001000,
 	PKG_UnversionedProperties		= 0x00002000,   ///< Uses unversioned property serialization instead of versioned tagged property serialization
 	PKG_ContainsMapData				= 0x00004000,   ///< Contains map data (UObjects only referenced by a single ULevel) but is stored in a different package
-//	PKG_Unused						= 0x00008000,
+	PKG_IsSaving					= 0x00008000,	///< Temporarily set on a package while it is being saved.
 	PKG_Compiling					= 0x00010000,	///< package is currently being compiled
 	PKG_ContainsMap					= 0x00020000,	///< Set if the package contains a ULevel/ UWorld object
 	PKG_RequiresLocalizationGather	= 0x00040000,	///< Set if the package contains any data to be gathered by localization
diff --git a/Engine/Source/Runtime/CoreUObject/Public/UObject/PackageFileSummary.h b/Engine/Source/Runtime/CoreUObject/Public/UObject/PackageFileSummary.h
--- a/Engine/Source/Runtime/CoreUObject/Public/UObject/PackageFileSummary.h
+++ b/Engine/Source/Runtime/CoreUObject/Public/UObject/PackageFileSummary.h
@@ -58,6 +58,11 @@ private:
 	/* Custom version numbers. Keyed off a unique tag for each custom component. */
 	FCustomVersionContainer CustomVersionContainer;
 
+	/**
+	 * The flags for the package
+	 */
+	uint32	PackageFlags;
+
 public:
 	/**
 	* Total size of all information that needs to be read in to create a FLinkerLoad. This includes
@@ -65,11 +70,6 @@ public:
 	*/
 	int32		TotalHeaderSize;
 
-	/**
-	* The flags for the package
-	*/
-	uint32	PackageFlags;
-
 	/**
 	* The Generic Browser folder name that this package lives in
 	*/
@@ -241,6 +241,15 @@ public:
 		bUnversioned = bInSaveUnversioned;
 	}
 
+	/** Get the summary package flags. */
+	uint32 GetPackageFlags() const
+	{
+		return PackageFlags;
+	}
+
+	/** Set the summary package flags while stripping out temporary flags (i.e. NewlyCreated, IsSaving) */
+	void SetPackageFlags(uint32 InPackageFlags);
+
 	/** I/O functions */
 	friend COREUOBJECT_API FArchive& operator<<(FArchive& Ar, FPackageFileSummary& Sum);
 	friend COREUOBJECT_API void operator<<(FStructuredArchive::FSlot Slot, FPackageFileSummary& Sum);
diff --git a/Engine/Source/Runtime/CoreUObject/Public/UObject/UObjectGlobals.h b/Engine/Source/Runtime/CoreUObject/Public/UObject/UObjectGlobals.h
--- a/Engine/Source/Runtime/CoreUObject/Public/UObject/UObjectGlobals.h
+++ b/Engine/Source/Runtime/CoreUObject/Public/UObject/UObjectGlobals.h
@@ -49,7 +49,12 @@ DECLARE_DWORD_COUNTER_STAT_EXTERN(TEXT("FindObjectFast"),STAT_FindObjectFast,STA
 
 // Private system wide variables.
 
-/** Set while in SavePackage() to detect certain operations that are illegal while saving */
+/** 
+ * Set while in SavePackage() to advertise that a package is being saved
+ * Deprecated, use `IsSavingPackage` instead
+ * @see IsSavingPackage()
+ */
+//UE_DEPRECATED(5.0, "Use the IsSavingPackage() function instead.")
 extern COREUOBJECT_API bool					GIsSavingPackage;
 /** This allows loading unversioned cooked content in the editor */
 extern COREUOBJECT_API int32				GAllowUnversionedContentInEditor;
@@ -160,6 +165,20 @@ struct FObjectDuplicationParameters
 	COREUOBJECT_API FObjectDuplicationParameters( UObject* InSourceObject, UObject* InDestOuter );
 };
 
+// @note: should we start putting the code in the UE namespace?
+namespace UE
+{
+	/**
+	 * Return if the passed in package is currently saving,
+	 * or if any packages are saving if no parameter is passed.
+	 * @param InOuter The object which we want to check if its package is saving
+	 * @returns true if the package is saving, or if any packages are saving if no parameter is passed in.
+	 */
+	COREUOBJECT_API bool IsSavingPackage(UObject* InOuter = nullptr);
+
+} // namespace UE
+
+
 /** Parses a bit mask of property flags into an array of string literals that match the flags */
 COREUOBJECT_API TArray<const TCHAR*> ParsePropertyFlags(EPropertyFlags Flags);
 
@@ -212,14 +231,14 @@ COREUOBJECT_API int32 UpdateSuffixForNextNewObject(UObject* Parent, const UClass
  * @param	Class			The to be found object's class
  * @param	InOuter			Outer object to look inside, if null this will only look for top level packages
  * @param	InName			Object name to look for relative to InOuter
- * @param	ExactClass		Whether to require an exact match with the passed in class
- * @param	AnyPackage		Whether to look in any package
+ * @param	bExactClass		Whether to require an exact match with the passed in class
+ * @param	bAnyPackage		Whether to look in any package
  * @param	ExclusiveFlags	Ignores objects that contain any of the specified exclusive flags
- * @param ExclusiveInternalFlags  Ignores objects that contain any of the specified internal exclusive flags
+ * @param	ExclusiveInternalFlags  Ignores objects that contain any of the specified internal exclusive flags
  *
  * @return	Returns a pointer to the found object or null if none could be found
  */
-COREUOBJECT_API UObject* StaticFindObjectFast(UClass* Class, UObject* InOuter, FName InName, bool ExactClass = false, bool AnyPackage = false, EObjectFlags ExclusiveFlags = RF_NoFlags, EInternalObjectFlags ExclusiveInternalFlags = EInternalObjectFlags::None);
+COREUOBJECT_API UObject* StaticFindObjectFast(UClass* Class, UObject* InOuter, FName InName, bool bExactClass = false, bool bAnyPackage = false, EObjectFlags ExclusiveFlags = RF_NoFlags, EInternalObjectFlags ExclusiveInternalFlags = EInternalObjectFlags::None);
 
 /**
  * Fast and safe version of StaticFindObject that relies on the passed in FName being the object name without any group/package qualifiers.
@@ -229,14 +248,14 @@ COREUOBJECT_API UObject* StaticFindObjectFast(UClass* Class, UObject* InOuter, F
  * @param	Class			The to be found object's class
  * @param	InOuter			Outer object to look inside, if null this will only look for top level packages
  * @param	InName			Object name to look for relative to InOuter
- * @param	ExactClass		Whether to require an exact match with the passed in class
- * @param	AnyPackage		Whether to look in any package
+ * @param	bExactClass		Whether to require an exact match with the passed in class
+ * @param	bAnyPackage		Whether to look in any package
  * @param	ExclusiveFlags	Ignores objects that contain any of the specified exclusive flags
- * @param ExclusiveInternalFlags  Ignores objects that contain any of the specified internal exclusive flags
+ * @param	ExclusiveInternalFlags  Ignores objects that contain any of the specified internal exclusive flags
  *
  * @return	Returns a pointer to the found object or null if none could be found
  */
-COREUOBJECT_API UObject* StaticFindObjectFastSafe(UClass* Class, UObject* InOuter, FName InName, bool ExactClass = false, bool AnyPackage = false, EObjectFlags ExclusiveFlags = RF_NoFlags, EInternalObjectFlags ExclusiveInternalFlags = EInternalObjectFlags::None);
+COREUOBJECT_API UObject* StaticFindObjectFastSafe(UClass* Class, UObject* InOuter, FName InName, bool bExactClass = false, bool bAnyPackage = false, EObjectFlags ExclusiveFlags = RF_NoFlags, EInternalObjectFlags ExclusiveInternalFlags = EInternalObjectFlags::None);
 
 
 /**
