Add missing overload for serializing an FObjectPtr in FNetBitReader (to match the existing one that already exists in FNetBitWriter).  Both use the default codepath in FArchiveUObject which just resolves the wrapped object pointer and goes through the raw pointer serialization path.

diff --git a/Engine/Source/Runtime/CoreUObject/Private/UObject/CoreNet.cpp b/Engine/Source/Runtime/CoreUObject/Private/UObject/CoreNet.cpp
--- a/Engine/Source/Runtime/CoreUObject/Private/UObject/CoreNet.cpp
+++ b/Engine/Source/Runtime/CoreUObject/Private/UObject/CoreNet.cpp
@@ -504,6 +504,11 @@ FArchive& FNetBitReader::operator<<(FSoftObjectPtr& Value)
 	return FArchiveUObject::SerializeSoftObjectPtr(*this, Value);
 }
 
+FArchive& FNetBitReader::operator<<(FObjectPtr& Value)
+{
+	return FArchiveUObject::SerializeObjectPtr(*this, Value);
+}
+
 FArchive& FNetBitReader::operator<<(struct FWeakObjectPtr& WeakObjectPtr)
 {
 	return FArchiveUObject::SerializeWeakObjectPtr(*this, WeakObjectPtr);
diff --git a/Engine/Source/Runtime/CoreUObject/Public/UObject/CoreNet.h b/Engine/Source/Runtime/CoreUObject/Public/UObject/CoreNet.h
--- a/Engine/Source/Runtime/CoreUObject/Public/UObject/CoreNet.h
+++ b/Engine/Source/Runtime/CoreUObject/Public/UObject/CoreNet.h
@@ -369,6 +369,7 @@ public:
 	virtual FArchive& operator<<(UObject*& Object) override;
 	virtual FArchive& operator<<(FSoftObjectPath& Value) override;
 	virtual FArchive& operator<<(FSoftObjectPtr& Value) override;
+	virtual FArchive& operator<<(FObjectPtr& Value) override;
 	virtual FArchive& operator<<(struct FWeakObjectPtr& Value) override;
 
 	virtual void CountMemory(FArchive& Ar) const override;
