Added EditConditionHides to the list of PropertyMetadata.

diff --git a/Engine/Source/Runtime/CoreUObject/Public/UObject/ObjectMacros.h b/Engine/Source/Runtime/CoreUObject/Public/UObject/ObjectMacros.h
--- a/Engine/Source/Runtime/CoreUObject/Public/UObject/ObjectMacros.h
+++ b/Engine/Source/Runtime/CoreUObject/Public/UObject/ObjectMacros.h
@@ -1191,6 +1191,9 @@ namespace UM
 		/// [PropertyMetadata] Specifies a boolean property that is used to indicate whether editing of this property is disabled.
 		EditCondition,
 
+		/// [PropertyMetadata] Specifies a boolean property that is used to indicate whether this property should be hidden when the EditCondition evaluates to false. Requires an EditCondition to be set.
+		EditConditionHides,
+
 		/// [PropertyMetadata] Keeps the elements of an array from being reordered by dragging 
 		EditFixedOrder,
 		
