AssetRegistry - Add CustomVersion information to AssetRegistry's PackageData information. Add IAssetRegistry::WaitOnPackage.
These features will be used in the future by the Editor Domain.
#rnx

diff --git a/Engine/Source/Runtime/AssetRegistry/Private/AssetDataGatherer.cpp b/Engine/Source/Runtime/AssetRegistry/Private/AssetDataGatherer.cpp
--- a/Engine/Source/Runtime/AssetRegistry/Private/AssetDataGatherer.cpp
+++ b/Engine/Source/Runtime/AssetRegistry/Private/AssetDataGatherer.cpp
@@ -30,7 +30,7 @@ namespace AssetDataGathererConstants
 	constexpr int32 SingleThreadFilesPerBatch = 3;
 	constexpr int32 ExpectedMaxBatchSize = 100;
 	constexpr int32 MinSecondsToElapseBeforeCacheWrite = 60;
-	static constexpr uint32 CacheSerializationMagic = 0xCBA78339; 
+	static constexpr uint32 CacheSerializationMagic = 0x89ABCDF0;
 }
 
 namespace UE
@@ -2760,6 +2760,24 @@ void FAssetDataGatherer::GetAndTrimSearchResults(bool& bOutIsSearching, TRingBuf
 	bOutIsSearching = !bIsIdle;
 }
 
+void FAssetDataGatherer::GetPackageResults(TRingBuffer<FAssetData*>& OutAssetResults, TRingBuffer<FPackageDependencyData>& OutDependencyResults)
+{
+	FGathererScopeLock ResultsScopeLock(&ResultsLock);
+
+	auto MoveAppendRangeToRingBuffer = [](auto& OutRingBuffer, auto& InArray)
+	{
+		OutRingBuffer.Reserve(OutRingBuffer.Num() + InArray.Num());
+		for (auto& Element : InArray)
+		{
+			OutRingBuffer.Add(MoveTemp(Element));
+		}
+		InArray.Reset();
+	};
+
+	MoveAppendRangeToRingBuffer(OutAssetResults, AssetResults);
+	MoveAppendRangeToRingBuffer(OutDependencyResults, DependencyResults);
+}
+
 void FAssetDataGatherer::WaitOnPath(FStringView InPath)
 {
 	{
diff --git a/Engine/Source/Runtime/AssetRegistry/Private/AssetDataGatherer.h b/Engine/Source/Runtime/AssetRegistry/Private/AssetDataGatherer.h
--- a/Engine/Source/Runtime/AssetRegistry/Private/AssetDataGatherer.h
+++ b/Engine/Source/Runtime/AssetRegistry/Private/AssetDataGatherer.h
@@ -91,6 +91,8 @@ public:
 	void GetAndTrimSearchResults(bool& bOutIsSearching, TRingBuffer<FAssetData*>& OutAssetResults, TRingBuffer<FString>& OutPathResults,
 		TRingBuffer<FPackageDependencyData>& OutDependencyResults, TRingBuffer<FString>& OutCookedPackageNamesWithoutAssetDataResults,
 		TArray<double>& OutSearchTimes, int32& OutNumFilesToSearch, int32& OutNumPathsToSearch, bool& OutIsDiscoveringFiles);
+	/** Gets just the AssetResults and DependencyResults from the data gatherer. */
+	void GetPackageResults(TRingBuffer<FAssetData*>& OutAssetResults, TRingBuffer<FPackageDependencyData>& OutDependencyResults);
 	/** Wait for all monitored assets under the given path to be added to search results. Returns immediately if the given path is not monitored. */
 	void WaitOnPath(FStringView LocalPath);
 	/**
diff --git a/Engine/Source/Runtime/AssetRegistry/Private/AssetRegistry.cpp b/Engine/Source/Runtime/AssetRegistry/Private/AssetRegistry.cpp
--- a/Engine/Source/Runtime/AssetRegistry/Private/AssetRegistry.cpp
+++ b/Engine/Source/Runtime/AssetRegistry/Private/AssetRegistry.cpp
@@ -869,6 +869,23 @@ void UAssetRegistryImpl::WaitForCompletion()
 	}
 }
 
+void UAssetRegistryImpl::WaitForPackage(const FString& PackageName)
+{
+	check(IsInGameThread());
+	if (!IsLoadingAssets() || !GlobalGatherer.IsValid())
+	{
+		return;
+	}
+
+	FString LocalPath;
+	if (!FPackageName::TryConvertLongPackageNameToFilename(PackageName, LocalPath))
+	{
+		return;
+	}
+	GlobalGatherer->WaitOnPath(LocalPath);
+	TickGatherPackage(PackageName);
+}
+
 bool UAssetRegistryImpl::HasAssets(const FName PackagePath, const bool bRecursive) const
 {
 	bool bHasAssets = State.HasAssets(PackagePath);
@@ -2037,6 +2054,50 @@ void UAssetRegistryImpl::TickGatherer(float DeltaTime, TOptional<UAssetRegistryI
 #endif
 }
 
+void UAssetRegistryImpl::TickGatherPackage(const FString& PackageName)
+{
+	check(IsInGameThread()); // AssetRegistry is not threadsafe
+	if (!GlobalGatherer.IsValid())
+	{
+		return;
+	}
+	FName PackageFName(PackageName);
+
+	// Gather results from the background search
+	GlobalGatherer->GetPackageResults(BackgroundAssetResults, BackgroundDependencyResults);
+
+	TRingBuffer<FAssetData*> PackageAssets;
+	TRingBuffer<FPackageDependencyData> PackageDependencyDatas;
+	for (int n = 0; n < BackgroundAssetResults.Num(); ++n)
+	{
+		FAssetData* Asset = BackgroundAssetResults[n];
+		if (Asset->PackageName == PackageFName)
+		{
+			Swap(BackgroundAssetResults[n], BackgroundAssetResults.Last());
+			PackageAssets.Add(BackgroundAssetResults.PopValue());
+			--n;
+		}
+	}
+	for (int n = 0; n < BackgroundDependencyResults.Num(); ++n)
+	{
+		FPackageDependencyData& DependencyData = BackgroundDependencyResults[n];
+		if (DependencyData.PackageName == PackageFName)
+		{
+			Swap(BackgroundDependencyResults[n], BackgroundDependencyResults.Last());
+			PackageDependencyDatas.Add(BackgroundDependencyResults.PopValue());
+			--n;
+		}
+	}
+	if (PackageAssets.Num() > 0)
+	{
+		AssetSearchDataGathered(-1., PackageAssets);
+	}
+	if (PackageDependencyDatas.Num() > 0)
+	{
+		DependencyDataGathered(-1., PackageDependencyDatas);
+	}
+}
+
 void UAssetRegistryImpl::Serialize(FArchive& Ar)
 {
 	if (Ar.IsLoading())
diff --git a/Engine/Source/Runtime/AssetRegistry/Private/AssetRegistry.h b/Engine/Source/Runtime/AssetRegistry/Private/AssetRegistry.h
--- a/Engine/Source/Runtime/AssetRegistry/Private/AssetRegistry.h
+++ b/Engine/Source/Runtime/AssetRegistry/Private/AssetRegistry.h
@@ -93,6 +93,7 @@ public:
 	virtual bool PathExists(const FName PathToTest) const override;
 	virtual void SearchAllAssets(bool bSynchronousSearch) override;
 	virtual void WaitForCompletion() override;
+	virtual void WaitForPackage(const FString& PackageName) override;
 	virtual void ScanPathsSynchronous(const TArray<FString>& InPaths, bool bForceRescan = false, bool bIgnoreBlackListScanFilters = false) override;
 	virtual void ScanFilesSynchronous(const TArray<FString>& InFilePaths, bool bForceRescan = false) override;
 	virtual void PrioritizeSearchPath(const FString& PathToPrioritize) override;
@@ -180,6 +181,8 @@ private:
 	typedef TFunctionRef<void(const TRingBuffer<FAssetData*>&)> FAssetsFoundCallback;
 	/** Read any results from the gatherer, if it is running asynchronously */
 	void TickGatherer(float DeltaTime, TOptional<FAssetsFoundCallback> AssetsFoundCallback = TOptional<FAssetsFoundCallback>());
+	/** Look for and load a single AssetData result from the gatherer. */
+	void TickGatherPackage(const FString& PackageName);
 	/** Internal handler for ScanPathsSynchronous */
 	void ScanPathsSynchronousInternal(const TArray<FString>& InDirs, const TArray<FString>& InFiles, bool bForceRescan, bool bIgnoreBlackListScanFilters, TArray<FName>* OutFoundAssets);
 
diff --git a/Engine/Source/Runtime/AssetRegistry/Private/AssetRegistryState.cpp b/Engine/Source/Runtime/AssetRegistry/Private/AssetRegistryState.cpp
--- a/Engine/Source/Runtime/AssetRegistry/Private/AssetRegistryState.cpp
+++ b/Engine/Source/Runtime/AssetRegistry/Private/AssetRegistryState.cpp
@@ -1120,38 +1120,40 @@ void FAssetRegistryState::Load(Archive&& Ar, FAssetRegistryVersion::Type Version
 
 	if (LocalNumPackageData > 0)
 	{
+		FAssetPackageData SerializedElement;
+		TArrayView<FAssetPackageData> PreallocatedPackageDataBuffer;
 		if (Options.bLoadPackageData)
 		{
-			TArrayView<FAssetPackageData> PreallocatedPackageDataBuffer(new FAssetPackageData[LocalNumPackageData], LocalNumPackageData);
+			PreallocatedPackageDataBuffer = TArrayView<FAssetPackageData>(new FAssetPackageData[LocalNumPackageData], LocalNumPackageData);
 			PreallocatedPackageDataBuffers.Add(PreallocatedPackageDataBuffer.GetData());
 			CachedPackageData.Reserve(LocalNumPackageData);
-			for (FAssetPackageData& NewPackageData : PreallocatedPackageDataBuffer)
-			{
-				FName PackageName;
-				Ar << PackageName;
-		
-				if (Version < FAssetRegistryVersion::AddedCookedMD5Hash)
-				{
-					Ar << NewPackageData.DiskSize;
-					Ar << NewPackageData.PackageGuid;
-				}
-				else
-				{
-					NewPackageData.SerializeForCache(Ar);
-				}
-
-				CachedPackageData.Add(PackageName, &NewPackageData);
-			}	
 		}
-		else
+		for (int32 PackageDataIndex = 0; PackageDataIndex < LocalNumPackageData; PackageDataIndex++)
 		{
-			for (int32 PackageDataIndex = 0; PackageDataIndex < LocalNumPackageData; PackageDataIndex++)
+			FName PackageName;
+			Ar << PackageName;
+			FAssetPackageData* NewPackageData;
+			if (Options.bLoadPackageData)
 			{
-				FName PackageName;
-				Ar << PackageName;
-
-				FAssetPackageData FakeData;
-				FakeData.SerializeForCache(Ar);
+				NewPackageData = &PreallocatedPackageDataBuffer[PackageDataIndex];
+				CachedPackageData.Add(PackageName, NewPackageData);
+			}
+			else
+			{
+				NewPackageData = &SerializedElement;
+			}
+			if (Version >= FAssetRegistryVersion::WorkspaceDomain)
+			{
+				NewPackageData->SerializeForCache(Ar);
+			}
+			else
+			{
+				Ar << NewPackageData->DiskSize;
+				Ar << NewPackageData->PackageGuid;
+				if (Version >= FAssetRegistryVersion::AddedCookedMD5Hash)
+				{
+					Ar << NewPackageData->CookedHash;
+				}
 			}
 		}
 	}
diff --git a/Engine/Source/Runtime/AssetRegistry/Private/PackageReader.cpp b/Engine/Source/Runtime/AssetRegistry/Private/PackageReader.cpp
--- a/Engine/Source/Runtime/AssetRegistry/Private/PackageReader.cpp
+++ b/Engine/Source/Runtime/AssetRegistry/Private/PackageReader.cpp
@@ -1,14 +1,16 @@
 // Copyright Epic Games, Inc. All Rights Reserved.
 
 #include "PackageReader.h"
+
+#include "AssetRegistry.h"
+#include "AssetRegistryPrivate.h"
+#include "AssetRegistry/AssetData.h"
 #include "HAL/FileManager.h"
 #include "Internationalization/Internationalization.h"
 #include "Logging/MessageLog.h"
+#include "Misc/Guid.h"
 #include "Misc/PackageName.h"
 #include "UObject/Class.h"
-#include "AssetRegistryPrivate.h"
-#include "AssetRegistry/AssetData.h"
-#include "AssetRegistry.h"
 
 FPackageReader::FPackageReader()
 	: Loader(nullptr)
@@ -348,8 +350,13 @@ bool FPackageReader::ReadDependencyData(FPackageDependencyData& OutDependencyDat
 	}
 
 	OutDependencyData.PackageName = FName(*PackageNameString);
-	OutDependencyData.PackageData.DiskSize = PackageFileSize;
-	OutDependencyData.PackageData.PackageGuid = PackageFileSummary.Guid;
+	FAssetPackageData& PackageData = OutDependencyData.PackageData;
+	PackageData.DiskSize = PackageFileSize;
+	PackageData.PackageGuid = PackageFileSummary.Guid;
+	PackageData.SetCustomVersions(PackageFileSummary.GetCustomVersionContainer().GetAllVersions());
+	PackageData.FileVersionUE4 = PackageFileSummary.GetFileVersionUE4();
+	PackageData.FileVersionLicenseeUE4 = PackageFileSummary.GetFileVersionLicenseeUE4();
+	PackageData.SetIsLicenseeVersion(PackageFileSummary.SavedByEngineVersion.IsLicenseeVersion());
 
 	if (!SerializeNameMap())
 	{
diff --git a/Engine/Source/Runtime/AssetRegistry/Public/AssetRegistry/IAssetRegistry.h b/Engine/Source/Runtime/AssetRegistry/Public/AssetRegistry/IAssetRegistry.h
--- a/Engine/Source/Runtime/AssetRegistry/Public/AssetRegistry/IAssetRegistry.h
+++ b/Engine/Source/Runtime/AssetRegistry/Public/AssetRegistry/IAssetRegistry.h
@@ -395,6 +395,10 @@ public:
 	UFUNCTION(BlueprintCallable, Category = "AssetRegistry")
 	virtual void WaitForCompletion() = 0;
 
+	/** Wait for the scan of a specific package to be complete */
+	UFUNCTION(BlueprintCallable, Category = "AssetRegistry")
+	virtual void WaitForPackage(const FString& PackageName) = 0;
+
 	/** If assets are currently being asynchronously scanned in the specified path, this will cause them to be scanned before other assets. */
 	UFUNCTION(BlueprintCallable, Category = "AssetRegistry")
 	virtual void PrioritizeSearchPath(const FString& PathToPrioritize) = 0;
diff --git a/Engine/Source/Runtime/CoreUObject/Private/AssetRegistry/AssetData.cpp b/Engine/Source/Runtime/CoreUObject/Private/AssetRegistry/AssetData.cpp
--- a/Engine/Source/Runtime/CoreUObject/Private/AssetRegistry/AssetData.cpp
+++ b/Engine/Source/Runtime/CoreUObject/Private/AssetRegistry/AssetData.cpp
@@ -1,7 +1,12 @@
 // Copyright Epic Games, Inc. All Rights Reserved.
 
 #include "AssetRegistry/AssetData.h"
+
+#include "Algo/Sort.h"
 #include "AssetRegistry/ARFilter.h"
+#include "Containers/Set.h"
+#include "HAL/CriticalSection.h"
+#include "Misc/ScopeRWLock.h"
 #include "Serialization/CustomVersion.h"
 #include "String/Find.h"
 #include "UObject/PropertyPortFlags.h"
@@ -173,3 +178,107 @@ bool FAssetRegistryVersion::SerializeVersion(FArchive& Ar, FAssetRegistryVersion
 
 	return !Ar.IsError();
 }
+
+namespace UE
+{
+namespace AssetRegistry
+{
+
+uint32 GetTypeHash(const TArray<FPackageCustomVersion>& Versions)
+{
+	constexpr uint32 HashPrime = 23;
+	uint32 Hash = 0;
+	for (const FPackageCustomVersion& Version : Versions)
+	{
+		Hash = Hash * HashPrime + GetTypeHash(Version.Key);
+		Hash = Hash * HashPrime + Version.Version;
+	}
+	return Hash;
+}
+
+class FPackageCustomVersionRegistry
+{
+public:
+	FPackageCustomVersionsHandle FindOrAdd(TArray<FPackageCustomVersion>&& InVersions)
+	{
+		FPackageCustomVersionsHandle Result;
+		Algo::Sort(InVersions);
+		uint32 Hash = GetTypeHash(InVersions);
+		{
+			FReadScopeLock ScopeLock(Lock);
+			TArray<FPackageCustomVersion>* Existing = RegisteredValues.FindByHash(Hash, InVersions);
+			if (Existing)
+			{
+				// We return a TArrayView with a pointer to the allocation managed by the element in the Set
+				// The element in the set may be destroyed and a moved copy recreated when the set changes size,
+				// but since TSet uses move constructors during the resize, the allocation will be unchanged,
+				// so we can safely refer to it from external handles.
+				Result.Ptr = TConstArrayView<FPackageCustomVersion>(*Existing);
+				return Result;
+			}
+		}
+		{
+			FWriteScopeLock ScopeLock(Lock);
+			TArray<FPackageCustomVersion>& Existing = RegisteredValues.FindOrAddByHash(Hash, MoveTemp(InVersions));
+			Result.Ptr = TConstArrayView<FPackageCustomVersion>(Existing);
+			return Result;
+		}
+	}
+
+private:
+	TSet<TArray<FPackageCustomVersion>> RegisteredValues;
+	FRWLock Lock;
+} GFPackageCustomVersionRegistry;
+
+FPackageCustomVersionsHandle FPackageCustomVersionsHandle::FindOrAdd(TConstArrayView<FCustomVersion> InVersions)
+{
+	TArray<FPackageCustomVersion> PackageFormat;
+	PackageFormat.Reserve(InVersions.Num());
+	for (const FCustomVersion& Version : InVersions)
+	{
+		PackageFormat.Emplace(Version.Key, Version.Version);
+	}
+	return GFPackageCustomVersionRegistry.FindOrAdd(MoveTemp(PackageFormat));
+}
+
+FPackageCustomVersionsHandle FPackageCustomVersionsHandle::FindOrAdd(TConstArrayView<FPackageCustomVersion> InVersions)
+{
+	return GFPackageCustomVersionRegistry.FindOrAdd(TArray<FPackageCustomVersion>(InVersions));
+}
+
+FPackageCustomVersionsHandle FPackageCustomVersionsHandle::FindOrAdd(TArray<FPackageCustomVersion>&& InVersions)
+{
+	return GFPackageCustomVersionRegistry.FindOrAdd(MoveTemp(InVersions));
+}
+
+FArchive& operator<<(FArchive& Ar, UE::AssetRegistry::FPackageCustomVersionsHandle& Handle)
+{
+	using namespace UE::AssetRegistry;
+
+	if (Ar.IsLoading())
+	{
+		int32 NumCustomVersions;
+		Ar << NumCustomVersions;
+		TArray<UE::AssetRegistry::FPackageCustomVersion> CustomVersions;
+		CustomVersions.SetNum(NumCustomVersions);
+		for (UE::AssetRegistry::FPackageCustomVersion& CustomVersion : CustomVersions)
+		{
+			Ar << CustomVersion;
+		}
+		Handle = FPackageCustomVersionsHandle::FindOrAdd(MoveTemp(CustomVersions));
+	}
+	else
+	{
+		TConstArrayView<UE::AssetRegistry::FPackageCustomVersion> CustomVersions = Handle.Get();
+		int32 NumCustomVersions = CustomVersions.Num();
+		Ar << NumCustomVersions;
+		for (UE::AssetRegistry::FPackageCustomVersion CustomVersion : CustomVersions)
+		{
+			Ar << CustomVersion;
+		}
+	}
+	return Ar;
+}
+
+}
+}
diff --git a/Engine/Source/Runtime/CoreUObject/Public/AssetRegistry/AssetData.h b/Engine/Source/Runtime/CoreUObject/Public/AssetRegistry/AssetData.h
--- a/Engine/Source/Runtime/CoreUObject/Public/AssetRegistry/AssetData.h
+++ b/Engine/Source/Runtime/CoreUObject/Public/AssetRegistry/AssetData.h
@@ -3,20 +3,25 @@
 #pragma once
 
 #include "CoreMinimal.h"
-#include "UObject/ObjectMacros.h"
-#include "UObject/Object.h"
-#include "UObject/Class.h"
-#include "UObject/SoftObjectPath.h"
-#include "UObject/Package.h"
-#include "UObject/ObjectRedirector.h"
+
+#include "AssetRegistry/AssetBundleData.h"
+#include "AssetRegistry/AssetDataTagMap.h"
+#include "Containers/ArrayView.h"
+#include "Containers/StringView.h"
 #include "Misc/PackageName.h"
 #include "Misc/StringBuilder.h"
-#include "Containers/StringView.h"
+#include "Templates/UniquePtr.h"
+#include "UObject/Class.h"
 #include "UObject/LinkerLoad.h"
-#include "AssetRegistry/AssetBundleData.h"
-#include "AssetRegistry/AssetDataTagMap.h"
+#include "UObject/SoftObjectPath.h"
+#include "UObject/Object.h"
+#include "UObject/ObjectMacros.h"
+#include "UObject/ObjectRedirector.h"
+#include "UObject/Package.h"
 #include "UObject/PrimaryAssetId.h"
 
+struct FCustomVersion;
+
 COREUOBJECT_API DECLARE_LOG_CATEGORY_EXTERN(LogAssetData, Log, All);
 
 /** Version used for serializing asset registry caches, both runtime and editor */
@@ -44,6 +49,7 @@ struct COREUOBJECT_API FAssetRegistryVersion
 								// * Switched from FName table to seek-free and more optimized FName batch loading
 								// * Removed global tag storage, a tag map reference-counts one store per asset registry
 								// * All configs can mix fixed and loose tag maps 
+		WorkspaceDomain,		// Added Version information to AssetPackageData
 
 		// -----<new versions can be added above this line>-------------------------------------------------
 		VersionPlusOne,
@@ -545,23 +551,124 @@ inline FAssetRegistryExportPath FAssetData::GetTagValueRef<FAssetRegistryExportP
 	return FoundValue.IsSet() ? FoundValue.AsExportPath() : FAssetRegistryExportPath();
 }
 
+namespace UE
+{
+namespace AssetRegistry
+{
+
+/** Low-memory version of FCustomVersion; holds only Guid and integer version. */
+struct FPackageCustomVersion
+{
+	FGuid Key;
+	int32 Version = 0;
+
+	FPackageCustomVersion() = default;
+	FPackageCustomVersion(const FGuid& InKey, const int32 InVersion)
+		: Key(InKey)
+		, Version(InVersion)
+	{
+	}
+	bool operator<(const FPackageCustomVersion& RHS) const
+	{
+		if (Key < RHS.Key) return true;
+		if (RHS.Key < Key) return false;
+		return Version < RHS.Version;
+	}
+	bool operator==(const FPackageCustomVersion& RHS) const
+	{
+		return Key == RHS.Key && Version == RHS.Version;
+	}
+	friend FArchive& operator<<(FArchive& Ar, FPackageCustomVersion& CustomVersion)
+	{
+		return Ar << CustomVersion.Key << CustomVersion.Version;
+	}
+};
+
+/** A handle to a deduplicated, sorted array of FPackageCustomVersion. */
+class FPackageCustomVersionsHandle
+{
+public:
+	TConstArrayView<FPackageCustomVersion> Get() const { return Ptr; }
+
+	COREUOBJECT_API static FPackageCustomVersionsHandle FindOrAdd(TConstArrayView<FCustomVersion> InCustomVersions);
+	COREUOBJECT_API static FPackageCustomVersionsHandle FindOrAdd(TConstArrayView<FPackageCustomVersion> InCustomVersions);
+	COREUOBJECT_API static FPackageCustomVersionsHandle FindOrAdd(TArray<FPackageCustomVersion>&& InCustomVersions);
+	COREUOBJECT_API friend FArchive& operator<<(FArchive& Ar, FPackageCustomVersionsHandle& Handle);
+
+private:
+	TConstArrayView<FPackageCustomVersion> Ptr;
+	friend class FPackageCustomVersionRegistry;
+};
+
+}
+}
+
 /** A class to hold data about a package on disk, this data is updated on save/load and is not updated when an asset changes in memory */
 class FAssetPackageData
 {
 public:
-	/** Total size of this asset on disk */
-	int64 DiskSize;
-
 	/** Guid of the source package, uniquely identifies an asset package */
 	FGuid PackageGuid;
 
 	/** MD5 of the cooked package on disk, for tracking nondeterministic changes */
 	FMD5Hash CookedHash;
 
+	/** Total size of this asset on disk */
+	int64 DiskSize;
+
+	/** UE4 file version that the package was saved with */
+	int32 FileVersionUE4;
+
+	/** Licensee file version that the package was saved with */
+	int32 FileVersionLicenseeUE4;
+
+private:
+	UE::AssetRegistry::FPackageCustomVersionsHandle CustomVersions;
+	/** Bit storage for flags */
+	uint32 Flags;
+
+public:
+
 	FAssetPackageData()
 		: DiskSize(0)
+		, FileVersionUE4(-1)
+		, FileVersionLicenseeUE4(-1)
+		, Flags(0)
+	{
+	}
+	FAssetPackageData(FAssetPackageData&& Other) = default;
+	FAssetPackageData(const FAssetPackageData& Other)
+		: FAssetPackageData()
 	{
+		*this = Other;
 	}
+	FAssetPackageData& operator=(FAssetPackageData&& Other) = default;
+	FAssetPackageData& operator=(const FAssetPackageData& Other) = default;
+
+	/**
+	 * Custom versions used by the package, used to check whether we need to update the package for the current binary.
+	 * The array is sorted by FPackageCustomVersion::operator<.
+	 */
+	TConstArrayView<UE::AssetRegistry::FPackageCustomVersion> GetCustomVersions() const
+	{
+		return CustomVersions.Get();
+	}
+	void SetCustomVersions(TConstArrayView<FCustomVersion> InCustomVersions)
+	{
+		CustomVersions = UE::AssetRegistry::FPackageCustomVersionsHandle::FindOrAdd(InCustomVersions);
+	}
+	void SetCustomVersions(TConstArrayView<UE::AssetRegistry::FPackageCustomVersion> InCustomVersions)
+	{
+		CustomVersions = UE::AssetRegistry::FPackageCustomVersionsHandle::FindOrAdd(InCustomVersions);
+	}
+	void SetCustomVersions(TArray<UE::AssetRegistry::FPackageCustomVersion>&& InCustomVersions)
+	{
+		CustomVersions = UE::AssetRegistry::FPackageCustomVersionsHandle::FindOrAdd(MoveTemp(InCustomVersions));
+	}
+
+	/** Whether the package was saved from a licensee executable, used to tell whether non-matching FileVersionLicenseeUE4 requires a resave */
+	bool IsLicenseeVersion() const { return (Flags & FLAG_LICENSEE_VERSION) != 0; }
+	void SetIsLicenseeVersion(bool bValue) { Flags = (Flags & ~FLAG_LICENSEE_VERSION) | (bValue ? FLAG_LICENSEE_VERSION : 0); }
 
 	/**
 	 * Serialize as part of the registry cache. This is not meant to be serialized as part of a package so  it does not handle versions normally
@@ -572,7 +679,17 @@ public:
 		Ar << DiskSize;
 		Ar << PackageGuid;
 		Ar << CookedHash;
+		Ar << FileVersionUE4;
+		Ar << FileVersionLicenseeUE4;
+		Ar << Flags;
+		Ar << CustomVersions;
 	}
+
+private:
+	enum
+	{
+		FLAG_LICENSEE_VERSION = 0x1,
+	};
 };
 
 /**
@@ -827,4 +944,3 @@ struct FAssetIdentifier
 	}
 };
 
-
