Do not capture objects to transaction when they are marked with any async flags
Do not allow new transaction to be created while in the midst of loading packages or post loading objects
Temporarily shelve an opened transaction while ticking async loading in the main thread and consider this scope as "loading packages"

diff --git a/Engine/Source/Editor/UnrealEd/Private/EditorServer.cpp b/Engine/Source/Editor/UnrealEd/Private/EditorServer.cpp
--- a/Engine/Source/Editor/UnrealEd/Private/EditorServer.cpp
+++ b/Engine/Source/Editor/UnrealEd/Private/EditorServer.cpp
@@ -1082,7 +1082,7 @@ int32 UEditorEngine::BeginTransaction(const TCHAR* TransactionContext, const FTe
 {
 	int32 Index = INDEX_NONE;
 
-	if (Trans)
+	if (CanTransact())
 	{
 		// generate transaction context
 		Index = Trans->Begin(TransactionContext, Description);
@@ -1093,12 +1093,16 @@ int32 UEditorEngine::BeginTransaction(const TCHAR* TransactionContext, const FTe
 
 int32 UEditorEngine::BeginTransaction(const FText& Description)
 {
-	return BeginTransaction(NULL, Description, NULL);
+	return BeginTransaction(nullptr, Description, nullptr);
 }
 
 bool UEditorEngine::CanTransact()
 {
-	return Trans != nullptr;
+	// we can transact if we have a transaction buffer and aren't currently loading packages or  routing postload.
+	// No transaction should be created during loading
+	return Trans != nullptr &&
+		!GIsEditorLoadingPackage &&
+		!FUObjectThreadContext::Get().IsRoutingPostLoad;
 }
 
 int32 UEditorEngine::EndTransaction()
diff --git a/Engine/Source/Editor/UnrealEd/Private/ScopedTransaction.cpp b/Engine/Source/Editor/UnrealEd/Private/ScopedTransaction.cpp
--- a/Engine/Source/Editor/UnrealEd/Private/ScopedTransaction.cpp
+++ b/Engine/Source/Editor/UnrealEd/Private/ScopedTransaction.cpp
@@ -17,7 +17,7 @@ FScopedTransaction::FScopedTransaction(const TCHAR* TransactionContext, const FT
 
 void FScopedTransaction::Construct (const TCHAR* TransactionContext, const FText& SessionName, UObject* PrimaryObject, const bool bShouldActuallyTransact)
 {
-	if( bShouldActuallyTransact && GEditor && GEditor->Trans && ensure(!GIsTransacting))
+	if( bShouldActuallyTransact && GEditor && GEditor->CanTransact() && ensure(!GIsTransacting))
 	{
 		Index = GEditor->BeginTransaction( TransactionContext, SessionName, PrimaryObject );
 		check( IsOutstanding() );
diff --git a/Engine/Source/Runtime/CoreUObject/Private/UObject/Obj.cpp b/Engine/Source/Runtime/CoreUObject/Private/UObject/Obj.cpp
--- a/Engine/Source/Runtime/CoreUObject/Private/UObject/Obj.cpp
+++ b/Engine/Source/Runtime/CoreUObject/Private/UObject/Obj.cpp
@@ -1214,7 +1214,8 @@ bool UObject::Modify( bool bAlwaysMarkDirty/*=true*/ )
 		// Do not consider script packages, as they should never end up in the
 		// transaction buffer and we don't want to mark them dirty here either.
 		// We do want to consider PIE objects however
-		if (GetOutermost()->HasAnyPackageFlags(PKG_ContainsScript | PKG_CompiledIn) == false || GetClass()->HasAnyClassFlags(CLASS_DefaultConfig | CLASS_Config))
+		if ((GetOutermost()->HasAnyPackageFlags(PKG_ContainsScript | PKG_CompiledIn) == false || GetClass()->HasAnyClassFlags(CLASS_DefaultConfig | CLASS_Config)) &&
+			!HasAnyInternalFlags(EInternalObjectFlags::Async | EInternalObjectFlags::AsyncLoading))
 		{
 			// Attempt to mark the package dirty and save a copy of the object to the transaction
 			// buffer. The save will fail if there isn't a valid transactor, the object isn't
diff --git a/Engine/Source/Runtime/CoreUObject/Private/UObject/UObjectGlobals.cpp b/Engine/Source/Runtime/CoreUObject/Private/UObject/UObjectGlobals.cpp
--- a/Engine/Source/Runtime/CoreUObject/Private/UObject/UObjectGlobals.cpp
+++ b/Engine/Source/Runtime/CoreUObject/Private/UObject/UObjectGlobals.cpp
@@ -2212,6 +2212,7 @@ UObject* StaticDuplicateObjectEx( FObjectDuplicationParameters& Parameters )
 
 bool SaveToTransactionBuffer(UObject* Object, bool bMarkDirty)
 {
+	check(!Object->HasAnyInternalFlags(EInternalObjectFlags::Async | EInternalObjectFlags::AsyncLoading));
 	bool bSavedToTransactionBuffer = false;
 
 	// Script packages should not end up in the transaction buffer.
@@ -3376,7 +3377,11 @@ UObject* StaticConstructObject_Internal(const FStaticConstructObjectParameters&
 		(*InClass->ClassConstructor)(FObjectInitializer(Result, Params));
 	}
 	
-	if( GIsEditor && GUndo && (InFlags & RF_Transactional) && !(InFlags & RF_NeedLoad) && !InClass->IsChildOf(UField::StaticClass()) )
+	if (GIsEditor && GUndo && 
+		(InFlags & RF_Transactional) && !(InFlags & RF_NeedLoad) && 
+		!InClass->IsChildOf(UField::StaticClass()) &&
+		// Do not consider object creation in transaction if the object is marked as async or in being async loaded 
+		!Result->HasAnyInternalFlags(EInternalObjectFlags::Async|EInternalObjectFlags::AsyncLoading))
 	{
 		// Set RF_PendingKill and update the undo buffer so an undo operation will set RF_PendingKill on the newly constructed object.
 		Result->MarkPendingKill();
