Fix ArchiveUObject SerializeObjectPtr not updating pointers when modifying reference and the pointer to be set is nullptr

diff --git a/Engine/Source/Runtime/CoreUObject/Private/Serialization/ArchiveUObject.cpp b/Engine/Source/Runtime/CoreUObject/Private/Serialization/ArchiveUObject.cpp
--- a/Engine/Source/Runtime/CoreUObject/Private/Serialization/ArchiveUObject.cpp
+++ b/Engine/Source/Runtime/CoreUObject/Private/Serialization/ArchiveUObject.cpp
@@ -67,7 +67,7 @@ FArchive& FArchiveUObject::SerializeObjectPtr(FArchive& Ar, FObjectPtr& Value)
 	// allows FObjectPtrs to be treated like raw UObject*'s by default.
 	UObject* Object = Value.Get();
 	Ar << Object;
-	if (Ar.IsLoading() || (Object && Ar.IsModifyingWeakAndStrongReferences()))
+	if (Ar.IsLoading() || Ar.IsModifyingWeakAndStrongReferences())
 	{
 		Value = Object;
 	}
