Avoid installing test access tracker on test construction as that causes it to get used in editor when no tests have been run.  Instead install on instantiation of the metrics tracking struct.

diff --git a/Engine/Source/Runtime/CoreUObject/Private/Tests/ObjectRefTrackingTestBase.h b/Engine/Source/Runtime/CoreUObject/Private/Tests/ObjectRefTrackingTestBase.h
--- a/Engine/Source/Runtime/CoreUObject/Private/Tests/ObjectRefTrackingTestBase.h
+++ b/Engine/Source/Runtime/CoreUObject/Private/Tests/ObjectRefTrackingTestBase.h
@@ -14,7 +14,6 @@ public:
 	FObjectRefTrackingTestBase(const FString& InName, const bool bInComplexTask)
 	: FAutomationTestBase(InName, bInComplexTask)
 	{
-		ConditionalInstallCallbacks();
 	}
 
 	uint32 GetNumResolves() const { return NumResolves; }
@@ -29,7 +28,9 @@ public:
 		, OriginalNumResolves(Test.GetNumResolves())
 		, OriginalNumFailedResolves(Test.GetNumFailedResolves())
 		, OriginalNumReads(Test.GetNumReads())
-		{}
+		{
+			Test.ConditionalInstallCallbacks();
+		}
 
 		bool TestNumResolves(const TCHAR* What, uint32 ExpectedDelta)
 		{
