Fix compile errors from IPackageResourceManager refactor if WITH_IOSTORE_IN_EDITOR.
#rnx

diff --git a/Engine/Source/Runtime/CoreUObject/Private/Serialization/AsyncLoading.cpp b/Engine/Source/Runtime/CoreUObject/Private/Serialization/AsyncLoading.cpp
--- a/Engine/Source/Runtime/CoreUObject/Private/Serialization/AsyncLoading.cpp
+++ b/Engine/Source/Runtime/CoreUObject/Private/Serialization/AsyncLoading.cpp
@@ -5781,8 +5781,8 @@ EAsyncPackageState::Type FAsyncPackage::CreateLinker()
 					SCOPED_LOADTIMER(CreateLinker_DoesExist);
 					bDoesPackageExist = FPackageName::DoesPackageExist(PackagePath, Guid, false /* bMatchCaseOnDisk */, &PackagePath);
 #if WITH_IOSTORE_IN_EDITOR
-				// Only look for non cooked packages on disk
-				DoesPackageExist &= !DoesPackageExistInIoStore(NameToLoadFName);
+					// Only look for non cooked packages on disk
+					bDoesPackageExist &= !DoesPackageExistInIoStore(NameToLoadFName);
 #endif
 				}
 
diff --git a/Engine/Source/Runtime/CoreUObject/Private/Serialization/EditorPackageLoader.cpp b/Engine/Source/Runtime/CoreUObject/Private/Serialization/EditorPackageLoader.cpp
--- a/Engine/Source/Runtime/CoreUObject/Private/Serialization/EditorPackageLoader.cpp
+++ b/Engine/Source/Runtime/CoreUObject/Private/Serialization/EditorPackageLoader.cpp
@@ -43,28 +43,26 @@ public:
 	}
 
 	virtual int32 LoadPackage(
-			const FString& InPackageName,
-			const FGuid* InGuid,
-			const TCHAR* InPackageToLoadFrom,
-			FLoadPackageAsyncDelegate InCompletionDelegate,
-			EPackageFlags InPackageFlags,
-			int32 InPIEInstanceID,
-			int32 InPackagePriority,
-			const FLinkerInstancingContext* InstancingContext) override
+		const FPackagePath& PackagePath,
+		FName CustomPackageName,
+		FLoadPackageAsyncDelegate InCompletionDelegate,
+		const FGuid* InGuid,
+		EPackageFlags InPackageFlags,
+		int32 InPIEInstanceID,
+		int32 InPackagePriority,
+		const FLinkerInstancingContext* InstancingContext) override
 	{
-		const TCHAR* PackageName = InPackageToLoadFrom ? InPackageToLoadFrom : *InPackageName;
-
 		// Use the old loader if an uncooked package exists on disk
-		const bool bDoesUncookedPackageExist = FPackageName::DoesPackageExist(InPackageName, nullptr, nullptr, true) && !DoesPackageExistInIoStore(FName(*InPackageName));
+		const bool bDoesUncookedPackageExist = FPackageName::DoesPackageExist(PackagePath) && !DoesPackageExistInIoStore(PackagePath.GetPackageFName());
 		if (bDoesUncookedPackageExist)
 		{
-			UE_LOG(LogEditorPackageLoader, Verbose, TEXT("Loading uncooked package '%s' from filesystem"), PackageName);
-			return UncookedPackageLoader->LoadPackage(InPackageName, InGuid, InPackageToLoadFrom, InCompletionDelegate, InPackageFlags, InPIEInstanceID, InPackagePriority, InstancingContext);
+			UE_LOG(LogEditorPackageLoader, Verbose, TEXT("Loading uncooked package '%s' from filesystem"), *PackagePath.GetDebugName());
+			return UncookedPackageLoader->LoadPackage(PackagePath, CustomPackageName, InCompletionDelegate, InGuid, InPackageFlags, InPIEInstanceID, InPackagePriority, InstancingContext);
 		}
 		else
 		{
-			UE_LOG(LogEditorPackageLoader, Verbose, TEXT("Loading cooked package '%s' from I/O Store"), PackageName);
-			return CookedPackageLoader->LoadPackage(InPackageName, InGuid, InPackageToLoadFrom, InCompletionDelegate, InPackageFlags, InPIEInstanceID, InPackagePriority, InstancingContext);
+			UE_LOG(LogEditorPackageLoader, Verbose, TEXT("Loading cooked package '%s' from I/O Store"), *PackagePath.GetDebugName());
+			return CookedPackageLoader->LoadPackage(PackagePath, CustomPackageName, InCompletionDelegate, InGuid, InPackageFlags, InPIEInstanceID, InPackagePriority, InstancingContext);
 		}
 	}
 
diff --git a/Engine/Source/Runtime/CoreUObject/Private/UObject/UObjectGlobals.cpp b/Engine/Source/Runtime/CoreUObject/Private/UObject/UObjectGlobals.cpp
--- a/Engine/Source/Runtime/CoreUObject/Private/UObject/UObjectGlobals.cpp
+++ b/Engine/Source/Runtime/CoreUObject/Private/UObject/UObjectGlobals.cpp
@@ -1118,7 +1118,7 @@ UPackage* LoadPackageInternal(UPackage* InOuter, const FPackagePath& PackagePath
 
 #if WITH_IOSTORE_IN_EDITOR
 		// Use the old loader if an uncooked package exists on disk
-		const bool bDoesUncookedPackageExist = FPackageName::DoesPackageExist(PackagePath, nullptr) && !DoesPackageExistInIoStore(PackageName);
+		const bool bDoesUncookedPackageExist = FPackageName::DoesPackageExist(PackagePath) && !DoesPackageExistInIoStore(PackageName);
 		if (!bDoesUncookedPackageExist)
 #endif
 		{
