FindPackagesRecursive needs to handle IPlatformFilePak returning relative paths.
FPackageName::SearchForPackageOnDisk (and therefore the simpler version of FindPackagesRecursive that it uses) needs to consider only headersegment packages.
#rnx

diff --git a/Engine/Source/Runtime/CoreUObject/Private/UObject/PackageResourceManager.cpp b/Engine/Source/Runtime/CoreUObject/Private/UObject/PackageResourceManager.cpp
--- a/Engine/Source/Runtime/CoreUObject/Private/UObject/PackageResourceManager.cpp
+++ b/Engine/Source/Runtime/CoreUObject/Private/UObject/PackageResourceManager.cpp
@@ -82,7 +82,10 @@ void IPackageResourceManager::FindPackagesRecursive(TArray<FPackagePath>& OutPac
 	OutPackages.Reserve(PackageSegments.Num());
 	for (const TPair<FPackagePath, EPackageSegment>& PackageSegment : PackageSegments)
 	{
-		OutPackages.Add(PackageSegment.Get<0>());
+		if (PackageSegment.Get<1>() == EPackageSegment::Header)
+		{
+			OutPackages.Add(PackageSegment.Get<0>());
+		}
 	}
 }
 
diff --git a/Engine/Source/Runtime/CoreUObject/Private/UObject/PackageResourceManagerFile.cpp b/Engine/Source/Runtime/CoreUObject/Private/UObject/PackageResourceManagerFile.cpp
--- a/Engine/Source/Runtime/CoreUObject/Private/UObject/PackageResourceManagerFile.cpp
+++ b/Engine/Source/Runtime/CoreUObject/Private/UObject/PackageResourceManagerFile.cpp
@@ -336,16 +336,32 @@ void FPackageResourceManagerFile::FindPackagesRecursive(TArray<TPair<FPackagePat
 		BasenameWildcardNullTerminated.ToString(), true /* Files */, false /* Directories */);
 	OutPackages.Reserve(OutPackages.Num() + FoundFilenames.Num());
 	TStringBuilder<256> ResultFile;
+	FString BufferPath;
 	for (const FString& Filename : FoundFilenames)
 	{
-		if (!Filename.StartsWith(FileMountAbsPath))
+		FStringView RelPath;
+		if (Filename.StartsWith(FileMountAbsPath, ESearchCase::IgnoreCase))
 		{
-			UE_LOG(LogPackageResourceManager, Warning,
-				TEXT("FindPackagesRecursive: Filename \"%s\" returned from FindFilesRecursive does not start with RootPath \"%s\". Ignoring it."),
-				*Filename, *FileMountAbsPath);
-			continue;
+			RelPath = FStringView(Filename);
+		}
+		else
+		{
+			// Check whether the reason it doesn't start with FileMountAbsPath is because it is not normalized or is e.g. a relative path
+			// ConvertRelativePathToFull will normalize it in addition to converting it to an absolute path
+			BufferPath = FPaths::ConvertRelativePathToFull(Filename);
+			if (BufferPath.StartsWith(FileMountAbsPath, ESearchCase::IgnoreCase))
+			{
+				RelPath = FStringView(BufferPath);
+			}
+			else
+			{
+				UE_LOG(LogPackageResourceManager, Warning,
+					TEXT("FindPackagesRecursive: Filename \"%s\" returned from FindFilesRecursive does not start with RootPath \"%s\". Ignoring it."),
+					*Filename, *FileMountAbsPath);
+				continue;
+			}
 		}
-		FStringView RelPath(FStringView(Filename).RightChop(FileMountAbsPath.Len()));
+		RelPath.RightChopInline(FileMountAbsPath.Len());
 		int32 ExtensionStart;
 		EPackageExtension Extension = FPackagePath::ParseExtension(RelPath, &ExtensionStart);
 		if (Extension == EPackageExtension::Custom || // Files with unrecognized extensions on disk are not returned for IteratePackages
