Fix non-unity build error in ArchiveCrc32.cpp

diff --git a/Engine/Source/Runtime/CoreUObject/Private/Serialization/ArchiveCrc32.cpp b/Engine/Source/Runtime/CoreUObject/Private/Serialization/ArchiveCrc32.cpp
--- a/Engine/Source/Runtime/CoreUObject/Private/Serialization/ArchiveCrc32.cpp
+++ b/Engine/Source/Runtime/CoreUObject/Private/Serialization/ArchiveCrc32.cpp
@@ -3,6 +3,9 @@
 #include "Serialization/ArchiveCrc32.h"
 #include "Serialization/ArchiveObjectCrc32.h"
 
+#include "UObject/Object.h"
+
+
 FArchiveCrc32::FArchiveCrc32(uint32 InCRC, UObject* InRootObject)
 	: CRC(InCRC)
 	, RootObject(InRootObject)
