FUntypedBulkData::Serialize: refactor to ensure that all options for BulkData location (Inline, end of file, various separate files) are supported in all serialization environments.
#rnx

diff --git a/Engine/Source/Runtime/Core/Public/UObject/DebugSerializationFlags.h b/Engine/Source/Runtime/Core/Public/UObject/DebugSerializationFlags.h
--- a/Engine/Source/Runtime/Core/Public/UObject/DebugSerializationFlags.h
+++ b/Engine/Source/Runtime/Core/Public/UObject/DebugSerializationFlags.h
@@ -15,6 +15,11 @@ enum EDebugSerializationFlags
 	/** No special flags */
 	DSF_None =					0x00000000,
 
+	/**
+	 * If FDiffSerializeArchive is being used, instruct it to NOT report diffs while this flag is set.
+	 * This is used e.g. when serializing offsets that are likely to change when there is any other change
+	 * in the serialization of the package.
+	 */
 	DSF_IgnoreDiff =			0x00000001,
 	DSF_EnableCookerWarnings =	0x00000002,
 };
diff --git a/Engine/Source/Runtime/CoreUObject/Private/Serialization/BulkData.cpp b/Engine/Source/Runtime/CoreUObject/Private/Serialization/BulkData.cpp
--- a/Engine/Source/Runtime/CoreUObject/Private/Serialization/BulkData.cpp
+++ b/Engine/Source/Runtime/CoreUObject/Private/Serialization/BulkData.cpp
@@ -445,7 +445,12 @@ bool FUntypedBulkData::DoesExist() const
  */
 FName FUntypedBulkData::GetDecompressionFormat() const
 {
-	return (BulkDataFlags & BULKDATA_SerializeCompressedZLIB) ? NAME_Zlib: NAME_None;
+	return GetDecompressionFormat(this->BulkDataFlags);
+}
+
+FName FUntypedBulkData::GetDecompressionFormat(EBulkDataFlags InFlags)
+{
+	return (InFlags & BULKDATA_SerializeCompressedZLIB) ? NAME_Zlib : NAME_None;
 }
 
 /**
@@ -517,7 +522,7 @@ void FUntypedBulkData::GetCopy( void** Dest, bool bDiscardInternalCopy )
 		// Data isn't currently loaded so we need to load it from disk.
 		else
 		{
-			LoadDataIntoMemory( *Dest );
+			ensure(TryLoadDataIntoMemory(*Dest));
 		}
 	}
 	// Passed in memory is NULL so we need to allocate some.
@@ -562,7 +567,7 @@ void FUntypedBulkData::GetCopy( void** Dest, bool bDiscardInternalCopy )
 				*Dest = FMemory::Malloc( BulkDataSize, BulkDataAlignment );
 
 				// ... and directly load into it.
-				LoadDataIntoMemory( *Dest );
+				ensure(TryLoadDataIntoMemory(*Dest));
 			}
 			else
 			{
@@ -796,12 +801,12 @@ bool FUntypedBulkData::StartAsyncLoading()
  */
 void FUntypedBulkData::SetBulkDataFlags( uint32 BulkDataFlagsToSet )
 {
-	BulkDataFlags = EBulkDataFlags(BulkDataFlags | BulkDataFlagsToSet);
+	BulkDataFlags = static_cast<EBulkDataFlags>(BulkDataFlags | BulkDataFlagsToSet);
 }
 
 void FUntypedBulkData::ResetBulkDataFlags(uint32 BulkDataFlagsToSet)
 {
-	BulkDataFlags = (EBulkDataFlags)BulkDataFlagsToSet;
+	BulkDataFlags = static_cast<EBulkDataFlags>(BulkDataFlagsToSet);
 }
 
 /**
@@ -841,7 +846,7 @@ uint32 FUntypedBulkData::GetBulkDataAlignment() const
  */
 void FUntypedBulkData::ClearBulkDataFlags( uint32 BulkDataFlagsToClear )
 {
-	BulkDataFlags = EBulkDataFlags(BulkDataFlags & ~BulkDataFlagsToClear);
+	BulkDataFlags = static_cast<EBulkDataFlags>(BulkDataFlags & ~BulkDataFlagsToClear);
 }
 
 void FUntypedBulkData::AsyncLoadBulkData()
@@ -854,7 +859,7 @@ void FUntypedBulkData::AsyncLoadBulkData()
 
 	// Seek to the beginning of the bulk data in the file.
 	Result.Archive->Seek(BulkDataOffsetInFile);
-	SerializeBulkData(*Result.Archive, BulkDataAsync.Get());
+	SerializeBulkData(*Result.Archive, BulkDataAsync.Get(), BulkDataFlags);
 }
 
 /*-----------------------------------------------------------------------------
@@ -889,22 +894,25 @@ static FAutoConsoleVariableRef CVarMinimumBulkDataSizeForAsyncLoading(
 	ECVF_Default
 	);
 
-bool FUntypedBulkData::ShouldStreamBulkData()
+bool FUntypedBulkData::ShouldStreamBulkData(FArchive& Ar)
 {
-	if (GEventDrivenLoaderEnabled && !(BulkDataFlags&BULKDATA_PayloadAtEndOfFile))
+	if (Ar.IsUsingEventDrivenLoader())
 	{
-		return false; // if it is inline, it is already precached, so use it
-	}
-
-	if (GEventDrivenLoaderEnabled)
-	{
-		const bool bSeperateFile = !!(BulkDataFlags&BULKDATA_PayloadInSeperateFile);
+#if WITH_IOSTORE_IN_EDITOR
+		// Streaming not yet implemented
+		return false;
+#else
+		if (!(BulkDataFlags & BULKDATA_PayloadAtEndOfFile))
+		{
+			return false; // if it is inline, it is already precached, so use it
+		}
 
-		if (!bSeperateFile)
+		if (BulkDataFlags & BULKDATA_PayloadInSeperateFile)
 		{
-			check(!"Bulk data should either be inline or stored in a separate file for the new uobject loader.");
+			checkf(false, TEXT("Bulk data should either be inline or stored in a separate file for the new uobject loader."));
 			return false; // if it is not in a separate file, then we can't easily find the correct offset in the uexp file; we don't want this case anyway!
 		}
+#endif
 	}
 
 	const bool bForceStream = !!(BulkDataFlags & BULKDATA_ForceStreamPayload);
@@ -920,6 +928,16 @@ bool FUntypedBulkData::NeedsOffsetFixup() const
 	return (BulkDataFlags & BULKDATA_NoOffsetFixUp) == 0;
 }
 
+void FUntypedBulkData::SetBulkDataFlagsOn(EBulkDataFlags& InOutAccumulator, EBulkDataFlags FlagsToSet)
+{
+	InOutAccumulator = static_cast<EBulkDataFlags>(InOutAccumulator | FlagsToSet);
+}
+
+void FUntypedBulkData::ClearBulkDataFlagsOn(EBulkDataFlags& InOutAccumulator, EBulkDataFlags FlagsToClear)
+{
+	InOutAccumulator = static_cast<EBulkDataFlags>(InOutAccumulator & ~FlagsToClear);
+}
+
 void FUntypedBulkData::Serialize( FArchive& Ar, UObject* Owner, int32 Idx, bool bAttemptFileMapping, EFileRegionType FileRegionType)
 {
 	DECLARE_SCOPE_CYCLE_COUNTER(TEXT("FUntypedBulkData::Serialize"), STAT_UBD_Serialize, STATGROUP_Memory);
@@ -930,18 +948,19 @@ void FUntypedBulkData::Serialize( FArchive& Ar, UObject* Owner, int32 Idx, bool
 
 	check(!bAttemptFileMapping || Ar.IsLoading()); // makes no sense to map unless we are loading
 
-	if(Ar.IsTransacting())
+	if (Ar.IsTransacting())
 	{
 		// Special case for transacting bulk data arrays.
 
-		// constructing the object during load will save it to the transaction buffer. If it tries to load the bulk data now it will try to break it.
+		// Constructing the object during load will save it to the transaction buffer.
+		// We need to cancel that save because trying to load the bulk data now would break.
 		bool bActuallySave = Ar.IsSaving() && (!Owner || !Owner->HasAnyFlags(RF_NeedLoad));
 
 		Ar << bActuallySave;
 
 		if (bActuallySave)
 		{
-			if(Ar.IsLoading())
+			if (Ar.IsLoading())
 			{
 				// Flags for bulk data.
 				Ar << BulkDataFlags;
@@ -952,9 +971,9 @@ void FUntypedBulkData::Serialize( FArchive& Ar, UObject* Owner, int32 Idx, bool
 				BulkData.Reallocate( GetBulkDataSize(), BulkDataAlignment );
 
 				// Deserialize bulk data.
-				SerializeBulkData( Ar, BulkData.Get() );
+				SerializeBulkData(Ar, BulkData.Get(), BulkDataFlags);
 			}
-			else if(Ar.IsSaving())
+			else if (Ar.IsSaving())
 			{
 				// Flags for bulk data.
 				Ar << BulkDataFlags;
@@ -970,105 +989,91 @@ void FUntypedBulkData::Serialize( FArchive& Ar, UObject* Owner, int32 Idx, bool
 					MakeSureBulkDataIsLoaded();
 
 					// Serialize bulk data.
-					SerializeBulkData(Ar, BulkData.Get());
+					SerializeBulkData(Ar, BulkData.Get(), BulkDataFlags);
 				}
 			}
 		}
 	}
-	else if( Ar.IsPersistent() && !Ar.IsObjectReferenceCollector() && !Ar.ShouldSkipBulkData() )
+	else if ( Ar.IsPersistent() && !Ar.IsObjectReferenceCollector() && !Ar.ShouldSkipBulkData() )
 	{
 #if TRACK_BULKDATA_USE
 		FThreadSafeBulkDataToObjectMap::Get().Add( this, Owner );
 #endif
-		// When saving, this block mutates values in order to write them out to the archive, even if just cooking. Store the existing version of them so we can restore them if necessary after writing out the bytes.
-		EBulkDataFlags SavedBulkDataFlags = BulkDataFlags;
-
-		// Offset where the bulkdata flags are stored
-		int64 SavedBulkDataFlagsPos = Ar.Tell();
+		// We're loading from the persistent archive.
+		if (Ar.IsLoading())
 		{
-			FArchive::FScopeSetDebugSerializationFlags S(Ar, DSF_IgnoreDiff);
-
-			// If the bulk data size is greater than can be held in an int32, then potentially the ElementCount
-			// and BulkDataSizeOnDisk need to be held as int64s, so set a flag indicating the new format.
-			if (Ar.IsSaving() && GetBulkDataSize() >= (1LL << 31))
+			EBulkDataFlags AddedFlags = static_cast<EBulkDataFlags>(0);
+			EBulkDataFlags RemovedFlags = static_cast<EBulkDataFlags>(0);
+			auto SetLocalBulkDataFlags = [this, &AddedFlags, &RemovedFlags](EBulkDataFlags BulkDataFlagsToSet)
 			{
-				SetBulkDataFlags(BULKDATA_Size64Bit);
-			}
-			Ar << BulkDataFlags;
-		}
+				AddedFlags = static_cast<EBulkDataFlags>(AddedFlags | BulkDataFlagsToSet);
+				RemovedFlags = static_cast<EBulkDataFlags>(RemovedFlags & ~BulkDataFlagsToSet);
+				SetBulkDataFlags(BulkDataFlagsToSet);
+			};
+			auto ClearLocalBulkDataFlags = [this, &AddedFlags, &RemovedFlags](EBulkDataFlags BulkDataFlagsToClear)
+			{
+				AddedFlags = static_cast<EBulkDataFlags>(AddedFlags & ~BulkDataFlagsToClear);
+				RemovedFlags = static_cast<EBulkDataFlags>(RemovedFlags | BulkDataFlagsToClear);
+				ClearBulkDataFlags(BulkDataFlagsToClear);
+			};
 
+			Ar << BulkDataFlags;
+			SerializeBulkDataSizeInt(Ar, ElementCount, BulkDataFlags);
 
-		// Number of elements in array.
-		if (BulkDataFlags & BULKDATA_Size64Bit)
-		{
-			Ar << ElementCount;
-		}
-		else
-		{
-			int32 ElementCountAsInt32 = ElementCount;
-			Ar << ElementCountAsInt32;
-			ElementCount = ElementCountAsInt32;
-		}
+			// Size on disk, which in the case of compression is != GetBulkDataSize()
+			SerializeBulkDataSizeInt(Ar, BulkDataSizeOnDisk, BulkDataFlags);
+			Ar << BulkDataOffsetInFile;
 
-		// We're loading from the persistent archive.
-		if (Ar.IsLoading())
-		{
-			PackagePath = FPackagePath();
-			PackageSegment = EPackageSegment::Header;
+			// Oct 2019: We erroneously wrote uint16 ChunkID during the saving of BulkData; it should've been written only when saving AND cooking to avoid writing it to the uasset or DDC.
+			// To fix this we use a flag that was also set by those bad versions - BULKDATA_BadDataVersion - to identify them here when loading and load and discard the value.
+			// TODO: Increase the DDC keys for the data types that write BulkData to the DDC and remove this check and BULKDATA_BadDataVersion.
+			if ((BulkDataFlags & BULKDATA_BadDataVersion) != 0)
+			{
+				uint16 DummyValue;
+				Ar << DummyValue;
+				ClearLocalBulkDataFlags(BULKDATA_BadDataVersion);
+			}
 
+			EBulkDataFlags DuplicateDataFlags = static_cast<EBulkDataFlags>(0);
+			int64 DuplicateSizeOnDisk = 0;
+			int64 DuplicateDataOffsetInFile = 0;
+			if (BulkDataFlags & BULKDATA_DuplicateNonOptionalPayload)
+			{
+				Ar << DuplicateDataFlags;
+				SerializeBulkDataSizeInt(Ar, DuplicateSizeOnDisk, DuplicateDataFlags);
+				Ar << DuplicateDataOffsetInFile;
+			}
+			
 			// @todo when Landscape (and others?) only Lock/Unlock once, we can enable this
 			if (false) // FPlatformProperties::RequiresCookedData())
 			{
 				// Bulk data that is being serialized via seekfree loading is single use only. This allows us 
 				// to free the memory as e.g. the bulk data won't be attached to an archive in the case of
 				// seek free loading.
-				SetBulkDataFlags(BULKDATA_SingleUse);
+				SetLocalBulkDataFlags(BULKDATA_SingleUse);
 			}
 
 			// Hacky fix for using cooked data in editor. Cooking sets BULKDATA_SingleUse for textures, but PIEing needs to keep bulk data around.
 			if (GIsEditor)
 			{
-				ClearBulkDataFlags(BULKDATA_SingleUse);
-			}
-
-			// Size on disk, which in the case of compression is != GetBulkDataSize()
-			if (BulkDataFlags & BULKDATA_Size64Bit)
-			{
-				Ar << BulkDataSizeOnDisk;
-			}
-			else
-			{
-				int32 BulkDataSizeOnDiskAsInt32;
-				Ar << BulkDataSizeOnDiskAsInt32;
-				BulkDataSizeOnDisk = BulkDataSizeOnDiskAsInt32;
+				ClearLocalBulkDataFlags(BULKDATA_SingleUse);
 			}
 
-			Ar << BulkDataOffsetInFile;
-
-			if ((BulkDataFlags & BULKDATA_BadDataVersion) != 0)
-			{
-				uint16 DummyValue;
-				Ar << DummyValue;
+			const bool bPayloadInline = !(BulkDataFlags & BULKDATA_PayloadAtEndOfFile);
+			const bool bPayloadInSeparateFile = !bPayloadInline && (BulkDataFlags & BULKDATA_PayloadInSeperateFile);
 
-				ClearBulkDataFlags(BULKDATA_BadDataVersion);	
-			}
-			
-			// determine whether the payload is stored inline or at the end of the file
-			const bool bPayloadInline = !(BulkDataFlags&BULKDATA_PayloadAtEndOfFile);
+			// GetLinker
+			bool bUseIOStore = false;
+#if WITH_EDITOR
 #if WITH_IOSTORE_IN_EDITOR
-			if (Ar.IsUsingEventDrivenLoader())
+			if (Ar.IsUsingEventDrivenLoader() && bPayloadInSeparateFile)
 			{
-				check(Owner);
-				check(Owner->GetPackage()->GetPackageId().IsValid());
-				SetBulkDataFlags(BULKDATA_UsesIoDispatcher);
+				check(Owner && Owner->GetPackage()->GetPackageId().IsValid());
+				SetLocalBulkDataFlags(BULKDATA_UsesIoDispatcher);
 				PackageId = Owner->GetPackage()->GetPackageId();
-				SerializeFromIoStore(Ar, Owner, Idx, bAttemptFileMapping);
-				return;
+				bUseIOStore = true;
 			}
 #endif // WITH_IOSTORE_IN_EDITOR
-
-			// GetLinker
-#if WITH_EDITOR
 			if (Owner != nullptr)
 			{
 				Linker = Owner->GetLinker();
@@ -1080,139 +1085,116 @@ void FUntypedBulkData::Serialize( FArchive& Ar, UObject* Owner, int32 Idx, bool
 				Package = Owner->GetOutermost();
 				check(Package.IsValid());
 				Linker = FLinkerLoad::FindExistingLinkerForPackage(Package.Get());
+				check(!Owner->GetLinker() || Owner->GetLinker() == Linker);
 			}
 #endif
 
-			// fix up the file offset, but only if not stored inline
-			int64 OffsetInFileFixup = 0;
-			if (Owner != NULL && Linker && !bPayloadInline && NeedsOffsetFixup())
+			// fix up the file offset if the offset is relative to the file's bulkdata section
+			if (!bPayloadInline && NeedsOffsetFixup())
 			{
-				OffsetInFileFixup = Linker->Summary.BulkDataStartOffset;
+				check(Linker);
+				check(!bUseIOStore);
+				BulkDataOffsetInFile += Linker->Summary.BulkDataStartOffset;
 			}
-			BulkDataOffsetInFile += OffsetInFileFixup;
 
-			// We're allowing deferred serialization.
-			FArchive* CacheableArchive = Ar.GetCacheableArchive();
-			if (Ar.IsAllowingLazyLoading() && Owner != NULL && CacheableArchive)
-			{				
-#if WITH_EDITOR
-				check(Linker);
-				CacheableArchive->AttachBulkData(Owner, this);
-				check(!CacheableArchive->IsTextFormat());
-				AttachedAr = CacheableArchive;
-				PackagePath = Linker->GetPackagePath();
-				PackageSegment = EPackageSegment::Header;
-#else
-				check(Linker);
+			// Get the PackagePath and PackageSegment
+			PackagePath = FPackagePath();
+			if (Linker)
+			{
 				PackagePath = Linker->GetPackagePath();
-				PackageSegment = EPackageSegment::Header;
-
-				check(!Owner->GetLinker() || Owner->GetLinker() == Linker);
-#endif // WITH_EDITOR
-
-				if (bPayloadInline)
-				{
-					if (ShouldStreamBulkData())
-					{
-						// Start serializing immediately
-						StartSerializingBulkData(Ar, Owner, Idx, bPayloadInline);
-					}
-					else
-					{
-						bool bWasMapped = false;
-						if (bAttemptFileMapping)
-						{
-							UE_LOG(LogSerialization, Error, TEXT("Attempt to file map inline bulk data. This is not desireable. File %s"), *PackagePath.GetDebugName(PackageSegment));
-
-							if (GEventDrivenLoaderEnabled && PackageSegment == EPackageSegment::Header)
-							{
-								BulkDataOffsetInFile -= IPackageResourceManager::Get().FileSize(PackagePath, PackageSegment);
-								check(BulkDataOffsetInFile >= 0);
-								PackageSegment = EPackageSegment::Exports;
-							}
-							bWasMapped = BulkData.MapFile(PackagePath, PackageSegment, BulkDataOffsetInFile, GetBulkDataSize());
-						}
-						if (bWasMapped)
-						{
-							// we need to seek past the inline bulk data 
-							// @todo, we really don't want to do this with inline data
-							Ar.Seek(Ar.Tell() + GetBulkDataSize());
-						}
-						else
-						{
-							// Force non-lazy loading of inline bulk data to prevent PostLoad spikes.
-							BulkData.Reallocate(GetBulkDataSize(), BulkDataAlignment);
-							// if the payload is stored inline, just serialize it
-							SerializeBulkData(Ar, BulkData.Get());
-						}
-					}
-				}
-				else if (BulkDataFlags & BULKDATA_DuplicateNonOptionalPayload)
+			}
+			PackageSegment = EPackageSegment::Header;
+			if (bPayloadInSeparateFile)
+			{
+				if (BulkDataFlags & BULKDATA_DuplicateNonOptionalPayload)
 				{
-					// Load from optional payload instead if optional payload is available
-					if (IPackageResourceManager::Get().DoesPackageExist(PackagePath, EPackageSegment::BulkDataOptional))
+					// The (required) payload is stored in both the default segment and the optional segment.
+					// Load it from the optional segment; this is preferable to loading it from the default segment
+					// because it reduces seek times if we want to load optional bulk data elsewhere in this package.
+					if (!bUseIOStore && IPackageResourceManager::Get().DoesPackageExist(PackagePath, EPackageSegment::BulkDataOptional))
 					{
 						PackageSegment = EPackageSegment::BulkDataOptional;
-						Ar << BulkDataFlags;
-						SetBulkDataFlags(BULKDATA_OptionalPayload);
-
-						if (BulkDataFlags & BULKDATA_Size64Bit)
-						{
-							Ar << BulkDataSizeOnDisk;
-						}
-						else
+						BulkDataFlags = static_cast<EBulkDataFlags>((DuplicateDataFlags | AddedFlags) & ~RemovedFlags);
+						SetBulkDataFlags(BULKDATA_OptionalPayload | BULKDATA_PayloadInSeperateFile | BULKDATA_PayloadAtEndOfFile);
+						BulkDataOffsetInFile = DuplicateDataOffsetInFile;
+						if (NeedsOffsetFixup())
 						{
-							int32 BulkDataSizeOnDiskAsInt32;
-							Ar << BulkDataSizeOnDiskAsInt32;
-							BulkDataSizeOnDisk = BulkDataSizeOnDiskAsInt32;
+							check(Linker);
+							BulkDataOffsetInFile += Linker->Summary.BulkDataStartOffset;
 						}
-
-						Ar << BulkDataOffsetInFile;
-						BulkDataOffsetInFile += OffsetInFileFixup;
 					}
 					else
 					{
 						PackageSegment = EPackageSegment::BulkDataDefault;
-						// Skip all bulkdata info of the copy in uptnl.
-						uint32 DummyBulkDataFlags;
-						Ar << DummyBulkDataFlags;
-
-						if (BulkDataFlags & BULKDATA_Size64Bit)
-						{
-							int64 DummyBulkDataSizeOnDisk;
-							Ar << DummyBulkDataSizeOnDisk;
-						}
-						else
-						{
-							int32 DummyBulkDataSizeOnDiskAsInt32;
-							Ar << DummyBulkDataSizeOnDiskAsInt32;
-						}
-
-						int64 DummyBulkDataOffsetInFile;
-						Ar << DummyBulkDataOffsetInFile;
 					}
 				}
 				else if (BulkDataFlags & BULKDATA_OptionalPayload)
 				{
 					PackageSegment = EPackageSegment::BulkDataOptional;
 				}
-				else if (BulkDataFlags & BULKDATA_PayloadInSeperateFile)
+				else if (BulkDataFlags & BULKDATA_MemoryMappedPayload)
+				{
+					PackageSegment = EPackageSegment::BulkDataMemoryMapped;
+				}
+				else
 				{
-					if (BulkDataFlags & BULKDATA_MemoryMappedPayload)
+					PackageSegment = EPackageSegment::BulkDataDefault;
+				}
+			}
+			else
+			{
+				check(PackageSegment == EPackageSegment::Header)
+				if (GEventDrivenLoaderEnabled)
+				{
+					BulkDataOffsetInFile -= IPackageResourceManager::Get().FileSize(PackagePath, PackageSegment);
+					check(BulkDataOffsetInFile >= 0);
+					PackageSegment = EPackageSegment::Exports;
+				}
+			}
+
+			FArchive* CacheableArchive = Ar.GetCacheableArchive();
+			if (Ar.IsAllowingLazyLoading() && Owner != nullptr && CacheableArchive && !bUseIOStore)
+			{
+				// Deferred serialization is allowed by the archive/caller. Set flags for how to read
+				// the data, but skip synchronous reading of it until TryLoadDataIntoMemory is called.
+#if WITH_EDITOR
+				CacheableArchive->AttachBulkData(Owner, this);
+				check(!CacheableArchive->IsTextFormat());
+				AttachedAr = CacheableArchive;
+#endif // WITH_EDITOR
+
+				if (bPayloadInline)
+				{
+					// Inline payloads load immediately even though the archive/caller is allowing deferred loads.
+
+					// Check whether we should load the data asynchronously before we proceed to synchronous or mapped load.
+					if (ShouldStreamBulkData(Ar))
+					{
+						// Serialize the bulk data asynchronously, and start the task immediately
+						StartSerializingBulkData(Ar, Owner, Idx, bPayloadInline);
+					}
+					// Check whether we should memory map the data before we proceed to synchronous load.
+					else if (bAttemptFileMapping &&
+						BulkData.MapFile(PackagePath, PackageSegment, BulkDataOffsetInFile, GetBulkDataSize()))
 					{
-						PackageSegment = EPackageSegment::BulkDataMemoryMapped;
+						UE_LOG(LogSerialization, Error, TEXT("Attempt to file map inline bulk data. This is not desireable. File %s"),
+							*PackagePath.GetDebugName(PackageSegment));
+						// we need to seek past the inline bulk data 
+						Ar.Seek(Ar.Tell() + GetBulkDataSize());
 					}
 					else
 					{
-						PackageSegment = EPackageSegment::BulkDataDefault;
+						// Force non-lazy loading of inline bulk data to prevent PostLoad spikes.
+						BulkData.Reallocate(GetBulkDataSize(), BulkDataAlignment);
+						// if the payload is stored inline, just serialize it
+						SerializeBulkData(Ar, BulkData.Get(), BulkDataFlags);
 					}
-
-
+				}
+				else
+				{
 					if (bAttemptFileMapping)
 					{
-						check(!FPlatformProperties::GetMemoryMappingAlignment() || IsAligned(BulkDataOffsetInFile, FPlatformProperties::GetMemoryMappingAlignment()));
-						bool bWasMapped = BulkData.MapFile(PackagePath, PackageSegment, BulkDataOffsetInFile, GetBulkDataSize());
-						if (!bWasMapped)
+						if (!BulkData.MapFile(PackagePath, PackageSegment, BulkDataOffsetInFile, GetBulkDataSize()))
 						{
 							// we failed to map when requested, so they will be looking for the memory, do a sync load now.
 							ForceBulkDataResident();
@@ -1220,15 +1202,13 @@ void FUntypedBulkData::Serialize( FArchive& Ar, UObject* Owner, int32 Idx, bool
 					}
 				}
 			}
-			// Serialize the bulk data right away.
 			else
 			{
-				if (Owner && Owner->GetLinker())
-				{
-					PackagePath = Owner->GetLinker()->GetPackagePath();
-					PackageSegment = EPackageSegment::Header;
-				}
-				if (ShouldStreamBulkData())
+				// Serialize the bulk data right away.
+				// bAttemptFileMapping is ignored in this case, since we are going to load the memory before returning anyway.
+
+				// Check whether we should load the data asynchronously before we proceed to synchronous load.
+				if (ShouldStreamBulkData(Ar))
 				{
 					StartSerializingBulkData(Ar, Owner, Idx, bPayloadInline);
 				}
@@ -1239,31 +1219,39 @@ void FUntypedBulkData::Serialize( FArchive& Ar, UObject* Owner, int32 Idx, bool
 					if (bPayloadInline)
 					{
 						// if the payload is stored inline, just serialize it
-						SerializeBulkData( Ar, BulkData.Get() );
+						SerializeBulkData(Ar, BulkData.Get(), BulkDataFlags);
 					}
 					else
 					{
 						// if the payload is NOT stored inline ...
-						if (BulkDataFlags & BULKDATA_PayloadInSeperateFile)
+						if (bPayloadInSeparateFile)
 						{
-							// open separate bulk data file
-							UE_CLOG(GEventDrivenLoaderEnabled, LogSerialization, Error, TEXT("Attempt to sync load bulk data with EDL enabled (separate file). This is not desireable. File %s"), *PackagePath.GetDebugName(PackageSegment));
-
-							if (GEventDrivenLoaderEnabled && PackageSegment == EPackageSegment::Header)
+#if WITH_IOSTORE_IN_EDITOR
+							if (bUseIOStore)
 							{
-								BulkDataOffsetInFile -= IPackageResourceManager::Get().FileSize(PackagePath, PackageSegment);
-								check(BulkDataOffsetInFile >= 0);
-								PackageSegment = EPackageSegment::Exports;
+								TUniquePtr<IBulkDataIORequest> Request = CreateBulkDataIoDispatcherRequest(CreateBulkDataChunkId(PackageId, BulkDataFlags));
+								Request->WaitCompletion();
+
+								FLargeMemoryReader MemoryAr(Request->GetReadResults(), Request->GetSize());
+								MemoryAr.Seek(BulkDataOffsetInFile);
+
+								SerializeBulkData(MemoryAr, BulkData.Get(), BulkDataFlags);
 							}
+							else
+#endif
+							{
+								// open separate bulk data file
+								UE_CLOG(GEventDrivenLoaderEnabled, LogSerialization, Error, TEXT("Attempt to sync load bulk data with EDL enabled (separate file). This is not desireable. File %s"), *PackagePath.GetDebugName(PackageSegment));
 
-							FOpenPackageResult TargetArchive = IPackageResourceManager::Get().OpenReadPackage(PackagePath, PackageSegment);
-							checkf(TargetArchive.Archive.IsValid() && TargetArchive.Format == EPackageFormat::Binary, TEXT("Attempted to load bulk data from an invalid PackagePath '%s': %s."),
-								*PackagePath.GetDebugName(PackageSegment), (!TargetArchive.Archive.IsValid() ? TEXT("could not find package") : TEXT("package is a TextAsset which is not supported")));
+								FOpenPackageResult TargetArchive = IPackageResourceManager::Get().OpenReadPackage(PackagePath, PackageSegment);
+								checkf(TargetArchive.Archive.IsValid() && TargetArchive.Format == EPackageFormat::Binary, TEXT("Attempted to load bulk data from an invalid PackagePath '%s': %s."),
+									*PackagePath.GetDebugName(PackageSegment), (!TargetArchive.Archive.IsValid() ? TEXT("could not find package") : TEXT("package is a TextAsset which is not supported")));
 
-							// seek to the location in the file where the payload is stored
-							TargetArchive.Archive->Seek(BulkDataOffsetInFile);
-							// serialize the payload
-							SerializeBulkData(*TargetArchive.Archive, BulkData.Get());
+								// seek to the location in the file where the payload is stored
+								TargetArchive.Archive->Seek(BulkDataOffsetInFile);
+								// serialize the payload
+								SerializeBulkData(*TargetArchive.Archive, BulkData.Get(), BulkDataFlags);
+							}
 						}
 						else
 						{
@@ -1274,7 +1262,7 @@ void FUntypedBulkData::Serialize( FArchive& Ar, UObject* Owner, int32 Idx, bool
 							// seek to the location in the file where the payload is stored
 							Ar.Seek(BulkDataOffsetInFile);
 							// serialize the payload
-							SerializeBulkData(Ar, BulkData.Get());
+							SerializeBulkData(Ar, BulkData.Get(), BulkDataFlags);
 							// seek to the location we came from
 							Ar.Seek(CurOffset);
 						}
@@ -1283,53 +1271,63 @@ void FUntypedBulkData::Serialize( FArchive& Ar, UObject* Owner, int32 Idx, bool
 			}
 		}
 		// We're saving to the persistent archive.
-		else if( Ar.IsSaving() )
+		else if ( Ar.IsSaving() )
 		{
-			// This block mutates values in order to write them out to the archive, even if just cooking. Store the existing version of them so we can restore them if necessary after writing out the bytes.
-			int64 SavedBulkDataSizeOnDisk = BulkDataSizeOnDisk;
-			int64 SavedBulkDataOffsetInFile = BulkDataOffsetInFile;
+			// Make sure bulk data is loaded.
+			MakeSureBulkDataIsLoaded();
+
+			// Make mutable copies of the bulkdata location variables
+			EBulkDataFlags LocalBulkDataFlags = BulkDataFlags;
+			int64 LocalBulkDataSizeOnDisk = BulkDataSizeOnDisk;
+			int64 LocalBulkDataOffsetInFile = BulkDataOffsetInFile;
 
+			// If the bulk data size is greater than can be held in an int32, then potentially the ElementCount
+			// and BulkDataSizeOnDisk need to be held as int64s, so set a flag indicating the new format.
+			if (GetBulkDataSize() >= (1LL << 31))
+			{
+				SetBulkDataFlagsOn(LocalBulkDataFlags, BULKDATA_Size64Bit);
+			}
+			else
+			{
+				ClearBulkDataFlagsOn(LocalBulkDataFlags, BULKDATA_Size64Bit);
+			}
 			// Remove single element serialization requirement before saving out bulk data flags.
-			ClearBulkDataFlags(BULKDATA_ForceSingleElementSerialization);
+			ClearBulkDataFlagsOn(LocalBulkDataFlags, BULKDATA_ForceSingleElementSerialization);
+
+			// Save offset where we are serializing BulkDataFlags and store a placeholder
+			int64 SavedBulkDataFlagsPos = Ar.Tell();
+			{
+				FArchive::FScopeSetDebugSerializationFlags S(Ar, DSF_IgnoreDiff);
+				Ar << LocalBulkDataFlags;
+			}
+
+			// Number of elements in array.
+			SerializeBulkDataSizeInt(Ar, ElementCount, LocalBulkDataFlags);
 
-			// Make sure bulk data is loaded.
-			MakeSureBulkDataIsLoaded();
-			
 			// Only serialize status information if wanted.
 			int64 SavedBulkDataSizeOnDiskPos	= INDEX_NONE;
 			int64 SavedBulkDataOffsetInFilePos	= INDEX_NONE;
 			
-			// Keep track of position we are going to serialize placeholder BulkDataSizeOnDisk.
-			SavedBulkDataSizeOnDiskPos = Ar.Tell();
-			BulkDataSizeOnDisk = INDEX_NONE;
-
 			{
 				FArchive::FScopeSetDebugSerializationFlags S(Ar, DSF_IgnoreDiff);
 
-				// And serialize the placeholder which is going to be overwritten later.
-				if (BulkDataFlags & BULKDATA_Size64Bit)
-				{
-					Ar << BulkDataSizeOnDisk;
-				}
-				else
-				{
-					int32 BulkDataSizeOnDiskAsInt32 = BulkDataSizeOnDisk;
-					Ar << BulkDataSizeOnDiskAsInt32;
-				}
-				// Keep track of position we are going to serialize placeholder BulkDataOffsetInFile.
+				// Save offset where we are serializing BulkDataSizeOnDisk and store a placeholder
+				SavedBulkDataSizeOnDiskPos = Ar.Tell();
+				LocalBulkDataSizeOnDisk = INDEX_NONE;
+				SerializeBulkDataSizeInt(Ar, LocalBulkDataSizeOnDisk, LocalBulkDataFlags);
+
+				// Save offset where we are serializing BulkDataOffsetInFile and store a placeholder
 				SavedBulkDataOffsetInFilePos = Ar.Tell();
-				BulkDataOffsetInFile = INDEX_NONE;
-				// And serialize the placeholder which is going to be overwritten later.
-				Ar << BulkDataOffsetInFile;
+				LocalBulkDataOffsetInFile = INDEX_NONE;
+				Ar << LocalBulkDataOffsetInFile;
 			}
 
 			// try to get the linkersave object
 			FLinkerSave* LinkerSave = Cast<FLinkerSave>(Ar.GetLinker());
 
 			// determine whether we are going to store the payload inline or not.
-			bool bStoreInline = !!(BulkDataFlags&BULKDATA_ForceInlinePayload) || !LinkerSave || Ar.IsTextFormat();
-
-			if (IsEventDrivenLoaderEnabledInCookedBuilds() && Ar.IsCooking() && !bStoreInline && !(BulkDataFlags&BULKDATA_Force_NOT_InlinePayload))
+			bool bStoreInline = !!(LocalBulkDataFlags & BULKDATA_ForceInlinePayload) || !LinkerSave || Ar.IsTextFormat();
+			if (IsEventDrivenLoaderEnabledInCookedBuilds() && Ar.IsCooking() && !(LocalBulkDataFlags & BULKDATA_Force_NOT_InlinePayload))
 			{
 				bStoreInline = true;
 			}
@@ -1337,59 +1335,51 @@ void FUntypedBulkData::Serialize( FArchive& Ar, UObject* Owner, int32 Idx, bool
 			if (!bStoreInline)
 			{
 				// set the flag indicating where the payload is stored
-				SetBulkDataFlags(BULKDATA_PayloadAtEndOfFile);
+				SetBulkDataFlagsOn(LocalBulkDataFlags, BULKDATA_PayloadAtEndOfFile);
+				ClearBulkDataFlagsOn(LocalBulkDataFlags, BULKDATA_PayloadInSeperateFile); // SavePackageUtilities::SaveBulkData will add this back if required
 
 				// with no LinkerSave we have to store the data inline
-				check(LinkerSave != NULL);				
-				
+				check(LinkerSave != NULL);
+
 				// add the bulkdata storage info object to the linkersave
-				int32 Index = LinkerSave->BulkDataToAppend.AddZeroed(1);
-				FLinkerSave::FBulkDataStorageInfo& BulkStore = LinkerSave->BulkDataToAppend[Index];
+				FLinkerSave::FBulkDataStorageInfo& BulkStore = LinkerSave->BulkDataToAppend.AddZeroed_GetRef();
 
 				BulkStore.BulkDataOffsetInFilePos = SavedBulkDataOffsetInFilePos;
 				BulkStore.BulkDataSizeOnDiskPos = SavedBulkDataSizeOnDiskPos;
 				BulkStore.BulkDataFlagsPos = SavedBulkDataFlagsPos;
-				BulkStore.BulkDataFlags = BulkDataFlags;
+				BulkStore.BulkDataFlags = LocalBulkDataFlags;
 				BulkStore.BulkDataFileRegionType = FileRegionType;
 				BulkStore.BulkData = this;
 
 				// If having flag BULKDATA_DuplicateNonOptionalPayload, duplicate bulk data in optional storage (.uptnl)
-				if (BulkDataFlags & BULKDATA_DuplicateNonOptionalPayload)
+				if (LocalBulkDataFlags & BULKDATA_DuplicateNonOptionalPayload)
 				{
 					int64 SavedDupeBulkDataFlagsPos = INDEX_NONE;
 					int64 SavedDupeBulkDataSizeOnDiskPos = INDEX_NONE;
 					int64 SavedDupeBulkDataOffsetInFilePos = INDEX_NONE;
 
-					int32 SavedDupeBulkDataFlags = (BulkDataFlags & ~BULKDATA_DuplicateNonOptionalPayload) | BULKDATA_OptionalPayload;
+					EBulkDataFlags SavedDupeBulkDataFlags = static_cast<EBulkDataFlags>(
+						(LocalBulkDataFlags & ~BULKDATA_DuplicateNonOptionalPayload) | BULKDATA_OptionalPayload);
 					{
 						FArchive::FScopeSetDebugSerializationFlags S(Ar, DSF_IgnoreDiff);
 
+						// Save offset where we are serializing SavedDupeBulkDataFlags and store a placeholder
 						SavedDupeBulkDataFlagsPos = Ar.Tell();
 						Ar << SavedDupeBulkDataFlags;
 
-						// And serialize the placeholder which is going to be overwritten later.
+						// Save offset where we are serializing SavedDupeBulkDataSizeOnDisk and store a placeholder
 						SavedDupeBulkDataSizeOnDiskPos = Ar.Tell();
-						BulkDataSizeOnDisk = INDEX_NONE;
-						if (BulkDataFlags & BULKDATA_Size64Bit)
-						{
-							Ar << BulkDataSizeOnDisk;
-						}
-						else
-						{
-							int32 BulkDataSizeOnDiskAsInt32 = BulkDataSizeOnDisk;
-							Ar << BulkDataSizeOnDiskAsInt32;
-						}
+						int64 DupeBulkDataSizeOnDisk = INDEX_NONE;
+						SerializeBulkDataSizeInt(Ar, DupeBulkDataSizeOnDisk, SavedDupeBulkDataFlags);
 
-						// Keep track of position we are going to serialize placeholder BulkDataOffsetInFile.
+						// Save offset where we are serializing SavedDupeBulkDataOffsetInFile and store a placeholder
 						SavedDupeBulkDataOffsetInFilePos = Ar.Tell();
-						BulkDataOffsetInFile = INDEX_NONE;
-						// And serialize the placeholder which is going to be overwritten later.
-						Ar << BulkDataOffsetInFile;
+						int64 DupeBulkDataOffsetInFile = INDEX_NONE;
+						Ar << DupeBulkDataOffsetInFile;
 					}
 
 					// add duplicate bulkdata with different flag
-					Index = LinkerSave->BulkDataToAppend.AddZeroed(1);
-					FLinkerSave::FBulkDataStorageInfo& DupeBulkStore = LinkerSave->BulkDataToAppend[Index];
+					FLinkerSave::FBulkDataStorageInfo& DupeBulkStore = LinkerSave->BulkDataToAppend.AddZeroed_GetRef();
 
 					DupeBulkStore.BulkDataOffsetInFilePos = SavedDupeBulkDataOffsetInFilePos;
 					DupeBulkStore.BulkDataSizeOnDiskPos = SavedDupeBulkDataSizeOnDiskPos;
@@ -1398,14 +1388,12 @@ void FUntypedBulkData::Serialize( FArchive& Ar, UObject* Owner, int32 Idx, bool
 					DupeBulkStore.BulkDataFileRegionType = FileRegionType;
 					DupeBulkStore.BulkData = this;
 				}
-				
-				// Serialize bulk data into the storage info
-				BulkDataSizeOnDisk = -1;
 			}
 			else
 			{
 				// set the flag indicating where the payload is stored
-				ClearBulkDataFlags(BULKDATA_PayloadAtEndOfFile);
+				ClearBulkDataFlagsOn(LocalBulkDataFlags,
+					static_cast<EBulkDataFlags>(BULKDATA_PayloadAtEndOfFile | BULKDATA_PayloadInSeperateFile));
 
 				int64 SavedBulkDataStartPos = Ar.Tell();
 
@@ -1414,7 +1402,7 @@ void FUntypedBulkData::Serialize( FArchive& Ar, UObject* Owner, int32 Idx, bool
 				{
 					Ar.PushFileRegionType(FileRegionType);
 				}
-				SerializeBulkData( Ar, BulkData.Get() );
+				SerializeBulkData(Ar, BulkData.Get(), LocalBulkDataFlags);
 				if (FileRegionType != EFileRegionType::None)
 				{
 					Ar.PopFileRegionType();
@@ -1424,79 +1412,48 @@ void FUntypedBulkData::Serialize( FArchive& Ar, UObject* Owner, int32 Idx, bool
 				int64 SavedBulkDataEndPos = Ar.Tell();
 
 				checkf(SavedBulkDataStartPos >= 0 && SavedBulkDataEndPos >= 0,
-					   TEXT("Bad archive positions for bulkdata. StartPos=%d EndPos=%d"),
-					   SavedBulkDataStartPos, SavedBulkDataEndPos);
+					TEXT("Bad archive positions for bulkdata. StartPos=%d EndPos=%d"),
+					SavedBulkDataStartPos, SavedBulkDataEndPos);
 
-				BulkDataSizeOnDisk		= SavedBulkDataEndPos - SavedBulkDataStartPos;
-				BulkDataOffsetInFile	= SavedBulkDataStartPos;
-			}
+				LocalBulkDataSizeOnDisk = SavedBulkDataEndPos - SavedBulkDataStartPos;
+				LocalBulkDataOffsetInFile = SavedBulkDataStartPos;
 
-			// store current file offset before seeking back
-			int64 CurrentFileOffset = Ar.Tell();
+				// Since we are storing inline we are not relying on SavePackageUtilities::SaveBulkData to update the placeholder
+				// location data, so we need to do it here.
 
-			{
-				FArchive::FScopeSetDebugSerializationFlags S(Ar, DSF_IgnoreDiff);
+				// store current file offset before seeking back
+				int64 CurrentFileOffset = Ar.Tell();
+				{
+					FArchive::FScopeSetDebugSerializationFlags S(Ar, DSF_IgnoreDiff);
 
-				// Seek back and overwrite the flags 
-				Ar.Seek(SavedBulkDataFlagsPos);
-				Ar << BulkDataFlags;
+					// Seek back and overwrite the flags 
+					Ar.Seek(SavedBulkDataFlagsPos);
+					Ar << LocalBulkDataFlags;
 
-				// Seek back and overwrite placeholder for BulkDataSizeOnDisk
-				Ar.Seek(SavedBulkDataSizeOnDiskPos);
-				if (BulkDataFlags & BULKDATA_Size64Bit)
-				{
-					Ar << BulkDataSizeOnDisk;
-				}
-				else
-				{
-					int32 BulkDataSizeOnDiskAsInt32 = BulkDataSizeOnDisk;
-					Ar << BulkDataSizeOnDiskAsInt32;
-				}
+					// Seek back and overwrite placeholder for BulkDataSizeOnDisk
+					Ar.Seek(SavedBulkDataSizeOnDiskPos);
+					SerializeBulkDataSizeInt(Ar, LocalBulkDataSizeOnDisk, LocalBulkDataFlags);
 
-				// Seek back and overwrite placeholder for BulkDataOffsetInFile
-				Ar.Seek(SavedBulkDataOffsetInFilePos);
-				Ar << BulkDataOffsetInFile;
+					// Seek back and overwrite placeholder for BulkDataOffsetInFile
+					Ar.Seek(SavedBulkDataOffsetInFilePos);
+					Ar << LocalBulkDataOffsetInFile;
+				}
+				// Seek to the end of written data so we don't clobber any data in subsequent writes
+				Ar.Seek(CurrentFileOffset);
 			}
 
-			// Seek to the end of written data so we don't clobber any data in subsequent write 
-			// operations
-			Ar.Seek(CurrentFileOffset);
-
-			// Restore the mutated values when we are serializing for cooking; saving to a cook target should not mutate the in-memory object
-			if (Ar.IsCooking())
+			// If we are overwriting the LoadedPath for the current package, set the location variables on *this equal to the new values
+			// that we are writing into the package on disk
+			if (!Ar.IsCooking())
 			{
-				BulkDataFlags = SavedBulkDataFlags;
-				BulkDataOffsetInFile = SavedBulkDataOffsetInFile;
-				BulkDataSizeOnDisk = SavedBulkDataSizeOnDisk;
+				BulkDataFlags = LocalBulkDataFlags;
+				BulkDataOffsetInFile = LocalBulkDataOffsetInFile;
+				BulkDataSizeOnDisk = LocalBulkDataSizeOnDisk;
 			}
 		}
 	}
 }
 
-#if WITH_IOSTORE_IN_EDITOR
-void FUntypedBulkData::SerializeFromIoStore( FArchive& Ar, UObject* Owner, int32 Idx, bool bAttemptFileMapping )
-{
-	BulkData.Reallocate(GetBulkDataSize(), BulkDataAlignment);
-
-	if (IsInlined())
-	{
-		SerializeBulkData(Ar, BulkData.Get());
-	}
-	else
-	{
-		check(!NeedsOffsetFixup());
-
-		TUniquePtr<IBulkDataIORequest> Request = CreateBulkDataIoDispatcherRequest(CreateBulkDataChunkId(PackageId, BulkDataFlags));
-		Request->WaitCompletion();
-
-		FLargeMemoryReader MemoryAr(Request->GetReadResults(), Request->GetSize());
-		MemoryAr.Seek(BulkDataOffsetInFile);
-
-		SerializeBulkData(MemoryAr, BulkData.Get());
-	}
-}
-#endif
-
 /*-----------------------------------------------------------------------------
 	Class specific virtuals.
 -----------------------------------------------------------------------------*/
@@ -1628,20 +1585,21 @@ void FUntypedBulkData::SerializeElements(FArchive& Ar, void* Data)
  *
  * @param	Ar					Archive to serialize with
  * @param	Data				Memory to serialize either to or from
+ * @param	InBulkDataFlags		Flags describing how the data was/shouldbe serialized
  */
-void FUntypedBulkData::SerializeBulkData( FArchive& Ar, void* Data )
+void FUntypedBulkData::SerializeBulkData(FArchive& Ar, void* Data, EBulkDataFlags InBulkDataFlags)
 {
 	SCOPED_LOADTIMER(BulkData_SerializeBulkData);
 
 	// skip serializing of unused data
-	if( BulkDataFlags & BULKDATA_Unused )
+	if (InBulkDataFlags & BULKDATA_Unused)
 	{
 		return;
 	}
 
 	// Skip serialization for bulk data of zero length
 	const int64 BulkDataSize = GetBulkDataSize();
-	if(BulkDataSize == 0)
+	if (BulkDataSize == 0)
 	{
 		return;
 	}
@@ -1649,22 +1607,23 @@ void FUntypedBulkData::SerializeBulkData( FArchive& Ar, void* Data )
 	// Allow backward compatible serialization by forcing bulk serialization off if required. Saving also always uses single
 	// element serialization so errors or oversight when changing serialization code is recoverable.
 	bool bSerializeInBulk = true;
-	if( RequiresSingleElementSerialization( Ar ) 
+	if (RequiresSingleElementSerialization(Ar) 
 	// Set when serialized like a lazy array.
-	|| (BulkDataFlags & BULKDATA_ForceSingleElementSerialization) 
+	|| (InBulkDataFlags & BULKDATA_ForceSingleElementSerialization)
 	// We use bulk serialization even when saving 1 byte types (texture & sound bulk data) as an optimization for those.
-	|| (Ar.IsSaving() && (GetElementSize() > 1) ) )
+	|| (Ar.IsSaving() && (GetElementSize() > 1)))
 	{
 		bSerializeInBulk = false;
 	}
 
-	// Raw serialize the bulk data without any possiblity for potential endian conversion.
-	if( bSerializeInBulk )
+	// Raw serialize the bulk data without any possibility for potential endian conversion.
+	if (bSerializeInBulk)
 	{
 		// Serialize data compressed.
-		if( BulkDataFlags & BULKDATA_SerializeCompressed )
+		if (InBulkDataFlags & BULKDATA_SerializeCompressed)
 		{
-			Ar.SerializeCompressed( Data, GetBulkDataSize(), GetDecompressionFormat(), COMPRESS_NoFlags, false );
+			Ar.SerializeCompressed(Data, GetBulkDataSize(),
+				GetDecompressionFormat(InBulkDataFlags), COMPRESS_NoFlags, false);
 		}
 		// Uncompressed/ regular serialization.
 		else
@@ -1678,35 +1637,37 @@ void FUntypedBulkData::SerializeBulkData( FArchive& Ar, void* Data )
 	else
 	{
 		// Serialize data compressed.
-		if( BulkDataFlags & BULKDATA_SerializeCompressed )
+		if (InBulkDataFlags & BULKDATA_SerializeCompressed)
 		{
 			// Loading, data is compressed in archive and needs to be decompressed.
-			if( Ar.IsLoading() )
+			if (Ar.IsLoading())
 			{
 				TUniquePtr<uint8[]> SerializedData = MakeUnique<uint8[]>( GetBulkDataSize() );
 
 				// Serialize data with passed in archive and compress.
-				Ar.SerializeCompressed( SerializedData.Get(), GetBulkDataSize(), GetDecompressionFormat(), COMPRESS_NoFlags, false );
-				
+				Ar.SerializeCompressed(SerializedData.Get(), GetBulkDataSize(),
+					GetDecompressionFormat(InBulkDataFlags), COMPRESS_NoFlags, false);
+
 				// Initialize memory reader with uncompressed data array and propagate forced byte swapping
-				FLargeMemoryReader MemoryReader( SerializedData.Get(), GetBulkDataSize(), ELargeMemoryReaderFlags::Persistent );
-				MemoryReader.SetByteSwapping( Ar.ForceByteSwapping() );
+				FLargeMemoryReader MemoryReader(SerializedData.Get(), GetBulkDataSize(), ELargeMemoryReaderFlags::Persistent);
+				MemoryReader.SetByteSwapping(Ar.ForceByteSwapping());
 
 				// Serialize each element individually via memory reader.
 				SerializeElements(MemoryReader, Data);
 			}
 			// Saving, data is uncompressed in memory and needs to be compressed.
-			else if( Ar.IsSaving() )
-			{			
+			else if (Ar.IsSaving())
+			{
 				// Initialize memory writer with blank data array and propagate forced byte swapping
-				FLargeMemoryWriter MemoryWriter( GetBulkDataSize(), true );
-				MemoryWriter.SetByteSwapping( Ar.ForceByteSwapping() );
+				FLargeMemoryWriter MemoryWriter(GetBulkDataSize(), true);
+				MemoryWriter.SetByteSwapping(Ar.ForceByteSwapping());
 
-				// Serialize each element individually via memory writer.			
+				// Serialize each element individually via memory writer.
 				SerializeElements(MemoryWriter, Data);
 
 				// Serialize data with passed in archive and compress.
-				Ar.SerializeCompressed( MemoryWriter.GetData(), GetBulkDataSize(), GetDecompressionFormat(), COMPRESS_NoFlags, false );
+				Ar.SerializeCompressed(MemoryWriter.GetData(), GetBulkDataSize(),
+					GetDecompressionFormat(InBulkDataFlags), COMPRESS_NoFlags, false);
 			}
 		}
 		// Uncompressed/ regular serialization.
@@ -1718,6 +1679,11 @@ void FUntypedBulkData::SerializeBulkData( FArchive& Ar, void* Data )
 	}
 }
 
+void FUntypedBulkData::SerializeBulkData(FArchive& Ar, void* Data)
+{
+	SerializeBulkData(Ar, Data, this->BulkDataFlags);
+}
+
 IAsyncReadFileHandle* FUntypedBulkData::OpenAsyncReadHandle() const
 {
 	return IPackageResourceManager::Get().OpenAsyncReadPackage(GetPackagePath(), GetPackageSegment());
@@ -1741,25 +1707,19 @@ IBulkDataIORequest* FUntypedBulkData::CreateStreamingRequest(int64 OffsetInBulkD
 #endif
 	check(PackagePath.IsEmpty() == false);
 
-	// If we are loading from a .uexp file then we need to adjust the segment and offset stored by BulkData in order to use them
-	// to access the data in the .uexp file. To keep the method const (due to a large number of places calling this, assuming that
-	// it is const) we take a copy of these data values which if needed can be adjusted and use them instead.
-	EPackageSegment AdjustedPackageSegment = PackageSegment;
-	int64 AdjustedBulkDataOffsetInFile = BulkDataOffsetInFile;
-
-	// Fix up the Segment/Offset to work with streaming if EDL is enabled and the segment is still referencing the header
-	if (GEventDrivenLoaderEnabled && AdjustedPackageSegment == EPackageSegment::Header)
+	if (GEventDrivenLoaderEnabled)
 	{
-		AdjustedBulkDataOffsetInFile -= IPackageResourceManager::Get().FileSize(PackagePath, AdjustedPackageSegment);
-		AdjustedPackageSegment = EPackageSegment::Exports;
-
-		UE_LOG(LogSerialization, Error, TEXT("Streaming from the .uexp file '%s' this MUST be in a ubulk instead for best performance."), *PackagePath.GetDebugName(AdjustedPackageSegment));
+		// Verify that the Serialize function converted a Header Segment and Offset to be relative to the Exports segment
+		check(PackageSegment != EPackageSegment::Header);
+		// Even the Exports segment is bad for performance
+		UE_CLOG(PackageSegment == EPackageSegment::Exports,
+			LogSerialization, Error, TEXT("Streaming from the .uexp file '%s' this MUST be in a ubulk instead for best performance."), *PackagePath.GetDebugName(PackageSegment));
 	}
 
-	UE_CLOG(IsStoredCompressedOnDisk(), LogSerialization, Fatal, TEXT("Package level compression is no longer supported (%s)."), *PackagePath.GetDebugName(AdjustedPackageSegment));
-	UE_CLOG(GetBulkDataSize() <= 0, LogSerialization, Error, TEXT("(%s) has invalid bulk data size."), *PackagePath.GetDebugName(AdjustedPackageSegment));
+	UE_CLOG(IsStoredCompressedOnDisk(), LogSerialization, Fatal, TEXT("Package level compression is no longer supported (%s)."), *PackagePath.GetDebugName(PackageSegment));
+	UE_CLOG(GetBulkDataSize() <= 0, LogSerialization, Error, TEXT("(%s) has invalid bulk data size."), *PackagePath.GetDebugName(PackageSegment));
 
-	IAsyncReadFileHandle* IORequestHandle = IPackageResourceManager::Get().OpenAsyncReadPackage(PackagePath, AdjustedPackageSegment);
+	IAsyncReadFileHandle* IORequestHandle = IPackageResourceManager::Get().OpenAsyncReadPackage(PackagePath, PackageSegment);
 	check(IORequestHandle); // this generally cannot fail because it is async
 
 	if (IORequestHandle == nullptr)
@@ -1767,7 +1727,7 @@ IBulkDataIORequest* FUntypedBulkData::CreateStreamingRequest(int64 OffsetInBulkD
 		return nullptr;
 	}
 
-	const int64 OffsetInFile = AdjustedBulkDataOffsetInFile + OffsetInBulkData;
+	const int64 OffsetInFile = BulkDataOffsetInFile + OffsetInBulkData;
 
 	FBulkDataIORequest* IORequest = new FBulkDataIORequest(IORequestHandle);
 
@@ -1943,7 +1903,10 @@ void FUntypedBulkData::MakeSureBulkDataIsLoaded()
 			// in which case it starts out with a size of zero.
 			if (BytesNeeded > 0)
 			{
-				LoadDataIntoMemory(BulkData.Get());
+				if (!TryLoadDataIntoMemory(BulkData.Get()))
+				{
+					BulkData.Deallocate();
+				}
 			}
 		}
 	}
@@ -1979,43 +1942,38 @@ bool FUntypedBulkData::FlushAsyncLoading()
  *
  * @param Dest Memory to serialize data into
  */
-void FUntypedBulkData::LoadDataIntoMemory( void* Dest )
+bool FUntypedBulkData::TryLoadDataIntoMemory(void* Dest)
 {
 	// Try flushing async loading before attempting to load
 	if (FlushAsyncLoading())
 	{
 		FMemory::Memcpy(Dest, BulkData.Get(), GetBulkDataSize());
-		return;
+		return true;
 	}
 
 #if WITH_EDITOR
 	checkf( AttachedAr, TEXT( "Attempted to load bulk data without an attached archive. Most likely the bulk data was loaded twice on console, which is not supported" ) );
 
+	TUniquePtr<FArchive> BulkDataLoadedFile;
 	FArchive* BulkDataArchive = nullptr;
 	if (Linker && Linker->GetAsyncLoader() && Linker->GetAsyncLoader()->IsCookedForEDLInEditor() &&
 		(BulkDataFlags & BULKDATA_PayloadInSeperateFile))
 	{
-		// The attached archive is a package cooked for EDL loaded in the editor so the actual bulk data sits in a separate ubulk file.
-		EPackageSegment BulkDataSegment;
+		// The Serialize function set the PackageSegment, and since the payload is in a separate file, we
+		// can assert that the PackageSegment that was set is one of the separate-file segments
+		check(PackageSegment != EPackageSegment::Header && PackageSegment != EPackageSegment::Exports);
 
-		if (BulkDataFlags & BULKDATA_MemoryMappedPayload)
-		{
-			BulkDataSegment = EPackageSegment::BulkDataMemoryMapped;
-		}
-		else if (IsOptional())
-		{
-			BulkDataSegment = EPackageSegment::BulkDataOptional;
-		}
-		else
+		FOpenPackageResult Result = IPackageResourceManager::Get().OpenReadPackage(PackagePath, PackageSegment);
+		if (!Result.Archive.IsValid() || Result.Format != EPackageFormat::Binary)
 		{
-			BulkDataSegment = EPackageSegment::BulkDataDefault;
+			UE_LOG(LogSerialization, Error, TEXT("Attempted to load bulk data from an invalid PackagePath '%s'%s."),
+				*PackagePath.GetDebugName(PackageSegment), (!Result.Archive.IsValid() ? TEXT("") : TEXT("; package is in non-binary format and this is not supported.")));
+			return false;
 		}
-		FOpenPackageResult OpenResult = IPackageResourceManager::Get().OpenReadPackage(PackagePath, BulkDataSegment);
-		BulkDataArchive = OpenResult.Archive.Release();
-		check(BulkDataArchive == nullptr || OpenResult.Format == EPackageFormat::Binary);
+		BulkDataLoadedFile = MoveTemp(Result.Archive);
+		BulkDataArchive = BulkDataLoadedFile.Get();
 	}
-
-	if (!BulkDataArchive)
+	else
 	{
 		BulkDataArchive = AttachedAr;
 	}
@@ -2023,19 +1981,14 @@ void FUntypedBulkData::LoadDataIntoMemory( void* Dest )
 	// Keep track of current position in file so we can restore it later.
 	int64 PushedPos = BulkDataArchive->Tell();
 	// Seek to the beginning of the bulk data in the file.
-	BulkDataArchive->Seek( BulkDataOffsetInFile );
-		
-	SerializeBulkData( *BulkDataArchive, Dest );
+	BulkDataArchive->Seek(BulkDataOffsetInFile);
+
+	SerializeBulkData(*BulkDataArchive, Dest, BulkDataFlags);
 
 	// Restore file pointer.
-	BulkDataArchive->Seek( PushedPos );
+	BulkDataArchive->Seek(PushedPos);
 	BulkDataArchive->FlushCache();
-
-	if (BulkDataArchive != AttachedAr)
-	{
-		delete BulkDataArchive;
-		BulkDataArchive = nullptr;
-	}
+	return true;
 
 #else
 	bool bWasLoadedSuccessfully = false;
@@ -2056,7 +2009,7 @@ void FUntypedBulkData::LoadDataIntoMemory( void* Dest )
 					Ar->Seek( BulkDataOffsetInFile );
 
 					// serialize the bulk data
-					SerializeBulkData( *Ar, Dest );
+					SerializeBulkData(*Ar, Dest, BulkDataFlags);
 
 					// seek back to the position the archive was before
 					Ar->Seek(CurPos);
@@ -2070,7 +2023,7 @@ void FUntypedBulkData::LoadDataIntoMemory( void* Dest )
 	// if we weren't able to load via linker, load directly by PackagePath
 	if (!bWasLoadedSuccessfully)
 	{
-		// load from the specied PackagePath when the linker has been cleared
+		// load from the specified PackagePath when the linker has been cleared
 		checkf(!PackagePath.IsEmpty(), TEXT( "Attempted to load bulk data without a proper PackagePath." ) );
 
 #if PLATFORM_SUPPORTS_TEXTURE_STREAMING
@@ -2086,21 +2039,22 @@ void FUntypedBulkData::LoadDataIntoMemory( void* Dest )
 		}
 #endif
 
-		if (GEventDrivenLoaderEnabled && PackageSegment == EPackageSegment::Header)
-		{
-			BulkDataOffsetInFile -= IPackageResourceManager::Get().FileSize(PackagePath, PackageSegment);
-			check(BulkDataOffsetInFile >= 0);
-			PackageSegment = EPackageSegment::Exports;
-		}
+		// Verify that the Serialize function converted a Header Segment and Offset to be relative to the Exports segment, if GEventDrivenLoaderEnabled
+		check(!GEventDrivenLoaderEnabled || PackageSegment != EPackageSegment::Header);
 
 		FOpenPackageResult Result = IPackageResourceManager::Get().OpenReadPackage(PackagePath, PackageSegment);
-		checkf( Result.Archive.IsValid() && Result.Format == EPackageFormat::Binary, TEXT( "Attempted to load bulk data from an invalid PackagePath '%s'%s." ),
-			*PackagePath.GetDebugName(PackageSegment), (!Result.Archive.IsValid() ? TEXT("") : TEXT("; package is in non-binary format and this is not supported.")));
+		if (!Result.Archive.IsValid() || Result.Format != EPackageFormat::Binary)
+		{
+			UE_LOG(LogSerialization, Error, TEXT("Attempted to load bulk data from an invalid PackagePath '%s'%s."),
+				*PackagePath.GetDebugName(PackageSegment), (!Result.Archive.IsValid() ? TEXT("") : TEXT("; package is in non-binary format and this is not supported.")));
+			return false;
+		}
 	
 		// Seek to the beginning of the bulk data in the file.
-		Result.Archive->Seek( BulkDataOffsetInFile );
-		SerializeBulkData( *Result.Archive, Dest );
+		Result.Archive->Seek(BulkDataOffsetInFile);
+		SerializeBulkData(*Result.Archive, Dest, BulkDataFlags);
 	}
+	return true;
 #endif // WITH_EDITOR
 }
 
@@ -2260,11 +2214,11 @@ void FFormatContainer::Serialize(FArchive& Ar, UObject* Owner, const TArray<FNam
 				if (bInline)
 				{
 					Bulk->SetBulkDataFlags(BULKDATA_ForceInlinePayload);
-					Bulk->ClearBulkDataFlags(BULKDATA_PayloadAtEndOfFile | BULKDATA_PayloadInSeperateFile | BULKDATA_Force_NOT_InlinePayload | BULKDATA_MemoryMappedPayload);
+					Bulk->ClearBulkDataFlags(BULKDATA_Force_NOT_InlinePayload | BULKDATA_MemoryMappedPayload);
 				}
 				else
 				{
-					Bulk->SetBulkDataFlags(BULKDATA_PayloadAtEndOfFile | BULKDATA_PayloadInSeperateFile | BULKDATA_Force_NOT_InlinePayload);
+					Bulk->SetBulkDataFlags(BULKDATA_Force_NOT_InlinePayload);
 					if (bMapped)
 					{
 						Bulk->SetBulkDataFlags(BULKDATA_MemoryMappedPayload);
diff --git a/Engine/Source/Runtime/CoreUObject/Private/Serialization/BulkData2.cpp b/Engine/Source/Runtime/CoreUObject/Private/Serialization/BulkData2.cpp
--- a/Engine/Source/Runtime/CoreUObject/Private/Serialization/BulkData2.cpp
+++ b/Engine/Source/Runtime/CoreUObject/Private/Serialization/BulkData2.cpp
@@ -693,7 +693,7 @@ FBulkDataBase& FBulkDataBase::operator=(const FBulkDataBase& Other)
 	BulkDataOffset = Other.BulkDataOffset;
 	BulkDataFlags = Other.BulkDataFlags;
 
-	if( !Other.IsDataMemoryMapped())
+	if (!Other.IsDataMemoryMapped() || !Other.IsInSeparateFile())
 	{
 		if (Other.GetDataBufferReadOnly())
 		{
@@ -709,7 +709,7 @@ FBulkDataBase& FBulkDataBase::operator=(const FBulkDataBase& Other)
 		{
 			TIoStatusOr<FIoMappedRegion> Status = IoDispatcher->OpenMapped(CreateChunkId(), FIoReadOptions());
 			FIoMappedRegion MappedRegion = Status.ConsumeValueOrDie();
-			DataAllocation.SetMemoryMappedData(this, MappedRegion.MappedFileHandle, MappedRegion.MappedFileRegion);		
+			DataAllocation.SetMemoryMappedData(this, MappedRegion.MappedFileHandle, MappedRegion.MappedFileRegion);
 		}
 		else
 		{
@@ -734,6 +734,40 @@ FBulkDataBase::~FBulkDataBase()
 	}
 }
 
+void FBulkDataBase::ConditionalSetInlineAlwaysAllowDiscard()
+{
+	// If PackagePath is null and we do not have BULKDATA_UsesIoDispatcher, we will not be able to reload the bulkdata.
+	// We will not have a PackagePath if the engine is using IoStore, unless the setting ShouldAllowBulkDataInIoStore()
+	// has been set to false.
+	// If IoStore with !ShouldAllowBulkDataInIoStore, we set the PackagePath, but only when BulkData is in a separate
+	// BulkData file. We can load the separate BulkData files, but we we can not load the file containing the
+	// exports (which is where inline BulkData or end-of-file BulkData is stored) because those files are managed
+	// separately by IoStore .
+	// So in IoStore  with BulkData stored either inline or in the end-of-file section, we can not reload the bulkdata.
+	// We do not need to consider the end-of-file bulkdata section because IoStore  guarantees during cooking
+	// (SaveBulkData) that only inlined data is in the package file; there is no end-of-file bulkdata section.
+
+	// In the Inlined IoStore case therefore we need to not discard the bulk data if !BULKDATA_SingleUse;
+	// we have to keep it in case it gets requested twice.
+	// However, some systems (Audio,Animation) have large inline data they for legacy reasons have not marked as
+	// BULKDATA_SingleUse.
+	// In the old loader this data was discarded for them since CanLoadFromDisk was true.
+	// Since CanLoadFromDisk is now false, we are keeping that data around, and this causes memory bloat.
+	// Licensees have asked us to fix this memory bloat.
+	//
+	// To hack-fix the memory bloat in these systems when using IoStore we will mark all inline BulkData
+	// as discardable when using IoStore.
+	// This will cause a bug when using IoStore in any systems that do actually need to reload inline data.
+	// Each project that uses IoStore needs to guarantee that they do not have any systems that need to
+	// reload inline data.
+	//
+	// TODO: We need to make the old loader and new loader consistent on how inline data is treated!
+	if (FIoDispatcher::IsInitialized())
+	{
+		SetBulkDataFlags(BULKDATA_AlwaysAllowDiscard);
+	}
+}
+
 void FBulkDataBase::Serialize(FArchive& Ar, UObject* Owner, int32 /*Index*/, bool bAttemptFileMapping, int32 ElementSize)
 {
 	DECLARE_SCOPE_CYCLE_COUNTER(TEXT("FBulkDataBase::Serialize"), STAT_UBD_Serialize, STATGROUP_Memory);
@@ -754,34 +788,27 @@ void FBulkDataBase::Serialize(FArchive& Ar, UObject* Owner, int32 /*Index*/, boo
 
 		BulkDataOffset = 0;
 
-		if (BulkDataFlags & BULKDATA_Size64Bit)
-		{
-			Ar << ElementCount;
-			Ar << BulkDataSizeOnDisk;
-		}
-		else
-		{
-			int32 Temp32ByteValue;
-
-			Ar << Temp32ByteValue;
-			ElementCount = Temp32ByteValue;
-
-			Ar << Temp32ByteValue;
-			BulkDataSizeOnDisk = Temp32ByteValue;
-		}
+		SerializeBulkDataSizeInt(Ar, ElementCount, BulkDataFlags);
+		SerializeBulkDataSizeInt(Ar, BulkDataSizeOnDisk, BulkDataFlags);
 
 		BulkDataSize = ElementCount * ElementSize;
 
 		Ar << BulkDataOffset;
 
-		const bool bUseIoDispatcher = IsIoDispatcherEnabled();
-
 		if ((BulkDataFlags & BULKDATA_BadDataVersion) != 0)
 		{
 			uint16 DummyValue;
 			Ar << DummyValue;
 		}
 
+		EBulkDataFlags DuplicateFlags = static_cast<EBulkDataFlags>(0);
+		int64 DuplicateSizeOnDisk = INDEX_NONE;
+		int64 DuplicateOffset = INDEX_NONE;
+		if (IsDuplicateNonOptional())
+		{
+			SerializeDuplicateData(Ar, DuplicateFlags, DuplicateSizeOnDisk, DuplicateOffset);
+		}
+
 		// We assume that Owner/Package/Linker are all valid. The old BulkData system would
 		// generally fail if any of these were nullptr but had plenty of inconsistent checks
 		// scattered throughout.
@@ -789,42 +816,79 @@ void FBulkDataBase::Serialize(FArchive& Ar, UObject* Owner, int32 /*Index*/, boo
 		const UPackage* Package = Owner->GetOutermost();
 		checkf(Package != nullptr, TEXT("FBulkDataBase::Serialize requires an Owner that returns a valid UPackage"));
 
-		if (!IsInlined() && bUseIoDispatcher)
+		const bool bEngineUsesIoStore = FIoDispatcher::IsInitialized();
+		const bool bSeparateFilesUseIoStore = bEngineUsesIoStore && ShouldAllowBulkDataInIoStore();
+		const FPackagePath* PackagePath = nullptr;
+		const FLinkerLoad* Linker = nullptr;
+		if (bEngineUsesIoStore)
 		{
-			Data.PackageID = Package->GetPackageId().Value();
-			SetRuntimeBulkDataFlags(BULKDATA_UsesIoDispatcher); // Indicates that this BulkData should use the FIoChunkId rather than a PackagePath
+			checkf(IsInlined() || !NeedsOffsetFixup(), TEXT("IODispatcher does not support offset fixups; 'LegacyBulkDataOffsets=true' will not work with the IoStore!"));
+			checkf(IsInlined() || IsInSeparateFile() || !GEventDrivenLoaderEnabled,
+				TEXT("IODispatcher does not support finding the file size of header segments, which is required if BulkData is at end-of-file and EDL is enabled. ")
+				TEXT("Non-inline BulkData must be stored in a separate file when EDL is enabled!"));
+			if (IsInSeparateFile())
+			{
+				if (bSeparateFilesUseIoStore)
+				{
+					Data.PackageID = Package->GetPackageId().Value();
+					SetRuntimeBulkDataFlags(BULKDATA_UsesIoDispatcher); // Indicates that this BulkData should use the FIoChunkId rather than a PackagePath
+				}
+				else
+				{
+					// When using iostore with !ShouldAllowBulkDataInIoStore, the separate files for BulkData are loadable,
+					// but we still cannot load the package files themselves (iostore handles those in a special manner).
+					PackagePath = &Package->GetLoadedPath();
+					if (PackagePath->IsEmpty())
+					{
+						// Note that this Bulkdata object will end up with an invalid token and will end up resolving to an empty file path!
+						UE_LOG(LogSerialization, Warning, TEXT("Could not get PackagePath from package for package %s!"), *Package->GetName());
+						PackagePath = nullptr;
+					}
+					// Reset the token even though it should already be invalid (it will be set later when registered)
+					Data.Token = InvalidToken;
+				}
+			}
+			else
+			{
+				// this->Data will be unused since we are loading from the current file; initialize it to invalid
+				Data.Token = InvalidToken;
+			}
 		}
 		else
-		{
-			// Reset the token even though it should already be invalid (it will be set later when registered)
-			Data.Token = InvalidToken;
-		}
-
-		const FPackagePath* PackagePath = nullptr;
-		const FLinkerLoad* Linker = nullptr;
-
-		if (bUseIoDispatcher == false)
 		{
 			Linker = FLinkerLoad::FindExistingLinkerForPackage(Package);
-
 			if (Linker != nullptr)
 			{
 				PackagePath = &Linker->GetPackagePath();
 			}
-			else if (!::ShouldAllowBulkDataInIoStore() && !IsInlined())
+			if (!PackagePath || PackagePath->IsEmpty())
 			{
-				PackagePath = &Package->GetLoadedPath();
-				if (PackagePath->IsEmpty())
-				{
-					// Note that this Bulkdata object will end up with an invalid token and will end up resolving to an empty file path!
-					UE_LOG(LogSerialization, Warning, TEXT("Could not get PackagePath from linker or package for package %s!"), *Package->GetName());
-					PackagePath = nullptr;
-				}
+				// Note that this Bulkdata object will end up with an invalid token and will end up resolving to an empty file path!
+				UE_LOG(LogSerialization, Warning, TEXT("Could not get PackagePath from linker for package %s!"), *Package->GetName());
+				PackagePath = nullptr;
+			}
+
+			if (!IsInlined() && NeedsOffsetFixup())
+			{
+				// The offset can be set relative to the start of the BulkData section during cooking. Convert it back to an offset from the start of the package
+				checkf(Linker != nullptr, TEXT("BulkData needs its offset fixed on load but no linker found"));
+				BulkDataOffset += Linker->Summary.BulkDataStartOffset;
+			}
+			if (!IsInSeparateFile() && GEventDrivenLoaderEnabled && PackagePath)
+			{
+				// The BulkDataOffset is written as an offset from the start of the combined PackageHeader+PackageExports+BulkData
+				// When EventDrivenLoader is enabled, the combined file is split into a header with linker tables (EPackageSegment::Header)
+				// and the exports file (EPackageSegment::Exports) with serialized bytes from exports and with the closing bulk data section (if there is one)
+				// We need to modify the BulkDataOffset to be relative to the start of the Exports file
+				BulkDataOffset -= IPackageResourceManager::Get().FileSize(*PackagePath, EPackageSegment::Header);
 			}
+
+			// Reset the token even though it should already be invalid (it will be set later when registered)
+			Data.Token = InvalidToken;
 		}
 
-		// Some failed paths require us to load the data before we return from ::Serialize but it is not
-		// safe to do so until the end of this method. By setting this flag to true we can indicate that 
+		// Some error handling cases in this function require us to load the data before we return from ::Serialize
+		// but it is not safe to do so until the end of this method. By setting this flag to true we can indicate that
 		// the load is required.
 		bool bShouldForceLoad = false;
 
@@ -836,32 +900,28 @@ void FBulkDataBase::Serialize(FArchive& Ar, UObject* Owner, int32 /*Index*/, boo
 			void* DataBuffer = AllocateData(BulkDataSize);
 			SerializeBulkData(Ar, DataBuffer, BulkDataSize);
 
-			// Inline data must be set to be allowed to always discard it's data if we are using the IoDispatcher
-			// since we will not be able to reload it and existing code might well rely on ::GetCopy being able to discard it.
-			// TODO: We need to make the old loader and new loader consistent on how inline data is treated!
-			if (FIoDispatcher::IsInitialized())
-			{
-				SetBulkDataFlags(BULKDATA_AlwaysAllowDiscard);
-			}
+			// Apply the legacy BULKDATA_AlwaysAllowDiscard fix for inline data in some conditions
+			ConditionalSetInlineAlwaysAllowDiscard();
 		}
 		else
 		{
 			if (IsDuplicateNonOptional())
 			{
-				ProcessDuplicateData(Ar, Package, PackagePath, BulkDataOffset);
+				ProcessDuplicateData(DuplicateFlags, DuplicateSizeOnDisk, DuplicateOffset, Package, PackagePath, Linker);
 			}
 
-			// Set during cook
-			if (NeedsOffsetFixup())
+			if (bAttemptFileMapping && !IsInSeparateFile() && (bEngineUsesIoStore || !Ar.IsAllowingLazyLoading()))
 			{
-				checkf(Linker != nullptr, TEXT("BulkData needs it's offset fixing on load but no linker found, data cooked with 'LegacyBulkDataOffsets=true' will not work with the IoStore!"));
-				BulkDataOffset += Linker->Summary.BulkDataStartOffset;
+				UE_CLOG(bAttemptFileMapping, LogSerialization, Error,
+					TEXT("Attempt to file map BulkData in end-of-package-file section, this is not supported when %s. Package '%s'"),
+					bEngineUsesIoStore ? TEXT("using IoDispatcher") : TEXT("archive does not support lazyload"), *Package->GetFName().ToString());
+				bShouldForceLoad = true; // Signal we want to force the BulkData to load
 			}
-
-			if (bAttemptFileMapping)
+			else if (bAttemptFileMapping)
 			{
-				if (bUseIoDispatcher)
+				if (bEngineUsesIoStore)
 				{
+					check(IsInSeparateFile());
 					TIoStatusOr<FIoMappedRegion> Status = IoDispatcher->OpenMapped(CreateChunkId(), FIoReadOptions());
 					if (Status.IsOk())
 					{
@@ -886,12 +946,27 @@ void FBulkDataBase::Serialize(FArchive& Ar, UObject* Owner, int32 /*Index*/, boo
 					}
 				}
 			}
-			else if (!Ar.IsAllowingLazyLoading() && !IsInSeparateFile())
+			else if (!IsInSeparateFile() && (bEngineUsesIoStore || !Ar.IsAllowingLazyLoading()))
 			{
+				// If the data is in the end-of-package-file section, and we can't load the package file again later either
+				// because we're using the IoDispatcher or because the archive does not support lazy loading then we have
+				// to load the data immediately.
+				bShouldForceLoad = true;
+			}
+		}
 
-				// If the archive does not support lazy loading and the data is not in a different file then we have to load 
-				// the data from the archive immediately as we won't get another chance.
+		// If we have a PackagePath then we need to make sure we can retrieve it later!
+		if (PackagePath != nullptr)
+		{
+			check(Data.Token == InvalidToken);
+			Data.Token = PackageTokenSystem::RegisterToken(Package->GetFName(), *PackagePath);
+		}
 
+		if (bShouldForceLoad)
+		{
+			if (!IsInSeparateFile())
+			{
+				check(!IsInlined()); // We checked for IsInlined up above and already loaded the data in that case
 				const int64 CurrentArchiveOffset = Ar.Tell();
 				Ar.Seek(BulkDataOffset);
 
@@ -900,17 +975,10 @@ void FBulkDataBase::Serialize(FArchive& Ar, UObject* Owner, int32 /*Index*/, boo
 
 				Ar.Seek(CurrentArchiveOffset); // Return back to the original point in the archive so future serialization can continue
 			}
-		}
-
-		// If we are not using the FIoDispatcher and we have a PackagePath then we need to make sure we can retrieve it later!
-		if (bUseIoDispatcher == false && PackagePath != nullptr)
-		{
-			Data.Token = PackageTokenSystem::RegisterToken(Package->GetFName(), *PackagePath);
-		}
-
-		if (bShouldForceLoad)
-		{
-			ForceBulkDataResident();
+			else
+			{
+				ForceBulkDataResident();
+			}
 		}
 	}
 #else
@@ -1042,6 +1110,7 @@ void FBulkDataBase::GetCopy(void** DstBuffer, bool bDiscardInternalCopy)
 FIoChunkId FBulkDataBase::CreateChunkId() const 
 { 
 	checkf(IsUsingIODispatcher(), TEXT("Calling ::CreateChunkId on Bulkdata that is not using the IoDispatcher"));
+	checkf(IsInSeparateFile(), TEXT("Calling ::CreateChunkId on BulkData that is stored in the package file rather than in a separately loadable file."));
 
 	const EIoChunkType Type =	IsOptional() ? EIoChunkType::OptionalBulkData :
 								IsFileMemoryMapped() ? EIoChunkType::MemoryMappedBulkData :
@@ -1114,9 +1183,18 @@ bool FBulkDataBase::CanLoadFromDisk() const
 
 bool FBulkDataBase::DoesExist() const
 {
-#if ALLOW_OPTIONAL_DATA
+#if !ALLOW_OPTIONAL_DATA
+	if (IsOptional())
+	{
+		return false;
+	}
+#endif
 	if (!IsUsingIODispatcher())
 	{
+		if (Data.Token == InvalidToken)
+		{
+			return false;
+		}
 		FPackagePath PackagePath = PackageTokenSystem::GetPackagePath(Data.Token);
 		return IPackageResourceManager::Get().DoesPackageExist(PackagePath, GetPackageSegmentFromFlags());
 	}
@@ -1124,9 +1202,6 @@ bool FBulkDataBase::DoesExist() const
 	{
 		return IoDispatcher->DoesChunkExist(CreateChunkId());
 	}
-#else
-	return false;
-#endif
 }
 
 bool FBulkDataBase::IsStoredCompressedOnDisk() const
@@ -1156,12 +1231,12 @@ bool FBulkDataBase::IsOptional() const
 
 bool FBulkDataBase::IsInlined() const
 {
-	return	(GetBulkDataFlags() & BULKDATA_PayloadAtEndOfFile) == 0;
+	return (GetBulkDataFlags() & BULKDATA_PayloadAtEndOfFile) == 0;
 }
 
 bool FBulkDataBase::IsInSeparateFile() const
 {
-	return	(GetBulkDataFlags() & BULKDATA_PayloadInSeperateFile) != 0;
+	return (GetBulkDataFlags() & BULKDATA_PayloadInSeperateFile) != 0;
 }
 
 bool FBulkDataBase::IsSingleUse() const
@@ -1188,6 +1263,9 @@ IAsyncReadFileHandle* FBulkDataBase::OpenAsyncReadHandle() const
 {
 	if (IsUsingIODispatcher())
 	{
+		checkf(IsInSeparateFile(),
+			TEXT("Attempting to OpenAsyncReadHandle on %s when the IoDispatcher is enabled, this operation is not supported!"),
+			IsInlined() ? TEXT("inline BulkData") : TEXT("BulkData in end-of-package-file section"));
 		return new FAsyncReadChunkIdHandle(CreateChunkId());
 	}
 	else
@@ -1208,6 +1286,9 @@ IBulkDataIORequest* FBulkDataBase::CreateStreamingRequest(int64 OffsetInBulkData
 	if (IsUsingIODispatcher())
 	{
 		checkf(OffsetInBulkData + BytesToRead <= BulkDataSize, TEXT("Attempting to read past the end of BulkData"));
+		checkf(IsInSeparateFile(),
+			TEXT("Attempting to CreateStreamingRequest on %s when the IoDispatcher is enabled, this operation is not supported!"),
+			IsInlined() ? TEXT("inline BulkData") : TEXT("BulkData in end-of-package-file section"));
 		FBulkDataIoDispatcherRequest* BulkDataIoDispatcherRequest = new FBulkDataIoDispatcherRequest(CreateChunkId(), BulkDataOffset + OffsetInBulkData, BytesToRead, CompleteCallback, UserSuppliedMemory);
 		BulkDataIoDispatcherRequest->StartAsyncWork();
 
@@ -1258,8 +1339,11 @@ IBulkDataIORequest* FBulkDataBase::CreateStreamingRequestForRange(const BulkData
 	{
 		const FBulkDataBase& End = *RangeArray[RangeArray.Num() - 1];
 
-		checkf(Start.CreateChunkId() == End.CreateChunkId(), TEXT("BulkData range does not come from the same file (%s vs %s)"),
-			*Start.GetPackagePath().GetDebugName(), *End.GetPackagePath().GetDebugName());
+		checkf(Start.IsInSeparateFile(),
+			TEXT("Attempting to CreateStreamingRequestForRange on %s when the IoDispatcher is enabled, this operation is not supported!"),
+			Start.IsInlined() ? TEXT("inline BulkData") : TEXT("BulkData in end-of-package-file section"));
+		checkf(End.IsInSeparateFile() && Start.CreateChunkId() == End.CreateChunkId(), TEXT("BulkData range does not come from the same file (%s vs %s)"),
+			*Start.GetPackagePath().GetDebugName(Start.GetPackageSegment()), *End.GetPackagePath().GetDebugName(End.GetPackageSegment()));
 
 		const int64 ReadOffset = Start.GetBulkDataOffsetInFile();
 		const int64 ReadLength = (End.GetBulkDataOffsetInFile() + End.GetBulkDataSize()) - ReadOffset;
@@ -1385,13 +1469,16 @@ int64 FBulkDataBase::GetBulkDataOffsetInFile() const
 
 FIoFilenameHash FBulkDataBase::GetIoFilenameHash() const
 {
-	if (!IsUsingIODispatcher())
+	if (IsUsingIODispatcher())
 	{
-		return MakeIoFilenameHash(PackageTokenSystem::GetPackagePath(Data.Token));
+		checkf(IsInSeparateFile(),
+			TEXT("Attempting to GetIoFilenameHash on %s when the IoDispatcher is enabled, this operation is not supported!"),
+			IsInlined() ? TEXT("inline BulkData") : TEXT("BulkData in end-of-package-file section"));
+		return MakeIoFilenameHash(CreateChunkId());
 	}
 	else
 	{
-		return MakeIoFilenameHash(CreateChunkId());
+		return MakeIoFilenameHash(PackageTokenSystem::GetPackagePath(Data.Token));
 	}
 }
 
@@ -1457,15 +1544,15 @@ void FBulkDataBase::LoadDataDirectly(void** DstBuffer)
 	{
 		InternalLoadFromPackageResource(DstBuffer);
 	}
-	else if (IsUsingIODispatcher())
+	else if (IsUsingIODispatcher() && IsInSeparateFile())
 	{
-
 		InternalLoadFromIoStore(DstBuffer);
 	}
 	else
 	{
 		// Note that currently this shouldn't be reachable as we should early out due to the ::CanLoadFromDisk check at the start of the method
-		UE_LOG(LogSerialization, Error, TEXT("Attempting to reload inline BulkData when the IoDispatcher is enabled, this operation is not supported! (%d)"), IsInlined());
+		UE_LOG(LogSerialization, Error, TEXT("Attempting to %s when the IoDispatcher is enabled, this operation is not supported!"),
+			IsInlined() ? TEXT("reload inline BulkData") : TEXT("load BulkData from end-of-package-file section"));
 	}
 }
 
@@ -1491,7 +1578,7 @@ void FBulkDataBase::LoadDataAsynchronously(AsyncCallback&& Callback)
 				Callback(Status);
 			});
 	}
-	else if (IsUsingIODispatcher())
+	else if (IsUsingIODispatcher() && IsInSeparateFile())
 	{
 		void* DummyPointer = nullptr;
 		InternalLoadFromIoStoreAsync(&DummyPointer, MoveTemp(Callback));
@@ -1499,29 +1586,17 @@ void FBulkDataBase::LoadDataAsynchronously(AsyncCallback&& Callback)
 	else
 	{
 		// Note that currently this shouldn't be reachable as we should early out due to the ::CanLoadFromDisk check at the start of the method
-		UE_LOG(LogSerialization, Error, TEXT("Attempting to reload inline BulkData when the IoDispatcher is enabled, this operation is not supported!"));
+		UE_LOG(LogSerialization, Error, TEXT("Attempting to %s when the IoDispatcher is enabled, this operation is not supported!"),
+			IsInlined() ? TEXT("reload inline BulkData") : TEXT("load BulkData from end-of-package-file section"));
 	}
 }
 
 void FBulkDataBase::InternalLoadFromPackageResource(void** DstBuffer)
 {
 	FPackagePath PackagePath = PackageTokenSystem::GetPackagePath(Data.Token);
-	int64 Offset = BulkDataOffset;
-
-	// Fix up the PackageSegment/Offset to work with streaming if inlined
-	EPackageSegment PackageSegment;
-	if (IsInlined() && GEventDrivenLoaderEnabled)
-	{
-		Offset -= IPackageResourceManager::Get().FileSize(PackagePath, EPackageSegment::Header);
-		PackageSegment = EPackageSegment::Exports;
-	}
-	else
-	{
-		PackageSegment = GetPackageSegmentFromFlags();
-	}
+	EPackageSegment PackageSegment = GetPackageSegmentFromFlags();
 
 	// If the data is inlined then we already loaded it during ::Serialize, this warning should help track cases where data is being discarded then re-requested.
-	// Disabled at the moment
 	UE_CLOG(IsInlined(), LogSerialization, Warning, TEXT("Reloading inlined bulk data directly from disk, this is detrimental to loading performance. PackagePath: '%s'."),
 		*PackagePath.GetDebugName(PackageSegment));
 
@@ -1530,7 +1605,7 @@ void FBulkDataBase::InternalLoadFromPackageResource(void** DstBuffer)
 		*PackagePath.GetDebugName(PackageSegment), (!Result.Archive.IsValid() ? TEXT("could not find package") : TEXT("package is a TextAsset which is not supported")));
 
 	// Seek to the beginning of the bulk data in the file.
-	Result.Archive->Seek(Offset);
+	Result.Archive->Seek(BulkDataOffset);
 
 	if (*DstBuffer == nullptr)
 	{
@@ -1579,40 +1654,42 @@ void FBulkDataBase::InternalLoadFromIoStoreAsync(void** DstBuffer, AsyncCallback
 	Batch.Issue();
 }
 
-void FBulkDataBase::ProcessDuplicateData(FArchive& Ar, const UPackage* Package, const FPackagePath* PackagePath, int64& InOutOffsetInFile)
+void FBulkDataBase::ProcessDuplicateData(EBulkDataFlags NewFlags, int64 NewSizeOnDisk, int64 NewOffset,
+	const UPackage* Package, const FPackagePath* PackagePath, const FLinkerLoad* Linker)
 {
 	// We need to load the optional bulkdata info as we might need to create a FIoChunkId based on it!
-	EBulkDataFlags NewFlags;
-	int64 NewSizeOnDisk;
-	int64 NewOffset;
-
-	SerializeDuplicateData(Ar, NewFlags, NewSizeOnDisk, NewOffset);
-
 #if ALLOW_OPTIONAL_DATA
-	if (IsUsingIODispatcher())
+	bool bUseOptionalSegment = false;
+	bool bUsingIODispatcher = IsUsingIODispatcher();
+	if (bUsingIODispatcher)
 	{
 		const FIoChunkId OptionalChunkId = CreateIoChunkId(Data.PackageID, 0, EIoChunkType::OptionalBulkData);
-
-		if (IoDispatcher->DoesChunkExist(OptionalChunkId))
-		{
-			BulkDataFlags = EBulkDataFlags(NewFlags | BULKDATA_UsesIoDispatcher);
-
-			checkf(BulkDataSize == NewSizeOnDisk, TEXT("Size mistach between original data size (%lld)and duplicate data size (%lld)"), BulkDataSize, NewSizeOnDisk);
-			InOutOffsetInFile = NewOffset;
-		}
+		bUseOptionalSegment = IoDispatcher->DoesChunkExist(OptionalChunkId);
 	}
 	else
 	{
 		// If we have no valid input packagepath then the package is broken anyway and we will not be able to find the optional data!
-		if (PackagePath != nullptr)
-		{
-			EPackageSegment PackageSegment = EPackageSegment::BulkDataOptional;
+		bUseOptionalSegment = PackagePath != nullptr && 
+			IPackageResourceManager::Get().DoesPackageExist(*PackagePath, EPackageSegment::BulkDataOptional);
+	}
 
-			if (IPackageResourceManager::Get().DoesPackageExist(*PackagePath, PackageSegment))
+	if (bUseOptionalSegment)
+	{
+		checkf(BulkDataSize == NewSizeOnDisk, TEXT("Size mismatch between original data size (%lld)and duplicate data size (%lld)"), BulkDataSize, NewSizeOnDisk);
+		BulkDataOffset = NewOffset;
+		BulkDataFlags = static_cast<EBulkDataFlags>((NewFlags & ~BULKDATA_DuplicateNonOptionalPayload)
+			| BULKDATA_OptionalPayload | BULKDATA_PayloadInSeperateFile | BULKDATA_PayloadAtEndOfFile);
+		if (bUsingIODispatcher)
+		{
+			checkf(!NeedsOffsetFixup(), TEXT("IODispatcher does not support offset fixups; 'LegacyBulkDataOffsets=true' will not work with the IoStore!"));
+			BulkDataFlags = static_cast<EBulkDataFlags>(BulkDataFlags | BULKDATA_UsesIoDispatcher);
+		}
+		else
+		{
+			if (NeedsOffsetFixup())
 			{
-				BulkDataFlags = NewFlags;
-				checkf(BulkDataSize == NewSizeOnDisk, TEXT("Size mistach between original data size (%lld)and duplicate data size (%lld)"), BulkDataSize, NewSizeOnDisk);
-				InOutOffsetInFile = NewOffset;
+				checkf(Linker != nullptr, TEXT("BulkData needs its offset fixed on load but no linker found"));
+				BulkDataOffset += Linker->Summary.BulkDataStartOffset;
 			}
 		}
 	}
@@ -1732,20 +1809,24 @@ void FBulkDataBase::FlushAsyncLoading()
 
 EPackageSegment FBulkDataBase::GetPackageSegmentFromFlags() const
 {
-	if (IsOptional())
-	{
-		// Optional data should be tested for first as we in theory can have data that would
-		// be marked as inline, also marked as optional and in this case we should treat it as
-		// optional data first.
-		return EPackageSegment::BulkDataOptional;
-	}
-	else if (!IsInSeparateFile())
+	if (!IsInSeparateFile())
 	{
-		return EPackageSegment::Header;
+		if (GEventDrivenLoaderEnabled)
+		{
+			// With EDL data, packages are split into EPackageSegment::Header (summary and linker tables) and
+			// EPackageSegment::Exports (serialized UObject bytes and the bulk data section)
+			// Inline and end-of-file bulk data is in the Exports section
+			return EPackageSegment::Exports;
+		}
+		else
+		{
+			// Without EDL, there is only a single combined file for Summary,linkertables,exports,and end-of-file bulk data
+			return EPackageSegment::Header;
+		}
 	}
-	else if (IsInlined())
+	else if (IsOptional())
 	{
-		return EPackageSegment::Exports;
+		return EPackageSegment::BulkDataOptional;
 	}
 	else if (IsFileMemoryMapped())
 	{
diff --git a/Engine/Source/Runtime/CoreUObject/Private/UObject/SavePackage/SavePackageUtilities.cpp b/Engine/Source/Runtime/CoreUObject/Private/UObject/SavePackage/SavePackageUtilities.cpp
--- a/Engine/Source/Runtime/CoreUObject/Private/UObject/SavePackage/SavePackageUtilities.cpp
+++ b/Engine/Source/Runtime/CoreUObject/Private/UObject/SavePackage/SavePackageUtilities.cpp
@@ -1841,11 +1841,12 @@ void SaveBulkData(FLinkerSave* Linker, const UPackage* InOuter, const TCHAR* Fil
 			MappedBulkArchive.Reset(new FLargeMemoryWriterWithRegions);
 		}
 
-		// If we are not allowing BulkData to go to the IoStore and we will be saving the BulkData to a separate file then 
-		// we cannot manipulate the offset as we cannot 'fix' it at runtime with the AsyncLoader2
-		// 
-		// We should remove the manipulated offset entirely, at least for separate files but for now we need to leave it to
-		// prevent larger patching sizes.
+		// If we are using IoStore (implemented in AsyncLoader2), then we will have no linkers and can not manipulate the
+		// offset as we will not have the Linker's Summary.BulkDataStartOffset information at runtime.
+		// This means that manipulated offsets are incompatible with IoStore.
+		// We now by default do not manipulate the offset for separate files, but we need to optionally leave it on
+		// to prevent larger patching sizes. In the future we may want to remove it even for end-of-file BulkData.
+		// Note that bForceLegacyOffsets is incompatible with using IoStore.
 		if (SavePackageContext != nullptr && SavePackageContext->bForceLegacyOffsets == false && bShouldUseSeparateBulkFile)
 		{
 			ExtraBulkDataFlags |= BULKDATA_NoOffsetFixUp;
@@ -1880,6 +1881,8 @@ void SaveBulkData(FLinkerSave* Linker, const UPackage* InOuter, const TCHAR* Fil
 				bBulkItemIsMapped = false;
 			}
 
+			// TODO: BulkData Serialize no longer requires BulkData flags to be set; it takes the BulkDataFlags for the save as arguments.
+			//       Remove this unnecessary set of the BulkData's BulkDataFlags.
 			BulkDataStorageInfo.BulkData->ClearBulkDataFlags(0xFFFFFFFF);
 			BulkDataStorageInfo.BulkData->SetBulkDataFlags(ModifiedBulkDataFlags);
 
@@ -1940,7 +1943,8 @@ void SaveBulkData(FLinkerSave* Linker, const UPackage* InOuter, const TCHAR* Fil
 
 			int64 StoredBulkStartOffset = (ModifiedBulkDataFlags & BULKDATA_NoOffsetFixUp) == 0 ? BulkStartOffset - StartOfBulkDataArea : BulkStartOffset;
 
-			BulkDataStorageInfo.BulkData->SerializeBulkData(*TargetArchive, BulkDataStorageInfo.BulkData->Lock(LOCK_READ_ONLY));
+			BulkDataStorageInfo.BulkData->SerializeBulkData(*TargetArchive, BulkDataStorageInfo.BulkData->Lock(LOCK_READ_ONLY),
+				static_cast<EBulkDataFlags>(ModifiedBulkDataFlags));
 
 			int64 BulkEndOffset = TargetArchive->Tell();
 			const int64 LinkerEndOffset = Linker->Tell();
@@ -1954,16 +1958,7 @@ void SaveBulkData(FLinkerSave* Linker, const UPackage* InOuter, const TCHAR* Fil
 			*Linker << StoredBulkStartOffset;
 
 			Linker->Seek(BulkDataStorageInfo.BulkDataSizeOnDiskPos);
-			if (ModifiedBulkDataFlags & BULKDATA_Size64Bit)
-			{
-				*Linker << SizeOnDisk;
-			}
-			else
-			{
-				check(SizeOnDisk < (1LL << 31));
-				int32 SizeOnDiskAsInt32 = SizeOnDisk;
-				*Linker << SizeOnDiskAsInt32;
-			}
+			SerializeBulkDataSizeInt(*Linker, SizeOnDisk, static_cast<EBulkDataFlags>(ModifiedBulkDataFlags));
 
 			if (SavePackageContext != nullptr && SavePackageContext->BulkDataManifest != nullptr)
 			{
@@ -1995,9 +1990,16 @@ void SaveBulkData(FLinkerSave* Linker, const UPackage* InOuter, const TCHAR* Fil
 			Linker->Seek(LinkerEndOffset);
 
 			// Restore BulkData flags to before serialization started
+			BulkDataStorageInfo.BulkData->Unlock();
 			BulkDataStorageInfo.BulkData->ClearBulkDataFlags(0xFFFFFFFF);
 			BulkDataStorageInfo.BulkData->SetBulkDataFlags(OldBulkDataFlags);
-			BulkDataStorageInfo.BulkData->Unlock();
+
+			// If we are overwriting the LoadedPath for the current package, the bulk data flags need to be updated to match
+			// the values set in the package on disk. This function is responsible for setting some of those flags, so set them now
+			if (!TargetPlatform)
+			{
+				BulkDataStorageInfo.BulkData->SetBulkDataFlags(ExtraBulkDataFlags);
+			}
 		}
 
 		if (BulkArchive)
diff --git a/Engine/Source/Runtime/CoreUObject/Public/Serialization/BulkData.h b/Engine/Source/Runtime/CoreUObject/Public/Serialization/BulkData.h
--- a/Engine/Source/Runtime/CoreUObject/Public/Serialization/BulkData.h
+++ b/Engine/Source/Runtime/CoreUObject/Public/Serialization/BulkData.h
@@ -642,7 +642,11 @@ public:
 	 *
 	 * @param	Ar					Archive to serialize with
 	 * @param	Data				Memory to serialize either to or from
-	 */ 
+	 * @param	InBulkDataFlags		Flags describing how the data was/shouldbe serialized
+	 */
+	void SerializeBulkData(FArchive& Ar, void* Data, EBulkDataFlags InBulkDataFlags);
+
+	UE_DEPRECATED(5.0, "Use the version that takes InBulkDataFlags")
 	void SerializeBulkData( FArchive& Ar, void* Data );
 
 	/*-----------------------------------------------------------------------------
@@ -709,6 +713,13 @@ public:
 	static IBulkDataIORequest* CreateStreamingRequestForRange(const FString& Filename, const BulkDataRangeArray& RangeArray, EAsyncIOPriorityAndFlags Priority, FBulkDataIORequestCallBack* CompleteCallback);
 #endif
 
+	/** Enable the given flags in the given accumulator variable. */
+	static void SetBulkDataFlagsOn(EBulkDataFlags& InOutAccumulator, EBulkDataFlags FlagsToSet);
+	/** Disable the given flags in the given accumulator variable. */
+	static void ClearBulkDataFlagsOn(EBulkDataFlags& InOutAccumulator, EBulkDataFlags FlagsToClear);
+	/** Returns decompress method flags specified by the given bulk data flags. */
+	static FName GetDecompressionFormat(EBulkDataFlags InFlags);
+
 	/*-----------------------------------------------------------------------------
 		Class specific virtuals.
 	-----------------------------------------------------------------------------*/
@@ -753,18 +764,6 @@ private:
 	void DetachFromArchive( FArchive* Ar, bool bEnsureBulkDataIsLoaded );
 #endif // WITH_EDITOR
 
-#if WITH_IOSTORE_IN_EDITOR
-	/**
-	 * Serialize bulk data structure from I/O store.
-	 *
-	 * @param Ar	Archive to serialize with
-	 * @param Owner	Object owning the bulk data
-	 * @param Idx	Index of bulk data item being serialized
-	 * @param bAttemptFileMapping	If true, attempt to map this instead of loading it into malloc'ed memory
-	 */
-	void SerializeFromIoStore( FArchive& Ar, UObject* Owner, int32 Idx, bool bAttemptFileMapping );
-#endif
-
 	/*-----------------------------------------------------------------------------
 		Internal helpers
 	-----------------------------------------------------------------------------*/
@@ -791,8 +790,10 @@ private:
 	 * archive we can use for serialization.
 	 *
 	 * @param Dest Memory to serialize data into
+	 *
+	 * @return Whether the load succeeded
 	 */
-	void LoadDataIntoMemory( void* Dest );
+	bool TryLoadDataIntoMemory(void* Dest);
 
 	/** Create the async load task */
 	void AsyncLoadBulkData();
@@ -810,7 +811,7 @@ private:
 	void ResetAsyncData();
 	
 	/** Returns true if bulk data should be loaded asynchronously */
-	bool ShouldStreamBulkData();
+	bool ShouldStreamBulkData(FArchive& Ar);
 
 	/** Returns if the offset needs fixing when serialized */
 	bool NeedsOffsetFixup() const;
diff --git a/Engine/Source/Runtime/CoreUObject/Public/Serialization/BulkData2.h b/Engine/Source/Runtime/CoreUObject/Public/Serialization/BulkData2.h
--- a/Engine/Source/Runtime/CoreUObject/Public/Serialization/BulkData2.h
+++ b/Engine/Source/Runtime/CoreUObject/Public/Serialization/BulkData2.h
@@ -15,6 +15,7 @@ struct FOwnedBulkDataPtr;
 class IMappedFileHandle;
 class IMappedFileRegion;
 class FBulkDataBase;
+class FLinkerLoad;
 
 /** A loose hash value that can be created from either a filenames or a FIoChunkId */
 using FIoFilenameHash = uint32;
@@ -262,7 +263,9 @@ private:
 	 */
 	bool CanDiscardInternalData() const;
 
-	void ProcessDuplicateData(FArchive& Ar, const UPackage* Package, const FPackagePath* PackagePath, int64& InOutOffsetInFile);
+	void ProcessDuplicateData(EBulkDataFlags NewFlags, int64 NewSizeOnDisk, int64 NewOffset, const UPackage* Package,
+		const FPackagePath* PackagePath, const FLinkerLoad* Linker);
+	void ConditionalSetInlineAlwaysAllowDiscard();
 	void SerializeDuplicateData(FArchive& Ar, EBulkDataFlags& OutBulkDataFlags, int64& OutBulkDataSizeOnDisk, int64& OutBulkDataOffsetInFile);
 	void SerializeBulkData(FArchive& Ar, void* DstBuffer, int64 DataLength);
 
diff --git a/Engine/Source/Runtime/CoreUObject/Public/Serialization/BulkDataCommon.h b/Engine/Source/Runtime/CoreUObject/Public/Serialization/BulkDataCommon.h
--- a/Engine/Source/Runtime/CoreUObject/Public/Serialization/BulkDataCommon.h
+++ b/Engine/Source/Runtime/CoreUObject/Public/Serialization/BulkDataCommon.h
@@ -2,6 +2,8 @@
 
 #pragma once
 
+#include "Math/NumericLimits.h"
+
 // NOTE: This file only needs to exist as long as we need to maintain the editor and runtime versions of Bulkdata.
 // Code common to both can be placed here.
 
@@ -12,7 +14,12 @@ enum EBulkDataFlags : uint32
 {
 	/** Empty flag set. */
 	BULKDATA_None = 0,
-	/** If set, payload is stored at the end of the file and not inline. */
+	/**
+	 * INTERNAL SET ONLY - callers of bulkdata should not set this flag on the bulk data
+	 * It is overwritten according to global configuration by Serialize.
+	 * If set, payload is stored not inline; it is stored either at the end of the file
+	 * or in a separate file.
+	 */
 	BULKDATA_PayloadAtEndOfFile = 1 << 0,
 	/** If set, payload should be [un]compressed using ZLIB during serialization. */
 	BULKDATA_SerializeCompressedZLIB = 1 << 1,
@@ -28,7 +35,11 @@ enum EBulkDataFlags : uint32
 	BULKDATA_SerializeCompressed = (BULKDATA_SerializeCompressedZLIB),
 	/** Forces the payload to be always streamed, regardless of its size. */
 	BULKDATA_ForceStreamPayload = 1 << 7,
-	/** If set, payload is stored in a .upack file alongside the uasset. */
+	/**
+	 * INTERNAL SET ONLY - callers of bulkdata should not set this flag on the bulk data
+	 * It is overwritten according to global configuration by Serialize.
+	 * If set, payload is stored in a separate file such as .ubulk.
+	 * */
 	BULKDATA_PayloadInSeperateFile = 1 << 8,
 	/** DEPRECATED: If set, payload is compressed using platform specific bit window. */
 	BULKDATA_SerializeCompressedBitWindow = 1 << 9,
@@ -69,6 +80,22 @@ FORCEINLINE FArchive& operator<<(FArchive& Ar, EBulkDataFlags& Flags)
 	return Ar;
 }
 
+/** Serialize the given Value as an int32 or int64 depending on InBulkDataFlags&BULKDATA_Size64Bit. */
+inline void SerializeBulkDataSizeInt(FArchive& Ar, int64& Value, EBulkDataFlags InBulkDataFlags)
+{
+	if (InBulkDataFlags & BULKDATA_Size64Bit)
+	{
+		Ar << Value;
+	}
+	else
+	{
+		check(!Ar.IsSaving() || (MIN_int32 <= Value && Value <= MAX_int32));
+		int32 ValueAsInt32 = Value;
+		Ar << ValueAsInt32;
+		Value = ValueAsInt32;
+	}
+}
+
 /**
  * Enumeration for bulk data lock status.
  */
diff --git a/Engine/Source/Runtime/Engine/Private/Animation/AnimCompressionTypes.cpp b/Engine/Source/Runtime/Engine/Private/Animation/AnimCompressionTypes.cpp
--- a/Engine/Source/Runtime/Engine/Private/Animation/AnimCompressionTypes.cpp
+++ b/Engine/Source/Runtime/Engine/Private/Animation/AnimCompressionTypes.cpp
@@ -623,7 +623,7 @@ void FCompressedAnimSequence::SerializeCompressedData(FArchive& Ar, bool bDDCDat
 			void* Dest = OptionalBulk.Realloc(NumBytes);
 			FMemory::Memcpy(Dest, &(SerializedData[0]), NumBytes);
 			OptionalBulk.Unlock();
-			OptionalBulk.SetBulkDataFlags(BULKDATA_PayloadAtEndOfFile | BULKDATA_PayloadInSeperateFile | BULKDATA_Force_NOT_InlinePayload | BULKDATA_MemoryMappedPayload);
+			OptionalBulk.SetBulkDataFlags(BULKDATA_Force_NOT_InlinePayload | BULKDATA_MemoryMappedPayload);
 			OptionalBulk.ClearBulkDataFlags(BULKDATA_ForceInlinePayload);
 			OptionalBulk.Serialize(Ar, DataOwner);
 #else
