Fix FindMostLikelyCulprit for Save2 to adequately report referencers information if a package ends up with illegal references

diff --git a/Engine/Source/Runtime/CoreUObject/Private/UObject/SavePackage/SavePackageUtilities.cpp b/Engine/Source/Runtime/CoreUObject/Private/UObject/SavePackage/SavePackageUtilities.cpp
--- a/Engine/Source/Runtime/CoreUObject/Private/UObject/SavePackage/SavePackageUtilities.cpp
+++ b/Engine/Source/Runtime/CoreUObject/Private/UObject/SavePackage/SavePackageUtilities.cpp
@@ -13,6 +13,7 @@
 #include "Misc/CommandLine.h"
 #include "Misc/Paths.h"
 #include "Misc/ScopedSlowTask.h"
+#include "SaveContext.h"
 #include "Serialization/BulkData.h"
 #include "Serialization/BulkDataManifest.h"
 #include "Serialization/LargeMemoryWriter.h"
@@ -455,12 +456,22 @@ void ConditionallyExcludeObjectForTarget(UObject* Obj, EObjectMark ExcludedObjec
  * @param	BadObjects	array of objects that are considered "bad" (e.g. non- RF_Public, in different map package, ...)
  * @return	UObject that is considered the most likely culprit causing them to be referenced or NULL
  */
-void FindMostLikelyCulprit(TArray<UObject*> BadObjects, UObject*& MostLikelyCulprit, FString& OutReferencer)
+void FindMostLikelyCulprit(const TArray<UObject*>& BadObjects, UObject*& MostLikelyCulprit, FString& OutReferencer, FSaveContext* InOptionalSaveContext)
 {
 	UObject* ArchetypeCulprit = nullptr;
 	UObject* ReferencedCulprit = nullptr;
 	const FProperty* ReferencedCulpritReferencer = nullptr;
 
+	auto IsObjectIncluded = [InOptionalSaveContext](UObject* InObject)
+	{
+		// if we passed in a SaveContext use that to validate if the object is an import or an export instead of marks
+		if (InOptionalSaveContext)
+		{
+			return InOptionalSaveContext->IsIncluded(InObject);
+		}
+		return InObject->HasAnyMarks(EObjectMark(OBJECTMARK_TagExp | OBJECTMARK_TagImp));
+	};
+
 	// Iterate over all objects that are marked as unserializable/ bad and print out their referencers.
 	for (int32 BadObjIndex = 0; BadObjIndex < BadObjects.Num(); BadObjIndex++)
 	{
@@ -501,7 +512,7 @@ void FindMostLikelyCulprit(TArray<UObject*> BadObjects, UObject*& MostLikelyCulp
 			for (int32 i = 0; i < Refs.ExternalReferences.Num(); i++)
 			{
 				UObject* RefObj = Refs.ExternalReferences[i].Referencer;
-				if (RefObj->HasAnyMarks(EObjectMark(OBJECTMARK_TagExp | OBJECTMARK_TagImp)))
+				if (IsObjectIncluded(RefObj))
 				{
 					if (RefObj->GetFName() == NAME_PersistentLevel || RefObj->GetClass()->GetFName() == NAME_World)
 					{
diff --git a/Engine/Source/Runtime/CoreUObject/Private/UObject/SavePackage/SavePackageUtilities.h b/Engine/Source/Runtime/CoreUObject/Private/UObject/SavePackage/SavePackageUtilities.h
--- a/Engine/Source/Runtime/CoreUObject/Private/UObject/SavePackage/SavePackageUtilities.h
+++ b/Engine/Source/Runtime/CoreUObject/Private/UObject/SavePackage/SavePackageUtilities.h
@@ -18,6 +18,7 @@
 
 class FMD5;
 class FSavePackageContext;
+class FSaveContext;
 template<typename StateType> class TAsyncWorkSequence;
 
 DECLARE_LOG_CATEGORY_EXTERN(LogSavePackage, Log, All);
@@ -347,7 +348,7 @@ namespace SavePackageUtilities
 	bool HasUnsaveableOuter(UObject* InObj, UPackage* InSavingPackage);
 	void CheckObjectPriorToSave(FArchiveUObject& Ar, UObject* InObj, UPackage* InSavingPackage);
 	void ConditionallyExcludeObjectForTarget(UObject* Obj, EObjectMark ExcludedObjectMarks, const ITargetPlatform* TargetPlatform);
-	void FindMostLikelyCulprit(TArray<UObject*> BadObjects, UObject*& MostLikelyCulprit, FString& OutReferencer);
+	void FindMostLikelyCulprit(const TArray<UObject*>& BadObjects, UObject*& MostLikelyCulprit, FString& OutReferencer, FSaveContext* InOptionalSaveContext = nullptr);
 	void AddFileToHash(FString const& Filename, FMD5& Hash);
 
 	void WriteToFile(const FString& Filename, const uint8* InDataPtr, int64 InDataSize);
diff --git a/Engine/Source/Runtime/CoreUObject/Private/UObject/SavePackage2.cpp b/Engine/Source/Runtime/CoreUObject/Private/UObject/SavePackage2.cpp
--- a/Engine/Source/Runtime/CoreUObject/Private/UObject/SavePackage2.cpp
+++ b/Engine/Source/Runtime/CoreUObject/Private/UObject/SavePackage2.cpp
@@ -519,7 +519,7 @@ ESavePackageResult ValidateIllegalReferences(FSaveContext& SaveContext, TArray<U
 		if (bFindCulprit)
 		{
 			FString Referencer;
-			SavePackageUtilities::FindMostLikelyCulprit(PrivateObjects, MostLikelyCulprit, Referencer);
+			SavePackageUtilities::FindMostLikelyCulprit(PrivateObjects, MostLikelyCulprit, Referencer, &SaveContext);
 			CulpritString = FString::Printf(TEXT("%s (%s)"),
 				(MostLikelyCulprit != nullptr) ? *MostLikelyCulprit->GetFullName() : TEXT("(unknown culprit)"),
 				*Referencer);
