Fix potential garbage collector starvation in editor when all worker threads are busy

  - Use ParallelFor for reference collection so the game-thread can make progress without the help of worker threads
  - This fixes a potential deadlock between worker threads using FGCScopeGuard and the game-thread trying to complete a GC
  - Add starvation unittest to detect if the garbage collector implementation is starvation free
  - Also protects against ABA problems that might occur when using TLockFreePointerListUnordered and going too wide

NOTES
  - The fix is only active in Editor to avoid introducing regression where big/little cores are important
  - Heavy use of FGCScopeGuard is currently not expected for non-editor workflows, this might get revisited later

TESTS
  - Confirm that FGarbageCollectorStarvationTest fails without the ParallelFor
  - Confirm that FGarbageCollectorStarvationTest succeeds when using the ParallelFor

diff --git a/Engine/Source/Runtime/CoreUObject/Public/UObject/FastReferenceCollector.h b/Engine/Source/Runtime/CoreUObject/Public/UObject/FastReferenceCollector.h
--- a/Engine/Source/Runtime/CoreUObject/Public/UObject/FastReferenceCollector.h
+++ b/Engine/Source/Runtime/CoreUObject/Public/UObject/FastReferenceCollector.h
@@ -12,6 +12,7 @@
 #include "Misc/ScopeLock.h"
 #include "HAL/PlatformProcess.h"
 #include "UObject/FieldPath.h"
+#include "Async/ParallelFor.h"
 #include "UObject/UObjectArray.h"
 #include "UObject/FastReferenceCollectorOptions.h"
 
@@ -342,6 +343,7 @@ private:
 
 		FORCENOINLINE void DoTask()
 		{
+			TRACE_CPUPROFILER_EVENT_SCOPE(FCollectorTaskQueue::DoTask);
 			{
 				FScopeLock Lock(&WaitingThreadsLock);
 				if (bDone)
@@ -535,7 +537,9 @@ public:
 				ENamedThreads::Type BackgroundThreadName = ENamedThreads::AnyBackgroundThreadNormalTask;
 
 				FPlatformProcess::ModifyThreadAssignmentForUObjectReferenceCollector(NumThreads, NumBackgroundThreads, NormalThreadName, BackgroundThreadName);
-				int32 NumTasks = NumThreads + NumBackgroundThreads;
+
+				// Avoid going over 26 threads as it may cause race conditions in TLockFreePointerListUnordered
+				int32 NumTasks = FMath::Min(NumThreads + NumBackgroundThreads, 26);
 
 				check(NumTasks > 0);
 				ChunkTasks.Empty(NumTasks);
@@ -550,6 +554,15 @@ public:
 					TaskQueue.AddTask(&ObjectsToCollectReferencesFor, StartIndex, NumPerChunk);
 					StartIndex += NumPerChunk;
 				}
+
+				// This is only enabled in editor for now because it has not been tested client mode neither with little/big cores.
+				// If this gets activated for non-editor, please also activate the FGarbageCollectorStarvationTest test in client mode.
+#if WITH_EDITOR
+				// Use ParallelFor to ensure we're making progress in this thread instead of just waiting for tasks to complete.
+				// This is especially important for good game-thread latency when the taskgraph is already filled and to avoid
+				// potential deadlock with workers making use of FGCScopeGuard.
+				ParallelForTemplate(NumTasks, [this](int32) { TaskQueue.DoTask(); }, EParallelForFlags::Unbalanced);
+#else
 				for (int32 Chunk = 0; Chunk < NumTasks; Chunk++)
 				{
 					extern CORE_API int32 GUseNewTaskBackend;
@@ -559,6 +572,7 @@ public:
 
 				QUICK_SCOPE_CYCLE_COUNTER(STAT_GC_Subtask_Wait);
 				FTaskGraphInterface::Get().WaitUntilTasksComplete(ChunkTasks, ENamedThreads::GameThread_Local);
+#endif
 				TaskQueue.CheckDone();
 			}
 		}
