Change relevant ensure in Save2 to use their ensureMsgf variant to provide additional error context.

diff --git a/Engine/Source/Runtime/CoreUObject/Private/UObject/SavePackage2.cpp b/Engine/Source/Runtime/CoreUObject/Private/UObject/SavePackage2.cpp
--- a/Engine/Source/Runtime/CoreUObject/Private/UObject/SavePackage2.cpp
+++ b/Engine/Source/Runtime/CoreUObject/Private/UObject/SavePackage2.cpp
@@ -286,8 +286,7 @@ ESavePackageResult HarvestPackage(FSaveContext& SaveContext)
 	else
 	{
 		// Validate that if an asset is provided it has the appropriate top level flags
-		ensure(!Asset || Asset->HasAnyFlags(TopLevelFlags));
-
+		ensureMsgf(!Asset || Asset->HasAnyFlags(TopLevelFlags), TEXT("The asset to save %s in package %s do not contain any of the provided object flags."), *Asset->GetName(), *SaveContext.GetPackage()->GetName());
 		ForEachObjectWithPackage(SaveContext.GetPackage(), [&Harvester, TopLevelFlags](UObject* InObject)
 			{
 				if (InObject->HasAnyFlags(TopLevelFlags))
@@ -576,6 +575,8 @@ ESavePackageResult ValidateImports(FSaveContext& SaveContext)
 
 	};
 
+	FString PackageName = SaveContext.GetPackage()->GetName();
+
 	// Warn for private objects & map object references
 	TArray<UObject*> PrivateObjects;
 	TArray<UObject*> ObjectsInOtherMaps;
@@ -583,15 +584,19 @@ ESavePackageResult ValidateImports(FSaveContext& SaveContext)
 	{
 		UPackage* ImportPackage = Import->GetPackage();
 		// All names should be properly harvested at this point
-		ensureAlways(SaveContext.NameExists(Import->GetFName().GetComparisonIndex()));
-		ensureAlways(SaveContext.NameExists(ImportPackage->GetFName().GetComparisonIndex()));
-		ensureAlways(SaveContext.NameExists(Import->GetClass()->GetFName().GetComparisonIndex()));
-		ensureAlways(SaveContext.NameExists(Import->GetClass()->GetOuter()->GetFName().GetComparisonIndex()));
+		ensureAlwaysMsgf(SaveContext.NameExists(Import->GetFName().GetComparisonIndex())
+			, TEXT("Missing import name %s while saving package %s. Did you rename an import during serialization?"), *Import->GetName(), *PackageName);
+		ensureAlwaysMsgf(SaveContext.NameExists(ImportPackage->GetFName().GetComparisonIndex())
+			, TEXT("Missing import package name %s while saving package %s. Did you rename an import during serialization?"), *ImportPackage->GetName(), *PackageName);
+		ensureAlwaysMsgf(SaveContext.NameExists(Import->GetClass()->GetFName().GetComparisonIndex())
+			, TEXT("Missing import class name %s while saving package %s"), *Import->GetClass()->GetName(), *PackageName);
+		ensureAlwaysMsgf(SaveContext.NameExists(Import->GetClass()->GetOuter()->GetFName().GetComparisonIndex())
+			, TEXT("Missing import class package name %s while saving package %s"), *Import->GetClass()->GetOuter()->GetName(), *PackageName);
 
 		// if the import is marked as a prestream package, we dont need to validate further
 		if (SaveContext.IsPrestreamPackage(ImportPackage))
 		{
-			ensureAlways(Import == ImportPackage);
+			ensureAlwaysMsgf(Import == ImportPackage, TEXT("Found an import refrence %s in a prestream package %s while saving package %s"), *Import->GetName(), *ImportPackage->GetName(), *PackageName);
 			// These are not errors
 			UE_LOG(LogSavePackage, Display, TEXT("Prestreaming package %s "), *ImportPackage->GetPathName()); //-V595
 			continue;
@@ -1019,7 +1024,6 @@ ESavePackageResult BuildLinker(FSaveContext& SaveContext)
 					if (ReplacedOuter && *ReplacedOuter)
 					{
 						Import.OuterIndex = Linker->MapObject(*ReplacedOuter);
-						ensure(Import.OuterIndex != FPackageIndex());
 					}
 					else
 					{
