Fix non-unity compile error due to missing include.

diff --git a/Engine/Source/Runtime/CoreUObject/Private/Tests/ObjectPtrTest.cpp b/Engine/Source/Runtime/CoreUObject/Private/Tests/ObjectPtrTest.cpp
--- a/Engine/Source/Runtime/CoreUObject/Private/Tests/ObjectPtrTest.cpp
+++ b/Engine/Source/Runtime/CoreUObject/Private/Tests/ObjectPtrTest.cpp
@@ -9,6 +9,7 @@
 #include "Templates/Models.h"
 #include "UObject/Interface.h"
 #include "UObject/MetaData.h"
+#include "UObject/ObjectRedirector.h"
 #include "UObject/SoftObjectPath.h"
 
 #include <type_traits>
