Fixed LinkerLoad not remapping imports when running with -game

diff --git a/Engine/Source/Runtime/CoreUObject/Private/UObject/LinkerLoad.cpp b/Engine/Source/Runtime/CoreUObject/Private/UObject/LinkerLoad.cpp
--- a/Engine/Source/Runtime/CoreUObject/Private/UObject/LinkerLoad.cpp
+++ b/Engine/Source/Runtime/CoreUObject/Private/UObject/LinkerLoad.cpp
@@ -5137,7 +5137,8 @@ UObject* FLinkerLoad::CreateImport( int32 Index )
 					// Import is a toplevel package.
 					if( Import.OuterIndex.IsNull() )
 					{
-						FindObject = CreatePackage(*Import.ObjectName.ToString());
+						FName ObjectName = InstancingContextRemap(Import.ObjectName);
+						FindObject = CreatePackage(*ObjectName.ToString());
 					}
 					// Import is regular import/ export.
 					else
@@ -5156,7 +5157,8 @@ UObject* FLinkerLoad::CreateImport( int32 Index )
 							// Outer is toplevel package, create/ find it.
 							else if( OuterImport.OuterIndex.IsNull() )
 							{
-								FindOuter = CreatePackage( *OuterImport.ObjectName.ToString() );
+								FName ObjectName = InstancingContextRemap(OuterImport.ObjectName);
+								FindOuter = CreatePackage( *ObjectName.ToString() );
 							}
 							// Outer is regular import/ export, use IndexToObject to potentially recursively load/ find it.
 							else
