Unshelved from pending changelist '15566751':

PR #7802: Allow spaces in CoreRedirect definitions. (Contributed by [REDACTED_PERSON])

diff --git a/Engine/Source/Runtime/CoreUObject/Private/UObject/CoreRedirects.cpp b/Engine/Source/Runtime/CoreUObject/Private/UObject/CoreRedirects.cpp
--- a/Engine/Source/Runtime/CoreUObject/Private/UObject/CoreRedirects.cpp
+++ b/Engine/Source/Runtime/CoreUObject/Private/UObject/CoreRedirects.cpp
@@ -160,7 +160,8 @@ int32 FCoreRedirectObjectName::MatchScore(const FCoreRedirectObjectName& Other)
 
 bool FCoreRedirectObjectName::HasValidCharacters() const
 {
-	static FString InvalidRedirectCharacters = TEXT("\"' ,|&!~\n\r\t@#(){}[]=;^%$`");
+	// ObjectNames in Blueprint may contain spaces.
+	static FString InvalidRedirectCharacters = TEXT("\"',|&!~\n\r\t@#(){}[]=;^%$`");
 
 	return ObjectName.IsValidXName(InvalidRedirectCharacters) && OuterName.IsValidXName(InvalidRedirectCharacters) && PackageName.IsValidXName(InvalidRedirectCharacters);
 }
