EditorDomainSave - Strip bulk data from EditorDomain copies of packages. BulkData types in the EditorDomain have flags set to load the data from workspace domain version of the file.
#rnx

diff --git a/Engine/Plugins/Runtime/HairStrands/Source/HairStrandsCore/Private/HairDescription.cpp b/Engine/Plugins/Runtime/HairStrands/Source/HairStrandsCore/Private/HairDescription.cpp
--- a/Engine/Plugins/Runtime/HairStrands/Source/HairStrandsCore/Private/HairDescription.cpp
+++ b/Engine/Plugins/Runtime/HairStrands/Source/HairStrandsCore/Private/HairDescription.cpp
@@ -92,13 +92,7 @@ void FHairDescriptionBulkData::Serialize(FArchive& Ar, UObject* Owner)
 	}
 	else
 	{
-		if (Ar.IsLoading())
-		{
-			// If loading, take the package custom version so it can be applied to the bulk data archive
-			// when unpacking HairDescription from it
-			CustomVersions = Ar.GetCustomVersions();
-		}
-		else if (Ar.IsSaving())
+		if (Ar.IsSaving())
 		{
 			// If the bulk data hasn't been updated since this was loaded, there's a possibility that it has old versioning.
 			// Explicitly load and resave the FHairDescription so that its version is in sync with the FHairDescriptionBulkData.
@@ -114,6 +108,13 @@ void FHairDescriptionBulkData::Serialize(FArchive& Ar, UObject* Owner)
 	BulkData.Serialize(Ar, Owner);
 
 	Ar << Guid;
+
+	if (!Ar.IsTransacting() && Ar.IsLoading())
+	{
+		// If loading, take the package custom version so it can be applied to the bulk data archive
+		// when unpacking HairDescription from it
+		CustomVersions = BulkData.GetCustomVersions(Ar);
+	}
 }
 
 void FHairDescriptionBulkData::SaveHairDescription(FHairDescription& HairDescription)
diff --git a/Engine/Source/Editor/UnrealEd/Private/EditorDomain/EditorDomain.cpp b/Engine/Source/Editor/UnrealEd/Private/EditorDomain/EditorDomain.cpp
--- a/Engine/Source/Editor/UnrealEd/Private/EditorDomain/EditorDomain.cpp
+++ b/Engine/Source/Editor/UnrealEd/Private/EditorDomain/EditorDomain.cpp
@@ -303,6 +303,22 @@ bool FEditorDomain::TryMatchCaseOnDisk(const FPackagePath& PackagePath, FPackage
 	return Workspace->TryMatchCaseOnDisk(PackagePath, OutNormalizedPath);
 }
 
+TUniquePtr<FArchive> FEditorDomain::OpenReadExternalResource(EPackageExternalResource ResourceType, FStringView Identifier)
+{
+	return Workspace->OpenReadExternalResource(ResourceType, Identifier);
+}
+
+bool FEditorDomain::DoesExternalResourceExist(EPackageExternalResource ResourceType, FStringView Identifier)
+{
+	return Workspace->DoesExternalResourceExist(ResourceType, Identifier);
+}
+
+IAsyncReadFileHandle* FEditorDomain::OpenAsyncReadExternalResource(
+	EPackageExternalResource ResourceType, FStringView Identifier)
+{
+	return Workspace->OpenAsyncReadExternalResource(ResourceType, Identifier);
+}
+
 void FEditorDomain::FindPackagesRecursive(TArray<TPair<FPackagePath, EPackageSegment>>& OutPackages,
 	FStringView PackageMount, FStringView FileMount, FStringView RootRelPath, FStringView BasenameWildcard)
 {
diff --git a/Engine/Source/Editor/UnrealEd/Private/EditorDomain/EditorDomain.h b/Engine/Source/Editor/UnrealEd/Private/EditorDomain/EditorDomain.h
--- a/Engine/Source/Editor/UnrealEd/Private/EditorDomain/EditorDomain.h
+++ b/Engine/Source/Editor/UnrealEd/Private/EditorDomain/EditorDomain.h
@@ -62,6 +62,10 @@ public:
 	virtual IMappedFileHandle* OpenMappedHandleToPackage(const FPackagePath& PackagePath,
 		EPackageSegment PackageSegment, FPackagePath* OutUpdatedPath = nullptr) override;
 	virtual bool TryMatchCaseOnDisk(const FPackagePath& PackagePath, FPackagePath* OutNormalizedPath = nullptr) override;
+	virtual TUniquePtr<FArchive> OpenReadExternalResource(EPackageExternalResource ResourceType, FStringView Identifier) override;
+	virtual bool DoesExternalResourceExist(EPackageExternalResource ResourceType, FStringView Identifier) override;
+	virtual IAsyncReadFileHandle* OpenAsyncReadExternalResource(
+		EPackageExternalResource ResourceType, FStringView Identifier) override;
 	virtual void FindPackagesRecursive(TArray<TPair<FPackagePath, EPackageSegment>>& OutPackages, FStringView PackageMount,
 		FStringView FileMount, FStringView RootRelPath, FStringView BasenameWildcard) override;
 	virtual void IteratePackagesInPath(FStringView PackageMount, FStringView FileMount, FStringView RootRelPath,
diff --git a/Engine/Source/Editor/UnrealEd/Private/EditorDomain/EditorDomainSave.cpp b/Engine/Source/Editor/UnrealEd/Private/EditorDomain/EditorDomainSave.cpp
--- a/Engine/Source/Editor/UnrealEd/Private/EditorDomain/EditorDomainSave.cpp
+++ b/Engine/Source/Editor/UnrealEd/Private/EditorDomain/EditorDomainSave.cpp
@@ -740,7 +740,10 @@ bool FEditorDomainSaveServer::TrySavePackage(const FPackagePath& PackagePath, FS
 	FString TempFilename = FPaths::Combine(FPaths::ProjectIntermediateDir(), FGuid::NewGuid().ToString());
 	ON_SCOPE_EXIT{ IFileManager::Get().Delete(*TempFilename); };
 
-	FSavePackageResultStruct Result = GEditor->Save(Package, nullptr, RF_Standalone, *TempFilename);
+	uint32 SaveFlags = SAVE_NoError | // Do not crash the SaveServer on an error
+		SAVE_BulkDataByReference; // EditorDomain saves reference bulkdata from the WorkspaceDomain rather than duplicating it
+	FSavePackageResultStruct Result = GEditor->Save(Package, nullptr, RF_Standalone, *TempFilename, GError,
+		nullptr /* Conform */, false /* bForceByteSwapping */, true /* bWarnOfLongFilename */, SaveFlags);
 	if (Result.Result != ESavePackageResult::Success)
 	{
 		OutErrorMessage = FString::Printf(TEXT("Package %s could not be saved to temporary file %s: ResultCode = %d"),
diff --git a/Engine/Source/Runtime/CoreUObject/Private/Serialization/BulkData.cpp b/Engine/Source/Runtime/CoreUObject/Private/Serialization/BulkData.cpp
--- a/Engine/Source/Runtime/CoreUObject/Private/Serialization/BulkData.cpp
+++ b/Engine/Source/Runtime/CoreUObject/Private/Serialization/BulkData.cpp
@@ -435,6 +435,11 @@ bool FUntypedBulkData::CanLoadFromDisk() const
 
 bool FUntypedBulkData::DoesExist() const
 {
+	if (IsInExternalResource())
+	{
+		return IPackageResourceManager::Get().DoesExternalResourceExist(
+			EPackageExternalResource::WorkspaceDomainFile, PackagePath.GetPackageName());
+	}
 	return IPackageResourceManager::Get().DoesPackageExist(PackagePath, PackageSegment);
 }
 
@@ -853,13 +858,25 @@ void FUntypedBulkData::AsyncLoadBulkData()
 {
 	BulkDataAsync.Reallocate(GetBulkDataSize(), BulkDataAlignment);
 
-	FOpenPackageResult Result = IPackageResourceManager::Get().OpenReadPackage(PackagePath, PackageSegment);
-	checkf(Result.Archive.IsValid() && Result.Format == EPackageFormat::Binary, TEXT("Attempted to load bulk data from an invalid package '%s'%s."),
-		*PackagePath.GetDebugName(PackageSegment), (Result.Archive == nullptr ? TEXT("") : TEXT(": Package Format is Text which is not supported")));
+	TUniquePtr<FArchive> BulkArchive;
+	if (IsInExternalResource())
+	{
+		BulkArchive = IPackageResourceManager::Get().OpenReadExternalResource(EPackageExternalResource::WorkspaceDomainFile,
+			PackagePath.GetPackageName());
+		checkf(BulkArchive.IsValid(), TEXT("Attempted to load bulk data from invalid WorkspaceDomain package '%s'."),
+			*PackagePath.GetPackageName());
+	}
+	else
+	{
+		FOpenPackageResult Result = IPackageResourceManager::Get().OpenReadPackage(PackagePath, PackageSegment);
+		checkf(Result.Archive.IsValid() && Result.Format == EPackageFormat::Binary, TEXT("Attempted to load bulk data from an invalid package '%s'%s."),
+			*PackagePath.GetDebugName(PackageSegment), (Result.Archive == nullptr ? TEXT("") : TEXT(": Package Format is Text which is not supported")));
+		BulkArchive = MoveTemp(Result.Archive);
+	}
 
 	// Seek to the beginning of the bulk data in the file.
-	Result.Archive->Seek(BulkDataOffsetInFile);
-	SerializeBulkData(*Result.Archive, BulkDataAsync.Get(), BulkDataFlags);
+	BulkArchive->Seek(BulkDataOffsetInFile);
+	SerializeBulkData(*BulkArchive, BulkDataAsync.Get(), BulkDataFlags);
 }
 
 /*-----------------------------------------------------------------------------
@@ -873,7 +890,8 @@ void FUntypedBulkData::StartSerializingBulkData(FArchive& Ar, UObject* Owner, in
 
 	SerializeFuture = Async(EAsyncExecution::ThreadPool, [=]() 
 	{ 
-		UE_CLOG(GEventDrivenLoaderEnabled, LogSerialization, Error, TEXT("Attempt to stream bulk data with EDL enabled. This is not desireable. Package %s"), *PackagePath.GetDebugName(PackageSegment));
+		UE_CLOG(GEventDrivenLoaderEnabled, LogSerialization, Error, TEXT("Attempt to stream bulk data with EDL enabled. This is not desireable. Package %s"),
+			*PackagePath.GetDebugName(PackageSegment));
 		AsyncLoadBulkData();
 
 		return true;
@@ -907,7 +925,7 @@ bool FUntypedBulkData::ShouldStreamBulkData(FArchive& Ar)
 			return false; // if it is inline, it is already precached, so use it
 		}
 
-		if (BulkDataFlags & BULKDATA_PayloadInSeperateFile)
+		if (!(BulkDataFlags & BULKDATA_PayloadInSeperateFile))
 		{
 			checkf(false, TEXT("Bulk data should either be inline or stored in a separate file for the new uobject loader."));
 			return false; // if it is not in a separate file, then we can't easily find the correct offset in the uexp file; we don't want this case anyway!
@@ -1069,6 +1087,9 @@ void FUntypedBulkData::Serialize( FArchive& Ar, UObject* Owner, int32 Idx, bool
 			if (Ar.IsUsingEventDrivenLoader() && bPayloadInSeparateFile)
 			{
 				check(Owner && Owner->GetPackage()->GetPackageId().IsValid());
+				checkf(!(BulkDataFlags& BULKDATA_WorkspaceDomainPayload),
+					TEXT("%s IsUsingEventDrivenLoader but has a bulkdata with BULKDATA_WorkspaceDomainPayload. ")
+					TEXT("BULKDATA_WorkspaceDomainPayload is not supported with iostore."), *Ar.GetArchiveName());
 				SetLocalBulkDataFlags(BULKDATA_UsesIoDispatcher);
 				PackageId = Owner->GetPackage()->GetPackageId();
 				bUseIOStore = true;
@@ -1136,6 +1157,11 @@ void FUntypedBulkData::Serialize( FArchive& Ar, UObject* Owner, int32 Idx, bool
 				{
 					PackageSegment = EPackageSegment::BulkDataMemoryMapped;
 				}
+				else if (BulkDataFlags & BULKDATA_WorkspaceDomainPayload)
+				{
+					// PackageSegment is set to BulkDataDefault for debug name purposes. The segment is not used for loading.
+					PackageSegment = EPackageSegment::BulkDataDefault;
+				}
 				else
 				{
 					PackageSegment = EPackageSegment::BulkDataDefault;
@@ -1153,8 +1179,10 @@ void FUntypedBulkData::Serialize( FArchive& Ar, UObject* Owner, int32 Idx, bool
 			}
 
 			FArchive* CacheableArchive = Ar.GetCacheableArchive();
-			if (Ar.IsAllowingLazyLoading() && Owner != nullptr && CacheableArchive && !bUseIOStore)
+			if (Ar.IsAllowingLazyLoading() && CacheableArchive && !bUseIOStore)
 			{
+				SetLocalBulkDataFlags(BULKDATA_LazyLoadable);
+
 				// Deferred serialization is allowed by the archive/caller. Set flags for how to read
 				// the data, but skip synchronous reading of it until TryLoadDataIntoMemory is called.
 #if WITH_EDITOR
@@ -1194,16 +1222,25 @@ void FUntypedBulkData::Serialize( FArchive& Ar, UObject* Owner, int32 Idx, bool
 				{
 					if (bAttemptFileMapping)
 					{
-						if (!BulkData.MapFile(PackagePath, PackageSegment, BulkDataOffsetInFile, GetBulkDataSize()))
+						if (IsInExternalResource())
 						{
-							// we failed to map when requested, so they will be looking for the memory, do a sync load now.
+							// FileMapping not supported, and they they will be looking for the memory, do a sync load now.
 							ForceBulkDataResident();
 						}
+						else
+						{
+							if (!BulkData.MapFile(PackagePath, PackageSegment, BulkDataOffsetInFile, GetBulkDataSize()))
+							{
+								// we failed to map when requested, so they will be looking for the memory, do a sync load now.
+								ForceBulkDataResident();
+							}
+						}
 					}
 				}
 			}
 			else
 			{
+				ClearLocalBulkDataFlags(BULKDATA_LazyLoadable);
 				// Serialize the bulk data right away.
 				// bAttemptFileMapping is ignored in this case, since we are going to load the memory before returning anyway.
 
@@ -1240,17 +1277,29 @@ void FUntypedBulkData::Serialize( FArchive& Ar, UObject* Owner, int32 Idx, bool
 							else
 #endif
 							{
+								SetLocalBulkDataFlags(BULKDATA_LazyLoadable); // Separate files are always lazy loadable
+
 								// open separate bulk data file
 								UE_CLOG(GEventDrivenLoaderEnabled, LogSerialization, Error, TEXT("Attempt to sync load bulk data with EDL enabled (separate file). This is not desireable. File %s"), *PackagePath.GetDebugName(PackageSegment));
-
-								FOpenPackageResult TargetArchive = IPackageResourceManager::Get().OpenReadPackage(PackagePath, PackageSegment);
-								checkf(TargetArchive.Archive.IsValid() && TargetArchive.Format == EPackageFormat::Binary, TEXT("Attempted to load bulk data from an invalid PackagePath '%s': %s."),
-									*PackagePath.GetDebugName(PackageSegment), (!TargetArchive.Archive.IsValid() ? TEXT("could not find package") : TEXT("package is a TextAsset which is not supported")));
+								TUniquePtr<FArchive> TargetArchive;
+								if (IsInExternalResource())
+								{
+									TargetArchive = IPackageResourceManager::Get().OpenReadExternalResource(EPackageExternalResource::WorkspaceDomainFile, PackagePath.GetPackageName());
+									checkf(TargetArchive.IsValid(), TEXT("Attempted to load bulk data from invalid WorkspaceDomain package '%s'."),
+										*PackagePath.GetPackageName());
+								}
+								else
+								{
+									FOpenPackageResult OpenResult = IPackageResourceManager::Get().OpenReadPackage(PackagePath, PackageSegment);
+									checkf(OpenResult.Archive.IsValid() && OpenResult.Format == EPackageFormat::Binary, TEXT("Attempted to load bulk data from an invalid PackagePath '%s': %s."),
+										*PackagePath.GetDebugName(PackageSegment), (!OpenResult.Archive.IsValid() ? TEXT("could not find package") : TEXT("package is a TextAsset which is not supported")));
+									TargetArchive = MoveTemp(OpenResult.Archive);
+								}
 
 								// seek to the location in the file where the payload is stored
-								TargetArchive.Archive->Seek(BulkDataOffsetInFile);
+								TargetArchive->Seek(BulkDataOffsetInFile);
 								// serialize the payload
-								SerializeBulkData(*TargetArchive.Archive, BulkData.Get(), BulkDataFlags);
+								SerializeBulkData(*TargetArchive, BulkData.Get(), BulkDataFlags);
 							}
 						}
 						else
@@ -1336,7 +1385,8 @@ void FUntypedBulkData::Serialize( FArchive& Ar, UObject* Owner, int32 Idx, bool
 			{
 				// set the flag indicating where the payload is stored
 				SetBulkDataFlagsOn(LocalBulkDataFlags, BULKDATA_PayloadAtEndOfFile);
-				ClearBulkDataFlagsOn(LocalBulkDataFlags, BULKDATA_PayloadInSeperateFile); // SavePackageUtilities::SaveBulkData will add this back if required
+				ClearBulkDataFlagsOn(LocalBulkDataFlags,
+					static_cast<EBulkDataFlags>(BULKDATA_PayloadInSeperateFile | BULKDATA_WorkspaceDomainPayload)); // SavePackageUtilities::SaveBulkData will add these back if required
 
 				// with no LinkerSave we have to store the data inline
 				check(LinkerSave != NULL);
@@ -1393,7 +1443,7 @@ void FUntypedBulkData::Serialize( FArchive& Ar, UObject* Owner, int32 Idx, bool
 			{
 				// set the flag indicating where the payload is stored
 				ClearBulkDataFlagsOn(LocalBulkDataFlags,
-					static_cast<EBulkDataFlags>(BULKDATA_PayloadAtEndOfFile | BULKDATA_PayloadInSeperateFile));
+					static_cast<EBulkDataFlags>(BULKDATA_PayloadAtEndOfFile | BULKDATA_PayloadInSeperateFile | BULKDATA_WorkspaceDomainPayload));
 
 				int64 SavedBulkDataStartPos = Ar.Tell();
 
@@ -1440,17 +1490,122 @@ void FUntypedBulkData::Serialize( FArchive& Ar, UObject* Owner, int32 Idx, bool
 				}
 				// Seek to the end of written data so we don't clobber any data in subsequent writes
 				Ar.Seek(CurrentFileOffset);
+
+#if WITH_EDITOR
+				// If we are overwriting the LoadedPath for the current package, set the location variables on *this equal to the new values
+				// that we are writing into the package on disk
+				if (LinkerSave && LinkerSave->bSavingNewLoadedPath)
+				{
+					SetFlagsFromDiskWrittenValues(LocalBulkDataFlags, LocalBulkDataOffsetInFile, LocalBulkDataSizeOnDisk,
+						INDEX_NONE /* LinkerSummaryBulkDataStartOffset, not applicable */);
+				}
+#endif
 			}
+		}
+	}
+}
+
+#if WITH_EDITOR
+
+void FUntypedBulkData::SetFlagsFromDiskWrittenValues(EBulkDataFlags InBulkDataFlags, int64 InBulkDataOffsetInFile, int64 InBulkDataSizeOnDisk,
+	int64 LinkerSummaryBulkDataStartOffset)
+{
+	auto SetLocalBulkDataFlags = [&InBulkDataFlags](EBulkDataFlags BulkDataFlagsToSet)
+	{
+		InBulkDataFlags = static_cast<EBulkDataFlags>(InBulkDataFlags | BulkDataFlagsToSet);
+	};
+	auto ClearLocalBulkDataFlags = [&InBulkDataFlags](EBulkDataFlags BulkDataFlagsToClear)
+	{
+		InBulkDataFlags = static_cast<EBulkDataFlags>(InBulkDataFlags & ~BulkDataFlagsToClear);
+	};
+
+	check(!(InBulkDataFlags & BULKDATA_BadDataVersion)); // This is a legacy flag that should no longer be set when saving
+	if (GIsEditor)
+	{
+		ClearLocalBulkDataFlags(BULKDATA_SingleUse); // SingleUse is not used in editor, see the notes in Serialize
+	}
+#if WITH_IOSTORE_IN_EDITOR
+	// We are no longer loading from iostore, even if we were before; our data is now stored in the loose file we saved
+	ClearLocalBulkDataFlags(BULKDATA_UsesIoDispatcher);
+#endif // WITH_IOSTORE_IN_EDITOR
 
-			// If we are overwriting the LoadedPath for the current package, set the location variables on *this equal to the new values
-			// that we are writing into the package on disk
-			if (!Ar.IsCooking())
+	const bool bPayloadInline = !(InBulkDataFlags & BULKDATA_PayloadAtEndOfFile);
+	const bool bPayloadInSeparateFile = !bPayloadInline && (InBulkDataFlags & BULKDATA_PayloadInSeperateFile);
+	// fix up the file offset if the offset is relative to the file's bulkdata section
+	if (!bPayloadInline && NeedsOffsetFixup())
+	{
+		check(LinkerSummaryBulkDataStartOffset >= 0);
+		InBulkDataOffsetInFile += LinkerSummaryBulkDataStartOffset;
+	}
+
+	// PackagePath does not change, but segment might
+	PackageSegment = EPackageSegment::Header;
+	if (bPayloadInSeparateFile)
+	{
+		if (InBulkDataFlags & BULKDATA_DuplicateNonOptionalPayload)
+		{
+			// Ignore the duplicate and use the non-optional data. The performance benefit of loading from optional data
+			// is not necessary in this edge case of reloading bulk data after a save.
+			PackageSegment = EPackageSegment::BulkDataDefault;
+		}
+		else if (InBulkDataFlags & BULKDATA_OptionalPayload)
+		{
+			PackageSegment = EPackageSegment::BulkDataOptional;
+		}
+		else if (InBulkDataFlags & BULKDATA_MemoryMappedPayload)
+		{
+			PackageSegment = EPackageSegment::BulkDataMemoryMapped;
+		}
+		else if (InBulkDataFlags & BULKDATA_WorkspaceDomainPayload)
+		{
+			// PackageSegment is set to BulkDataDefault for debug name purposes. The segment is not used for loading.
+			PackageSegment = EPackageSegment::BulkDataDefault;
+		}
+		else
+		{
+			PackageSegment = EPackageSegment::BulkDataDefault;
+		}
+	}
+	else
+	{
+		check(PackageSegment == EPackageSegment::Header);
+		// Note that we're assuming the file we wrote on disk is not split into header and exports.
+		// If we do need to support that, we will need to do a callback to update the BulkDataOffsetInFile once we know the header size
+	}
+
+	BulkDataFlags = InBulkDataFlags;
+	BulkDataOffsetInFile = InBulkDataOffsetInFile;
+	BulkDataSizeOnDisk = InBulkDataSizeOnDisk;
+}
+#endif
+
+FCustomVersionContainer FUntypedBulkData::GetCustomVersions(FArchive& InlineArchive)
+{
+	if (!IsInSeparateFile())
+	{
+		return InlineArchive.GetCustomVersions();
+	}
+	else if (!IsInExternalResource())
+	{
+		// The BulkData is in a sidecar file. These files were created with the same custom versions that
+		// the package file containing the BulkData used
+		return InlineArchive.GetCustomVersions();
+	}
+	else
+	{
+		// Read the CustomVersions out of the separate package file
+		TUniquePtr<FArchive> ExternalArchive = IPackageResourceManager::Get().OpenReadExternalResource(
+			EPackageExternalResource::WorkspaceDomainFile, PackagePath.GetPackageName());
+		if (ExternalArchive.IsValid())
+		{
+			FPackageFileSummary PackageFileSummary;
+			*ExternalArchive << PackageFileSummary;
+			if (PackageFileSummary.Tag == PACKAGE_FILE_TAG && !ExternalArchive->IsError())
 			{
-				BulkDataFlags = LocalBulkDataFlags;
-				BulkDataOffsetInFile = LocalBulkDataOffsetInFile;
-				BulkDataSizeOnDisk = LocalBulkDataSizeOnDisk;
+				return PackageFileSummary.GetCustomVersionContainer();
 			}
 		}
+		return FCustomVersionContainer();
 	}
 }
 
@@ -1686,7 +1841,15 @@ void FUntypedBulkData::SerializeBulkData(FArchive& Ar, void* Data)
 
 IAsyncReadFileHandle* FUntypedBulkData::OpenAsyncReadHandle() const
 {
-	return IPackageResourceManager::Get().OpenAsyncReadPackage(GetPackagePath(), GetPackageSegment());
+	if (IsInExternalResource())
+	{
+		return IPackageResourceManager::Get().OpenAsyncReadExternalResource(EPackageExternalResource::WorkspaceDomainFile,
+			GetPackagePath().GetPackageName());
+	}
+	else
+	{
+		return IPackageResourceManager::Get().OpenAsyncReadPackage(GetPackagePath(), GetPackageSegment());
+	}
 }
 
 IBulkDataIORequest* FUntypedBulkData::CreateStreamingRequest(EAsyncIOPriorityAndFlags Priority, FBulkDataIORequestCallBack* CompleteCallback, uint8* UserSuppliedMemory) const
@@ -1710,7 +1873,7 @@ IBulkDataIORequest* FUntypedBulkData::CreateStreamingRequest(int64 OffsetInBulkD
 	if (GEventDrivenLoaderEnabled)
 	{
 		// Verify that the Serialize function converted a Header Segment and Offset to be relative to the Exports segment
-		check(PackageSegment != EPackageSegment::Header);
+		check(IsInExternalResource() || PackageSegment != EPackageSegment::Header);
 		// Even the Exports segment is bad for performance
 		UE_CLOG(PackageSegment == EPackageSegment::Exports,
 			LogSerialization, Error, TEXT("Streaming from the .uexp file '%s' this MUST be in a ubulk instead for best performance."), *PackagePath.GetDebugName(PackageSegment));
@@ -1719,7 +1882,16 @@ IBulkDataIORequest* FUntypedBulkData::CreateStreamingRequest(int64 OffsetInBulkD
 	UE_CLOG(IsStoredCompressedOnDisk(), LogSerialization, Fatal, TEXT("Package level compression is no longer supported (%s)."), *PackagePath.GetDebugName(PackageSegment));
 	UE_CLOG(GetBulkDataSize() <= 0, LogSerialization, Error, TEXT("(%s) has invalid bulk data size."), *PackagePath.GetDebugName(PackageSegment));
 
-	IAsyncReadFileHandle* IORequestHandle = IPackageResourceManager::Get().OpenAsyncReadPackage(PackagePath, PackageSegment);
+	IAsyncReadFileHandle* IORequestHandle;
+	if (IsInExternalResource())
+	{
+		IORequestHandle = IPackageResourceManager::Get().OpenAsyncReadExternalResource(
+			EPackageExternalResource::WorkspaceDomainFile, PackagePath.GetPackageName());
+	}
+	else
+	{
+		IORequestHandle = IPackageResourceManager::Get().OpenAsyncReadPackage(PackagePath, PackageSegment);
+	}
 	check(IORequestHandle); // this generally cannot fail because it is async
 
 	if (IORequestHandle == nullptr)
@@ -1774,6 +1946,8 @@ IBulkDataIORequest* FUntypedBulkData::CreateStreamingRequestForRange(const FPack
 	check(Start.IsValid());
 	check(End.IsValid());
 
+	// TODO: The caller is assuming that their specified PackageSegments are the way in which bulkdata is stored for their package
+	// To allow more flexible bulkdata storage, we will need to eliminate this interface and have them provide an array of bulkdatas.
 	IAsyncReadFileHandle* IORequestHandle = IPackageResourceManager::Get().OpenAsyncReadPackage(PackagePath, PackageSegment);
 	check(IORequestHandle); // this generally cannot fail because it is async
 
@@ -1952,29 +2126,45 @@ bool FUntypedBulkData::TryLoadDataIntoMemory(void* Dest)
 	}
 
 #if WITH_EDITOR
-	checkf( AttachedAr, TEXT( "Attempted to load bulk data without an attached archive. Most likely the bulk data was loaded twice on console, which is not supported" ) );
-
 	TUniquePtr<FArchive> BulkDataLoadedFile;
 	FArchive* BulkDataArchive = nullptr;
-	if (Linker && Linker->GetAsyncLoader() && Linker->GetAsyncLoader()->IsCookedForEDLInEditor() &&
-		(BulkDataFlags & BULKDATA_PayloadInSeperateFile))
+	if (IsInSeparateFile())
 	{
-		// The Serialize function set the PackageSegment, and since the payload is in a separate file, we
-		// can assert that the PackageSegment that was set is one of the separate-file segments
-		check(PackageSegment != EPackageSegment::Header && PackageSegment != EPackageSegment::Exports);
-
-		FOpenPackageResult Result = IPackageResourceManager::Get().OpenReadPackage(PackagePath, PackageSegment);
-		if (!Result.Archive.IsValid() || Result.Format != EPackageFormat::Binary)
+		if (IsInExternalResource())
 		{
-			UE_LOG(LogSerialization, Error, TEXT("Attempted to load bulk data from an invalid PackagePath '%s'%s."),
-				*PackagePath.GetDebugName(PackageSegment), (!Result.Archive.IsValid() ? TEXT("") : TEXT("; package is in non-binary format and this is not supported.")));
-			return false;
+			BulkDataLoadedFile = IPackageResourceManager::Get().OpenReadExternalResource(
+				EPackageExternalResource::WorkspaceDomainFile, PackagePath.GetPackageName());
+			if (!BulkDataLoadedFile.IsValid())
+			{
+				UE_LOG(LogSerialization, Error, TEXT("Attempted to load bulk data from invalid WorkspaceDomain package '%s'."),
+					*PackagePath.GetPackageName());
+				return false;
+			}
+		}
+		else
+		{
+			// The Serialize function set the PackageSegment, and since the payload is in a separate file, we
+			// can assert that the PackageSegment that was set is one of the separate-file segments
+			check(PackageSegment != EPackageSegment::Header && PackageSegment != EPackageSegment::Exports);
+			FOpenPackageResult Result = IPackageResourceManager::Get().OpenReadPackage(PackagePath, PackageSegment);
+			if (!Result.Archive.IsValid() || Result.Format != EPackageFormat::Binary)
+			{
+				UE_LOG(LogSerialization, Error, TEXT("Attempted to load bulk data from an invalid PackagePath '%s'%s."),
+					*PackagePath.GetDebugName(PackageSegment), (!Result.Archive.IsValid() ? TEXT("") : TEXT("; package is in non-binary format and this is not supported.")));
+				return false;
+			}
+			BulkDataLoadedFile = MoveTemp(Result.Archive);
 		}
-		BulkDataLoadedFile = MoveTemp(Result.Archive);
 		BulkDataArchive = BulkDataLoadedFile.Get();
 	}
 	else
 	{
+		if (!AttachedAr)
+		{
+			UE_LOG(LogSerialization, Error, TEXT("Attempted to load bulk data without an attached archive. ")
+				TEXT("Most likely the bulk data was loaded twice on console, which is not supported"));
+			return false;
+		}
 		BulkDataArchive = AttachedAr;
 	}
 
@@ -1992,7 +2182,7 @@ bool FUntypedBulkData::TryLoadDataIntoMemory(void* Dest)
 
 #else
 	bool bWasLoadedSuccessfully = false;
-	if (((BulkDataFlags & BULKDATA_PayloadInSeperateFile) == 0) && IsInAsyncLoadingThread())
+	if (!IsInSeparateFile() && IsInAsyncLoadingThread())
 	{
 		if (UPackage* PackagePtr = Package.Get())
 		{
@@ -2033,26 +2223,41 @@ bool FUntypedBulkData::TryLoadDataIntoMemory(void* Dest)
 		// is it possible for streaming mips to be loaded in non streaming ways.
 		if (CVarTextureStreamingEnabled->GetValueOnAnyThread() != 0)
 		{
-			bool bIsBulkFile = PackageSegment == EPackageSegment::BulkDataDefault;
-			UE_CLOG(GEventDrivenLoaderEnabled && bIsBulkFile && (IsInGameThread() || IsInAsyncLoadingThread()), LogSerialization, Error,
+			UE_CLOG(GEventDrivenLoaderEnabled && IsInSeparateFile() && (IsInGameThread() || IsInAsyncLoadingThread()), LogSerialization, Error,
 				TEXT("Attempt to sync load bulk data with EDL enabled (LoadDataIntoMemory). This is not desireable. File %s"), *PackagePath.GetDebugName(PackageSegment));
 		}
 #endif
 
 		// Verify that the Serialize function converted a Header Segment and Offset to be relative to the Exports segment, if GEventDrivenLoaderEnabled
-		check(!GEventDrivenLoaderEnabled || PackageSegment != EPackageSegment::Header);
+		check(!GEventDrivenLoaderEnabled || IsInExternalResource() || PackageSegment != EPackageSegment::Header);
 
-		FOpenPackageResult Result = IPackageResourceManager::Get().OpenReadPackage(PackagePath, PackageSegment);
-		if (!Result.Archive.IsValid() || Result.Format != EPackageFormat::Binary)
+		TUniquePtr<FArchive> BulkArchive;
+		if (IsInExternalResource())
 		{
-			UE_LOG(LogSerialization, Error, TEXT("Attempted to load bulk data from an invalid PackagePath '%s'%s."),
-				*PackagePath.GetDebugName(PackageSegment), (!Result.Archive.IsValid() ? TEXT("") : TEXT("; package is in non-binary format and this is not supported.")));
-			return false;
+			BulkArchive = IPackageResourceManager::Get().OpenReadExternalResource(
+				EPackageExternalResource::WorkspaceDomainFile, PackagePath.GetPackageName());
+			if (!BulkArchive.IsValid())
+			{
+				UE_LOG(LogSerialization, Error, TEXT("Attempted to load bulk data from invalid WorkspaceDomain package '%s'."),
+					*PackagePath.GetPackageName());
+				return false;
+			}
+		}
+		else
+		{
+			FOpenPackageResult Result = IPackageResourceManager::Get().OpenReadPackage(PackagePath, PackageSegment);
+			if (!Result.Archive.IsValid() || Result.Format != EPackageFormat::Binary)
+			{
+				UE_LOG(LogSerialization, Error, TEXT("Attempted to load bulk data from an invalid PackagePath '%s'%s."),
+					*PackagePath.GetDebugName(PackageSegment), (!Result.Archive.IsValid() ? TEXT("") : TEXT("; package is in non-binary format and this is not supported.")));
+				return false;
+			}
+			BulkArchive = MoveTemp(Result.Archive);
 		}
 	
 		// Seek to the beginning of the bulk data in the file.
-		Result.Archive->Seek(BulkDataOffsetInFile);
-		SerializeBulkData(*Result.Archive, Dest, BulkDataFlags);
+		BulkArchive->Seek(BulkDataOffsetInFile);
+		SerializeBulkData(*BulkArchive, Dest, BulkDataFlags);
 	}
 	return true;
 #endif // WITH_EDITOR
diff --git a/Engine/Source/Runtime/CoreUObject/Private/Serialization/BulkData2.cpp b/Engine/Source/Runtime/CoreUObject/Private/Serialization/BulkData2.cpp
--- a/Engine/Source/Runtime/CoreUObject/Private/Serialization/BulkData2.cpp
+++ b/Engine/Source/Runtime/CoreUObject/Private/Serialization/BulkData2.cpp
@@ -808,6 +808,8 @@ void FBulkDataBase::Serialize(FArchive& Ar, UObject* Owner, int32 /*Index*/, boo
 		{
 			SerializeDuplicateData(Ar, DuplicateFlags, DuplicateSizeOnDisk, DuplicateOffset);
 		}
+		checkf(!(BulkDataFlags & BULKDATA_WorkspaceDomainPayload), TEXT("FBulkDataBase error on %s:")
+			TEXT("FBulkDataBase does not support BULKDATA_WorkspaceDomainPayload"), *Ar.GetArchiveName());
 
 		// We assume that Owner/Package/Linker are all valid. The old BulkData system would
 		// generally fail if any of these were nullptr but had plenty of inconsistent checks
diff --git a/Engine/Source/Runtime/CoreUObject/Private/UObject/PackageResourceManagerFile.cpp b/Engine/Source/Runtime/CoreUObject/Private/UObject/PackageResourceManagerFile.cpp
--- a/Engine/Source/Runtime/CoreUObject/Private/UObject/PackageResourceManagerFile.cpp
+++ b/Engine/Source/Runtime/CoreUObject/Private/UObject/PackageResourceManagerFile.cpp
@@ -45,6 +45,11 @@ public:
 		FPackagePath* OutUpdatedPath = nullptr) override;
 	virtual bool TryMatchCaseOnDisk(const FPackagePath& PackagePath, FPackagePath* OutPackagePath = nullptr) override;
 
+	virtual TUniquePtr<FArchive> OpenReadExternalResource(EPackageExternalResource ResourceType, FStringView Identifier) override;
+	virtual bool DoesExternalResourceExist(EPackageExternalResource ResourceType, FStringView Identifier) override;
+	virtual IAsyncReadFileHandle* OpenAsyncReadExternalResource(
+		EPackageExternalResource ResourceType, FStringView Identifier) override;
+
 	virtual void FindPackagesRecursive(TArray<TPair<FPackagePath, EPackageSegment>>& OutPackages,
 		FStringView PackageMount, FStringView FileMount, FStringView RootRelPath, FStringView BasenameWildcard) override;
 	virtual void IteratePackagesInPath(FStringView PackageMount, FStringView FileMount, FStringView RootRelPath,
@@ -315,6 +320,65 @@ bool FPackageResourceManagerFile::TryMatchCaseOnDisk(const FPackagePath& Package
 	}
 }
 
+TUniquePtr<FArchive> FPackageResourceManagerFile::OpenReadExternalResource(EPackageExternalResource ResourceType, FStringView Identifier)
+{
+	switch (ResourceType)
+	{
+	case EPackageExternalResource::WorkspaceDomainFile:
+	{
+		FPackagePath PackagePath;
+		if (!FPackagePath::TryFromPackageName(Identifier, PackagePath))
+		{
+			return TUniquePtr<FArchive>();
+		}
+		return OpenReadPackage(PackagePath, EPackageSegment::Header).Archive;
+	}
+	default:
+		checkNoEntry();
+		return TUniquePtr<FArchive>();
+	}
+}
+
+bool FPackageResourceManagerFile::DoesExternalResourceExist(EPackageExternalResource ResourceType, FStringView Identifier)
+{
+	switch (ResourceType)
+	{
+	case EPackageExternalResource::WorkspaceDomainFile:
+	{
+		FPackagePath PackagePath;
+		if (!FPackagePath::TryFromPackageName(Identifier, PackagePath))
+		{
+			return false;
+		}
+		return DoesPackageExist(PackagePath, EPackageSegment::Header);
+	}
+	default:
+		checkNoEntry();
+		return false;
+	}
+}
+
+IAsyncReadFileHandle* FPackageResourceManagerFile::OpenAsyncReadExternalResource(EPackageExternalResource ResourceType,
+	FStringView Identifier)
+{
+	switch (ResourceType)
+	{
+	case EPackageExternalResource::WorkspaceDomainFile:
+	{
+		FPackagePath PackagePath;
+		if (!FPackagePath::TryFromPackageName(Identifier, PackagePath))
+		{
+			return new FAsyncReadFileHandleNull();
+		}
+		return OpenAsyncReadPackage(PackagePath, EPackageSegment::Header);
+	}
+	default:
+		checkNoEntry();
+		return new FAsyncReadFileHandleNull();
+	}
+}
+
+
 void FPackageResourceManagerFile::FindPackagesRecursive(TArray<TPair<FPackagePath, EPackageSegment>>& OutPackages,
 	FStringView PackageMount, FStringView FileMount, FStringView RootRelPath, FStringView BasenameWildcard)
 {
diff --git a/Engine/Source/Runtime/CoreUObject/Private/UObject/SavePackage.cpp b/Engine/Source/Runtime/CoreUObject/Private/UObject/SavePackage.cpp
--- a/Engine/Source/Runtime/CoreUObject/Private/UObject/SavePackage.cpp
+++ b/Engine/Source/Runtime/CoreUObject/Private/UObject/SavePackage.cpp
@@ -2052,7 +2052,7 @@ FSavePackageResultStruct UPackage::Save(UPackage* InOuter, UObject* Base, EObjec
 	check(InOuter);
 	check(Filename);
 	const bool bIsCooking = TargetPlatform != nullptr;
-
+	FPackagePath TargetPackagePath = FPackagePath::FromLocalPath(Filename);
 
 #if WITH_EDITOR
 	TMap<UObject*, UObject*> ReplacedImportOuters;
@@ -2065,13 +2065,12 @@ FSavePackageResultStruct UPackage::Save(UPackage* InOuter, UObject* Base, EObjec
 
 	// if the in memory package filename is different the filename we are saving it to,
 	// regenerate a new persistent id for it.
-	FPackagePath TargetPackagePath = FPackagePath::FromLocalPath(Filename);
 	if (!bIsCooking && !InOuter->GetLoadedPath().IsEmpty() && InOuter->GetLoadedPath() != TargetPackagePath && !(SaveFlags & SAVE_FromAutosave))
 	{
 		InOuter->SetPersistentGuid(FGuid::NewGuid());
 	}
 #endif //WITH_EDITOR
-
+	const bool bSavingNewLoadedPath = SavePackageUtilities::IsSavingNewLoadedPath(bIsCooking, TargetPackagePath, SaveFlags);
 	const bool bSavingConcurrent = !!(SaveFlags & ESaveFlags::SAVE_Concurrent);
 
 	if (FPlatformProperties::HasEditorOnlyData())
@@ -2439,6 +2438,7 @@ FSavePackageResultStruct UPackage::Save(UPackage* InOuter, UObject* Base, EObjec
 					}
 					else
 #endif
+					{
 						if (bSaveAsync)
 						{
 							// Allocate the linker with a memory writer, forcing byte swapping if wanted.
@@ -2450,6 +2450,8 @@ FSavePackageResultStruct UPackage::Save(UPackage* InOuter, UObject* Base, EObjec
 							TempFilename = FPaths::CreateTempFilename(*FPaths::ProjectSavedDir(), *BaseFilename.Left(32));
 							Linker = TUniquePtr<FLinkerSave>(new FLinkerSave(InOuter, *TempFilename.GetValue(), bForceByteSwapping, bSaveUnversioned));
 						}
+					}
+					Linker->bSavingNewLoadedPath = bSavingNewLoadedPath;
 
 #if WITH_TEXT_ARCHIVE_SUPPORT
 					if (bTextFormat)
@@ -4148,8 +4150,13 @@ FSavePackageResultStruct UPackage::Save(UPackage* InOuter, UObject* Base, EObjec
 
 				SlowTask.EnterProgressFrame(1, NSLOCTEXT("Core", "SerializingBulkData", "Serializing bulk data"));
 
-				SavePackageUtilities::SaveBulkData(Linker.Get(), InOuter, Filename, TargetPlatform, SavePackageContext, bTextFormat, bDiffing,
+				ESavePackageResult SaveResult = 
+					SavePackageUtilities::SaveBulkData(Linker.Get(), InOuter, Filename, TargetPlatform, SavePackageContext, SaveFlags, bTextFormat, bDiffing,
 							 bComputeHash, AsyncWriteAndHashSequence, TotalPackageSizeUncompressed );
+				if (SaveResult != ESavePackageResult::Success)
+				{
+					return SaveResult;
+				}
 
 				AppendAdditionalData(*Linker);
 
@@ -4639,7 +4646,7 @@ FSavePackageResultStruct UPackage::Save(UPackage* InOuter, UObject* Base, EObjec
 		{
 			// if the save was successful, update the internal package filename path if we aren't currently cooking
 #if WITH_EDITOR
-			if (TargetPlatform == nullptr && TargetPackagePath.IsMountedPath())
+			if (bSavingNewLoadedPath)
 			{
 				InOuter->SetLoadedPath(TargetPackagePath);
 			}
diff --git a/Engine/Source/Runtime/CoreUObject/Private/UObject/SavePackage/SaveContext.h b/Engine/Source/Runtime/CoreUObject/Private/UObject/SavePackage/SaveContext.h
--- a/Engine/Source/Runtime/CoreUObject/Private/UObject/SavePackage/SaveContext.h
+++ b/Engine/Source/Runtime/CoreUObject/Private/UObject/SavePackage/SaveContext.h
@@ -101,6 +101,8 @@ public:
 		{
 			bIsProcessingPrestreamPackages = ProcessPrestreamingRequests->GetInt() > 0;
 		}
+
+		bSavingNewLoadedPath = SavePackageUtilities::IsSavingNewLoadedPath(IsCooking(), TargetPackagePath, SaveArgs.SaveFlags);
 	}
 
 	~FSaveContext()
@@ -187,6 +189,11 @@ public:
 		return SaveArgs.TargetPlatform != nullptr;
 	}
 
+	bool IsSavingNewLoadedPath() const
+	{
+		return bSavingNewLoadedPath;
+	}
+
 	bool IsFilterEditorOnly() const
 	{
 		return Package->HasAnyPackageFlags(PKG_FilterEditorOnly);
@@ -566,6 +573,7 @@ private:
 	bool bIsProcessingPrestreamPackages = false;
 	bool bNeedPreSaveCleanup = false;
 	bool bGenerateFileStub = false;
+	bool bSavingNewLoadedPath = false;
 
 	// Config classes shared with the old Save
 	FSavePackageDiffSettings DiffSettings;
diff --git a/Engine/Source/Runtime/CoreUObject/Private/UObject/SavePackage/SavePackageUtilities.cpp b/Engine/Source/Runtime/CoreUObject/Private/UObject/SavePackage/SavePackageUtilities.cpp
--- a/Engine/Source/Runtime/CoreUObject/Private/UObject/SavePackage/SavePackageUtilities.cpp
+++ b/Engine/Source/Runtime/CoreUObject/Private/UObject/SavePackage/SavePackageUtilities.cpp
@@ -691,6 +691,17 @@ void GetCDOSubobjects(UObject* CDO, TArray<UObject*>& Subobjects)
 	}
 }
 
+bool IsSavingNewLoadedPath(bool bIsCooking, const FPackagePath& TargetPackagePath, uint32 SaveFlags)
+{
+#if WITH_EDITOR
+	return !bIsCooking &&							// Do not update the loadedpath if we're cooking
+		TargetPackagePath.IsMountedPath() &&		// Do not update the loadedpath if the new path is not a viable mounted path
+		!(SaveFlags & SAVE_BulkDataByReference);	// Do not update the loadedpath if we're storing references to the old loadedpath
+#else
+	return false; // Saving when not in editor never updates the LoadedPath
+#endif
+}
+
 } // end namespace SavePackageUtilities
 
 bool IsEditorOnlyObject(const UObject* InObject, bool bCheckRecursive, bool bCheckMarks)
@@ -1793,115 +1804,188 @@ void SaveThumbnails(UPackage* InOuter, FLinkerSave* Linker, FStructuredArchive::
 #endif
 }
 
-void SaveBulkData(FLinkerSave* Linker, const UPackage* InOuter, const TCHAR* Filename, const ITargetPlatform* TargetPlatform,
-				  FSavePackageContext* SavePackageContext, const bool bTextFormat, const bool bDiffing, const bool bComputeHash, TAsyncWorkSequence<FMD5>& AsyncWriteAndHashSequence, int64& TotalPackageSizeUncompressed)
+ESavePackageResult SaveBulkData(FLinkerSave* Linker, const UPackage* InOuter, const TCHAR* Filename, const ITargetPlatform* TargetPlatform,
+				  FSavePackageContext* SavePackageContext, uint32 SaveFlags, const bool bTextFormat, const bool bDiffing,
+				  const bool bComputeHash, TAsyncWorkSequence<FMD5>& AsyncWriteAndHashSequence, int64& TotalPackageSizeUncompressed)
 {
 	// Now we write all the bulkdata that is supposed to be at the end of the package
 	// and fix up the offset
 	const int64 StartOfBulkDataArea = Linker->Tell();
 	Linker->Summary.BulkDataStartOffset = StartOfBulkDataArea;
 
-	check(!bTextFormat || Linker->BulkDataToAppend.Num() == 0);
-
-	if (!bTextFormat && Linker->BulkDataToAppend.Num() > 0)
+	if (Linker->BulkDataToAppend.Num() == 0)
 	{
-		COOK_STAT(FScopedDurationTimer SaveTimer(FSavePackageStats::SerializeBulkDataTimeSec));
+		return ESavePackageResult::Success;;
+	}
+	check(!bTextFormat);
 
-		FScopedSlowTask BulkDataFeedback(Linker->BulkDataToAppend.Num());
+	COOK_STAT(FScopedDurationTimer SaveTimer(FSavePackageStats::SerializeBulkDataTimeSec));
 
-		class FLargeMemoryWriterWithRegions : public FLargeMemoryWriter
-		{
-		public:
-			FLargeMemoryWriterWithRegions()
-				: FLargeMemoryWriter(0, /* IsPersistent */ true)
-			{}
+	FScopedSlowTask BulkDataFeedback(Linker->BulkDataToAppend.Num());
 
-			TArray<FFileRegion> FileRegions;
-		};
+	class FLargeMemoryWriterWithRegions : public FLargeMemoryWriter
+	{
+	public:
+		FLargeMemoryWriterWithRegions()
+			: FLargeMemoryWriter(0, /* IsPersistent */ true)
+		{}
 
-		TUniquePtr<FLargeMemoryWriterWithRegions> BulkArchive;
-		TUniquePtr<FLargeMemoryWriterWithRegions> OptionalBulkArchive;
-		TUniquePtr<FLargeMemoryWriterWithRegions> MappedBulkArchive;
+		TArray<FFileRegion> FileRegions;
+	};
+
+	TUniquePtr<FLargeMemoryWriterWithRegions> BulkArchive;
+	TUniquePtr<FLargeMemoryWriterWithRegions> OptionalBulkArchive;
+	TUniquePtr<FLargeMemoryWriterWithRegions> MappedBulkArchive;
 
-		uint32 ExtraBulkDataFlags = 0;
+	static const struct FCookerUseSeparateSegments
+	{
+		bool bEnable = false;
 
-		static const struct FUseSeparateBulkDataFiles
+		FCookerUseSeparateSegments()
 		{
-			bool bEnable = false;
+			GConfig->GetBool(TEXT("Core.System"), TEXT("UseSeperateBulkDataFiles"), /* out */ bEnable, GEngineIni);
 
-			FUseSeparateBulkDataFiles()
+			if (IsEventDrivenLoaderEnabledInCookedBuilds())
 			{
-				GConfig->GetBool(TEXT("Core.System"), TEXT("UseSeperateBulkDataFiles"), /* out */ bEnable, GEngineIni);
-
-				if (IsEventDrivenLoaderEnabledInCookedBuilds())
-				{
-					// Always split bulk data when splitting cooked files
-					bEnable = true;
-				}
+				// Always split bulk data when splitting cooked files
+				bEnable = true;
 			}
-		} ShouldUseSeparateBulkDataFiles;
+		}
+	} CookerUseSeparateSegments;
+	const bool bSeparateSegmentsEnabled = CookerUseSeparateSegments.bEnable && Linker->IsCooking();
+
+	if (bSeparateSegmentsEnabled)
+	{
+		BulkArchive.Reset(new FLargeMemoryWriterWithRegions);
+		OptionalBulkArchive.Reset(new FLargeMemoryWriterWithRegions);
+		MappedBulkArchive.Reset(new FLargeMemoryWriterWithRegions);
+	}
+	bool bRequestSaveByReference = SaveFlags & SAVE_BulkDataByReference;
 
-		const bool bShouldUseSeparateBulkFile = ShouldUseSeparateBulkDataFiles.bEnable && Linker->IsCooking();
+	bool bAlignBulkData = false;
+	bool bUseFileRegions = false;
+	int64 BulkDataAlignment = 0;
 
-		if (bShouldUseSeparateBulkFile)
+	if (TargetPlatform)
+	{
+		bAlignBulkData = TargetPlatform->SupportsFeature(ETargetPlatformFeatures::MemoryMappedFiles);
+		bUseFileRegions = TargetPlatform->SupportsFeature(ETargetPlatformFeatures::CookFileRegionMetadata);
+		BulkDataAlignment = TargetPlatform->GetMemoryMappingAlignment();
+	}
+	if (bRequestSaveByReference)
+	{
+		const TCHAR* FailureReason = nullptr;
+		if (SavePackageContext != nullptr && (SavePackageContext->BulkDataManifest != nullptr))
 		{
-			ExtraBulkDataFlags = BULKDATA_PayloadInSeperateFile;
-
-			BulkArchive.Reset(new FLargeMemoryWriterWithRegions);
-			OptionalBulkArchive.Reset(new FLargeMemoryWriterWithRegions);
-			MappedBulkArchive.Reset(new FLargeMemoryWriterWithRegions);
+			FailureReason = TEXT("SAVE_BulkDataByReference is incompatible with BulkDataManifest");
 		}
-
-		// If we are using IoStore (implemented in AsyncLoader2), then we will have no linkers and can not manipulate the
-		// offset as we will not have the Linker's Summary.BulkDataStartOffset information at runtime.
-		// This means that manipulated offsets are incompatible with IoStore.
-		// We now by default do not manipulate the offset for separate files, but we need to optionally leave it on
-		// to prevent larger patching sizes. In the future we may want to remove it even for end-of-file BulkData.
-		// Note that bForceLegacyOffsets is incompatible with using IoStore.
-		if (SavePackageContext != nullptr && SavePackageContext->bForceLegacyOffsets == false && bShouldUseSeparateBulkFile)
+		else if (Linker->bSavingNewLoadedPath)
+		{
+			FailureReason = TEXT("SAVE_BulkDataByReference is incompatible with bSavingNewLoadedPath");
+		}
+		else if (SavePackageContext != nullptr && SavePackageContext->bForceLegacyOffsets)
 		{
-			ExtraBulkDataFlags |= BULKDATA_NoOffsetFixUp;
+			FailureReason = TEXT("SAVE_BulkDataByReference is incompatible with bForceLegacyOffsets");
 		}
+		if (FailureReason)
+		{
+			UE_LOG(LogSavePackage, Error, TEXT("SaveBulkData failed for %s: %s."), Filename, FailureReason);
+			return ESavePackageResult::Error;
+		}
+	}
 
-		bool bAlignBulkData = false;
-		bool bUseFileRegions = false;
-		int64 BulkDataAlignment = 0;
+	for (FLinkerSave::FBulkDataStorageInfo& BulkDataStorageInfo : Linker->BulkDataToAppend)
+	{
+		BulkDataFeedback.EnterProgressFrame();
 
-		if (TargetPlatform)
+		// Set bulk data flags to what they were during initial serialization (they might have changed after that)
+		uint32 BulkDataFlags = BulkDataStorageInfo.BulkDataFlags;
+		FUntypedBulkData* BulkData = BulkDataStorageInfo.BulkData;
+		checkf(BulkDataFlags& BULKDATA_PayloadAtEndOfFile, TEXT("Inlined BulkData data should not have been added to BulkDataToAppend"));
+
+		const bool bBulkItemIsOptional = (BulkDataFlags & BULKDATA_OptionalPayload) != 0;
+		bool bBulkItemIsMapped = bAlignBulkData && ((BulkDataFlags & BULKDATA_MemoryMappedPayload) != 0);
+
+		if (bBulkItemIsMapped && bBulkItemIsOptional)
 		{
-			bAlignBulkData = TargetPlatform->SupportsFeature(ETargetPlatformFeatures::MemoryMappedFiles);
-			bUseFileRegions = TargetPlatform->SupportsFeature(ETargetPlatformFeatures::CookFileRegionMetadata);
-			BulkDataAlignment = TargetPlatform->GetMemoryMappingAlignment();
+			UE_LOG(LogSavePackage, Warning, TEXT("%s has bulk data that is both mapped and optional. This is not currently supported. Will not be mapped."), Filename);
+			BulkDataFlags &= ~BULKDATA_MemoryMappedPayload;
+			bBulkItemIsMapped = false;
 		}
-
-		uint16 BulkDataIndex = 1;
-		for (FLinkerSave::FBulkDataStorageInfo& BulkDataStorageInfo : Linker->BulkDataToAppend)
+		if (bBulkItemIsMapped && bRequestSaveByReference)
 		{
-			BulkDataFeedback.EnterProgressFrame();
-
-			// Set bulk data flags to what they were during initial serialization (they might have changed after that)
-			const uint32 OldBulkDataFlags = BulkDataStorageInfo.BulkData->GetBulkDataFlags();
-			uint32 ModifiedBulkDataFlags = BulkDataStorageInfo.BulkDataFlags | ExtraBulkDataFlags;
-			const bool bBulkItemIsOptional = (ModifiedBulkDataFlags & BULKDATA_OptionalPayload) != 0;
-			bool bBulkItemIsMapped = bAlignBulkData && ((ModifiedBulkDataFlags & BULKDATA_MemoryMappedPayload) != 0);
+			UE_LOG(LogSavePackage, Warning, TEXT("%s has bulk data that is mapped, but the save method is SAVE_BulkDataByReference.")
+				TEXT("This is not currently supported. Will not be mapped."), Filename);
+			BulkDataFlags &= ~BULKDATA_MemoryMappedPayload;
+			bBulkItemIsMapped = false;
+		}
 
-			if (bBulkItemIsMapped && bBulkItemIsOptional)
+		enum ESaveLocation
+		{
+			Invalid,
+			EndOfPackage,
+			SeparateSegment,
+			Reference,
+		};
+		ESaveLocation SaveLocation = ESaveLocation::Invalid;
+		if (SaveLocation == ESaveLocation::Invalid && bRequestSaveByReference)
+		{
+			if (BulkData->GetBulkDataOffsetInFile() != INDEX_NONE &&
+				// We don't support yet loading from a separate file
+				!BulkData->IsInSeparateFile() &&
+				// It is possible to have a BulkData marked as optional without putting it into a separate file, and we
+				// assume that if BulkData is optional and in a separate file, then it is in the BulkDataOptional
+				// segment. Rather than changing that assumption to support optional ExternalResource bulkdata, we
+				// instead require that optional inlined/endofpackagedata BulkDatas can not be read from an
+				// ExternalResource and must remain inline.
+				!BulkData->IsOptional() &&
+				// Inline or end-of-package-file data can only be loaded from the workspace domain package file if the
+				// archive used by the bulk data was actually from the package file; BULKDATA_LazyLoadable is set by
+				// Serialize iff that is the case										
+				(BulkData->GetBulkDataFlags() & BULKDATA_LazyLoadable) 
+				)
 			{
-				UE_LOG(LogSavePackage, Warning, TEXT("%s has bulk data that is both mapped and optional. This is not currently supported. Will not be mapped."), Filename);
-				ModifiedBulkDataFlags &= ~BULKDATA_MemoryMappedPayload;
-				bBulkItemIsMapped = false;
+				SaveLocation = ESaveLocation::Reference;
 			}
+		}
+		if (SaveLocation == ESaveLocation::Invalid && bSeparateSegmentsEnabled)
+		{
+			SaveLocation = ESaveLocation::SeparateSegment;
+		}
+		if (SaveLocation == ESaveLocation::Invalid)
+		{
+			SaveLocation = ESaveLocation::EndOfPackage;
+		}
 
-			// TODO: BulkData Serialize no longer requires BulkData flags to be set; it takes the BulkDataFlags for the save as arguments.
-			//       Remove this unnecessary set of the BulkData's BulkDataFlags.
-			BulkDataStorageInfo.BulkData->ClearBulkDataFlags(0xFFFFFFFF);
-			BulkDataStorageInfo.BulkData->SetBulkDataFlags(ModifiedBulkDataFlags);
+		int64 BulkDataOffsetInFile;
+		int64 BulkDataSizeOnDisk;
+		if (SaveLocation == ESaveLocation::Reference)
+		{
+			BulkDataFlags |= BULKDATA_PayloadInSeperateFile | BULKDATA_WorkspaceDomainPayload;
+			BulkDataFlags |= BULKDATA_NoOffsetFixUp; // We don't use legacy offset fixups when referencing
+			BulkDataFlags &= ~BULKDATA_MemoryMappedPayload; // We don't support memory mapping from referenced
 
+			BulkDataOffsetInFile = BulkData->GetBulkDataOffsetInFile();
+			BulkDataSizeOnDisk = BulkData->GetBulkDataSizeOnDisk();
+		}
+		else
+		{
 			TArray<FFileRegion>* TargetRegions = &Linker->FileRegions;
 			FArchive* TargetArchive = Linker;
-
-			if (bShouldUseSeparateBulkFile)
+			if (SaveLocation == ESaveLocation::SeparateSegment)
 			{
+				BulkDataFlags |= BULKDATA_PayloadInSeperateFile;
+				// If we are using IoStore (implemented in AsyncLoader2), then we will have no linkers and can not manipulate the
+				// offset as we will not have the Linker's Summary.BulkDataStartOffset information at runtime.
+				// This means that manipulated offsets are incompatible with IoStore.
+				// We now by default do not manipulate the offset for separate files, but we need to optionally leave it on
+				// to prevent larger patching sizes. In the future we may want to remove it even for end-of-file BulkData.
+				// Note that bForceLegacyOffsets is incompatible with using IoStore.
+				if (SavePackageContext != nullptr && SavePackageContext->bForceLegacyOffsets == false)
+				{
+					BulkDataFlags |= BULKDATA_NoOffsetFixUp;
+				}
+
 				if (bBulkItemIsOptional)
 				{
 					TargetArchive = OptionalBulkArchive.Get();
@@ -1951,25 +2035,12 @@ void SaveBulkData(FLinkerSave* Linker, const UPackage* InOuter, const TCHAR* Fil
 			}
 
 			const int64 BulkStartOffset = TargetArchive->Tell();
+			BulkData->SerializeBulkData(*TargetArchive, BulkData->Lock(LOCK_READ_ONLY), static_cast<EBulkDataFlags>(BulkDataFlags));
+			BulkData->Unlock();
+			const int64 BulkEndOffset = TargetArchive->Tell();
 
-			int64 StoredBulkStartOffset = (ModifiedBulkDataFlags & BULKDATA_NoOffsetFixUp) == 0 ? BulkStartOffset - StartOfBulkDataArea : BulkStartOffset;
-
-			BulkDataStorageInfo.BulkData->SerializeBulkData(*TargetArchive, BulkDataStorageInfo.BulkData->Lock(LOCK_READ_ONLY),
-				static_cast<EBulkDataFlags>(ModifiedBulkDataFlags));
-
-			int64 BulkEndOffset = TargetArchive->Tell();
-			const int64 LinkerEndOffset = Linker->Tell();
-
-			int64 SizeOnDisk = BulkEndOffset - BulkStartOffset;
-
-			Linker->Seek(BulkDataStorageInfo.BulkDataFlagsPos);
-			*Linker << ModifiedBulkDataFlags;
-
-			Linker->Seek(BulkDataStorageInfo.BulkDataOffsetInFilePos);
-			*Linker << StoredBulkStartOffset;
-
-			Linker->Seek(BulkDataStorageInfo.BulkDataSizeOnDiskPos);
-			SerializeBulkDataSizeInt(*Linker, SizeOnDisk, static_cast<EBulkDataFlags>(ModifiedBulkDataFlags));
+			BulkDataOffsetInFile = (BulkDataFlags & BULKDATA_NoOffsetFixUp) == 0 ? BulkStartOffset - StartOfBulkDataArea : BulkStartOffset;
+			BulkDataSizeOnDisk = BulkEndOffset - BulkStartOffset;
 
 			if (SavePackageContext != nullptr && SavePackageContext->BulkDataManifest != nullptr)
 			{
@@ -1979,114 +2050,122 @@ void SaveBulkData(FLinkerSave* Linker, const UPackage* InOuter, const TCHAR* Fil
 					{
 						return FPackageStoreBulkDataManifest::EBulkdataType::MemoryMapped;
 					}
-
-					if (BulkDataFlags & BULKDATA_OptionalPayload)
+					else if (BulkDataFlags & BULKDATA_OptionalPayload)
 					{
 						return FPackageStoreBulkDataManifest::EBulkdataType::Optional;
 					}
-
-					return FPackageStoreBulkDataManifest::EBulkdataType::Normal;
+					else
+					{
+						check(!(BulkDataFlags & BULKDATA_WorkspaceDomainPayload));
+						return FPackageStoreBulkDataManifest::EBulkdataType::Normal;
+					}
 				};
 
 				const FPackageStoreBulkDataManifest::EBulkdataType Type = BulkDataTypeFromFlags(BulkDataStorageInfo.BulkDataFlags);
-
-				SavePackageContext->BulkDataManifest->AddFileAccess(Filename, Type, StoredBulkStartOffset, BulkStartOffset, SizeOnDisk);
+				SavePackageContext->BulkDataManifest->AddFileAccess(Filename, Type,
+					BulkDataOffsetInFile /* Use the serialized offset as the ChunkId */, BulkStartOffset, BulkDataSizeOnDisk);
 			}
 
-			if (bUseFileRegions && BulkDataStorageInfo.BulkDataFileRegionType != EFileRegionType::None && SizeOnDisk > 0)
+			if (bUseFileRegions && BulkDataStorageInfo.BulkDataFileRegionType != EFileRegionType::None && BulkDataSizeOnDisk > 0)
 			{
-				TargetRegions->Add(FFileRegion(BulkStartOffset, SizeOnDisk, BulkDataStorageInfo.BulkDataFileRegionType));
+				TargetRegions->Add(FFileRegion(BulkStartOffset, BulkDataSizeOnDisk, BulkDataStorageInfo.BulkDataFileRegionType));
 			}
+		}
 
-			Linker->Seek(LinkerEndOffset);
-
-			// Restore BulkData flags to before serialization started
-			BulkDataStorageInfo.BulkData->Unlock();
-			BulkDataStorageInfo.BulkData->ClearBulkDataFlags(0xFFFFFFFF);
-			BulkDataStorageInfo.BulkData->SetBulkDataFlags(OldBulkDataFlags);
+		const int64 SavedLinkerOffset = Linker->Tell();
+		Linker->Seek(BulkDataStorageInfo.BulkDataFlagsPos);
+		*Linker << BulkDataFlags;
+		Linker->Seek(BulkDataStorageInfo.BulkDataSizeOnDiskPos);
+		SerializeBulkDataSizeInt(*Linker, BulkDataSizeOnDisk, static_cast<EBulkDataFlags>(BulkDataFlags));
+		Linker->Seek(BulkDataStorageInfo.BulkDataOffsetInFilePos);
+		*Linker << BulkDataOffsetInFile;
+		Linker->Seek(SavedLinkerOffset);
 
-			// If we are overwriting the LoadedPath for the current package, the bulk data flags need to be updated to match
-			// the values set in the package on disk. This function is responsible for setting some of those flags, so set them now
-			if (!TargetPlatform)
-			{
-				BulkDataStorageInfo.BulkData->SetBulkDataFlags(ExtraBulkDataFlags);
-			}
+#if WITH_EDITOR
+		// If we are overwriting the LoadedPath for the current package, the bulk data flags and location need to be updated to match
+		// the values set in the package on disk.
+		if (Linker->bSavingNewLoadedPath)
+		{
+			BulkData->SetFlagsFromDiskWrittenValues(static_cast<EBulkDataFlags>(BulkDataFlags), BulkDataOffsetInFile,
+				BulkDataSizeOnDisk, StartOfBulkDataArea);
 		}
+#endif
+	}
 
-		if (BulkArchive)
-		{
-			check(OptionalBulkArchive);
-			check(MappedBulkArchive);
+	if (BulkArchive)
+	{
+		check(OptionalBulkArchive);
+		check(MappedBulkArchive);
 
-			const bool bWriteBulkToDisk = !bDiffing;
+		const bool bWriteBulkToDisk = !bDiffing;
 
-			if (SavePackageContext != nullptr && SavePackageContext->PackageStoreWriter != nullptr && bWriteBulkToDisk)
+		if (SavePackageContext != nullptr && SavePackageContext->PackageStoreWriter != nullptr && bWriteBulkToDisk)
+		{
+			auto AddSizeAndConvertToIoBuffer = [&TotalPackageSizeUncompressed](FLargeMemoryWriter* Writer)
 			{
-				auto AddSizeAndConvertToIoBuffer = [&TotalPackageSizeUncompressed](FLargeMemoryWriter* Writer)
-				{
-					const int64 TotalSize = Writer->TotalSize();
-					TotalPackageSizeUncompressed += TotalSize;
-					return FIoBuffer(FIoBuffer::AssumeOwnership, Writer->ReleaseOwnership(), TotalSize);
-				};
+				const int64 TotalSize = Writer->TotalSize();
+				TotalPackageSizeUncompressed += TotalSize;
+				return FIoBuffer(FIoBuffer::AssumeOwnership, Writer->ReleaseOwnership(), TotalSize);
+			};
 
-				IPackageStoreWriter::FBulkDataInfo BulkInfo;
-				BulkInfo.PackageName = InOuter->GetFName();
+			IPackageStoreWriter::FBulkDataInfo BulkInfo;
+			BulkInfo.PackageName = InOuter->GetFName();
 				
-				if (BulkArchive->TotalSize())
-				{
-					BulkInfo.BulkdataType = IPackageStoreWriter::FBulkDataInfo::Standard;
-					BulkInfo.LooseFilePath = FPaths::ChangeExtension(Filename, TEXT(".ubulk"));
-					SavePackageContext->PackageStoreWriter->WriteBulkdata(BulkInfo, AddSizeAndConvertToIoBuffer(BulkArchive.Get()), BulkArchive->FileRegions);
-				}
-				if (OptionalBulkArchive->TotalSize())
-				{
-					BulkInfo.BulkdataType = IPackageStoreWriter::FBulkDataInfo::Optional;
-					BulkInfo.LooseFilePath = FPaths::ChangeExtension(Filename, TEXT(".uptnl"));
-					SavePackageContext->PackageStoreWriter->WriteBulkdata(BulkInfo, AddSizeAndConvertToIoBuffer(OptionalBulkArchive.Get()), OptionalBulkArchive->FileRegions);
-				}
-				if (MappedBulkArchive->TotalSize())
-				{
-					BulkInfo.BulkdataType = IPackageStoreWriter::FBulkDataInfo::Mmap;
-					BulkInfo.LooseFilePath = FPaths::ChangeExtension(Filename, TEXT(".m.ubulk"));
-					SavePackageContext->PackageStoreWriter->WriteBulkdata(BulkInfo, AddSizeAndConvertToIoBuffer(MappedBulkArchive.Get()), MappedBulkArchive->FileRegions);
-				}
+			if (BulkArchive->TotalSize())
+			{
+				BulkInfo.BulkdataType = IPackageStoreWriter::FBulkDataInfo::Standard;
+				BulkInfo.LooseFilePath = FPaths::ChangeExtension(Filename, LexToString(EPackageExtension::BulkDataDefault));
+				SavePackageContext->PackageStoreWriter->WriteBulkdata(BulkInfo, AddSizeAndConvertToIoBuffer(BulkArchive.Get()), BulkArchive->FileRegions);
 			}
-			else
+			if (OptionalBulkArchive->TotalSize())
 			{
-				auto WriteBulkData = [&](FLargeMemoryWriterWithRegions* Archive, const TCHAR* BulkFileExtension)
+				BulkInfo.BulkdataType = IPackageStoreWriter::FBulkDataInfo::Optional;
+				BulkInfo.LooseFilePath = FPaths::ChangeExtension(Filename, LexToString(EPackageExtension::BulkDataOptional));
+				SavePackageContext->PackageStoreWriter->WriteBulkdata(BulkInfo, AddSizeAndConvertToIoBuffer(OptionalBulkArchive.Get()), OptionalBulkArchive->FileRegions);
+			}
+			if (MappedBulkArchive->TotalSize())
+			{
+				BulkInfo.BulkdataType = IPackageStoreWriter::FBulkDataInfo::Mmap;
+				BulkInfo.LooseFilePath = FPaths::ChangeExtension(Filename, LexToString(EPackageExtension::BulkDataMemoryMapped));
+				SavePackageContext->PackageStoreWriter->WriteBulkdata(BulkInfo, AddSizeAndConvertToIoBuffer(MappedBulkArchive.Get()), MappedBulkArchive->FileRegions);
+			}
+		}
+		else
+		{
+			auto WriteBulkData = [&](FLargeMemoryWriterWithRegions* Archive, const TCHAR* BulkFileExtension)
+			{
+				if (const int64 DataSize = Archive->TotalSize())
 				{
-					if (const int64 DataSize = Archive->TotalSize())
+					TotalPackageSizeUncompressed += DataSize;
+
+					if (bComputeHash || bWriteBulkToDisk)
 					{
-						TotalPackageSizeUncompressed += DataSize;
+						FLargeMemoryPtr DataPtr(Archive->ReleaseOwnership());
+
+						const FString ArchiveFilename = FPaths::ChangeExtension(Filename, BulkFileExtension);
 
-						if (bComputeHash || bWriteBulkToDisk)
+						EAsyncWriteOptions WriteOptions(EAsyncWriteOptions::None);
+						if (bComputeHash)
 						{
-							FLargeMemoryPtr DataPtr(Archive->ReleaseOwnership());
-
-							const FString ArchiveFilename = FPaths::ChangeExtension(Filename, BulkFileExtension);
-
-							EAsyncWriteOptions WriteOptions(EAsyncWriteOptions::None);
-							if (bComputeHash)
-							{
-								WriteOptions |= EAsyncWriteOptions::ComputeHash;
-							}
-							if (bWriteBulkToDisk)
-							{
-								WriteOptions |= EAsyncWriteOptions::WriteFileToDisk;
-							}
-							SavePackageUtilities::AsyncWriteFile(AsyncWriteAndHashSequence, MoveTemp(DataPtr), DataSize, *ArchiveFilename, WriteOptions, Archive->FileRegions);
+							WriteOptions |= EAsyncWriteOptions::ComputeHash;
 						}
+						if (bWriteBulkToDisk)
+						{
+							WriteOptions |= EAsyncWriteOptions::WriteFileToDisk;
+						}
+						SavePackageUtilities::AsyncWriteFile(AsyncWriteAndHashSequence, MoveTemp(DataPtr), DataSize, *ArchiveFilename, WriteOptions, Archive->FileRegions);
 					}
-				};
+				}
+			};
 
-				WriteBulkData(BulkArchive.Get(), TEXT(".ubulk"));			// Regular separate bulk data file
-				WriteBulkData(OptionalBulkArchive.Get(), TEXT(".uptnl"));	// Optional bulk data
-				WriteBulkData(MappedBulkArchive.Get(), TEXT(".m.ubulk"));	// Memory-mapped bulk data
-			}
+			WriteBulkData(BulkArchive.Get(), LexToString(EPackageExtension::BulkDataDefault));
+			WriteBulkData(OptionalBulkArchive.Get(), LexToString(EPackageExtension::BulkDataOptional));
+			WriteBulkData(MappedBulkArchive.Get(), LexToString(EPackageExtension::BulkDataMemoryMapped));
 		}
 	}
 
 	Linker->BulkDataToAppend.Empty();
+	return ESavePackageResult::Success;
 }
 
 void SaveWorldLevelInfo(UPackage* InOuter, FLinkerSave* Linker, FStructuredArchive::FRecord Record)
diff --git a/Engine/Source/Runtime/CoreUObject/Private/UObject/SavePackage/SavePackageUtilities.h b/Engine/Source/Runtime/CoreUObject/Private/UObject/SavePackage/SavePackageUtilities.h
--- a/Engine/Source/Runtime/CoreUObject/Private/UObject/SavePackage/SavePackageUtilities.h
+++ b/Engine/Source/Runtime/CoreUObject/Private/UObject/SavePackage/SavePackageUtilities.h
@@ -12,6 +12,7 @@
 #include "Serialization/ArchiveStackTrace.h"
 #include "Serialization/FileRegions.h"
 #include "UObject/NameTypes.h"
+#include "UObject/Package.h"
 #include "UObject/UObjectMarks.h"
 
 // This file contains private utilities shared by UPackage::Save and UPackage::Save2 
@@ -341,8 +342,9 @@ namespace SavePackageUtilities
 	void DecrementOutstandingAsyncWrites();
 
 	void SaveThumbnails(UPackage* InOuter, FLinkerSave* Linker, FStructuredArchive::FSlot Slot);
-	void SaveBulkData(FLinkerSave* Linker, const UPackage* InOuter, const TCHAR* Filename, const ITargetPlatform* TargetPlatform,
-		FSavePackageContext* SavePackageContext, const bool bTextFormat, const bool bDiffing, const bool bComputeHash, TAsyncWorkSequence<FMD5>& AsyncWriteAndHashSequence, int64& TotalPackageSizeUncompressed);
+	ESavePackageResult SaveBulkData(FLinkerSave* Linker, const UPackage* InOuter, const TCHAR* Filename, const ITargetPlatform* TargetPlatform,
+		FSavePackageContext* SavePackageContext, uint32 SaveFlags, const bool bTextFormat, const bool bDiffing,
+		const bool bComputeHash, TAsyncWorkSequence<FMD5>& AsyncWriteAndHashSequence, int64& TotalPackageSizeUncompressed);
 	void SaveWorldLevelInfo(UPackage* InOuter, FLinkerSave* Linker, FStructuredArchive::FRecord Record);
 	EObjectMark GetExcludedObjectMarksForTargetPlatform(const class ITargetPlatform* TargetPlatform);
 	bool HasUnsaveableOuter(UObject* InObj, UPackage* InSavingPackage);
@@ -356,6 +358,12 @@ namespace SavePackageUtilities
 	void AsyncWriteFileWithSplitExports(TAsyncWorkSequence<FMD5>& AsyncWriteAndHashSequence, FLargeMemoryPtr Data, const int64 DataSize, const int64 HeaderSize, const TCHAR* Filename, EAsyncWriteOptions Options, TArrayView<const FFileRegion> InFileRegions);
 
 	void GetCDOSubobjects(UObject* CDO, TArray<UObject*>& Subobjects);
+
+	/**
+	 * Return whether the given save parameters indicate the LoadedPath of the package being saved should be updated.
+	 * This allows us to update the in-memory package when it is saved in editor to match its new save file.
+	 */
+	bool IsSavingNewLoadedPath(bool bIsCooking, const FPackagePath& TargetPackagePath, uint32 SaveFlags);
 }
 
 #if ENABLE_COOK_STATS
diff --git a/Engine/Source/Runtime/CoreUObject/Private/UObject/SavePackage2.cpp b/Engine/Source/Runtime/CoreUObject/Private/UObject/SavePackage2.cpp
--- a/Engine/Source/Runtime/CoreUObject/Private/UObject/SavePackage2.cpp
+++ b/Engine/Source/Runtime/CoreUObject/Private/UObject/SavePackage2.cpp
@@ -700,6 +700,7 @@ ESavePackageResult CreateLinker(FSaveContext& SaveContext)
 				SaveContext.Linker = MakeUnique<FLinkerSave>(SaveContext.GetPackage(), *SaveContext.TempFilename.GetValue(), SaveContext.IsForceByteSwapping(), SaveContext.IsSaveUnversioned());
 			}
 		}
+		SaveContext.Linker->bSavingNewLoadedPath = SaveContext.IsSavingNewLoadedPath();
 
 #if WITH_TEXT_ARCHIVE_SUPPORT
 		if (SaveContext.IsTextFormat())
@@ -1923,7 +1924,7 @@ void PostSavePackage(FSaveContext& SaveContext)
 	// update the internal package filename path if we're saving to a valid mounted path and we aren't currently cooking
 #if WITH_EDITOR
 	const FPackagePath& PackagePath = SaveContext.GetTargetPackagePath();
-	if (!SaveContext.IsCooking() && PackagePath.IsMountedPath())
+	if (SaveContext.IsSavingNewLoadedPath())
 	{
 		Package->SetLoadedPath(PackagePath);
 	}
@@ -2035,8 +2036,13 @@ ESavePackageResult InnerSave(FSaveContext& SaveContext)
 
 	// Save Bulk Data
 	SlowTask.EnterProgressFrame();
-	SavePackageUtilities::SaveBulkData(SaveContext.Linker.Get(), SaveContext.GetPackage(), SaveContext.GetFilename(), SaveContext.GetTargetPlatform(), SaveContext.GetSavePackageContext(), SaveContext.IsTextFormat(), SaveContext.IsDiffing(),
+	SaveContext.Result = SavePackageUtilities::SaveBulkData(SaveContext.Linker.Get(), SaveContext.GetPackage(), SaveContext.GetFilename(), SaveContext.GetTargetPlatform(),
+		SaveContext.GetSavePackageContext(), SaveContext.GetSaveArgs().SaveFlags, SaveContext.IsTextFormat(), SaveContext.IsDiffing(),
 		SaveContext.IsComputeHash(), SaveContext.AsyncWriteAndHashSequence, SaveContext.TotalPackageSizeUncompressed);
+	if (SaveContext.Result != ESavePackageResult::Success)
+	{
+		return SaveContext.Result;
+	}
 
 	// Write Additional files from export
 	SlowTask.EnterProgressFrame();
diff --git a/Engine/Source/Runtime/CoreUObject/Public/Serialization/BulkData.h b/Engine/Source/Runtime/CoreUObject/Public/Serialization/BulkData.h
--- a/Engine/Source/Runtime/CoreUObject/Public/Serialization/BulkData.h
+++ b/Engine/Source/Runtime/CoreUObject/Public/Serialization/BulkData.h
@@ -33,6 +33,7 @@
 
 class IMappedFileHandle;
 class IMappedFileRegion;
+class UPackage;
 
 /*-----------------------------------------------------------------------------
 	Base version of untyped bulk data.
@@ -444,6 +445,15 @@ public:
 		return	(GetBulkDataFlags() & BULKDATA_PayloadInSeperateFile) != 0;
 	}
 
+	/**
+	 * Returns whether this bulk data is stored in a PackageExternalResource rather than in
+	 * a neighboring segment of its owner's PackagePath. 
+	 */
+	bool IsInExternalResource() const
+	{
+		return IsInSeparateFile() && (GetBulkDataFlags() & BULKDATA_WorkspaceDomainPayload);
+	}
+
 	/**
 	* Returns whether this bulk data is accessed via the IoDispatcher or not.
 	* @return false as the old BulkData API does not support it
@@ -649,6 +659,27 @@ public:
 	UE_DEPRECATED(5.0, "Use the version that takes InBulkDataFlags")
 	void SerializeBulkData( FArchive& Ar, void* Data );
 
+	/**
+	 * Get the CustomVersions used in the file containing the BulkData payload. If !IsInSeparateFile, this will be
+	 * the custom versions from the archive used to serialize the FBulkDataInterface, which the caller must provide.
+	 * Otherwise, the CustomVersions come from the separate file and this function will look them up.
+	 * 
+	 * @param InlineArchive The archive that was used to load this BulkData object
+	 * 
+	 * @return The CustomVersions that apply to the interpretation of the BulkData's payload.
+	 */
+	FCustomVersionContainer GetCustomVersions(FArchive& InlineArchive);
+
+#if WITH_EDITOR
+	/**
+	 * When saving BulkData, if we are overwriting the file we need to update the BulkData's (flags,offset,size) to be
+	 * able to load from the new file. But SerializeBulkData modifies those values when loading, so the in-memory values after
+	 * loading from disk are not the same as the values on disk.
+	 * This function handles running the same steps that SerializeBulkData does, but skips the deserialization of the BulkData.
+	 * */
+	void SetFlagsFromDiskWrittenValues(EBulkDataFlags InBulkDataFlags, int64 InBulkDataOffsetInFile, int64 InBulkDataSizeOnDisk, int64 LinkerSummaryBulkDataStartOffset);
+#endif
+
 	/*-----------------------------------------------------------------------------
 		Async Streaming Interface.
 	-----------------------------------------------------------------------------*/
diff --git a/Engine/Source/Runtime/CoreUObject/Public/Serialization/BulkDataCommon.h b/Engine/Source/Runtime/CoreUObject/Public/Serialization/BulkDataCommon.h
--- a/Engine/Source/Runtime/CoreUObject/Public/Serialization/BulkDataCommon.h
+++ b/Engine/Source/Runtime/CoreUObject/Public/Serialization/BulkDataCommon.h
@@ -57,6 +57,16 @@ enum EBulkDataFlags : uint32
 	BULKDATA_BadDataVersion = 1 << 15,
 	/** BulkData did not have it's offset changed during the cook and does not need the fix up at load time */
 	BULKDATA_NoOffsetFixUp = 1 << 16,
+	/**
+	 * INTERNAL SET ONLY - callers of bulkdata should not set this flag on the bulk data
+	 * If set, payload is stored in the workspace domain version of the file.
+	 */
+	BULKDATA_WorkspaceDomainPayload = 1 << 17,
+	/**
+	 * INTERNAL SET ONLY - callers of bulkdata should not set this flag on the bulk data
+	 * If true, the BulkData can be loaded from its file at any time
+	 */
+	BULKDATA_LazyLoadable = 1 << 18,
 
 	/* Runtime only flags below this point! Note that they take the high bits in reverse order! */
 
diff --git a/Engine/Source/Runtime/CoreUObject/Public/UObject/LinkerSave.h b/Engine/Source/Runtime/CoreUObject/Public/UObject/LinkerSave.h
--- a/Engine/Source/Runtime/CoreUObject/Public/UObject/LinkerSave.h
+++ b/Engine/Source/Runtime/CoreUObject/Public/UObject/LinkerSave.h
@@ -77,6 +77,12 @@ public:
 	 */
 	TArray<AdditionalDataCallback> AdditionalDataToAppend;
 
+	/**
+	 * Stores whether the target of the save is the path PackagePath that already is/will be stored in the Package's LoadedPath
+	 * This is used e.g. to decide whether to update the in-memory file offsets for BulkData.
+	 */
+	bool bSavingNewLoadedPath = false;
+
 	/** A mapping of package name to generated script SHA keys */
 	COREUOBJECT_API static TMap<FString, TArray<uint8> > PackagesToScriptSHAMap;
 
diff --git a/Engine/Source/Runtime/CoreUObject/Public/UObject/ObjectMacros.h b/Engine/Source/Runtime/CoreUObject/Public/UObject/ObjectMacros.h
--- a/Engine/Source/Runtime/CoreUObject/Public/UObject/ObjectMacros.h
+++ b/Engine/Source/Runtime/CoreUObject/Public/UObject/ObjectMacros.h
@@ -97,6 +97,7 @@ enum ESaveFlags
 	SAVE_DiffCallstack  = 0x00000400,	///< Serializes the package to a special memory archive that compares all differences against a file on disk and dumps relevant callstacks
 	SAVE_ComputeHash    = 0x00000800,	///< Compute the MD5 hash of the cooked data
 	SAVE_CompareLinker	= 0x00001000,	///< Return the linker save to compare against another
+	SAVE_BulkDataByReference = 0x00002000, ///< When saving to a different file than the package's LoadedPath, point bulkdata in the new file to be loaded from the original file.
 };
 
 /** Package flags, passed into UPackage::SetPackageFlags and related functions */
diff --git a/Engine/Source/Runtime/CoreUObject/Public/UObject/PackageResourceManager.h b/Engine/Source/Runtime/CoreUObject/Public/UObject/PackageResourceManager.h
--- a/Engine/Source/Runtime/CoreUObject/Public/UObject/PackageResourceManager.h
+++ b/Engine/Source/Runtime/CoreUObject/Public/UObject/PackageResourceManager.h
@@ -48,6 +48,14 @@ struct FOpenPackageResult
 	EPackageFormat Format;
 };
 
+enum class EPackageExternalResource
+{
+	/**
+	 * Open a Package file from the Workspace Domain, even if the Manager is providing packages from a different domain.
+	 * The identifier is the PackageName to open.
+	 */
+	WorkspaceDomainFile,
+};
 
 /**
  * Provides directory queries and Archive payloads for PackagePaths and their PackageSegments from a repository that
@@ -153,6 +161,41 @@ public:
 	 */
 	virtual bool TryMatchCaseOnDisk(const FPackagePath& PackagePath, FPackagePath* OutNormalizedPath = nullptr) = 0;
 
+	/**
+	 * Open a seekable binary FArchive to read the bytes of the given External Resource.
+	 * An ExternalResource is in a separate domain from the one out of which this ResourceManager serves PackagePaths.
+	 *
+	 * @param ResourceType Id for the method used to map the identifier to an archive.
+	 * @param Identifier Id for which resource to return within the ResourceType's domain.
+	 * @return The opened FArchive, or nullptr if it wasn't found.
+	 */
+	virtual TUniquePtr<FArchive> OpenReadExternalResource(EPackageExternalResource ResourceType, FStringView Identifier) = 0;
+
+	/**
+	 * Open an IAsyncReadFileHandle to asynchronously read the bytes of the given ExternalResource.
+	 * An ExternalResource is in a separate domain from the one out of which this ResourceManager serves PackagePaths.
+	 *
+	 * TODO: This call should not hit the disk/network or block, but it currently does, to find the extension.
+	 * This call will always return a non-null handle, even if the package does not exist
+	 *
+	 * @param ResourceType Id for the method used to map the identifier to an archive.
+	 * @param Identifier Id for which resource to return within the ResourceType's domain.
+	 * @return An IAsyncReadFileHandle that will read from the package if it exists, or will be in the
+	 *         canceled state if the package does not exist
+	 */
+	virtual IAsyncReadFileHandle* OpenAsyncReadExternalResource(
+		EPackageExternalResource ResourceType, FStringView Identifier) = 0;
+
+	/**
+	 * Report whether a given ExternalResource exists. The same behavior as OpenReadExternalResource != nullptr,
+	 * but more performant.
+	 * 
+	 * @param ResourceType Id for the method used to map the identifier to an archive.
+	 * @param Identifier Id for which resource to return within the ResourceType's domain.
+	 * @return true if the ExternalResource exists, else false
+	 */
+	virtual bool DoesExternalResourceExist(EPackageExternalResource ResourceType, FStringView Identifier) = 0;
+
 	/**
 	 * Search the given subdirectory of a package mount for all packages with the given package basename in the
 	 * package repository
diff --git a/Engine/Source/Runtime/Engine/Private/Rendering/SkeletalMeshLODImporterData.cpp b/Engine/Source/Runtime/Engine/Private/Rendering/SkeletalMeshLODImporterData.cpp
--- a/Engine/Source/Runtime/Engine/Private/Rendering/SkeletalMeshLODImporterData.cpp
+++ b/Engine/Source/Runtime/Engine/Private/Rendering/SkeletalMeshLODImporterData.cpp
@@ -478,13 +478,6 @@ void FReductionBaseSkeletalMeshBulkData::Serialize(FArchive& Ar, UObject* Owner)
 	}
 	else
 	{
-		if (Ar.IsLoading())
-		{
-			//Save the custom version so we can load FReductionSkeletalMeshData later
-			SerializeLoadingCustomVersionContainer = Ar.GetCustomVersions();
-			bUseSerializeLoadingCustomVersion = true;
-		}
-
 		if (Ar.IsSaving() && bUseSerializeLoadingCustomVersion == true)
 		{
 			//We need to update the FReductionSkeletalMeshData serialize version to the latest in case we save the Parent bulkdata
@@ -496,6 +489,12 @@ void FReductionBaseSkeletalMeshBulkData::Serialize(FArchive& Ar, UObject* Owner)
 	}
 
 	BulkData.Serialize(Ar, Owner);
+	if (!Ar.IsTransacting() && Ar.IsLoading())
+	{
+		//Save the custom version so we can load FReductionSkeletalMeshData later
+		SerializeLoadingCustomVersionContainer = BulkData.GetCustomVersions(Ar);
+		bUseSerializeLoadingCustomVersion = true;
+	}
 }
 
 void FReductionBaseSkeletalMeshBulkData::SaveReductionData(FSkeletalMeshLODModel& BaseLODModel, TMap<FString, TArray<FMorphTargetDelta>>& BaseLODMorphTargetData, UObject* Owner)
@@ -737,13 +736,6 @@ void FRawSkeletalMeshBulkData::Serialize(FArchive& Ar, UObject* Owner)
 	}
 	else
 	{
-		if (Ar.IsLoading())
-		{
-			//Save the custom version so we can load FReductionSkeletalMeshData later
-			SerializeLoadingCustomVersionContainer = Ar.GetCustomVersions();
-			bUseSerializeLoadingCustomVersion = true;
-		}
-
 		if (Ar.IsSaving() && bUseSerializeLoadingCustomVersion == true)
 		{
 			//We need to update the FReductionSkeletalMeshData serialize version to the latest in case we save the Parent bulkdata
@@ -768,6 +760,13 @@ void FRawSkeletalMeshBulkData::Serialize(FArchive& Ar, UObject* Owner)
 	BulkData.Serialize(Ar, Owner);
 	Ar << Guid;
 	Ar << bGuidIsHash;
+
+	if (!Ar.IsTransacting() && Ar.IsLoading())
+	{
+		//Save the custom version so we can load FReductionSkeletalMeshData later
+		SerializeLoadingCustomVersionContainer = BulkData.GetCustomVersions(Ar);
+		bUseSerializeLoadingCustomVersion = true;
+	}
 }
 
 void FRawSkeletalMeshBulkData::SaveRawMesh(FSkeletalMeshImportData& InMesh)
diff --git a/Engine/Source/Runtime/MeshDescription/Private/MeshDescription.cpp b/Engine/Source/Runtime/MeshDescription/Private/MeshDescription.cpp
--- a/Engine/Source/Runtime/MeshDescription/Private/MeshDescription.cpp
+++ b/Engine/Source/Runtime/MeshDescription/Private/MeshDescription.cpp
@@ -1970,7 +1970,7 @@ void FMeshDescriptionBulkData::Serialize( FArchive& Ar, UObject* Owner )
 		return;
 	}
 
-	if( Ar.IsTransacting() )
+	if ( Ar.IsTransacting() )
 	{
 		// If transacting, keep these members alive the other side of an undo, otherwise their values will get lost
 		CustomVersions.Serialize( Ar );
@@ -1978,13 +1978,7 @@ void FMeshDescriptionBulkData::Serialize( FArchive& Ar, UObject* Owner )
 	}
 	else
 	{
-		if( Ar.IsLoading() )
-		{
-			// If loading, take a copy of the package custom version container, so it can be applied when unpacking
-			// MeshDescription from the bulk data.
-			CustomVersions = Ar.GetCustomVersions();
-		}
-		else if( Ar.IsSaving() )
+		if ( Ar.IsSaving() )
 		{
 			// If the bulk data hasn't been updated since this was loaded, there's a possibility that it has old versioning.
 			// Explicitly load and resave the FMeshDescription so that its version is in sync with the FMeshDescriptionBulkData.
@@ -2031,6 +2025,13 @@ void FMeshDescriptionBulkData::Serialize( FArchive& Ar, UObject* Owner )
 		Ar << bGuidIsHash;
 	}
 
+	if (!Ar.IsTransacting() && Ar.IsLoading())
+	{
+		// If loading, take a copy of the package custom version container, so it can be applied when unpacking
+		// MeshDescription from the bulk data.
+		CustomVersions = BulkData.GetCustomVersions(Ar);
+	}
+
 #if UE_USE_VIRTUALBULKDATA
 	// Needs to be after the guid is serialized
 	if (bSerializedOldDataTypes)
