Ensure that the class source object is marked standalone so it doesn't get GC'd in the editor. This is needed for a BPGC asset in a cooked package.

diff --git a/Engine/Source/Runtime/CoreUObject/Private/Blueprint/BlueprintSupport.cpp b/Engine/Source/Runtime/CoreUObject/Private/Blueprint/BlueprintSupport.cpp
--- a/Engine/Source/Runtime/CoreUObject/Private/Blueprint/BlueprintSupport.cpp
+++ b/Engine/Source/Runtime/CoreUObject/Private/Blueprint/BlueprintSupport.cpp
@@ -723,6 +723,15 @@ bool FLinkerLoad::RegenerateBlueprintClass(UClass* LoadClass, UObject* ExportObj
 				// Fix up the linker so that the RegeneratedClass is used
 				LoadClass->ClearFlags(RF_NeedLoad | RF_NeedPostLoad | RF_NeedPostLoadSubobjects);
 			}
+
+#if WITH_EDITOR
+			// Ensure that the class source object is marked standalone so it doesn't get GC'd in the editor.
+			// In particular, this is needed for a BPGC asset in a cooked package.
+			if (LoadClass->bCooked)
+			{
+				ClassSourceObject->SetFlags(RF_Standalone);
+			}
+#endif //if WITH_EDITOR
 		}
 	}
 
