Control lazy load of TObjectPtr properties through the use of metadata on the referenced class type or the individual property.  Not currently enabled without the presence of "-AllowLazyResolve" on the command line.
Ensured loads that need to happen from a reference resolve off the game thread get shunted to the game thread.
Code changes to avoid easily-avoidable object reference resolves that would otherwise reduce benefits of lazy load.
Further work pending to avoid object resolution during serialization phases at cook time and define cook-specific lazy load configuration.

Tested:
ShooterGame editor, cook, PIE

diff --git a/Engine/Source/Editor/UnrealEd/Private/Kismet2/Kismet2.cpp b/Engine/Source/Editor/UnrealEd/Private/Kismet2/Kismet2.cpp
--- a/Engine/Source/Editor/UnrealEd/Private/Kismet2/Kismet2.cpp
+++ b/Engine/Source/Editor/UnrealEd/Private/Kismet2/Kismet2.cpp
@@ -881,26 +881,29 @@ static void ConformComponentsUtils::ConformRemovedNativeComponents(UObject* BpCd
 	};
 
 	// 
-	for (TFieldIterator<FObjectProperty> ObjPropIt(NativeSuperClass); ObjPropIt; ++ObjPropIt)
+	if (DestroyedComponents.Num() > 0)
 	{
-		FObjectProperty* ObjectProp = *ObjPropIt;
-		UObject* PropObjValue = ObjectProp->GetObjectPropertyValue_InContainer(ActorCDO);
-
-		if (DestroyedComponents.Contains(PropObjValue))
+		for (TFieldIterator<FObjectProperty> ObjPropIt(NativeSuperClass); ObjPropIt; ++ObjPropIt)
 		{
-			// Get the "new" value that's currently set on the native parent CDO. We need the Blueprint CDO to reflect this update in property value.
-			UObject* SuperObjValue = ObjectProp->GetObjectPropertyValue_InContainer(NativeCDO);
-			if (SuperObjValue && SuperObjValue->IsA<UActorComponent>())
+			FObjectProperty* ObjectProp = *ObjPropIt;
+			UObject* PropObjValue = ObjectProp->GetObjectPropertyValue_InContainer(ActorCDO);
+
+			if (DestroyedComponents.Contains(PropObjValue))
 			{
-				// For components, make sure we use the instance that's owned by the Blueprint CDO and not the native parent CDO's instance.
-				if (UActorComponent** ComponentTemplatePtr = FindComponentTemplateByNameInActorCDO(SuperObjValue->GetFName()))
+				// Get the "new" value that's currently set on the native parent CDO. We need the Blueprint CDO to reflect this update in property value.
+				UObject* SuperObjValue = ObjectProp->GetObjectPropertyValue_InContainer(NativeCDO);
+				if (SuperObjValue && SuperObjValue->IsA<UActorComponent>())
 				{
-					SuperObjValue = *ComponentTemplatePtr;
+					// For components, make sure we use the instance that's owned by the Blueprint CDO and not the native parent CDO's instance.
+					if (UActorComponent** ComponentTemplatePtr = FindComponentTemplateByNameInActorCDO(SuperObjValue->GetFName()))
+					{
+						SuperObjValue = *ComponentTemplatePtr;
+					}
 				}
-			}
 			
-			// Update the Blueprint CDO to match the native parent CDO.
-			ObjectProp->SetObjectPropertyValue_InContainer(ActorCDO, SuperObjValue);
+				// Update the Blueprint CDO to match the native parent CDO.
+				ObjectProp->SetObjectPropertyValue_InContainer(ActorCDO, SuperObjValue);
+			}
 		}
 	}
 
diff --git a/Engine/Source/Runtime/CoreUObject/Private/UObject/LinkerLoad.cpp b/Engine/Source/Runtime/CoreUObject/Private/UObject/LinkerLoad.cpp
--- a/Engine/Source/Runtime/CoreUObject/Private/UObject/LinkerLoad.cpp
+++ b/Engine/Source/Runtime/CoreUObject/Private/UObject/LinkerLoad.cpp
@@ -1964,6 +1964,90 @@ FPackageIndex FLinkerLoad::FindOrCreateImport(const FName InObjectName, const FN
 	return FPackageIndex::FromImport(ImportMap.Num() - 1);
 }
 
+#if UE_WITH_OBJECT_HANDLE_LATE_RESOLVE
+FLinkerLoad::EImportLoadBehavior FLinkerLoad::GetCurrentPropertyImportLoadBehavior(FPackageIndex ImportIndex)
+{
+	const FObjectImport& Import = Imp(ImportIndex);
+	if (Import.bImportSearchedFor)
+	{
+		// If it was something that's been searched for, we've already attempted a resolve, might as well use it
+		return EImportLoadBehavior::Eager;
+	}
+
+	const FArchiveSerializedPropertyChain* PropChain = GetSerializedPropertyChain();
+	if (!PropChain)
+	{
+		// In the absence of property information, the safe choice is to not lazy load
+		return EImportLoadBehavior::Eager;
+	}
+
+	const int32 NumPropsInChain = PropChain->GetNumProperties();
+	const FProperty* CurrentProperty = NumPropsInChain > 0 ? PropChain->GetPropertyFromStack(0) : nullptr;
+	if (!CurrentProperty)
+	{
+		// In the absence of property information, the safe choice is to not lazy load
+		return EImportLoadBehavior::Eager;
+	}
+
+	const FProperty* ParentProperty = NumPropsInChain > 1 ? PropChain->GetPropertyFromStack(1) : nullptr;
+	// Exclude properties that are in hash sensitive containers as they will require the object to be loaded to compute a hash
+	// @TODO: OBJPTR: Forcing resolve for ObjectPtr hashing.  See GetTypeHash in ObjectHandle.h.
+	if (ParentProperty && 
+		(
+			ParentProperty->IsA<FSetProperty>() ||
+			(ParentProperty->IsA<FMapProperty>() && (CurrentProperty == static_cast<const FMapProperty*>(ParentProperty)->GetKeyProperty()))
+		))
+	{
+		return EImportLoadBehavior::Eager;
+	}
+
+	static bool bAllowLazyResolve = FParse::Param(FCommandLine::Get(), TEXT("AllowLazyResolve"));
+	if (!bAllowLazyResolve || !IsAllowingLazyLoading() || GIsPlayInEditorWorld)
+	{
+		return EImportLoadBehavior::Eager;
+	}
+
+	static const bool bLazyResolveAllImports = FParse::Param(FCommandLine::Get(), TEXT("LazyResolveAllImports"));
+	if (bLazyResolveAllImports)
+	{
+		return EImportLoadBehavior::LazyOnDemand;
+	}
+
+	// If we're not blanket lazy resolving all imports, decide if this one has the metadata on the property or the referenced class to allow lazy resolve
+	EImportLoadBehavior LoadBehavior = EImportLoadBehavior::Eager;
+	static const FName Name_LoadBehavior("LoadBehavior");
+	const FString* LoadBehaviorMeta = CurrentProperty->FindMetaData(Name_LoadBehavior);
+	if (!LoadBehaviorMeta)
+	{
+		// Attempt to get the meta from the referenced class.  This only looks in already loaded classes.  May need to resolve the class in the future.
+		UObject* ClassPackage = FindObject<UPackage>(nullptr, *Import.ClassPackage.ToString());
+		if (UClass* FindClass = ClassPackage ? FindObject<UClass>(ClassPackage, *Import.ClassName.ToString()) : nullptr)
+		{
+			LoadBehaviorMeta = FindClass->FindMetaData(Name_LoadBehavior);
+		}
+
+	}
+	if (LoadBehaviorMeta)
+	{
+		if (*LoadBehaviorMeta == "Eager")
+		{
+			LoadBehavior = EImportLoadBehavior::Eager;
+		}
+		// @TODO: OBJPTR: we want to permit lazy background loading in the future
+		// else if (*LoadBehaviorMeta == "LazyBackground")
+		// {
+		// 	LoadBehavior = EImportLoadBehavior::LazyBackground;
+		// }
+		else if (*LoadBehaviorMeta == "LazyOnDemand")
+		{
+			LoadBehavior = EImportLoadBehavior::LazyOnDemand;
+		}
+	}
+
+	return LoadBehavior;
+}
+#endif // UE_WITH_OBJECT_HANDLE_LATE_RESOLVE
+
 FString ExtractObjectName(const FString& InFullPath)
 {
 	FString ObjectName = InFullPath;
@@ -5519,34 +5603,19 @@ FArchive& FLinkerLoad::operator<<( FObjectPtr& ObjectPtr )
 	Ar << Index;
 
 #if UE_WITH_OBJECT_HANDLE_LATE_RESOLVE
-	static bool bLazyResolveAllImports = FParse::Param(FCommandLine::Get(), TEXT("LazyResolveAllImports"));
-	if (bLazyResolveAllImports && Index.IsImport())
+	if (Index.IsImport())
 	{
-		const FObjectImport& Import = Imp(Index);
-		if (!Import.bImportSearchedFor)
+		EImportLoadBehavior LoadBehavior = GetCurrentPropertyImportLoadBehavior(Index);
+		if (LoadBehavior == EImportLoadBehavior::LazyOnDemand)
 		{
-			bool bIsInHashSensitiveContainer = false;
-			if (const FArchiveSerializedPropertyChain* PropChain = GetSerializedPropertyChain())
-			{
-				const int32 NumPropsInChain = PropChain->GetNumProperties();
-				const FProperty* CurrentProperty = NumPropsInChain > 0 ? PropChain->GetPropertyFromStack(0) : nullptr;
-				if (const FProperty* ParentProperty = NumPropsInChain > 1 ? PropChain->GetPropertyFromStack(1) : nullptr)
-				{
-					// @TODO: OBJPTR: Forcing resolve for ObjectPtr hashing.  See GetTypeHash in ObjectHandle.h.
-					bIsInHashSensitiveContainer = ParentProperty->IsA<FSetProperty>() || (ParentProperty->IsA<FMapProperty>() && (CurrentProperty == static_cast<const FMapProperty*>(ParentProperty)->GetKeyProperty()));
-				}
-			}
+			const FObjectImport& Import = Imp(Index);
+			FObjectPathId ObjectPath;
+			FName PackageName = FObjectPathId::MakeImportPathIdAndPackageName(Import, *this, ObjectPath);
+			FObjectRef ImportRef{PackageName, Import.ClassPackage, Import.ClassName, ObjectPath};
 
-			if (!bIsInHashSensitiveContainer)
-			{
-				FObjectPathId ObjectPath;
-				FName PackageName = FObjectPathId::MakeImportPathIdAndPackageName(Import, *this, ObjectPath);
-				FObjectRef ImportRef{PackageName, Import.ClassPackage, Import.ClassName, ObjectPath};
+			ObjectPtr = FObjectPtr(ImportRef);
 
-				ObjectPtr = FObjectPtr(ImportRef);
-
-				return *this;
-			}
+			return *this;
 		}
 	}
 #endif
diff --git a/Engine/Source/Runtime/CoreUObject/Private/UObject/ObjectHandle.cpp b/Engine/Source/Runtime/CoreUObject/Private/UObject/ObjectHandle.cpp
--- a/Engine/Source/Runtime/CoreUObject/Private/UObject/ObjectHandle.cpp
+++ b/Engine/Source/Runtime/CoreUObject/Private/UObject/ObjectHandle.cpp
@@ -1,6 +1,7 @@
 // Copyright Epic Games, Inc. All Rights Reserved.
 
 #include "UObject/ObjectHandle.h"
+#include "Async/TaskGraphInterfaces.h"
 #include "HAL/CriticalSection.h"
 #include "HAL/PlatformAtomics.h"
 #include "Misc/PackageName.h"
@@ -362,6 +363,26 @@ UClass* ResolveObjectRefClass(const FObjectRef& ObjectRef, uint32 LoadFlags /*=
 	return ClassObject;
 }
 
+class FFullyLoadPackageOnHandleResolveTask
+{
+public:
+	FFullyLoadPackageOnHandleResolveTask(UPackage* InPackage): Package(InPackage)
+	{
+	}
+
+	FORCEINLINE TStatId GetStatId() const { RETURN_QUICK_DECLARE_CYCLE_STAT(FFullyLoadPackageOnHandleResolveTask, STATGROUP_TaskGraphTasks); }
+	ENamedThreads::Type GetDesiredThread() { return ENamedThreads::Type::GameThread; }
+	static ESubsequentsMode::Type GetSubsequentsMode() { return ESubsequentsMode::FireAndForget; }
+
+	void DoTask(ENamedThreads::Type CurrentThread, const FGraphEventRef& MyCompletionGraphEvent)
+	{
+		Package->LinkerLoad->LoadAllObjects(true);
+	}
+
+private:
+	UPackage* Package = nullptr;
+};
+
 UObject* ResolveObjectRef(const FObjectRef& ObjectRef, uint32 LoadFlags /*= LOAD_None*/)
 {
 	TRACE_CPUPROFILER_EVENT_SCOPE(ResolveObjectRef);
@@ -404,9 +425,17 @@ UObject* ResolveObjectRef(const FObjectRef& ObjectRef, uint32 LoadFlags /*= LOAD
 			CurrentObject = Redirector->DestinationObject;
 		}
 
-		if (!CurrentObject && IsInGameThread() && !TargetPackage->IsFullyLoaded() && TargetPackage->LinkerLoad && TargetPackage->LinkerLoad->IsLoading())
+		if (!CurrentObject && !TargetPackage->IsFullyLoaded() && TargetPackage->LinkerLoad && TargetPackage->LinkerLoad->IsLoading())
 		{
-			TargetPackage->LinkerLoad->LoadAllObjects(true);
+			if (IsInAsyncLoadingThread() || IsInGameThread())
+			{
+				TargetPackage->LinkerLoad->LoadAllObjects(true);
+			}
+			else
+			{
+				// Shunt the load request to happen on the game thread and block on its completion.  This is a deadlock risk!  The game thread may be blocked waiting on this thread.
+				TGraphTask<FFullyLoadPackageOnHandleResolveTask>::CreateTask().ConstructAndDispatchWhenReady(TargetPackage)->Wait();
+			}
 
 			CurrentObject = StaticFindObjectFastInternal(nullptr, PreviousOuter, ResolvedNames[ObjectPathIndex]);
 			if (UObjectRedirector* Redirector = dynamic_cast<UObjectRedirector*>(CurrentObject))
diff --git a/Engine/Source/Runtime/CoreUObject/Private/UObject/UObjectGlobals.cpp b/Engine/Source/Runtime/CoreUObject/Private/UObject/UObjectGlobals.cpp
--- a/Engine/Source/Runtime/CoreUObject/Private/UObject/UObjectGlobals.cpp
+++ b/Engine/Source/Runtime/CoreUObject/Private/UObject/UObjectGlobals.cpp
@@ -3560,7 +3560,7 @@ protected:
 
 	virtual FArchive& operator<<(FObjectPtr& Object) override
 	{
-		if (Object && IsObjectHandleResolved(Object.GetHandle()))
+		if (IsObjectHandleResolved(Object.GetHandle()) && Object)
 		{
 			FReferenceCollector& CurrentCollector = GetCollector();
 			FProperty* OldCollectorSerializedProperty = CurrentCollector.GetSerializedProperty();
diff --git a/Engine/Source/Runtime/CoreUObject/Public/UObject/LinkerLoad.h b/Engine/Source/Runtime/CoreUObject/Public/UObject/LinkerLoad.h
--- a/Engine/Source/Runtime/CoreUObject/Public/UObject/LinkerLoad.h
+++ b/Engine/Source/Runtime/CoreUObject/Public/UObject/LinkerLoad.h
@@ -1041,6 +1041,16 @@ private:
 	/** For the given object and class info, find or create an associated import record. Used when loading text assets which only store object paths */
 	FPackageIndex FindOrCreateImport(const FName InObjectName, const FName InClassName, const FName InClassPackageName);
 
+#if UE_WITH_OBJECT_HANDLE_LATE_RESOLVE
+	enum class EImportLoadBehavior
+	{
+		Eager = 0,
+		// @TODO: OBJPTR: we want to permit lazy background loading in the future
+		//LazyBackground,
+		LazyOnDemand,
+	};
+	EImportLoadBehavior GetCurrentPropertyImportLoadBehavior(FPackageIndex ImportIndex);
+#endif // UE_WITH_OBJECT_HANDLE_LATE_RESOLVE
 public:
 	/**
 	 * Serializes the gatherable text data container.
diff --git a/Engine/Source/Runtime/CoreUObject/Public/UObject/ObjectMacros.h b/Engine/Source/Runtime/CoreUObject/Public/UObject/ObjectMacros.h
--- a/Engine/Source/Runtime/CoreUObject/Public/UObject/ObjectMacros.h
+++ b/Engine/Source/Runtime/CoreUObject/Public/UObject/ObjectMacros.h
@@ -1083,6 +1083,9 @@ namespace UM
 		/// [ClassMetadata] Comma delimited list of blueprint events that are not be allowed to be overridden in classes of this type
 		KismetHideOverrides,
 
+		/// [ClassMetadata] [PropertyMetadata] Used by TObjectPtr properties.  Specifies that references to objects of this class should be eager/lazy/background loaded in editor by default.  Can be overridden by LoadBehavior specified on the referencing property field.
+		LoadBehavior,
+
 		/// [ClassMetadata] Specifies interfaces that are not compatible with the class.
 		ProhibitedInterfaces,
 
@@ -1237,6 +1240,9 @@ namespace UM
 		/// [PropertyMetadata] Signifies that the bool property is only displayed inline as an edit condition toggle in other properties, and should not be shown on its own row.
 		InlineEditConditionToggle,
 
+		/// [ClassMetadata] [PropertyMetadata] Used by TObjectPtr properties.  Allows the referenced object to be eager/lazy/background loaded in editor.
+		// LoadBehavior, (Commented out so as to avoid duplicate name with version in the Class section, but still show in the property section)
+
 		/// [PropertyMetadata] Used by FDirectoryPath properties.  Converts the path to a long package name
 		LongPackageName,
 
diff --git a/Engine/Source/Runtime/Engine/Classes/Sound/SoundCue.h b/Engine/Source/Runtime/Engine/Classes/Sound/SoundCue.h
--- a/Engine/Source/Runtime/Engine/Classes/Sound/SoundCue.h
+++ b/Engine/Source/Runtime/Engine/Classes/Sound/SoundCue.h
@@ -84,7 +84,7 @@ public:
 /**
  * The behavior of audio playback is defined within Sound Cues.
  */
-UCLASS(hidecategories=object, BlueprintType)
+UCLASS(hidecategories=object, BlueprintType, meta= (LoadBehavior = "LazyOnDemand"))
 class ENGINE_API USoundCue : public USoundBase
 {
 	GENERATED_UCLASS_BODY()
diff --git a/Engine/Source/Runtime/Engine/Classes/Sound/SoundWave.h b/Engine/Source/Runtime/Engine/Classes/Sound/SoundWave.h
--- a/Engine/Source/Runtime/Engine/Classes/Sound/SoundWave.h
+++ b/Engine/Source/Runtime/Engine/Classes/Sound/SoundWave.h
@@ -370,7 +370,7 @@ private:
 };
 
 
-UCLASS(hidecategories=Object, editinlinenew, BlueprintType)
+UCLASS(hidecategories=Object, editinlinenew, BlueprintType, meta= (LoadBehavior = "LazyOnDemand"))
 class ENGINE_API USoundWave : public USoundBase, public IAudioProxyDataFactory
 {
 	GENERATED_UCLASS_BODY()
