UnversionedPropertySerialization: Support unversioned in uncooked packages so that EditorDomain can use it.
#rnx

diff --git a/Engine/Source/Runtime/CoreUObject/Private/Serialization/UnversionedPropertySerialization.cpp b/Engine/Source/Runtime/CoreUObject/Private/Serialization/UnversionedPropertySerialization.cpp
--- a/Engine/Source/Runtime/CoreUObject/Private/Serialization/UnversionedPropertySerialization.cpp
+++ b/Engine/Source/Runtime/CoreUObject/Private/Serialization/UnversionedPropertySerialization.cpp
@@ -281,12 +281,14 @@ struct FUnversionedStructSchema
 	uint32 Num;
 	FUnversionedPropertySerializer Serializers[0];
 
-	static FUnversionedStructSchema* Create(const UStruct* Struct)
+	FORCEINLINE static FUnversionedStructSchema* Create(const UStruct* Struct, bool bSkipEditorOnly)
 	{
 		TArray<FUnversionedPropertySerializer, TInlineAllocator<256>> Serializers;
 		for (FProperty* Property = Struct->PropertyLink; Property; Property = Property->PropertyLinkNext)
 		{
-			if (!Property->IsEditorOnlyProperty())
+#if WITH_EDITORONLY_DATA
+			if (!bSkipEditorOnly || !Property->IsEditorOnlyProperty())
+#endif
 			{
 				for (int32 ArrayIdx = 0, ArrayDim = Property->ArrayDim; ArrayIdx < ArrayDim; ++ArrayIdx)
 				{
@@ -304,16 +306,25 @@ struct FUnversionedStructSchema
 	}
 };
 
-const FUnversionedStructSchema& GetOrCreateUnversionedSchema(const UStruct* Struct)
+static const FUnversionedStructSchema*& GetUnversionedSchema(const UStruct* Struct, bool bSkipEditorOnly)
 {
-	if (const FUnversionedStructSchema* ExistingSchema = Struct->UnversionedSchema)
+#if WITH_EDITORONLY_DATA
+	return bSkipEditorOnly ? Struct->UnversionedGameSchema : Struct->UnversionedEditorSchema;
+#else
+	return Struct->UnversionedGameSchema;
+#endif
+}
+
+const FUnversionedStructSchema& GetOrCreateUnversionedSchema(const UStruct* Struct, bool bSkipEditorOnly)
+{
+	if (const FUnversionedStructSchema* ExistingSchema = GetUnversionedSchema(Struct, bSkipEditorOnly))
 	{
 		return *ExistingSchema;
 	}
-	
-	FUnversionedStructSchema* CreatedSchema = FUnversionedStructSchema::Create(Struct);
 
-	void** CachedSchemaPtr = reinterpret_cast<void**>(const_cast<FUnversionedStructSchema**>(&Struct->UnversionedSchema));
+	FUnversionedStructSchema* CreatedSchema = FUnversionedStructSchema::Create(Struct, bSkipEditorOnly);
+
+	void** CachedSchemaPtr = reinterpret_cast<void**>(const_cast<FUnversionedStructSchema**>(&GetUnversionedSchema(Struct, bSkipEditorOnly)));
 	if (const FUnversionedStructSchema* ExistingSchema = reinterpret_cast<const FUnversionedStructSchema*>(FPlatformAtomics::InterlockedCompareExchangePointer(CachedSchemaPtr, CreatedSchema, nullptr)))
 	{
 		delete CreatedSchema;
@@ -331,20 +342,23 @@ struct FLinkWalkingSchemaIterator
 {
 	FProperty* Property = nullptr;
 	uint32 ArrayIndex = 0;
+	bool bSkipEditorOnly = true;
 
 	FLinkWalkingSchemaIterator() {}
 
-	explicit FLinkWalkingSchemaIterator(FProperty* FirstProperty)
-		: Property(SkipEditorOnlyProperties(FirstProperty))
-	{}
+	FORCEINLINE explicit FLinkWalkingSchemaIterator(FProperty* FirstProperty, bool bInSkipEditorOnly)
+		: Property(SkipEditorOnlyProperties(FirstProperty, bInSkipEditorOnly))
+		, bSkipEditorOnly(bInSkipEditorOnly)
+	{
+	}
 
-	void operator++()
+	FORCEINLINE void operator++()
 	{
 		check(Property);
 
 		if (ArrayIndex + 1 == Property->ArrayDim)
 		{
-			Property = SkipEditorOnlyProperties(Property->PropertyLinkNext);
+			Property = SkipEditorOnlyProperties(Property->PropertyLinkNext, bSkipEditorOnly);
 			ArrayIndex = 0;
 		}
 		else
@@ -353,7 +367,7 @@ struct FLinkWalkingSchemaIterator
 		}
 	}
 
-	void operator+=(uint32 Num)
+	FORCEINLINE void operator+=(uint32 Num)
 	{
 		while (Num--)
 		{
@@ -371,12 +385,15 @@ struct FLinkWalkingSchemaIterator
 		return (Property != Rhs.Property) | (ArrayIndex != Rhs.ArrayIndex);
 	}
 
-	static FProperty* SkipEditorOnlyProperties(FProperty* Property)
+	FORCEINLINE static FProperty* SkipEditorOnlyProperties(FProperty* Property, bool bSkipEditorOnly)
 	{
 #if WITH_EDITORONLY_DATA
-		while (Property && Property->IsEditorOnlyProperty())
+		if (bSkipEditorOnly)
 		{
-			Property = Property->PropertyLinkNext;
+			while (Property && Property->IsEditorOnlyProperty())
+			{
+				Property = Property->PropertyLinkNext;
+			}
 		}
 #endif
 		return Property;
@@ -389,14 +406,14 @@ using FUnversionedSchemaIterator = FLinkWalkingSchemaIterator;
 
 struct FUnversionedSchemaRange
 {
-	explicit FUnversionedSchemaRange(const UStruct* Struct)
+	FORCEINLINE explicit FUnversionedSchemaRange(const UStruct* Struct, bool bSkipEditorOnly)
 	{
 #if CACHE_UNVERSIONED_PROPERTY_SCHEMA
-		const FUnversionedStructSchema& Schema = GetOrCreateUnversionedSchema(Struct);
+		const FUnversionedStructSchema& Schema = GetOrCreateUnversionedSchema(Struct, bSkipEditorOnly);
 		Begin = Schema.Serializers;
 		End = Schema.Serializers + Schema.Num;	
 #else
-		Begin = FUnversionedSchemaIterator(Struct->PropertyLink);
+		Begin = FUnversionedSchemaIterator(Struct->PropertyLink, bSkipEditorOnly);
 		// End is default-initialized
 #endif
 	}
@@ -761,11 +778,21 @@ bool CanUseUnversionedPropertySerialization(const ITargetPlatform* Target)
 void DestroyUnversionedSchema(const UStruct* Struct)
 {
 #if CACHE_UNVERSIONED_PROPERTY_SCHEMA
-	delete Struct->UnversionedSchema;
-	Struct->UnversionedSchema = nullptr;
+	delete Struct->UnversionedGameSchema;
+	Struct->UnversionedGameSchema = nullptr;
+#if WITH_EDITORONLY_DATA
+	delete Struct->UnversionedEditorSchema;
+	Struct->UnversionedEditorSchema = nullptr;
+#endif
 #endif
 }
 
+#if WITH_EDITORONLY_DATA
+static bool SkipEditorOnlyFields(FArchive& Ar) { return Ar.IsFilterEditorOnly(); }
+#else
+static constexpr bool SkipEditorOnlyFields(FArchive& Ar) { return true; }
+#endif
+
 void SerializeUnversionedProperties(const UStruct* Struct, FStructuredArchive::FSlot Slot, uint8* Data, UStruct* DefaultsStruct, uint8* DefaultsData)
 {
 	FArchive& UnderlyingArchive = Slot.GetUnderlyingArchive();
@@ -780,7 +807,7 @@ void SerializeUnversionedProperties(const UStruct* Struct, FStructuredArchive::F
 
 		if (Header.HasValues())
 		{
-			FUnversionedSchemaRange Schema(Struct);
+			FUnversionedSchemaRange Schema(Struct, SkipEditorOnlyFields(UnderlyingArchive));
 
 			if (Header.HasNonZeroValues())
 			{
@@ -819,7 +846,7 @@ void SerializeUnversionedProperties(const UStruct* Struct, FStructuredArchive::F
 		const bool bDense = !UnderlyingArchive.DoDelta() || UnderlyingArchive.IsTransacting() || (!DefaultsData && !dynamic_cast<const UClass*>(Struct));
 		FDefaultStruct Defaults(DefaultsData, DefaultsStruct);
 
-		FUnversionedSchemaRange Schema(Struct);
+		FUnversionedSchemaRange Schema(Struct, SkipEditorOnlyFields(UnderlyingArchive));
 		FUnversionedHeaderBuilder Header;
 		for (FUnversionedPropertySerializer Serializer : Schema)
 		{
diff --git a/Engine/Source/Runtime/CoreUObject/Public/UObject/Class.h b/Engine/Source/Runtime/CoreUObject/Public/UObject/Class.h
--- a/Engine/Source/Runtime/CoreUObject/Public/UObject/Class.h
+++ b/Engine/Source/Runtime/CoreUObject/Public/UObject/Class.h
@@ -343,8 +343,12 @@ public:
 	int32 FieldPathSerialNumber;
 #endif
 
-	/** Cached schema for optimized unversioned property serialization, owned by this. */
-	mutable const struct FUnversionedStructSchema* UnversionedSchema = nullptr;
+	/** Cached schema for optimized unversioned and filtereditoronly property serialization, owned by this. */
+	mutable const struct FUnversionedStructSchema* UnversionedGameSchema = nullptr;
+#if WITH_EDITORONLY_DATA
+	/** Cached schema for optimized unversioned property serialization, with editor data, owned by this. */
+	mutable const struct FUnversionedStructSchema* UnversionedEditorSchema = nullptr;
+#endif
 
 public:
 	// Constructors.
