Only ensure when the given package name is not empty to support legacy behavior before 14814912

#rnx

diff --git a/Engine/Source/Runtime/CoreUObject/Private/Misc/PackageName.cpp b/Engine/Source/Runtime/CoreUObject/Private/Misc/PackageName.cpp
--- a/Engine/Source/Runtime/CoreUObject/Private/Misc/PackageName.cpp
+++ b/Engine/Source/Runtime/CoreUObject/Private/Misc/PackageName.cpp
@@ -1272,7 +1272,13 @@ bool FPackageName::DoesPackageExist(const FString& LongPackageName, const FGuid*
 		{
 			FString Message = FString::Printf(TEXT("Illegal call to DoesPackageExist: %s"), *FormatErrorAsString(LongPackageName, FailureReason));
 			UE_LOG(LogPackageName, Error, TEXT("%s"), *Message);
-			ensureMsgf(false, TEXT("%s"), *Message);
+
+			// Only ensure if the given package name is not an empty string to 
+			// support legacy behavior
+			if(!LongPackageName.IsEmpty())
+			{
+				ensureMsgf(false, TEXT("%s"), *Message);
+			}
 			return false;
 		}
 		PackagePath = FPackagePath::FromMountedComponents(PackageNameRoot, FilePathRoot, RelPath, Extension, CustomExtension);
