UnrealInsights: Fixed missing package names when using the new loader

diff --git a/Engine/Source/Developer/TraceServices/Private/Analyzers/LoadTimeTraceAnalysis.cpp b/Engine/Source/Developer/TraceServices/Private/Analyzers/LoadTimeTraceAnalysis.cpp
--- a/Engine/Source/Developer/TraceServices/Private/Analyzers/LoadTimeTraceAnalysis.cpp
+++ b/Engine/Source/Developer/TraceServices/Private/Analyzers/LoadTimeTraceAnalysis.cpp
@@ -114,6 +114,11 @@ bool FAsyncLoadingTraceAnalyzer::OnEvent(uint16 RouteId, EStyle Style, const FOn
 		if (AsyncPackage)
 		{
 			FAnalysisSessionEditScope _(Session);
+			if (EventData.GetAttachmentSize())
+			{
+				const TCHAR* PackageName = reinterpret_cast<const TCHAR*>(EventData.GetAttachment());
+				AsyncPackage->PackageInfo->Name = Session.StoreString(PackageName);
+			}
 			FPackageSummaryInfo& Summary = AsyncPackage->PackageInfo->Summary;
 			Summary.TotalHeaderSize = EventData.GetValue<uint32>("TotalHeaderSize");
 			Summary.ImportCount = EventData.GetValue<uint32>("ImportCount");
@@ -309,7 +314,12 @@ bool FAsyncLoadingTraceAnalyzer::OnEvent(uint16 RouteId, EStyle Style, const FOn
 		uint64 AsyncPackagePtr = EventData.GetValue<uint64>("AsyncPackage");
 		//check(!ActivePackagesMap.Contains(AsyncPackagePtr));
 		FAsyncPackageState* AsyncPackageState = new FAsyncPackageState();
-		AsyncPackageState->PackageInfo = &LoadTimeProfilerProvider.EditPackageInfo(reinterpret_cast<const TCHAR*>(EventData.GetAttachment()));
+		AsyncPackageState->PackageInfo = &LoadTimeProfilerProvider.CreatePackage();
+		if (EventData.GetAttachmentSize()) // For backwards compatibility
+		{
+			const TCHAR* PackageName = reinterpret_cast<const TCHAR*>(EventData.GetAttachment());
+			AsyncPackageState->PackageInfo->Name = Session.StoreString(PackageName);
+		}
 		ActiveAsyncPackagesMap.Add(AsyncPackagePtr, AsyncPackageState);
 		break;
 	}
diff --git a/Engine/Source/Developer/TraceServices/Private/Model/LoadTimeProfiler.cpp b/Engine/Source/Developer/TraceServices/Private/Model/LoadTimeProfiler.cpp
--- a/Engine/Source/Developer/TraceServices/Private/Model/LoadTimeProfiler.cpp
+++ b/Engine/Source/Developer/TraceServices/Private/Model/LoadTimeProfiler.cpp
@@ -389,14 +389,13 @@ FLoadRequest& FLoadTimeProfilerProvider::CreateRequest()
 	return RequestInfo;
 }
 
-FPackageInfo& FLoadTimeProfilerProvider::EditPackageInfo(const TCHAR* PackageName)
+FPackageInfo& FLoadTimeProfilerProvider::CreatePackage()
 {
 	Session.WriteAccessCheck();
 
 	uint32 PackageId = Packages.Num();
 	FPackageInfo& Package = Packages.PushBack();
 	Package.Id = PackageId;
-	Package.Name = Session.StoreString(PackageName);
 	return Package;
 }
 
diff --git a/Engine/Source/Developer/TraceServices/Private/Model/LoadTimeProfilerPrivate.h b/Engine/Source/Developer/TraceServices/Private/Model/LoadTimeProfilerPrivate.h
--- a/Engine/Source/Developer/TraceServices/Private/Model/LoadTimeProfilerPrivate.h
+++ b/Engine/Source/Developer/TraceServices/Private/Model/LoadTimeProfilerPrivate.h
@@ -32,7 +32,7 @@ public:
 	virtual const ITable<FLoadRequest>& GetRequestsTable() const override { return RequestsTable; }
 	const FClassInfo& AddClassInfo(const TCHAR* ClassName);
 	FLoadRequest& CreateRequest();
-	FPackageInfo& EditPackageInfo(const TCHAR* PackageName);
+	FPackageInfo& CreatePackage();
 	uint64 BeginLoadPackage(const FPackageInfo& PackageInfo, double Time);
 	void EndLoadPackage(uint64 LoadHandle, double Time);
 	FPackageExportInfo& CreateExport();
diff --git a/Engine/Source/Runtime/CoreUObject/Private/Serialization/AsyncLoading.cpp b/Engine/Source/Runtime/CoreUObject/Private/Serialization/AsyncLoading.cpp
--- a/Engine/Source/Runtime/CoreUObject/Private/Serialization/AsyncLoading.cpp
+++ b/Engine/Source/Runtime/CoreUObject/Private/Serialization/AsyncLoading.cpp
@@ -1564,7 +1564,7 @@ void FAsyncPackage::Event_FinishLinker()
 			}
 		}
 
-		TRACE_LOADTIME_PACKAGE_SUMMARY(this, Linker->Summary.TotalHeaderSize, Linker->Summary.ImportCount, Linker->Summary.ExportCount);
+		TRACE_LOADTIME_PACKAGE_SUMMARY(this, Desc.Name, Linker->Summary.TotalHeaderSize, Linker->Summary.ImportCount, Linker->Summary.ExportCount);
 
 		check(AsyncPackageLoadingState == EAsyncPackageLoadingState::WaitingForSummary);
 		AsyncPackageLoadingState = EAsyncPackageLoadingState::StartImportPackages;
@@ -4682,10 +4682,6 @@ FAsyncLoadingThread::FAsyncLoadingThread(int32 InThreadIndex, IEDLBootNotificati
 			TEXT("Event driven async loader is NOT being used but it seems to be enabled in project settings."));
 	}
 
-#if LOADTIMEPROFILERTRACE_ENABLED
-	FLoadTimeProfilerTracePrivate::Init();
-#endif
-
 	PrecacheHandler = new FPrecacheCallbackHandler();
 	QueuedRequestsEvent = FPlatformProcess::GetSynchEventFromPool();
 	CancelLoadingEvent = FPlatformProcess::GetSynchEventFromPool();
@@ -5233,7 +5229,7 @@ FAsyncPackage::FAsyncPackage(FAsyncLoadingThread& InThread, const FAsyncPackageD
 , FinishObjectsTime(0.0)
 #endif // PERF_TRACK_DETAILED_ASYNC_STATS
 {
-	TRACE_LOADTIME_NEW_ASYNC_PACKAGE(this, InDesc.Name);
+	TRACE_LOADTIME_NEW_ASYNC_PACKAGE(this);
 	AddRequestID(InDesc.RequestID);
 }
 
diff --git a/Engine/Source/Runtime/CoreUObject/Private/Serialization/AsyncLoading2.cpp b/Engine/Source/Runtime/CoreUObject/Private/Serialization/AsyncLoading2.cpp
--- a/Engine/Source/Runtime/CoreUObject/Private/Serialization/AsyncLoading2.cpp
+++ b/Engine/Source/Runtime/CoreUObject/Private/Serialization/AsyncLoading2.cpp
@@ -3256,7 +3256,7 @@ EAsyncPackageState::Type FAsyncPackage2::Event_ProcessPackageSummary(FAsyncLoadi
 		Package->AllExportDataPtr = PackageSummaryData + PackageSummarySize;
 		Package->CurrentExportDataPtr = Package->AllExportDataPtr;
 
-		TRACE_LOADTIME_PACKAGE_SUMMARY(Package, PackageSummarySize, Package->ImportStore.ImportMap.Num(), Package->Data.ExportCount);
+		TRACE_LOADTIME_PACKAGE_SUMMARY(Package, Package->Desc.DiskPackageName, PackageSummarySize, Package->ImportStore.ImportMap.Num(), Package->Data.ExportCount);
 	}
 
 	if (GIsInitialLoad)
@@ -4462,10 +4462,6 @@ FAsyncLoadingThread2::FAsyncLoadingThread2(FIoDispatcher& InIoDispatcher)
 	GEventDrivenLoaderEnabled = true;
 #endif
 
-#if LOADTIMEPROFILERTRACE_ENABLED
-	FLoadTimeProfilerTracePrivate::Init();
-#endif
-
 	AltEventQueues.Add(&EventQueue);
 	for (FAsyncLoadEventQueue2* Queue : AltEventQueues)
 	{
@@ -5062,7 +5058,7 @@ FAsyncPackage2::FAsyncPackage2(
 , GraphAllocator(InGraphAllocator)
 , ImportStore(AsyncLoadingThread.GetPackageStore(), AsyncLoadingThread.GlobalImportStore, AsyncLoadingThread.LoadedPackageStore, Desc)
 {
-	TRACE_LOADTIME_NEW_ASYNC_PACKAGE(this, Desc.DiskPackageName);
+	TRACE_LOADTIME_NEW_ASYNC_PACKAGE(this);
 	AddRequestID(Desc.RequestID);
 
 	ExportBundlesSize = Desc.StoreEntry->ExportBundlesSize;
diff --git a/Engine/Source/Runtime/CoreUObject/Private/Serialization/LoadTimeTrace.cpp b/Engine/Source/Runtime/CoreUObject/Private/Serialization/LoadTimeTrace.cpp
--- a/Engine/Source/Runtime/CoreUObject/Private/Serialization/LoadTimeTrace.cpp
+++ b/Engine/Source/Runtime/CoreUObject/Private/Serialization/LoadTimeTrace.cpp
@@ -106,10 +106,6 @@ UE_TRACE_EVENT_BEGIN(LoadTime, ClassInfo, NoSync|Important)
 	UE_TRACE_EVENT_FIELD(UE::Trace::AnsiString, Name)
 UE_TRACE_EVENT_END()
 
-void FLoadTimeProfilerTracePrivate::Init()
-{
-}
-
 void FLoadTimeProfilerTracePrivate::OutputStartAsyncLoading()
 {
 	UE_TRACE_LOG(LoadTime, StartAsyncLoading, LoadTimeChannel)
@@ -144,21 +140,10 @@ void FLoadTimeProfilerTracePrivate::OutputEndRequest(uint64 RequestId)
 		<< EndRequest.RequestId(RequestId);
 }
 
-void FLoadTimeProfilerTracePrivate::OutputNewAsyncPackage(const void* AsyncPackage, FName PackageName)
+void FLoadTimeProfilerTracePrivate::OutputNewAsyncPackage(const void* AsyncPackage)
 {
-	TCHAR Buffer[FName::StringBufferSize];
-	uint16 NameSize = (PackageName.ToString(Buffer) + 1) * sizeof(TCHAR);
-	UE_TRACE_LOG(LoadTime, NewAsyncPackage, LoadTimeChannel, NameSize)
-		<< NewAsyncPackage.AsyncPackage(AsyncPackage)
-		<< NewAsyncPackage.Attachment(Buffer, NameSize);
-}
-
-void FLoadTimeProfilerTracePrivate::OutputNewAsyncPackage(const void* AsyncPackage, const FString& PackageName)
-{
-	uint16 NameSize = (PackageName.Len() + 1) * sizeof(TCHAR);
-	UE_TRACE_LOG(LoadTime, NewAsyncPackage, LoadTimeChannel, NameSize)
-		<< NewAsyncPackage.AsyncPackage(AsyncPackage)
-		<< NewAsyncPackage.Attachment(*PackageName, NameSize);
+	UE_TRACE_LOG(LoadTime, NewAsyncPackage, LoadTimeChannel)
+		<< NewAsyncPackage.AsyncPackage(AsyncPackage);
 }
 
 void FLoadTimeProfilerTracePrivate::OutputBeginLoadAsyncPackage(const void* AsyncPackage)
@@ -181,13 +166,17 @@ void FLoadTimeProfilerTracePrivate::OutputDestroyAsyncPackage(const void* AsyncP
 		<< DestroyAsyncPackage.AsyncPackage(AsyncPackage);
 }
 
-void FLoadTimeProfilerTracePrivate::OutputPackageSummary(const void* AsyncPackage, uint32 TotalHeaderSize, uint32 ImportCount, uint32 ExportCount)
+void FLoadTimeProfilerTracePrivate::OutputPackageSummary(const void* AsyncPackage, const FName& PackageName, uint32 TotalHeaderSize, uint32 ImportCount, uint32 ExportCount)
 {
-	UE_TRACE_LOG(LoadTime, PackageSummary, LoadTimeChannel)
+	TCHAR Buffer[FName::StringBufferSize];
+	uint16 NameSize = (PackageName.ToString(Buffer) + 1) * sizeof(TCHAR);
+
+	UE_TRACE_LOG(LoadTime, PackageSummary, LoadTimeChannel, NameSize)
 		<< PackageSummary.AsyncPackage(AsyncPackage)
 		<< PackageSummary.TotalHeaderSize(TotalHeaderSize)
 		<< PackageSummary.ImportCount(ImportCount)
-		<< PackageSummary.ExportCount(ExportCount);
+		<< PackageSummary.ExportCount(ExportCount)
+		<< PackageSummary.Attachment(Buffer, NameSize);
 }
 
 void FLoadTimeProfilerTracePrivate::OutputAsyncPackageRequestAssociation(const void* AsyncPackage, uint64 RequestId)
diff --git a/Engine/Source/Runtime/CoreUObject/Private/Serialization/LoadTimeTracePrivate.h b/Engine/Source/Runtime/CoreUObject/Private/Serialization/LoadTimeTracePrivate.h
--- a/Engine/Source/Runtime/CoreUObject/Private/Serialization/LoadTimeTracePrivate.h
+++ b/Engine/Source/Runtime/CoreUObject/Private/Serialization/LoadTimeTracePrivate.h
@@ -11,18 +11,16 @@ class FName;
 
 struct FLoadTimeProfilerTracePrivate
 {
-	static void Init();
 	static void OutputStartAsyncLoading();
 	static void OutputSuspendAsyncLoading();
 	static void OutputResumeAsyncLoading();
-	static void OutputNewAsyncPackage(const void* AsyncPackage, FName PackageName);
-	static void OutputNewAsyncPackage(const void* AsyncPackage, const FString& PackageName);
+	static void OutputNewAsyncPackage(const void* AsyncPackage);
 	static void OutputBeginLoadAsyncPackage(const void* AsyncPackage);
 	static void OutputEndLoadAsyncPackage(const void* AsyncPackage);
 	static void OutputDestroyAsyncPackage(const void* AsyncPackage);
 	static void OutputBeginRequest(uint64 RequestId);
 	static void OutputEndRequest(uint64 RequestId);
-	static void OutputPackageSummary(const void* AsyncPackage, uint32 TotalHeaderSize, uint32 ImportCount, uint32 ExportCount);
+	static void OutputPackageSummary(const void* AsyncPackage, const FName& PackageName, uint32 TotalHeaderSize, uint32 ImportCount, uint32 ExportCount);
 	static void OutputAsyncPackageImportDependency(const void* Package, const void* ImportedPackage);
 	static void OutputAsyncPackageRequestAssociation(const void* AsyncPackage, uint64 RequestId);
 	static void OutputClassInfo(const UClass* Class, const FName& Name);
@@ -66,8 +64,8 @@ struct FLoadTimeProfilerTracePrivate
 #define TRACE_LOADTIME_END_REQUEST(RequestId) \
 	FLoadTimeProfilerTracePrivate::OutputEndRequest(RequestId);
 
-#define TRACE_LOADTIME_NEW_ASYNC_PACKAGE(AsyncPackage, PackageName) \
-	FLoadTimeProfilerTracePrivate::OutputNewAsyncPackage(AsyncPackage, PackageName)
+#define TRACE_LOADTIME_NEW_ASYNC_PACKAGE(AsyncPackage) \
+	FLoadTimeProfilerTracePrivate::OutputNewAsyncPackage(AsyncPackage)
 
 #define TRACE_LOADTIME_BEGIN_LOAD_ASYNC_PACKAGE(AsyncPackage) \
 	FLoadTimeProfilerTracePrivate::OutputBeginLoadAsyncPackage(AsyncPackage)
@@ -78,8 +76,8 @@ struct FLoadTimeProfilerTracePrivate
 #define TRACE_LOADTIME_DESTROY_ASYNC_PACKAGE(AsyncPackage) \
 	FLoadTimeProfilerTracePrivate::OutputDestroyAsyncPackage(AsyncPackage);
 
-#define TRACE_LOADTIME_PACKAGE_SUMMARY(AsyncPackage, TotalHeaderSize, ImportCount, ExportCount) \
-	FLoadTimeProfilerTracePrivate::OutputPackageSummary(AsyncPackage, TotalHeaderSize, ImportCount, ExportCount);
+#define TRACE_LOADTIME_PACKAGE_SUMMARY(AsyncPackage, PackageName, TotalHeaderSize, ImportCount, ExportCount) \
+	FLoadTimeProfilerTracePrivate::OutputPackageSummary(AsyncPackage, PackageName, TotalHeaderSize, ImportCount, ExportCount);
 
 #define TRACE_LOADTIME_ASYNC_PACKAGE_REQUEST_ASSOCIATION(AsyncPackage, RequestId) \
 	FLoadTimeProfilerTracePrivate::OutputAsyncPackageRequestAssociation(AsyncPackage, RequestId);
