SoftObjectPath: Respect PPF_Delimited flag

diff --git a/Engine/Source/Runtime/CoreUObject/Private/UObject/SoftObjectPath.cpp b/Engine/Source/Runtime/CoreUObject/Private/UObject/SoftObjectPath.cpp
--- a/Engine/Source/Runtime/CoreUObject/Private/UObject/SoftObjectPath.cpp
+++ b/Engine/Source/Runtime/CoreUObject/Private/UObject/SoftObjectPath.cpp
@@ -336,7 +336,16 @@ bool FSoftObjectPath::ExportTextItem(FString& ValueStr, FSoftObjectPath const& D
 		FSoftObjectPath Temp = *this;
 		Temp.PreSavePath();
 
-		ValueStr += Temp.ToString();
+		if (PortFlags & PPF_Delimited)
+		{
+			ValueStr += TEXT("\"");
+			ValueStr += Temp.ToString();
+			ValueStr += TEXT("\"");
+		}
+		else
+		{
+			ValueStr += Temp.ToString();
+		}
 	}
 	else
 	{
