SoftObjectPath: Respect PPF_Delimited flag

diff --git a/Engine/Source/Runtime/CoreUObject/Private/UObject/SoftObjectPath.cpp b/Engine/Source/Runtime/CoreUObject/Private/UObject/SoftObjectPath.cpp
--- a/Engine/Source/Runtime/CoreUObject/Private/UObject/SoftObjectPath.cpp
+++ b/Engine/Source/Runtime/CoreUObject/Private/UObject/SoftObjectPath.cpp
@@ -339,7 +339,7 @@ bool FSoftObjectPath::ExportTextItem(FString& ValueStr, FSoftObjectPath const& D
 		if (PortFlags & PPF_Delimited)
 		{
 			ValueStr += TEXT("\"");
-			ValueStr += Temp.ToString();
+			ValueStr += Temp.ToString().ReplaceQuotesWithEscapedQuotes();
 			ValueStr += TEXT("\"");
 		}
 		else
