PKG_ContainsNoAsset package flag: Package doesn't contain any asset object (although asset tags can be present)

diff --git a/Engine/Source/Runtime/CoreUObject/Public/UObject/ObjectMacros.h b/Engine/Source/Runtime/CoreUObject/Public/UObject/ObjectMacros.h
--- a/Engine/Source/Runtime/CoreUObject/Public/UObject/ObjectMacros.h
+++ b/Engine/Source/Runtime/CoreUObject/Public/UObject/ObjectMacros.h
@@ -113,7 +113,7 @@ enum EPackageFlags
 	PKG_Developer					= 0x00000080,	///< Developer module
 	PKG_UncookedOnly				= 0x00000100,	///< Loaded only in uncooked builds (i.e. runtime in editor)
 	PKG_Cooked						= 0x00000200,	///< Package is cooked
-//	PKG_Unused						= 0x00000400,
+	PKG_ContainsNoAsset				= 0x00000400,	///< Package doesn't contain any asset object (although asset tags can be present)
 //	PKG_Unused						= 0x00000800,
 //	PKG_Unused						= 0x00001000,
 	PKG_UnversionedProperties		= 0x00002000,   ///< Uses unversioned property serialization instead of versioned tagged property serialization
