Add ability to check if a TObjectPtr is null without resolving it.
Add ability to check if a TObjectPtr is resolved.
Add GetPath method to TObjectPtr and use it from FSoftObjectPath.

diff --git a/Engine/Source/Runtime/CoreUObject/Private/UObject/ObjectPtr.cpp b/Engine/Source/Runtime/CoreUObject/Private/UObject/ObjectPtr.cpp
new file mode 100644
--- /dev/null
+++ b/Engine/Source/Runtime/CoreUObject/Private/UObject/ObjectPtr.cpp
@@ -0,0 +1,22 @@
+// Copyright Epic Games, Inc. All Rights Reserved.
+
+#include "UObject/ObjectPtr.h"
+#include "Misc/StringBuilder.h"
+
+FString FObjectPtr::GetPath() const
+{
+	FObjectRef Ref = MakeObjectRef(Handle);
+	if (Ref.PackageName == NAME_None)
+	{
+		return FString();
+	}
+	FObjectPathId::ResolvedNameContainerType ResolvedNames;
+	Ref.ObjectPath.Resolve(ResolvedNames);
+	TStringBuilder<FName::StringBufferSize> CompletePath;
+	CompletePath << Ref.PackageName;
+	for (int32 ResolvedNameIndex = 0; ResolvedNameIndex < ResolvedNames.Num(); ++ResolvedNameIndex)
+	{
+		CompletePath << (ResolvedNameIndex == 1 ? SUBOBJECT_DELIMITER_CHAR : '.') << ResolvedNames[ResolvedNameIndex];
+	}
+	return FString(CompletePath);
+}
diff --git a/Engine/Source/Runtime/CoreUObject/Private/UObject/SoftObjectPath.cpp b/Engine/Source/Runtime/CoreUObject/Private/UObject/SoftObjectPath.cpp
--- a/Engine/Source/Runtime/CoreUObject/Private/UObject/SoftObjectPath.cpp
+++ b/Engine/Source/Runtime/CoreUObject/Private/UObject/SoftObjectPath.cpp
@@ -19,52 +19,6 @@ FSoftObjectPath::FSoftObjectPath(const UObject* InObject)
 	}
 }
 
-void FSoftObjectPath::InitializeFromObjectHandle(FObjectHandle InObjectHandle)
-{
-	if (!IsObjectHandleNull(InObjectHandle))
-	{
-		if (IsObjectHandleResolved(InObjectHandle))
-		{
-			if (UObject* ResolvedObject = ResolveObjectHandle(InObjectHandle))
-			{
-				SetPath(ResolvedObject->GetPathName());
-			}
-		}
-		else
-		{
-			//Try to compose the path without forcing the referenced data to be loaded
-			FObjectRef Ref = MakeObjectRef(InObjectHandle);
-			FObjectPathId::ResolvedNameContainerType ResolvedNames;
-			Ref.ObjectPath.Resolve(ResolvedNames);
-			TStringBuilder<128> CompletePath;
-			CompletePath.Append(Ref.PackageName.ToString());
-			CompletePath.Append(TEXT('.'));
-			for (int32 ResolvedNameIndex = 0; ResolvedNameIndex < ResolvedNames.Num(); ++ResolvedNameIndex)
-			{
-				switch (ResolvedNameIndex)
-				{
-				case 0:
-				{
-					break;
-				}
-				case 1:
-				{
-					CompletePath << SUBOBJECT_DELIMITER_CHAR;
-					break;
-				}
-				default:
-				{
-					CompletePath << '.';
-					break;
-				}
-				}
-				CompletePath << ResolvedNames[ResolvedNameIndex];
-			}
-			SetPath(CompletePath.ToString());
-		}
-	}
-}
-
 FString FSoftObjectPath::ToString() const
 {
 	// Most of the time there is no sub path so we can do a single string allocation
diff --git a/Engine/Source/Runtime/CoreUObject/Public/UObject/ObjectPtr.h b/Engine/Source/Runtime/CoreUObject/Public/UObject/ObjectPtr.h
--- a/Engine/Source/Runtime/CoreUObject/Public/UObject/ObjectPtr.h
+++ b/Engine/Source/Runtime/CoreUObject/Public/UObject/ObjectPtr.h
@@ -135,8 +135,15 @@ public:
 
 	// Forcing resolve for null check.  Non null unresolved pointers can become null resolved pointers.
 	FORCEINLINE bool IsNull() const { return ResolveObjectHandle(Handle) == nullptr; }
+	FORCEINLINE bool IsNullNoResolve() const { return IsObjectHandleNull(Handle); }
+
 	FORCEINLINE bool operator!() const { return IsNull(); }
 	explicit FORCEINLINE operator bool() const { return !IsNull(); }
+
+	FORCEINLINE bool IsResolved() const { return IsObjectHandleResolved(Handle); }
+
+	COREUOBJECT_API FString GetPath() const;
+
 	FORCEINLINE FObjectHandle GetHandle() const { return Handle; }
 	FORCEINLINE FObjectHandle& GetHandleRef() const { return Handle; }
 
@@ -294,8 +301,11 @@ public:
 	explicit FORCEINLINE operator T*& () { return GetInternalRef(); }
 
 	using FObjectPtr::IsNull;
+	using FObjectPtr::IsNullNoResolve;
 	using FObjectPtr::operator!;
 	using FObjectPtr::operator bool;
+	using FObjectPtr::IsResolved;
+	using FObjectPtr::GetPath;
 	using FObjectPtr::GetHandle;
 
 	using FObjectPtr::IsA;
diff --git a/Engine/Source/Runtime/CoreUObject/Public/UObject/SoftObjectPath.h b/Engine/Source/Runtime/CoreUObject/Public/UObject/SoftObjectPath.h
--- a/Engine/Source/Runtime/CoreUObject/Public/UObject/SoftObjectPath.h
+++ b/Engine/Source/Runtime/CoreUObject/Public/UObject/SoftObjectPath.h
@@ -53,13 +53,13 @@ struct COREUOBJECT_API FSoftObjectPath
 
 	FSoftObjectPath(FObjectPtr InObject)
 	{
-		InitializeFromObjectHandle(InObject.GetHandle());
+		SetPath(InObject.GetPath());
 	}
 
 	template <typename T>
 	FSoftObjectPath(const TObjectPtr<T>& InObject)
 	{
-		InitializeFromObjectHandle(InObject.GetHandle());
+		SetPath(InObject.GetPath());
 	}
 
 	~FSoftObjectPath() {}
@@ -245,9 +245,6 @@ private:
 	/** Package names currently being duplicated, needed by FixupForPIE */
 	static TSet<FName> PIEPackageNames;
 
-	/** Initialize from a FObjectHandle/FObjectPtr/TObjectPtr which may or may not have the referenced object in memory */
-	void InitializeFromObjectHandle(FObjectHandle InObjectHandle);
-
 	UObject* ResolveObjectInternal() const;
 	UObject* ResolveObjectInternal(const TCHAR* PathString) const;
 
