Avoid resolving wrapped object pointer references when serializing to archive for memory counting purposes.

diff --git a/Engine/Source/Runtime/CoreUObject/Private/Serialization/ArchiveUObject.cpp b/Engine/Source/Runtime/CoreUObject/Private/Serialization/ArchiveUObject.cpp
--- a/Engine/Source/Runtime/CoreUObject/Private/Serialization/ArchiveUObject.cpp
+++ b/Engine/Source/Runtime/CoreUObject/Private/Serialization/ArchiveUObject.cpp
@@ -60,6 +60,11 @@ FArchive& FArchiveUObject::SerializeLazyObjectPtr(FArchive& Ar, FLazyObjectPtr&
 
 FArchive& FArchiveUObject::SerializeObjectPtr(FArchive& Ar, FObjectPtr& Value)
 {
+	if (Ar.IsCountingMemory() && !(Ar.IsLoading() || Ar.IsSaving()) && !Value.IsResolved())
+	{
+		return Ar;
+	}
+
 	// Default behavior is to fully resolve the reference (if we're not loading) and send it through
 	// the raw UObject* serialization codepath and ensure that the result is saved back into
 	// the FObjectPtr afterwards.  There will be many use cases where this is
diff --git a/Engine/Source/Runtime/CoreUObject/Private/Tests/ObjectPtrTest.cpp b/Engine/Source/Runtime/CoreUObject/Private/Tests/ObjectPtrTest.cpp
--- a/Engine/Source/Runtime/CoreUObject/Private/Tests/ObjectPtrTest.cpp
+++ b/Engine/Source/Runtime/CoreUObject/Private/Tests/ObjectPtrTest.cpp
@@ -179,7 +179,7 @@ bool FObjectPtrTestDefaultSerialize::RunTest(const FString& Parameters)
 	ObjectRefMetrics.TestNumFailedResolves(TEXT("Unexpected resolve failure after initializing an FObjectPtr"), 0);
 	ObjectRefMetrics.TestNumReads(TEXT("NumReads should not change when initializing an FObjectPtr"), 0);
 
-	FArchiveCountMem Writer(nullptr);
+	FArchiveUObject Writer;
 	Writer << DefaultTexturePtr;
 
 	ObjectRefMetrics.TestNumResolves(TEXT("Serializing an FObjectPtr should force it to resolve"), 1);
@@ -296,6 +296,32 @@ bool FObjectPtrTestHashConsistency::RunTest(const FString& Parameters)
 	return true;
 }
 
+// @TODO: OBJPTR: We should have a test that ensures that lazy loading of an object with an external package is handled correctly.
+//				  This should also include external packages in the outer chain of the target object.
+// IMPLEMENT_CUSTOM_SIMPLE_AUTOMATION_TEST(FObjectPtrTestExternalPackages, FObjectPtrTestBase, TEST_NAME_ROOT TEXT(".ExternalPackages"), ObjectPtrTestFlags)
+// bool FObjectPtrTestExternalPackages::RunTest(const FString& Parameters)
+// {
+// 	const FName TestExternalPackage1Name(TEXT("/Engine/Test/ObjectPtrExternalPackages1/Transient"));
+// 	UPackage* TestExternalPackage1 = NewObject<UPackage>(nullptr, TestExternalPackage1Name, RF_Transient);
+// 	TestExternalPackage1->SetPackageFlags(PKG_EditorOnly | PKG_ContainsMapData);
+// 	TestExternalPackage1->AddToRoot();
+
+// 	const FName TestPackage1Name(TEXT("/Engine/Test/ObjectPtrExternalPackages1/Transient"));
+// 	UPackage* TestPackage1 = NewObject<UPackage>(nullptr, TestPackage1Name, RF_Transient);
+// 	TestPackage1->AddToRoot();
+// 	UObject* TestOuter1 = NewObject<UTestDummyObject>(TestPackage1, TEXT("TestOuter1"));
+// 	UObject* TestOuter2 = NewObject<UTestDummyObject>(TestPackage1, TEXT("TestOuter2"));
+// 	UObject* TestOuter3 = NewObject<UTestDummyObject>(TestPackage1, TEXT("TestOuter3"));
+// 	UObject* TestOuter4 = NewObject<UTestDummyObject>(TestPackage1, TEXT("TestOuter4"));
+
+// 	UObject* TestPublicObject = NewObject<UTestDummyObject>(TestOuter4, TEXT("TestPublicObject"), RF_Public);
+// 	TestPublicObject->SetExternalPackage(TestExternalPackage1);
+
+// 	TestPackage1->RemoveFromRoot();
+// 	TestExternalPackage1->RemoveFromRoot();
+// 	return true;
+// }
+
 // @TODO: OBJPTR: We should have a test that ensures that we can (de)serialize an FObjectPtr to FLinkerSave/FLinkerLoad and that upon load the object
 //			pointer is not resolved if we are in a configuration that supports lazy load.  This is proving difficult due to the restrictions around how
 //			FLinkerSave/FLinkerLoad is used.
