Avoid resolving references where possible during common operations in the blueprint compilation and re-instancing flow.  This is to aid in the lazy loading of objects.

diff --git a/Engine/Source/Editor/UnrealEd/Private/Kismet2/KismetReinstanceUtilities.cpp b/Engine/Source/Editor/UnrealEd/Private/Kismet2/KismetReinstanceUtilities.cpp
--- a/Engine/Source/Editor/UnrealEd/Private/Kismet2/KismetReinstanceUtilities.cpp
+++ b/Engine/Source/Editor/UnrealEd/Private/Kismet2/KismetReinstanceUtilities.cpp
@@ -947,6 +947,28 @@ public:
 		}
 		return *this;
 	}
+
+	/**
+	 * Serializes a resolved or unresolved object reference
+	 */
+	FArchive& operator<<( FObjectPtr& Obj )
+	{
+		if (ShouldSkipReplacementCheckForObjectPtr(Obj, ReplacementMap, [] (const TPair<FFieldVariant, FFieldVariant>& ReplacementPair) -> const UObject*
+			{
+				if (ReplacementPair.Key.IsValid() && ReplacementPair.Key.IsUObject())
+				{
+					return ReplacementPair.Key.ToUObject();
+				}
+				return nullptr;
+			}))
+		{
+			return *this;
+		}
+
+		// Allow object references to go through the normal code path of resolving and running the raw pointer code path
+		return FArchiveReplaceObjectRefBase::operator<<(Obj);
+	}
+
 protected:
 	/** Map of objects to find references to -> object to replace references with */
 	const TMap<FFieldVariant, FFieldVariant>& ReplacementMap;
@@ -1378,6 +1400,19 @@ namespace InstancedPropertyUtils
 		}
 
 		//----------------------------------------------------------------------
+		FArchive& operator<<(FObjectPtr& Obj)
+		{
+			// Avoid resolving an FObjectPtr if it is not an instanced property
+			FProperty* SerializingProperty = GetSerializedProperty();
+			const bool bHasInstancedValue = SerializingProperty && SerializingProperty->HasAnyPropertyFlags(CPF_PersistentInstance);
+			if (!bHasInstancedValue)
+			{
+				return *this;
+			}
+
+			return FArchiveUObject::operator<<(Obj);
+		}
+
 		FArchive& operator<<(UObject*& Obj)
 		{
 			if (Obj != nullptr)
@@ -1446,6 +1481,19 @@ namespace InstancedPropertyUtils
 		}
 
 		//----------------------------------------------------------------------
+		FArchive& operator<<(FObjectPtr& Obj)
+		{
+			// Avoid resolving an FObjectPtr if it is not an instanced property
+			FProperty* SerializingProperty = GetSerializedProperty();
+			const bool bHasInstancedValue = SerializingProperty && SerializingProperty->HasAnyPropertyFlags(CPF_PersistentInstance);
+			if (!bHasInstancedValue)
+			{
+				return *this;
+			}
+
+			return FArchiveUObject::operator<<(Obj);
+		}
+
 		FArchive& operator<<(UObject*& Obj)
 		{
 			if (Obj == nullptr)
diff --git a/Engine/Source/Runtime/CoreUObject/Public/Serialization/ArchiveReplaceObjectRef.h b/Engine/Source/Runtime/CoreUObject/Public/Serialization/ArchiveReplaceObjectRef.h
--- a/Engine/Source/Runtime/CoreUObject/Public/Serialization/ArchiveReplaceObjectRef.h
+++ b/Engine/Source/Runtime/CoreUObject/Public/Serialization/ArchiveReplaceObjectRef.h
@@ -32,6 +32,64 @@ public:
 	virtual FString GetArchiveName() const { return TEXT("ReplaceObjectRef"); }
 	
 protected:
+
+	template <typename ContainerType, typename ElementToObjectType>
+	bool ShouldSkipReplacementCheckForObjectPtr(FObjectPtr& Obj, const ContainerType& ReplacementContainer, const ElementToObjectType& ElementToObject)
+	{
+		if (Obj.IsResolved())
+		{
+			return false;
+		}
+
+		if (!CanIgnoreUnresolvedImports.IsSet())
+		{
+			bool bCanIgnoreUnresolvedImportsLocal = ReplacementContainer.Num() == 0;
+			if (!bCanIgnoreUnresolvedImportsLocal)
+			{
+				bCanIgnoreUnresolvedImportsLocal = true; // Will be set to false if any of the criteria in the loop below are met.
+				for (const auto& ReplacementElem : ReplacementContainer)
+				{
+					if (const UObject* ReplacementObject = ElementToObject(ReplacementElem))
+					{
+						if (ReplacementObject->GetOutermost() != GetTransientPackage())
+						{
+							bCanIgnoreUnresolvedImportsLocal = false;
+							break;
+						}
+					}
+				}
+			}
+			CanIgnoreUnresolvedImports = bCanIgnoreUnresolvedImportsLocal;
+		}
+
+		if (CanIgnoreUnresolvedImports.GetValue())
+		{
+			return true;
+		}
+
+		if (UClass* ReferenceClass = Obj.GetClass())
+		{
+			bool bReferenceTypeIsRelatedToReplacementType = false;
+			for (const auto& ReplacementElem : ReplacementContainer)
+			{
+				if (const UObject* ReplacementObject = ElementToObject(ReplacementElem))
+				{
+					if (ReplacementObject->IsA(ReferenceClass) || Obj.IsA(ReplacementObject->GetClass()))
+					{
+						bReferenceTypeIsRelatedToReplacementType = true;
+						break;
+					}
+				}
+			}
+			if (!bReferenceTypeIsRelatedToReplacementType)
+			{
+				return true;
+			}
+		}
+
+		return false;
+	}
+
 	/**
 	* Serializes a single object
 	*/
@@ -57,6 +115,11 @@ protected:
 	* should be set to null
 	*/
 	bool bNullPrivateReferences;
+
+	/**
+	* Whether unresolved references to objects in other packages can be ignored when searching
+	*/
+	TOptional<bool> CanIgnoreUnresolvedImports;
 };
 
 /*----------------------------------------------------------------------------
@@ -176,6 +239,23 @@ public:
 		return *this;
 	}
 
+	/**
+	 * Serializes a resolved or unresolved object reference
+	 */
+	FArchive& operator<<( FObjectPtr& Obj )
+	{
+		if (ShouldSkipReplacementCheckForObjectPtr(Obj, ReplacementMap, [] (const TPair<T*, T*>& ReplacementPair) -> const UObject*
+			{
+				return ReplacementPair.Key;
+			}))
+		{
+			return *this;
+		}
+
+		// Allow object references to go through the normal code path of resolving and running the raw pointer code path
+		return FArchiveReplaceObjectRefBase::operator<<(Obj);
+	}
+
 protected:
 	/** Map of objects to find references to -> object to replace references with */
 	const TMap<T*,T*>& ReplacementMap;
diff --git a/Engine/Source/Runtime/Engine/Private/Actor.cpp b/Engine/Source/Runtime/Engine/Private/Actor.cpp
--- a/Engine/Source/Runtime/Engine/Private/Actor.cpp
+++ b/Engine/Source/Runtime/Engine/Private/Actor.cpp
@@ -738,7 +738,20 @@ void AActor::Serialize(FArchive& Ar)
 				// Ignore transient properties since they won't be serialized
 				if(!ObjProp->HasAnyPropertyFlags(CPF_Transient))
 				{
-					UActorComponent* ActorComponent = Cast<UActorComponent>(ObjProp->GetObjectPropertyValue_InContainer(this));
+					UActorComponent* ActorComponent = nullptr;
+					if (FObjectPtrProperty* ObjPtrProp = CastField<FObjectPtrProperty>(ObjProp))
+					{
+						FObjectPtr& ObjectPtr = (FObjectPtr&)ObjPtrProp->GetPropertyValue_InContainer(this);
+						if (!ObjectPtr.IsNullNoResolve() && ObjectPtr.IsA<UActorComponent>())
+						{
+							ActorComponent = Cast<UActorComponent>(ObjectPtr.Get());
+						}
+					}
+					else
+					{
+						ActorComponent = Cast<UActorComponent>(ObjProp->GetObjectPropertyValue_InContainer(this));
+					}
+
 					if(ActorComponent != nullptr && ActorComponent->CreationMethod == EComponentCreationMethod::Native)
 					{
 						NativeConstructedComponentToPropertyMap.Add(ActorComponent->GetFName(), ObjProp);
