Do not count shallow pointer comparisons with a TObjectPtr as "accesses" that go through the access tracker.
Remove pointless IsNull comparison before pointer comparison as IsNull has previously been modified to force resolve and it's no longer beneficial to attempt checking null before resolving the target address.
Modify USceneComponent::GetAttachChildren to return a TArray<TObjectPtr<T>> instead of TArray<T*> to avoid recording an access for every element in that potentially large and frequently accessed array as part of the conversion to a TArray<T*>.

diff --git a/Engine/Plugins/Runtime/Steam/SteamAudio/Source/SteamAudioEditor/Private/IndirectBaker.cpp b/Engine/Plugins/Runtime/Steam/SteamAudio/Source/SteamAudioEditor/Private/IndirectBaker.cpp
--- a/Engine/Plugins/Runtime/Steam/SteamAudio/Source/SteamAudioEditor/Private/IndirectBaker.cpp
+++ b/Engine/Plugins/Runtime/Steam/SteamAudio/Source/SteamAudioEditor/Private/IndirectBaker.cpp
@@ -258,7 +258,7 @@ namespace SteamAudio
 							//UE_LOG(LogTemp, Warning, TEXT("Audio Component [%s] UserID is now set to [%s]"), *AudioComponent->GetFullName(), *AudioComponent->AudioComponentUserID.ToString());
 
 							// Check if there's a phonon source child for this audio component
-							for (auto* ChildComp : AudioComponent->GetAttachChildren())
+							for (auto ChildComp : AudioComponent->GetAttachChildren())
 							{
 								auto* ChildPhononComponent = Cast<UPhononSourceComponent>(ChildComp);
 
diff --git a/Engine/Plugins/Runtime/Steam/SteamAudio/Source/SteamAudioEditor/Private/PhononSourceComponentDetails.cpp b/Engine/Plugins/Runtime/Steam/SteamAudio/Source/SteamAudioEditor/Private/PhononSourceComponentDetails.cpp
--- a/Engine/Plugins/Runtime/Steam/SteamAudio/Source/SteamAudioEditor/Private/PhononSourceComponentDetails.cpp
+++ b/Engine/Plugins/Runtime/Steam/SteamAudio/Source/SteamAudioEditor/Private/PhononSourceComponentDetails.cpp
@@ -126,7 +126,7 @@ namespace SteamAudio
 									//UE_LOG(LogTemp, Warning, TEXT("Audio Component [%s] UserID is now set to [%s]"), *AudioComponent->GetFullName(), *AudioComponent->AudioComponentUserID.ToString());
 
 									// Check if there's a phonon source child for this audio component
-									for (auto* ChildComp : AudioComponent->GetAttachChildren())
+									for (auto ChildComp : AudioComponent->GetAttachChildren())
 									{
 										auto* ChildPhononComponent = Cast<UPhononSourceComponent>(ChildComp);
 
diff --git a/Engine/Source/Runtime/CoreUObject/Public/UObject/ObjectHandle.h b/Engine/Source/Runtime/CoreUObject/Public/UObject/ObjectHandle.h
--- a/Engine/Source/Runtime/CoreUObject/Public/UObject/ObjectHandle.h
+++ b/Engine/Source/Runtime/CoreUObject/Public/UObject/ObjectHandle.h
@@ -97,6 +97,7 @@ inline UObject* ReadObjectHandlePointerNoCheck(FObjectHandle Handle);
 inline FPackedObjectRef ReadObjectHandlePackedObjectRefNoCheck(FObjectHandle Handle);
 
 inline UObject* ResolveObjectHandle(FObjectHandle& Handle);
+inline UObject* ResolveObjectHandleNoRead(FObjectHandle& Handle);
 inline UClass* ResolveObjectHandleClass(FObjectHandle Handle);
 
 /** Read the handle as a pointer if resolved, and otherwise return null. */
@@ -324,13 +325,12 @@ inline FObjectHandle MakeObjectHandle(UObject* Object) { return Object; }
 
 #endif
 
-inline UObject* ResolveObjectHandle(FObjectHandle& Handle)
+inline UObject* ResolveObjectHandleNoRead(FObjectHandle& Handle)
 {
 	FObjectHandle LocalHandle = Handle;
 	if (IsObjectHandleResolved(LocalHandle))
 	{
 		UObject* ResolvedObject = ReadObjectHandlePointerNoCheck(LocalHandle);
-		ObjectHandle_Private::OnHandleRead(ResolvedObject);
 		return ResolvedObject;
 	}
 	else
@@ -338,11 +338,17 @@ inline UObject* ResolveObjectHandle(FObjectHandle& Handle)
 		LocalHandle = MakeObjectHandle(ResolvePackedObjectRef(ReadObjectHandlePackedObjectRefNoCheck(LocalHandle)));
 		UObject* ResolvedObject = ReadObjectHandlePointerNoCheck(LocalHandle);
 		Handle = LocalHandle;
-		ObjectHandle_Private::OnHandleRead(ResolvedObject);
 		return ResolvedObject;
 	}
 }
 
+inline UObject* ResolveObjectHandle(FObjectHandle& Handle)
+{
+	UObject* ResolvedObject = ResolveObjectHandleNoRead(Handle);
+	ObjectHandle_Private::OnHandleRead(ResolvedObject);
+	return ResolvedObject;
+}
+
 inline UClass* ResolveObjectHandleClass(FObjectHandle Handle)
 {
 	if (IsObjectHandleResolved(Handle))
diff --git a/Engine/Source/Runtime/CoreUObject/Public/UObject/ObjectPtr.h b/Engine/Source/Runtime/CoreUObject/Public/UObject/ObjectPtr.h
--- a/Engine/Source/Runtime/CoreUObject/Public/UObject/ObjectPtr.h
+++ b/Engine/Source/Runtime/CoreUObject/Public/UObject/ObjectPtr.h
@@ -134,7 +134,7 @@ public:
 	FORCEINLINE UObject& operator*() const { return *Get(); }
 
 	// Forcing resolve for null check.  Non null unresolved pointers can become null resolved pointers.
-	FORCEINLINE bool IsNull() const { return ResolveObjectHandle(Handle) == nullptr; }
+	FORCEINLINE bool IsNull() const { return ResolveObjectHandleNoRead(Handle) == nullptr; }
 	FORCEINLINE bool IsNullNoResolve() const { return IsObjectHandleNull(Handle); }
 
 	FORCEINLINE bool operator!() const { return IsNull(); }
@@ -163,6 +163,12 @@ public:
 		return IsA(T::StaticClass());
 	}
 
+protected:
+	FORCEINLINE UObject* GetNoRead() const
+	{
+		return ResolveObjectHandleNoRead(Handle);
+	}
+
 private:
 	friend FORCEINLINE uint32 GetTypeHash(const FObjectPtr& Object)
 	{
@@ -183,6 +189,12 @@ struct TIsTObjectPtr
 	enum { Value = false };
 };
 
+namespace ObjectPtr_Private
+{
+	template <typename T, typename U>
+	FORCEINLINE bool IsObjectPtrEqualToRawPtrOfRelatedType(const TObjectPtr<T>& Ptr, const U* Other);
+};
+
 template <typename T>
 struct TObjectPtr : private FObjectPtr
 {
@@ -328,7 +340,11 @@ public:
 
 	friend struct FObjectPtr;
 	template <typename U> friend struct TObjectPtr;
+	template <typename U, typename V> friend bool ObjectPtr_Private::IsObjectPtrEqualToRawPtrOfRelatedType(const TObjectPtr<U>& Ptr, const V* Other);
 private:
+
+	FORCEINLINE T* GetNoRead() const { return (T*)(FObjectPtr::GetNoRead()); }
+
 	// @TODO: OBJPTR: There is a risk of a gap in access tracking here.  The caller may get a mutable pointer, write to it, then
 	//			read from it.  That last read would happen without an access being recorded.  Not sure if there is a good way
 	//			to handle this case without forcing the calling code to be modified.
@@ -412,6 +428,12 @@ namespace ObjectPtr_Private
 		return Other.Get();
 	}
 
+	template <typename T, typename U>
+	FORCEINLINE bool IsObjectPtrEqualToRawPtrOfRelatedType(const TObjectPtr<T>& Ptr, const U* Other)
+	{
+		return Ptr.GetNoRead() == Other;
+	}
+
 	/** Perform shallow equality check between a TObjectPtr and another (non TObjectPtr) type that we can coerce to a pointer. */
 	template <
 		typename T,
@@ -423,8 +445,9 @@ namespace ObjectPtr_Private
 	>
 	FORCEINLINE bool IsObjectPtrEqual(const TObjectPtr<T>& Ptr, U&& Other)
 	{
-		auto OtherPtr = ObjectPtr_Private::CoerceToPointer<T>(Other);
-		return (Ptr.IsNull() == (OtherPtr == nullptr)) && (Ptr.Get() == OtherPtr);
+		// This function deliberately avoids the tracking code path as we are only doing
+		// a shallow pointer comparison.
+		return IsObjectPtrEqualToRawPtrOfRelatedType<T>(Ptr, ObjectPtr_Private::CoerceToPointer<T>(Other));
 	}
 }
 
diff --git a/Engine/Source/Runtime/Engine/Classes/Components/SceneComponent.h b/Engine/Source/Runtime/Engine/Classes/Components/SceneComponent.h
--- a/Engine/Source/Runtime/Engine/Classes/Components/SceneComponent.h
+++ b/Engine/Source/Runtime/Engine/Classes/Components/SceneComponent.h
@@ -631,7 +631,7 @@ public:
 	virtual bool IsAnySimulatingPhysics() const;
 
 	/** Get the SceneComponents that are attached to this component. */
-	const TArray<USceneComponent*>& GetAttachChildren() const;
+	const TArray<TObjectPtr<USceneComponent>>& GetAttachChildren() const;
 
 	/** Get the SceneComponent we are attached to. */
 	UFUNCTION(BlueprintCallable, Category="Utilities|Transformation")
@@ -1564,7 +1564,7 @@ private:
 //////////////////////////////////////////////////////////////////////////
 // USceneComponent inlines
 
-FORCEINLINE const TArray<USceneComponent*>& USceneComponent::GetAttachChildren() const
+FORCEINLINE const TArray<TObjectPtr<USceneComponent>>& USceneComponent::GetAttachChildren() const
 {
 	return AttachChildren;
 }
