Clear the LazyLoadable flag when clearing or setting the bulk data, as the bulkdata bytes no longer match the data on disk.
This fixes a mismatch between bulkdata flags and the bulkdata payload that occurred when loading a package from the EditorDomain that was saved with BULKData_WorkspaceDomainPayload even though the bulkdata was modified in memory on load.
#rnx

diff --git a/Engine/Source/Runtime/CoreUObject/Private/Serialization/BulkData.cpp b/Engine/Source/Runtime/CoreUObject/Private/Serialization/BulkData.cpp
--- a/Engine/Source/Runtime/CoreUObject/Private/Serialization/BulkData.cpp
+++ b/Engine/Source/Runtime/CoreUObject/Private/Serialization/BulkData.cpp
@@ -609,6 +609,7 @@ void* FUntypedBulkData::Lock( uint32 LockFlags )
 			check( AttachedAr == NULL );
 		}
 #endif // WITH_EDITOR
+		ClearBulkDataFlags(BULKDATA_LazyLoadable);
 	}
 	// Only read operations are allowed on returned memory.
 	else if( LockFlags & LOCK_READ_ONLY )
@@ -696,6 +697,7 @@ void FUntypedBulkData::RemoveBulkData()
 	// Resize to 0 elements.
 	ElementCount	= 0;
 	BulkData.Deallocate();
+	ClearBulkDataFlags(BULKDATA_LazyLoadable);
 }
 
 /**
