TryMatchCaseOnDisk: downgrade Unexpected non-matching LocalPath warning to a verbose log, as it has a known failure when loading packages from a directory junction.
#rnx

diff --git a/Engine/Source/Runtime/CoreUObject/Private/UObject/PackageResourceManagerFile.cpp b/Engine/Source/Runtime/CoreUObject/Private/UObject/PackageResourceManagerFile.cpp
--- a/Engine/Source/Runtime/CoreUObject/Private/UObject/PackageResourceManagerFile.cpp
+++ b/Engine/Source/Runtime/CoreUObject/Private/UObject/PackageResourceManagerFile.cpp
@@ -306,8 +306,8 @@ bool FPackageResourceManagerFile::TryMatchCaseOnDisk(const FPackagePath& Package
 		{
 			if (!FPackagePath::TryMatchCase(PackagePath, FilenameOnDisk, *OutPackagePath))
 			{
-				UE_LOG(LogPackageName, Warning,
-					TEXT("TryMatchCaseOnDisk: Unexpected non-matching LocalPath \"%s\" found when search for PackagePath \"%s\". Case will not be normalized."),
+				UE_LOG(LogPackageName, Verbose,
+					TEXT("TryMatchCaseOnDisk: Unexpected non-matching LocalPath \"%s\" found when searching for PackagePath \"%s\". Case will not be normalized."),
 					*FilenameOnDisk, *PackagePath.GetLocalFullPath());
 			}
 			OutPackagePath->SetHeaderExtension(ExtensionOnDisk, PackagePath.GetCustomExtension());
