Make a static version of FAssetData::IsUAsset to avoid the need to expensively create an FAssetData to test whether an object is the UAsset asset of its package.
#rnx

diff --git a/Engine/Source/Editor/PackagesDialog/Private/PackagesDialog.cpp b/Engine/Source/Editor/PackagesDialog/Private/PackagesDialog.cpp
--- a/Engine/Source/Editor/PackagesDialog/Private/PackagesDialog.cpp
+++ b/Engine/Source/Editor/PackagesDialog/Private/PackagesDialog.cpp
@@ -208,9 +208,10 @@ void FPackagesDialogModule::AddPackageItem(UPackage* InPackage, ECheckBoxState I
 	{
 		if (InnerObject->IsAsset())
 		{
-			if (FAssetData(InnerObject).IsUAsset())
+			if (FAssetData::IsUAsset(InnerObject))
 			{
 				// If we found the primary asset, use it
+				FoundAsset = InnerObject;
 				return false;
 			}
 			// Otherwise, keep the first found asset but keep looking for a primary asset
diff --git a/Engine/Source/Editor/UnrealEd/Private/AssetDeleteModel.cpp b/Engine/Source/Editor/UnrealEd/Private/AssetDeleteModel.cpp
--- a/Engine/Source/Editor/UnrealEd/Private/AssetDeleteModel.cpp
+++ b/Engine/Source/Editor/UnrealEd/Private/AssetDeleteModel.cpp
@@ -604,9 +604,8 @@ FPendingDelete::FPendingDelete(UObject* InObject)
 		}
 	}
 
-	FAssetData InAssetData(InObject);
 	// Filter out any non assets
-	if ( !InAssetData.IsUAsset() )
+	if ( !FAssetData::IsUAsset(InObject) )
 	{
 		bIsInternal = true;
 	}
diff --git a/Engine/Source/Runtime/CoreUObject/Private/AssetRegistry/AssetData.cpp b/Engine/Source/Runtime/CoreUObject/Private/AssetRegistry/AssetData.cpp
--- a/Engine/Source/Runtime/CoreUObject/Private/AssetRegistry/AssetData.cpp
+++ b/Engine/Source/Runtime/CoreUObject/Private/AssetRegistry/AssetData.cpp
@@ -112,6 +112,24 @@ FAssetData::FAssetData(const UObject* InAsset, bool bAllowBlueprintClass)
 	}
 }
 
+bool FAssetData::IsUAsset(UObject* InAsset)
+{
+	if (InAsset == nullptr)
+	{
+		return false;
+	}
+
+	const UPackage* Package = InAsset->GetPackage();
+
+	TStringBuilder<FName::StringBufferSize> AssetNameStrBuilder;
+	InAsset->GetPathName(Package, AssetNameStrBuilder);
+
+	TStringBuilder<FName::StringBufferSize> PackageNameStrBuilder;
+	Package->GetFName().AppendString(PackageNameStrBuilder);
+
+	return DetectIsUAssetByNames(PackageNameStrBuilder, AssetNameStrBuilder);
+}
+
 void FAssetData::SetTagsAndAssetBundles(FAssetDataTagMap&& Tags)
 {
 	using namespace UE::AssetData::Private;
diff --git a/Engine/Source/Runtime/CoreUObject/Public/AssetRegistry/AssetData.h b/Engine/Source/Runtime/CoreUObject/Public/AssetRegistry/AssetData.h
--- a/Engine/Source/Runtime/CoreUObject/Public/AssetRegistry/AssetData.h
+++ b/Engine/Source/Runtime/CoreUObject/Public/AssetRegistry/AssetData.h
@@ -157,20 +157,12 @@ public:
 
 		TStringBuilder<FName::StringBufferSize> PackageNameStrBuilder;
 		PackageName.ToString(PackageNameStrBuilder);
-
-		FStringView PackageAssetNameStr;
-		{
-			// Get everything after the last slash
-			FStringView PackageNameStr(PackageNameStrBuilder);
-			int32 IndexOfLastSlash = INDEX_NONE;
-			PackageNameStr.FindLastChar(TEXT('/'), IndexOfLastSlash);
-			PackageAssetNameStr = PackageNameStr.Mid(IndexOfLastSlash + 1);
-		}
-
-		FStringView AssetNameStr(AssetNameStrBuilder);
-		return PackageAssetNameStr == AssetNameStr;
+		return DetectIsUAssetByNames(PackageNameStrBuilder, AssetNameStrBuilder);
 	}
 
+	/** Returns true if the given UObject is the primary asset in a package, true for maps and assets but false for secondary objects like class redirectors */
+	COREUOBJECT_API static bool IsUAsset(UObject* Object);
+
 	void Shrink()
 	{
 		ChunkIDs.Shrink();
@@ -441,6 +433,20 @@ public:
 		Ar << ChunkIDs;
 		Ar << PackageFlags;
 	}
+
+private:
+	static bool DetectIsUAssetByNames(FStringView PackageName, FStringView ObjectPathName)
+	{
+		FStringView PackageBaseName;
+		{
+			// Get everything after the last slash
+			int32 IndexOfLastSlash = INDEX_NONE;
+			PackageName.FindLastChar(TEXT('/'), IndexOfLastSlash);
+			PackageBaseName = PackageName.Mid(IndexOfLastSlash + 1);
+		}
+
+		return PackageBaseName.Equals(ObjectPathName, ESearchCase::IgnoreCase);
+	}
 };
 
 FORCEINLINE uint32 GetTypeHash(const FAssetData& AssetData)
