Fixing MSVC warning issue that was accidently removed with prior change.

diff --git a/Engine/Source/Runtime/CoreUObject/Private/UObject/ScriptCore.cpp b/Engine/Source/Runtime/CoreUObject/Private/UObject/ScriptCore.cpp
--- a/Engine/Source/Runtime/CoreUObject/Private/UObject/ScriptCore.cpp
+++ b/Engine/Source/Runtime/CoreUObject/Private/UObject/ScriptCore.cpp
@@ -711,6 +711,7 @@ COREUOBJECT_API uint8 GRegisterNative( int32 NativeBytecodeIndex, const FNativeF
 	{
 		if( NativeBytecodeIndex<0 || (uint32)NativeBytecodeIndex>UE_ARRAY_COUNT(GNatives) || GNatives[NativeBytecodeIndex]!=&UObject::execUndefined) 
 		{
+			CA_SUPPRESS(6385)
 			if (!ReloadNotifyFunctionRemap(Func, GNatives[NativeBytecodeIndex]))
 			{
 				GNativeDuplicate = NativeBytecodeIndex;
