Renamed FGCObjectScopeGuard to TGCObjectScopeGuard

This allows it to store the scoped object as its actual type, so that it can be accessed if needed via a Get() function. FGCObjectScopeGuard remains as an alias for TGCObjectScopeGuard<const UObject>.

diff --git a/Engine/Source/Runtime/CoreUObject/Public/UObject/GCObjectScopeGuard.h b/Engine/Source/Runtime/CoreUObject/Public/UObject/GCObjectScopeGuard.h
--- a/Engine/Source/Runtime/CoreUObject/Public/UObject/GCObjectScopeGuard.h
+++ b/Engine/Source/Runtime/CoreUObject/Public/UObject/GCObjectScopeGuard.h
@@ -8,21 +8,24 @@
  * Specific implementation of FGCObject that prevents a single UObject-based pointer from being GC'd while this guard is in scope.
  * @note This is the lean version of TStrongObjectPtr which uses an inline FGCObject so *can't* safely be used with containers that treat types as trivially relocatable.
  */
-class FGCObjectScopeGuard : public FGCObject
+template <typename ObjectType>
+class TGCObjectScopeGuard : public FGCObject
 {
+	static_assert(TPointerIsConvertibleFromTo<ObjectType, const UObjectBase>::Value, "TGCObjectScopeGuard: ObjectType must be pointers to a type derived from UObject");
+
 public:
-	explicit FGCObjectScopeGuard(const UObject* InObject)
+	explicit TGCObjectScopeGuard(ObjectType* InObject)
 		: Object(InObject)
 	{
 	}
 
-	virtual ~FGCObjectScopeGuard()
+	virtual ~TGCObjectScopeGuard()
 	{
 	}
 
 	/** Non-copyable */
-	FGCObjectScopeGuard(const FGCObjectScopeGuard&) = delete;
-	FGCObjectScopeGuard& operator=(const FGCObjectScopeGuard&) = delete;
+	TGCObjectScopeGuard(const TGCObjectScopeGuard&) = delete;
+	TGCObjectScopeGuard& operator=(const TGCObjectScopeGuard&) = delete;
 
 	//~ FGCObject interface
 	virtual void AddReferencedObjects(FReferenceCollector& Collector) override
@@ -32,13 +35,20 @@ public:
 
 	virtual FString GetReferencerName() const override
 	{
-		return TEXT("FGCObjectScopeGuard");
+		return TEXT("TGCObjectScopeGuard");
+	}
+
+	FORCEINLINE_DEBUGGABLE ObjectType* Get() const
+	{
+		return Object;
 	}
 
 private:
-	const UObject* Object;
+	ObjectType* Object;
 };
 
+using FGCObjectScopeGuard = TGCObjectScopeGuard<const UObject>;
+
 /**
  * Specific implementation of FGCObject that prevents an array of UObject-based pointers from being GC'd while this guard is in scope.
  */
@@ -72,6 +82,11 @@ public:
 		return TEXT("TGCObjectsScopeGuard");
 	}
 
+	FORCEINLINE_DEBUGGABLE const TArray<ObjectType*>& Get() const
+	{
+		return Objects;
+	}
+
 private:
 	TArray<ObjectType*> Objects;
 };
