Don't set package loaded path when auto-saving.

diff --git a/Engine/Source/Runtime/CoreUObject/Private/UObject/SavePackage/SavePackageUtilities.cpp b/Engine/Source/Runtime/CoreUObject/Private/UObject/SavePackage/SavePackageUtilities.cpp
--- a/Engine/Source/Runtime/CoreUObject/Private/UObject/SavePackage/SavePackageUtilities.cpp
+++ b/Engine/Source/Runtime/CoreUObject/Private/UObject/SavePackage/SavePackageUtilities.cpp
@@ -694,9 +694,9 @@ void GetCDOSubobjects(UObject* CDO, TArray<UObject*>& Subobjects)
 bool IsSavingNewLoadedPath(bool bIsCooking, const FPackagePath& TargetPackagePath, uint32 SaveFlags)
 {
 #if WITH_EDITOR
-	return !bIsCooking &&							// Do not update the loadedpath if we're cooking
-		TargetPackagePath.IsMountedPath() &&		// Do not update the loadedpath if the new path is not a viable mounted path
-		!(SaveFlags & SAVE_BulkDataByReference);	// Do not update the loadedpath if we're storing references to the old loadedpath
+	return !bIsCooking &&												// Do not update the loadedpath if we're cooking
+		TargetPackagePath.IsMountedPath() &&							// Do not update the loadedpath if the new path is not a viable mounted path
+		!(SaveFlags & (SAVE_BulkDataByReference | SAVE_FromAutosave));	// Do not update the loadedpath if we're storing references to the old loadedpath or autosaving
 #else
 	return false; // Saving when not in editor never updates the LoadedPath
 #endif
