Interim BlueprintGeneratedClass::IsAsset change:
All code that is enumerating unfiltered assets in a package/outer needs to handle multiple assets, or if it can't, needs to call ShouldSkipAsset or IsUAsset to skip over the non-UAsset assets.

Move UE::AssetRegistry::FFiltering into CoreUObject to make it accessible to FindAssetInPackage.
#rnx

diff --git a/Engine/Source/Developer/Merge/Private/MergeUtils.cpp b/Engine/Source/Developer/Merge/Private/MergeUtils.cpp
--- a/Engine/Source/Developer/Merge/Private/MergeUtils.cpp
+++ b/Engine/Source/Developer/Merge/Private/MergeUtils.cpp
@@ -4,6 +4,7 @@
 #include "IAssetTypeActions.h"
 #include "UObject/UObjectHash.h"
 #include "UObject/Package.h"
+#include "Misc/AssetRegistryInterface.h"
 #include "Misc/PackageName.h"
 #include "SourceControlOperations.h"
 #include "ISourceControlProvider.h"
@@ -70,7 +71,7 @@ UObject const* FMergeToolUtils::LoadRevision(const FString& PackageName, const F
 
 		for (UObject* PackageObj : PackageObjs)
 		{
-			if (PackageObj->IsAsset())
+			if (PackageObj->IsAsset() && !UE::AssetRegistry::FFiltering::ShouldSkipAsset(PackageObj))
 			{
 				AssetRevision = LoadRevision(PackageObj, DesiredRevision);
 				break;
diff --git a/Engine/Source/Editor/UnrealEd/Private/EditorCommandLineUtils.cpp b/Engine/Source/Editor/UnrealEd/Private/EditorCommandLineUtils.cpp
--- a/Engine/Source/Editor/UnrealEd/Private/EditorCommandLineUtils.cpp
+++ b/Engine/Source/Editor/UnrealEd/Private/EditorCommandLineUtils.cpp
@@ -5,6 +5,7 @@
 #include "Misc/Paths.h"
 #include "Stats/Stats.h"
 #include "Misc/App.h"
+#include "Misc/AssetRegistryInterface.h"
 #include "Modules/ModuleManager.h"
 #include "UObject/UObjectHash.h"
 #include "Misc/PackageName.h"
@@ -696,7 +697,7 @@ static UObject* EditorCommandLineUtilsImpl::ExtractAssetFromPackage(UPackage* Pa
 	UObject* FoundAsset = nullptr;
 	for (UObject* PackageObj : ObjectsWithOuter)
 	{
-		if (PackageObj->IsAsset())
+		if (PackageObj->IsAsset() && !UE::AssetRegistry::FFiltering::ShouldSkipAsset(PackageObj))
 		{
 			FoundAsset = PackageObj;
 			break;
diff --git a/Engine/Source/Editor/UnrealEd/Private/EditorEngine.cpp b/Engine/Source/Editor/UnrealEd/Private/EditorEngine.cpp
--- a/Engine/Source/Editor/UnrealEd/Private/EditorEngine.cpp
+++ b/Engine/Source/Editor/UnrealEd/Private/EditorEngine.cpp
@@ -8,6 +8,7 @@
 #include "Misc/ScopedSlowTask.h"
 #include "Misc/CoreDelegates.h"
 #include "Misc/App.h"
+#include "Misc/AssetRegistryInterface.h"
 #include "Modules/ModuleManager.h"
 #include "UObject/MetaData.h"
 #include "UObject/ConstructorHelpers.h"
@@ -4276,7 +4277,7 @@ FSavePackageResultStruct UEditorEngine::Save( UPackage* InOuter, UObject* InBase
 			GetObjectsWithPackage(InOuter, PotentialAssets, false);
 			for (UObject* Object : PotentialAssets)
 			{
-				if (Object->IsAsset())
+				if (Object->IsAsset() && !UE::AssetRegistry::FFiltering::ShouldSkipAsset(Object))
 				{
 					Base = Object;
 					break;
diff --git a/Engine/Source/Runtime/AssetRegistry/Private/AssetRegistry.cpp b/Engine/Source/Runtime/AssetRegistry/Private/AssetRegistry.cpp
--- a/Engine/Source/Runtime/AssetRegistry/Private/AssetRegistry.cpp
+++ b/Engine/Source/Runtime/AssetRegistry/Private/AssetRegistry.cpp
@@ -232,133 +232,6 @@ EChunkProgressReportingType::Type GetChunkAvailabilityProgressType(EAssetAvailab
 	return ChunkReportType;
 }
 
-
-namespace UE
-{
-namespace AssetRegistry
-{
-#if WITH_ENGINE && WITH_EDITOR
-	TSet<FName> SkipUncookedClasses;
-	TSet<FName> SkipCookedClasses;
-	bool bInitializedSkipClasses = false;
-#endif //if WITH_ENGINE && WITH_EDITOR
-
-	// TODO: replace with a function in CoreGlobals
-	static bool IsRunningCookCommandlet()
-	{
-		FString Commandline = FCommandLine::Get();
-		const bool bIsCookCommandlet = IsRunningCommandlet() && Commandline.Contains(TEXT("run=cook"));
-		return bIsCookCommandlet;
-	}
-
-	bool FFiltering::ShouldSkipAsset(FName AssetClass, uint32 PackageFlags)
-	{
-#if WITH_ENGINE && WITH_EDITOR
-
-		TRACE_CPUPROFILER_EVENT_SCOPE(AssetRegistry::FFiltering::ShouldSkipAsset)
-
-		// We do not yet support having UBlueprintGeneratedClasses be assets when the UBlueprint is also
-		// an asset; the content browser does not handle the multiple assets correctly and displays this
-		// class asset as if it is in a separate package. Revisit when we have removed the UBlueprint as an asset
-		// or when we support multiple assets.
-		if (!bInitializedSkipClasses)
-		{
-			// Since we only collect these the first on-demand time, it is possible we will miss subclasses
-			// from plugins that load later. This flaw is a rare edge case, though, and this solution will
-			// be replaced eventually, so leaving it for now.
-			if (GIsEditor && (!IsRunningCommandlet() || IsRunningCookCommandlet()))
-			{
-				static const FName NAME_EnginePackage("/Script/Engine");
-				UPackage* EnginePackage = Cast<UPackage>(StaticFindObjectFast(UPackage::StaticClass(), nullptr, NAME_EnginePackage));
-
-				{
-					SkipUncookedClasses.Reset();
-
-					static const FName NAME_BlueprintGeneratedClass("BlueprintGeneratedClass");
-					UClass* BlueprintGeneratedClass = nullptr;
-					if (EnginePackage)
-					{
-						BlueprintGeneratedClass = Cast<UClass>(StaticFindObjectFast(UClass::StaticClass(), EnginePackage, NAME_BlueprintGeneratedClass));
-					}
-					if (!BlueprintGeneratedClass)
-					{
-						UE_LOG(LogAssetRegistry, Warning, TEXT("Could not find BlueprintGeneratedClass; will not be able to filter uncooked BPGC"));
-					}
-					else
-					{
-						SkipUncookedClasses.Add(NAME_BlueprintGeneratedClass);
-						for (TObjectIterator<UClass> It; It; ++It)
-						{
-							if (It->IsChildOf(BlueprintGeneratedClass) && !It->HasAnyClassFlags(CLASS_Abstract))
-							{
-								SkipUncookedClasses.Add(It->GetFName());
-							}
-						}
-					}
-				}
-				{
-					SkipCookedClasses.Reset();
-
-					static const FName NAME_Blueprint("Blueprint");
-					UClass* BlueprintClass = nullptr;
-					if (EnginePackage)
-					{
-						BlueprintClass = Cast<UClass>(StaticFindObjectFast(UClass::StaticClass(), EnginePackage, NAME_Blueprint));
-					}
-					if (!BlueprintClass)
-					{
-						UE_LOG(LogAssetRegistry, Warning, TEXT("Could not find BlueprintClass; will not be able to filter cooked BP"));
-					}
-					else
-					{
-						SkipCookedClasses.Add(NAME_Blueprint);
-						for (TObjectIterator<UClass> It; It; ++It)
-						{
-							if (It->IsChildOf(BlueprintClass) && !It->HasAnyClassFlags(CLASS_Abstract))
-							{
-								SkipCookedClasses.Add(It->GetFName());
-							}
-						}
-					}
-				}
-			}
-
-			bInitializedSkipClasses = true;
-		}
-
-		if (PackageFlags & PKG_ContainsNoAsset)
-		{
-			return true;
-		}
-
-		const bool bIsCooked = (PackageFlags & PKG_FilterEditorOnly);
-
-		if ((bIsCooked && SkipCookedClasses.Contains(AssetClass)) ||
-			(!bIsCooked && SkipUncookedClasses.Contains(AssetClass)))
-		{
-			return true;
-		}
-#endif //if WITH_ENGINE && WITH_EDITOR
-
-		return false;
-	}
-
-	bool FFiltering::ShouldSkipAsset(const UObject* InAsset)
-	{
-		if (!InAsset)
-		{
-			return false;
-		}
-		UPackage* Package = InAsset->GetPackage();
-		if (!Package)
-		{
-			return false;
-		}
-		return ShouldSkipAsset(InAsset->GetClass()->GetFName(), Package->GetPackageFlags());
-	}
-}
-}
-
 UAssetRegistry::UAssetRegistry(const FObjectInitializer& ObjectInitializer)
 	: Super(ObjectInitializer)
 {
@@ -673,7 +546,7 @@ void UAssetRegistryImpl::OnPluginLoadingPhaseComplete(ELoadingPhase::Type Loadin
 	// Reparse the skip classes the next time ShouldSkipAsset is called, since available classes
 	// for the search over all classes may have changed
 #if WITH_EDITOR
-	UE::AssetRegistry::bInitializedSkipClasses = false;
+	UE::AssetRegistry::FFiltering::MarkDirty();
 #endif
 	IPluginManager::Get().OnLoadingPhaseComplete().RemoveAll(this);
 }
diff --git a/Engine/Source/Runtime/AssetRegistry/Public/AssetRegistry/IAssetRegistry.h b/Engine/Source/Runtime/AssetRegistry/Public/AssetRegistry/IAssetRegistry.h
--- a/Engine/Source/Runtime/AssetRegistry/Public/AssetRegistry/IAssetRegistry.h
+++ b/Engine/Source/Runtime/AssetRegistry/Public/AssetRegistry/IAssetRegistry.h
@@ -579,15 +579,5 @@ namespace AssetRegistry
 	// WritePackageData is declared in AssetRegistryInterface.h, in the CoreUObject module, because it is needed by SavePackage in CoreUObject
 	ASSETREGISTRY_API bool ReadPackageDataMain(FArchive& BinaryArchive, const FString& PackageName, const FPackageFileSummary& PackageFileSummary, int64& OutDependencyDataOffset, TArray<FAssetData*>& OutAssetDataList, EReadPackageDataMainErrorCode& OutError);
 	ASSETREGISTRY_API bool ReadPackageDataDependencies(FArchive& BinaryArchive, TBitArray<>& OutImportUsedInGame, TBitArray<>& OutSoftPackageUsedInGame);
-
-	class ASSETREGISTRY_API FFiltering
-	{
-	public:
-		/** Called to check whether we should filter out assets of the given class and package flags from the editor's asset registry */
-		static bool ShouldSkipAsset(FName AssetClass, uint32 PackageFlags);
-
-		/** Called to check whether we should filter out the given object (assumed to be an asset) from the editor's asset registry */
-		static bool ShouldSkipAsset(const UObject* InAsset);
-	};
 }
 }
diff --git a/Engine/Source/Runtime/CoreUObject/Private/Misc/AssetRegistryInterface.cpp b/Engine/Source/Runtime/CoreUObject/Private/Misc/AssetRegistryInterface.cpp
--- a/Engine/Source/Runtime/CoreUObject/Private/Misc/AssetRegistryInterface.cpp
+++ b/Engine/Source/Runtime/CoreUObject/Private/Misc/AssetRegistryInterface.cpp
@@ -2,6 +2,18 @@
 
 #include "Misc/AssetRegistryInterface.h"
 
+#include "Containers/Set.h"
+#include "CoreGlobals.h"
+#include "Misc/CommandLine.h"
+#include "ProfilingDebugging/CpuProfilerTrace.h"
+#include "Templates/Casts.h"
+#include "UObject/Class.h"
+#include "UObject/NameTypes.h"
+#include "UObject/ObjectMacros.h"
+#include "UObject/Package.h"
+#include "UObject/UObjectGlobals.h"
+#include "UObject/UObjectIterator.h"
+
 IAssetRegistryInterface* IAssetRegistryInterface::Default = nullptr;
 IAssetRegistryInterface* IAssetRegistryInterface::GetPtr()
 {
@@ -16,5 +28,132 @@ namespace Private
 {
 	IAssetRegistry* IAssetRegistrySingleton::Singleton = nullptr;
 }
+
+#if WITH_ENGINE && WITH_EDITOR
+	TSet<FName> SkipUncookedClasses;
+	TSet<FName> SkipCookedClasses;
+	bool bInitializedSkipClasses = false;
+#endif //if WITH_ENGINE && WITH_EDITOR
+
+	// TODO: replace with a function in CoreGlobals
+	static bool IsRunningCookCommandlet()
+	{
+		FString Commandline = FCommandLine::Get();
+		const bool bIsCookCommandlet = IsRunningCommandlet() && Commandline.Contains(TEXT("run=cook"));
+		return bIsCookCommandlet;
+	}
+
+	bool FFiltering::ShouldSkipAsset(FName AssetClass, uint32 PackageFlags)
+	{
+#if WITH_ENGINE && WITH_EDITOR
+
+		TRACE_CPUPROFILER_EVENT_SCOPE(AssetRegistry::FFiltering::ShouldSkipAsset)
+
+		// We do not yet support having UBlueprintGeneratedClasses be assets when the UBlueprint is also
+		// an asset; the content browser does not handle the multiple assets correctly and displays this
+		// class asset as if it is in a separate package. Revisit when we have removed the UBlueprint as an asset
+		// or when we support multiple assets.
+		if (!bInitializedSkipClasses)
+		{
+			// Since we only collect these the first on-demand time, it is possible we will miss subclasses
+			// from plugins that load later. This flaw is a rare edge case, though, and this solution will
+			// be replaced eventually, so leaving it for now.
+			if (GIsEditor && (!IsRunningCommandlet() || IsRunningCookCommandlet()))
+			{
+				static const FName NAME_EnginePackage("/Script/Engine");
+				UPackage* EnginePackage = Cast<UPackage>(StaticFindObjectFast(UPackage::StaticClass(), nullptr, NAME_EnginePackage));
+
+				{
+					SkipUncookedClasses.Reset();
+
+					static const FName NAME_BlueprintGeneratedClass("BlueprintGeneratedClass");
+					UClass* BlueprintGeneratedClass = nullptr;
+					if (EnginePackage)
+					{
+						BlueprintGeneratedClass = Cast<UClass>(StaticFindObjectFast(UClass::StaticClass(), EnginePackage, NAME_BlueprintGeneratedClass));
+					}
+					if (!BlueprintGeneratedClass)
+					{
+						UE_LOG(LogCore, Warning, TEXT("Could not find BlueprintGeneratedClass; will not be able to filter uncooked BPGC"));
+					}
+					else
+					{
+						SkipUncookedClasses.Add(NAME_BlueprintGeneratedClass);
+						for (TObjectIterator<UClass> It; It; ++It)
+						{
+							if (It->IsChildOf(BlueprintGeneratedClass) && !It->HasAnyClassFlags(CLASS_Abstract))
+							{
+								SkipUncookedClasses.Add(It->GetFName());
+							}
+						}
+					}
+				}
+				{
+					SkipCookedClasses.Reset();
+
+					static const FName NAME_Blueprint("Blueprint");
+					UClass* BlueprintClass = nullptr;
+					if (EnginePackage)
+					{
+						BlueprintClass = Cast<UClass>(StaticFindObjectFast(UClass::StaticClass(), EnginePackage, NAME_Blueprint));
+					}
+					if (!BlueprintClass)
+					{
+						UE_LOG(LogCore, Warning, TEXT("Could not find BlueprintClass; will not be able to filter cooked BP"));
+					}
+					else
+					{
+						SkipCookedClasses.Add(NAME_Blueprint);
+						for (TObjectIterator<UClass> It; It; ++It)
+						{
+							if (It->IsChildOf(BlueprintClass) && !It->HasAnyClassFlags(CLASS_Abstract))
+							{
+								SkipCookedClasses.Add(It->GetFName());
+							}
+						}
+					}
+				}
+			}
+
+			bInitializedSkipClasses = true;
+		}
+
+		if (PackageFlags & PKG_ContainsNoAsset)
+		{
+			return true;
+		}
+
+		const bool bIsCooked = (PackageFlags & PKG_FilterEditorOnly);
+
+		if ((bIsCooked && SkipCookedClasses.Contains(AssetClass)) ||
+			(!bIsCooked && SkipUncookedClasses.Contains(AssetClass)))
+		{
+			return true;
+		}
+#endif //if WITH_ENGINE && WITH_EDITOR
+
+		return false;
+	}
+
+	bool FFiltering::ShouldSkipAsset(const UObject* InAsset)
+	{
+		if (!InAsset)
+		{
+			return false;
+		}
+		UPackage* Package = InAsset->GetPackage();
+		if (!Package)
+		{
+			return false;
+		}
+		return ShouldSkipAsset(InAsset->GetClass()->GetFName(), Package->GetPackageFlags());
+	}
+
+	void FFiltering::MarkDirty()
+	{
+#if WITH_ENGINE && WITH_EDITOR
+		bInitializedSkipClasses = false;
+#endif
+	}
 }
 }
diff --git a/Engine/Source/Runtime/CoreUObject/Private/UObject/Package.cpp b/Engine/Source/Runtime/CoreUObject/Private/UObject/Package.cpp
--- a/Engine/Source/Runtime/CoreUObject/Private/UObject/Package.cpp
+++ b/Engine/Source/Runtime/CoreUObject/Private/UObject/Package.cpp
@@ -1,12 +1,14 @@
 // Copyright Epic Games, Inc. All Rights Reserved.
 
 #include "UObject/Package.h"
+
 #include "HAL/FileManager.h"
+#include "Misc/AssetRegistryInterface.h"
 #include "Misc/ITransaction.h"
-#include "UObject/MetaData.h"
 #include "Misc/PackageName.h"
 #include "UObject/LinkerLoad.h"
 #include "UObject/LinkerManager.h"
+#include "UObject/MetaData.h"
 #include "UObject/PackageResourceManager.h"
 #include "UObject/UObjectHash.h"
 #include "UObject/UObjectThreadContext.h"
@@ -119,7 +121,7 @@ UObject* UPackage::FindAssetInPackage() const
 	UObject* Asset = nullptr;
 	ForEachObjectWithPackage(this, [&Asset](UObject* Object)
 		{
-			if (Object->IsAsset())
+			if (Object->IsAsset() && !UE::AssetRegistry::FFiltering::ShouldSkipAsset(Object))
 			{
 				ensure(Asset == nullptr);
 				Asset = Object;
diff --git a/Engine/Source/Runtime/CoreUObject/Private/UObject/SavePackage.cpp b/Engine/Source/Runtime/CoreUObject/Private/UObject/SavePackage.cpp
--- a/Engine/Source/Runtime/CoreUObject/Private/UObject/SavePackage.cpp
+++ b/Engine/Source/Runtime/CoreUObject/Private/UObject/SavePackage.cpp
@@ -2400,7 +2400,7 @@ FSavePackageResultStruct UPackage::Save(UPackage* InOuter, UObject* Base, EObjec
 						UObject* Asset = nullptr;
 						ForEachObjectWithOuter(Package, [&Asset](UObject* Object)
 							{
-								if (!Asset && Object->IsAsset())
+								if (!Asset && Object->IsAsset() && !UE::AssetRegistry::FFiltering::ShouldSkipAsset(Object))
 								{
 									Asset = Object;
 								}
diff --git a/Engine/Source/Runtime/CoreUObject/Private/UObject/SavePackage/SaveContext.h b/Engine/Source/Runtime/CoreUObject/Private/UObject/SavePackage/SaveContext.h
--- a/Engine/Source/Runtime/CoreUObject/Private/UObject/SavePackage/SaveContext.h
+++ b/Engine/Source/Runtime/CoreUObject/Private/UObject/SavePackage/SaveContext.h
@@ -4,6 +4,7 @@
 
 #include "HAL/FileManager.h"
 #include "HAL/IConsoleManager.h"
+#include "Misc/AssetRegistryInterface.h"
 #include "Misc/Optional.h"
 #include "Misc/PackageName.h"
 #include "Serialization/CustomVersion.h"
@@ -77,7 +78,7 @@ public:
 		{
 			ForEachObjectWithPackage(Package, [this](UObject* InObject)
 				{
-					if (InObject->IsAsset())
+					if (InObject->IsAsset() && !UE::AssetRegistry::FFiltering::ShouldSkipAsset(InObject))
 					{
 						Asset = InObject;
 						return false;
diff --git a/Engine/Source/Runtime/CoreUObject/Public/Misc/AssetRegistryInterface.h b/Engine/Source/Runtime/CoreUObject/Public/Misc/AssetRegistryInterface.h
--- a/Engine/Source/Runtime/CoreUObject/Public/Misc/AssetRegistryInterface.h
+++ b/Engine/Source/Runtime/CoreUObject/Public/Misc/AssetRegistryInterface.h
@@ -179,6 +179,20 @@ namespace AssetRegistry
 			friend class ::UAssetRegistryImpl;
 		};
 	}
+
+	class COREUOBJECT_API FFiltering
+	{
+	public:
+		/** Called to check whether we should filter out assets of the given class and package flags from the editor's asset registry */
+		static bool ShouldSkipAsset(FName AssetClass, uint32 PackageFlags);
+
+		/** Called to check whether we should filter out the given object (assumed to be an asset) from the editor's asset registry */
+		static bool ShouldSkipAsset(const UObject* InAsset);
+
+		/** Call to invalidate the list of skip assets and cause their next use to recreate them on demand */
+		static void MarkDirty();
+	};
+
 }
 }
 
