Add a way to reset only exports in a linker without detaching the linker itself nor its bulkdata
Use that new function when moving actors between package during cooking
Bumped the UObject::SetLinker Warning up to Error as they should not happen anymore.

diff --git a/Engine/Source/Runtime/CoreUObject/Private/UObject/Linker.cpp b/Engine/Source/Runtime/CoreUObject/Private/UObject/Linker.cpp
--- a/Engine/Source/Runtime/CoreUObject/Private/UObject/Linker.cpp
+++ b/Engine/Source/Runtime/CoreUObject/Private/UObject/Linker.cpp
@@ -480,6 +480,11 @@ PRAGMA_ENABLE_DEPRECATION_WARNINGS
 	Global functions
 -----------------------------------------------------------------------------*/
 
+void ResetLinkerExports(UPackage* InPackage)
+{
+	FLinkerManager::Get().ResetLinkerExports(InPackage);
+}
+
 void ResetLoaders(UObject* InPkg)
 {
 	if (IsAsyncLoading())
diff --git a/Engine/Source/Runtime/CoreUObject/Private/UObject/LinkerLoad.cpp b/Engine/Source/Runtime/CoreUObject/Private/UObject/LinkerLoad.cpp
--- a/Engine/Source/Runtime/CoreUObject/Private/UObject/LinkerLoad.cpp
+++ b/Engine/Source/Runtime/CoreUObject/Private/UObject/LinkerLoad.cpp
@@ -5423,39 +5423,39 @@ UObject* FLinkerLoad::IndexToObject( FPackageIndex Index )
 // Detach an export from this linker.
 void FLinkerLoad::DetachExport( int32 i )
 {
-	FObjectExport& E = ExportMap[ i ];
-	check(E.Object);
-	if( !E.Object->IsValidLowLevel() )
+	FObjectExport& Export = ExportMap[ i ];
+	check(Export.Object);
+	if( !Export.Object->IsValidLowLevel() )
 	{
-		UE_LOG(LogLinker, Fatal, TEXT("Linker object %s %s.%s is invalid"), *GetExportClassName(i).ToString(), *LinkerRoot->GetName(), *E.ObjectName.ToString() );
+		UE_LOG(LogLinker, Fatal, TEXT("Linker object %s %s.%s is invalid"), *GetExportClassName(i).ToString(), *LinkerRoot->GetName(), *Export.ObjectName.ToString() );
 	}
 	{
-		const FLinkerLoad* ActualLinker = E.Object->GetLinker();
+		const FLinkerLoad* ActualLinker = Export.Object->GetLinker();
 		// TODO: verify the condition
 		const bool DynamicType = !ActualLinker
-			&& (E.Object->HasAnyFlags(RF_Dynamic)
-			|| (E.Object->GetClass()->HasAnyFlags(RF_Dynamic) && E.Object->HasAnyFlags(RF_ClassDefaultObject) ));
+			&& (Export.Object->HasAnyFlags(RF_Dynamic)
+			|| (Export.Object->GetClass()->HasAnyFlags(RF_Dynamic) && Export.Object->HasAnyFlags(RF_ClassDefaultObject) ));
 		if ((ActualLinker != this) && !DynamicType)
 		{
-			UObject* Object = E.Object;
+			UObject* Object = Export.Object;
 			UE_LOG(LogLinker, Log, TEXT("Object            : %s"), *Object->GetFullName());
 			//UE_LOG(LogLinker, Log, TEXT("Object Linker     : %s"), *Object->GetLinker()->GetFullName() );
 			UE_LOG(LogLinker, Log, TEXT("Linker LinkerRoot : %s"), Object->GetLinker() ? *Object->GetLinker()->LinkerRoot->GetFullName() : TEXT("None"));
 			//UE_LOG(LogLinker, Log, TEXT("Detach Linker     : %s"), *GetFullName() );
 			UE_LOG(LogLinker, Log, TEXT("Detach LinkerRoot : %s"), *LinkerRoot->GetFullName());
-			UE_LOG(LogLinker, Fatal, TEXT("Linker object %s %s.%s mislinked!"), *GetExportClassName(i).ToString(), *LinkerRoot->GetName(), *E.ObjectName.ToString());
+			UE_LOG(LogLinker, Fatal, TEXT("Linker object %s %s.%s mislinked!"), *GetExportClassName(i).ToString(), *LinkerRoot->GetName(), *Export.ObjectName.ToString());
 		}
 	}
 
-	if (E.Object->GetLinkerIndex() == -1)
+	if (Export.Object->GetLinkerIndex() == -1)
 	{
-		UE_LOG(LogLinker, Warning, TEXT("Linker object %s %s.%s was already detached."), *GetExportClassName(i).ToString(), *LinkerRoot->GetName(), *E.ObjectName.ToString());
+		UE_LOG(LogLinker, Warning, TEXT("Linker object %s %s.%s was already detached."), *GetExportClassName(i).ToString(), *LinkerRoot->GetName(), *Export.ObjectName.ToString());
 	}
 	else
 	{
-		checkf(E.Object->GetLinkerIndex() == i, TEXT("Mismatched linker index in FLinkerLoad::DetachExport for %s in %s. Linker index was supposed to be %d, was %d"), *GetExportClassName(i).ToString(), *LinkerRoot->GetName(), i, E.Object->GetLinkerIndex());
+		checkf(Export.Object->GetLinkerIndex() == i, TEXT("Mismatched linker index in FLinkerLoad::DetachExport for %s in %s. Linker index was supposed to be %d, was %d"), *GetExportClassName(i).ToString(), *LinkerRoot->GetName(), i, Export.Object->GetLinkerIndex());
 	}
-	ExportMap[i].Object->SetLinker( NULL, INDEX_NONE );
+	Export.Object->SetLinker(nullptr, INDEX_NONE);
 }
 
 void FLinkerLoad::LoadAndDetachAllBulkData()
@@ -5480,6 +5480,18 @@ void FLinkerLoad::DestroyLoader()
 	bIsDestroyingLoader = false;
 }
 
+void FLinkerLoad::DetachExports()
+{
+	// Detach all objects linked with this linker.
+	for (int32 ExportIndex = 0; ExportIndex < ExportMap.Num(); ++ExportIndex)
+	{
+		if (ExportMap[ExportIndex].Object)
+		{
+			DetachExport(ExportIndex);
+		}
+	}
+}
+
 void FLinkerLoad::Detach()
 {
 #if WITH_EDITOR
@@ -5489,13 +5501,7 @@ void FLinkerLoad::Detach()
 #endif
 
 	// Detach all objects linked with this linker.
-	for (int32 ExportIndex = 0; ExportIndex < ExportMap.Num(); ++ExportIndex)
-	{	
-		if (ExportMap[ExportIndex].Object)
-		{
-			DetachExport(ExportIndex);
-		}
-	}
+	DetachExports();
 
 	// Remove from object manager, if it has been added.
 	FLinkerManager::Get().RemoveLoaderFromObjectLoadersAndLoadersWithNewImports(this);
diff --git a/Engine/Source/Runtime/CoreUObject/Private/UObject/LinkerManager.cpp b/Engine/Source/Runtime/CoreUObject/Private/UObject/LinkerManager.cpp
--- a/Engine/Source/Runtime/CoreUObject/Private/UObject/LinkerManager.cpp
+++ b/Engine/Source/Runtime/CoreUObject/Private/UObject/LinkerManager.cpp
@@ -94,16 +94,29 @@ bool FLinkerManager::Exec(class UWorld* InWorld, const TCHAR* Cmd, FOutputDevice
 	return false;
 }
 
+void FLinkerManager::ResetLinkerExports(UPackage* InPackage)
+{
+	if (FLinkerLoad* LinkerToReset = FLinkerLoad::FindExistingLinkerForPackage(InPackage))
+	{
+		// if the linker owner thread is not the main thread, we need to flush async loading (todo: for that package) before we can reset the linker
+		if (LinkerToReset->GetOwnerThreadId() != GGameThreadId)
+		{
+			FlushAsyncLoading();
+		}
+		LinkerToReset->DetachExports();
+	}
+}
+
 void FLinkerManager::ResetLoaders(UObject* InPkg)
 {
 	// Top level package to reset loaders for.
-	UObject*		TopLevelPackage = InPkg ? InPkg->GetOutermost() : NULL;
+	UObject*		TopLevelPackage = InPkg ? InPkg->GetOutermost() : nullptr;
 
 	// Find loader/ linker associated with toplevel package. We do this upfront as Detach resets LinkerRoot.
 	if (TopLevelPackage)
 	{
 		// Linker to reset/ detach.
-		auto LinkerToReset = FLinkerLoad::FindExistingLinkerForPackage(CastChecked<UPackage>(TopLevelPackage));
+		FLinkerLoad* LinkerToReset = FLinkerLoad::FindExistingLinkerForPackage(CastChecked<UPackage>(TopLevelPackage));
 		if (LinkerToReset)
 		{
 			{
@@ -119,7 +132,7 @@ void FLinkerManager::ResetLoaders(UObject* InPkg)
 						{
 							if (Import.SourceLinker == LinkerToReset)
 							{
-								Import.SourceLinker = NULL;
+								Import.SourceLinker = nullptr;
 								Import.SourceIndex = INDEX_NONE;
 							}
 						}
diff --git a/Engine/Source/Runtime/CoreUObject/Private/UObject/LinkerManager.h b/Engine/Source/Runtime/CoreUObject/Private/UObject/LinkerManager.h
--- a/Engine/Source/Runtime/CoreUObject/Private/UObject/LinkerManager.h
+++ b/Engine/Source/Runtime/CoreUObject/Private/UObject/LinkerManager.h
@@ -110,6 +110,9 @@ public:
 	// FSelfRegisteringExec interface
 	virtual bool Exec(class UWorld* InWorld, const TCHAR* Cmd, FOutputDevice& Ar) override;
 
+	/** Reset only the linker exports associated with this package. */
+	void ResetLinkerExports(UPackage* InPackage);
+
 	/** Empty the loaders */
 	void ResetLoaders(UObject* InPkg);
 
diff --git a/Engine/Source/Runtime/CoreUObject/Private/UObject/UObjectLinker.cpp b/Engine/Source/Runtime/CoreUObject/Private/UObject/UObjectLinker.cpp
--- a/Engine/Source/Runtime/CoreUObject/Private/UObject/UObjectLinker.cpp
+++ b/Engine/Source/Runtime/CoreUObject/Private/UObject/UObjectLinker.cpp
@@ -108,8 +108,7 @@ void UObject::SetLinker( FLinkerLoad* LinkerLoad, int32 LinkerIndex, bool bShoul
 	// Detach from existing linker.
 	if( Existing.Linker && bShouldDetachExisting )
 	{
-		// HACK-FH: temporarily downgrading assert to a warning until OFPA name collision problem are resolved for good.
-		UE_CLOG(HasAnyFlags(RF_NeedLoad|RF_NeedPostLoad), LogUObjectLinker, Warning,
+		UE_CLOG(HasAnyFlags(RF_NeedLoad|RF_NeedPostLoad), LogUObjectLinker, Error,
 			TEXT("Detaching from existing linker for %s while object %s needs loading from linker %s."),
 			*Existing.Linker->GetArchiveName(),
 			*GetFullName(),
diff --git a/Engine/Source/Runtime/CoreUObject/Public/UObject/Linker.h b/Engine/Source/Runtime/CoreUObject/Public/UObject/Linker.h
--- a/Engine/Source/Runtime/CoreUObject/Public/UObject/Linker.h
+++ b/Engine/Source/Runtime/CoreUObject/Public/UObject/Linker.h
@@ -612,6 +612,12 @@ typedef uint32 ELazyLoaderFlags;
 	Global functions
 -----------------------------------------------------------------------------*/
 
+/**
+ * Reset the linker exports associated with the package
+ * @note, this might flush async loading if the linker is owned by the loading thread
+ */
+COREUOBJECT_API void ResetLinkerExports(UPackage* InPackage);
+
 /**
  * Remove references to the linker for the given package and delete the linker. 
  * Can be called after the package has finished loading.
diff --git a/Engine/Source/Runtime/CoreUObject/Public/UObject/LinkerLoad.h b/Engine/Source/Runtime/CoreUObject/Public/UObject/LinkerLoad.h
--- a/Engine/Source/Runtime/CoreUObject/Public/UObject/LinkerLoad.h
+++ b/Engine/Source/Runtime/CoreUObject/Public/UObject/LinkerLoad.h
@@ -833,10 +833,15 @@ public:
 	COREUOBJECT_API void LoadAndDetachAllBulkData();
 
 	/**
-	* Detaches linker from bulk data/ exports and removes itself from array of loaders.
-	*/
+	 * Detaches linker from bulk data and exports and removes itself from array of loaders.
+	 */
 	COREUOBJECT_API void Detach();
 
+	/**
+	 * Only detaches the linker from its exports and the reset their cached state in the linker without touching bulkdata or underlying loader
+	 */
+	COREUOBJECT_API void DetachExports();
+
 private:
 
 	FORCEINLINE virtual void Seek(int64 InPos) override
diff --git a/Engine/Source/Runtime/Engine/Private/ActorEditor.cpp b/Engine/Source/Runtime/Engine/Private/ActorEditor.cpp
--- a/Engine/Source/Runtime/Engine/Private/ActorEditor.cpp
+++ b/Engine/Source/Runtime/Engine/Private/ActorEditor.cpp
@@ -875,8 +875,8 @@ void AActor::SetPackageExternal(bool bExternal, bool bShouldDirty)
 	else 
 	{
 		UPackage* ActorPackage = GetExternalPackage();
-		// Detach the linker from the actor package so that the actor won't keep references to it if we wanted to delete the package
-		ResetLoaders(ActorPackage);
+		// Detach the linker exports so it doesn't resolve to this actor anymore
+		ResetLinkerExports(ActorPackage);
 		SetExternalPackage(nullptr);
 	}
 
diff --git a/Engine/Source/Runtime/Engine/Private/WorldPartition/WorldPartitionLevelHelper.cpp b/Engine/Source/Runtime/Engine/Private/WorldPartition/WorldPartitionLevelHelper.cpp
--- a/Engine/Source/Runtime/Engine/Private/WorldPartition/WorldPartitionLevelHelper.cpp
+++ b/Engine/Source/Runtime/Engine/Private/WorldPartition/WorldPartitionLevelHelper.cpp
@@ -60,7 +60,7 @@ void FWorldPartitionLevelHelper::MoveExternalActorsToLevel(const TArray<FWorldPa
 		{
 			const bool bSameOuter = (InLevel == Actor->GetOuter());
 			Actor->SetPackageExternal(false, false);
-			Actor->Rename(nullptr, InLevel);
+			Actor->Rename(nullptr, InLevel, REN_ForceNoResetLoaders);
 			check(Actor->GetPackage() == LevelPackage);
 			if (bSameOuter && !InLevel->Actors.Contains(Actor))
 			{
