Fix issue with loading of unversioned data in editor - we can not cache the properties in BlueprintGeneratedClasses, because they may be recompiled in the editor.
#rnx

diff --git a/Engine/Source/Runtime/CoreUObject/Private/UObject/Class.cpp b/Engine/Source/Runtime/CoreUObject/Private/UObject/Class.cpp
--- a/Engine/Source/Runtime/CoreUObject/Private/UObject/Class.cpp
+++ b/Engine/Source/Runtime/CoreUObject/Private/UObject/Class.cpp
@@ -1648,6 +1648,7 @@ void UStruct::DestroyChildPropertiesAndResetPropertyLinks()
 #if WITH_EDITORONLY_DATA
 	FieldPathSerialNumber = GetNextFieldPathSerialNumber();
 #endif // WITH_EDITORONLY_DATA
+	DestroyUnversionedSchema(this);
 }
 
 UStruct::~UStruct()
