Add check to make sure we never try to detach a bulkdata that is being used

diff --git a/Engine/Source/Runtime/CoreUObject/Private/Serialization/BulkData.cpp b/Engine/Source/Runtime/CoreUObject/Private/Serialization/BulkData.cpp
--- a/Engine/Source/Runtime/CoreUObject/Private/Serialization/BulkData.cpp
+++ b/Engine/Source/Runtime/CoreUObject/Private/Serialization/BulkData.cpp
@@ -1641,6 +1641,7 @@ void FUntypedBulkData::DetachFromArchive( FArchive* Ar, bool bEnsureBulkDataIsLo
 {
 	check( Ar );
 	check( Ar == AttachedAr || AttachedAr->IsProxyOf(Ar) );
+	check( LockStatus == LOCKSTATUS_Unlocked );
 
 	// Make sure bulk data is loaded.
 	if( bEnsureBulkDataIsLoaded )
