Change the lockstatus after the detach to avoid tripping the new DetachBulkData check

diff --git a/Engine/Source/Runtime/CoreUObject/Private/Serialization/BulkData.cpp b/Engine/Source/Runtime/CoreUObject/Private/Serialization/BulkData.cpp
--- a/Engine/Source/Runtime/CoreUObject/Private/Serialization/BulkData.cpp
+++ b/Engine/Source/Runtime/CoreUObject/Private/Serialization/BulkData.cpp
@@ -597,8 +597,6 @@ void* FUntypedBulkData::Lock( uint32 LockFlags )
 	// Read-write operations are allowed on returned memory.
 	if( LockFlags & LOCK_READ_WRITE )
 	{
-		LockStatus = LOCKSTATUS_ReadWriteLock;
-
 #if WITH_EDITOR
 		// We need to detach from the archive to not be able to clobber changes by serializing
 		// over them.
@@ -609,6 +607,8 @@ void* FUntypedBulkData::Lock( uint32 LockFlags )
 			check( AttachedAr == NULL );
 		}
 #endif // WITH_EDITOR
+		// This has to be set after the DetachBulkData because we can't detach a locked bulkdata
+		LockStatus = LOCKSTATUS_ReadWriteLock;
 		ClearBulkDataFlags(BULKDATA_LazyLoadable);
 	}
 	// Only read operations are allowed on returned memory.
