Systematic changes needed for SoundCue and SoundWave to successfully lazy load in editor in most circumstances when using the -LazyLoadImports commandline switch.
Changes control of the feature to uniformly require the -LazyLoadImports commandline switch.
Ensures that serialization of ObjectPtr fields for blueprints doesn't differ between raw pointer (Object) fields.  INSTANCES of these types were already serialized identically (wrote out the same type identifier in tagged property serialization), but the serialized type descriptors (used by blueprints) write out the FFieldClass::GetName, and that differed between FObjectProperty and FObjectPtrProperty, and this change addresses that discrepancy.
Allow ObjectPtr fields to be instantiated in place of Object fields in script/blueprint in editor builds while ensuring that they are serialization compatible and considered equal types.

diff --git a/Engine/Source/Editor/KismetCompiler/Private/KismetCompilerMisc.cpp b/Engine/Source/Editor/KismetCompiler/Private/KismetCompilerMisc.cpp
--- a/Engine/Source/Editor/KismetCompiler/Private/KismetCompilerMisc.cpp
+++ b/Engine/Source/Editor/KismetCompiler/Private/KismetCompilerMisc.cpp
@@ -945,6 +945,10 @@ FProperty* FKismetCompilerUtilities::CreatePrimitiveProperty(FFieldVariant Prope
 				{
 					NewPropertyObj = new FWeakObjectProperty(PropertyScope, ValidatedPropertyName, ObjectFlags);
 				}
+				else if (FLinkerLoad::IsImportLazyLoadEnabled())
+				{
+					NewPropertyObj = new FObjectPtrProperty(PropertyScope, ValidatedPropertyName, ObjectFlags);
+				}
 				else
 				{
 					NewPropertyObj = new FObjectProperty(PropertyScope, ValidatedPropertyName, ObjectFlags);
diff --git a/Engine/Source/Runtime/CoreUObject/Private/Blueprint/BlueprintSupport.cpp b/Engine/Source/Runtime/CoreUObject/Private/Blueprint/BlueprintSupport.cpp
--- a/Engine/Source/Runtime/CoreUObject/Private/Blueprint/BlueprintSupport.cpp
+++ b/Engine/Source/Runtime/CoreUObject/Private/Blueprint/BlueprintSupport.cpp
@@ -1781,10 +1781,7 @@ void FLinkerLoad::FinalizeBlueprint(UClass* LoadClass)
 	// have to)... we do however need it here in FinalizeBlueprint(), because
 	// we need it ran for any super-classes before we regen
 
-#if UE_WITH_OBJECT_HANDLE_LATE_RESOLVE
-	static bool bLazyLoadAllImports = FParse::Param(FCommandLine::Get(), TEXT("LazyLoadAllImports"));
-	if (!bLazyLoadAllImports)
-#endif
+	if (!IsImportLazyLoadEnabled())
 	{
 		// @TODO: OBJPTR: Need to find other options for solving this issue of placeholder classes during blueprint compile without forcing all imports to resolve always
 		ResolveAllImports();
diff --git a/Engine/Source/Runtime/CoreUObject/Private/Serialization/AsyncLoading.cpp b/Engine/Source/Runtime/CoreUObject/Private/Serialization/AsyncLoading.cpp
--- a/Engine/Source/Runtime/CoreUObject/Private/Serialization/AsyncLoading.cpp
+++ b/Engine/Source/Runtime/CoreUObject/Private/Serialization/AsyncLoading.cpp
@@ -6026,13 +6026,10 @@ EAsyncPackageState::Type FAsyncPackage::LoadImports(FFlushTree* FlushTree)
 	TRACE_CPUPROFILER_EVENT_SCOPE(FAsyncPackage::LoadImports);
 	SCOPE_CYCLE_COUNTER(STAT_FAsyncPackage_LoadImports);
 
-#if UE_WITH_OBJECT_HANDLE_LATE_RESOLVE
-	static bool bLazyLoadAllImports = FParse::Param(FCommandLine::Get(), TEXT("LazyLoadAllImports"));
-	if (bLazyLoadAllImports)
+	if (FLinkerLoad::IsImportLazyLoadEnabled())
 	{
 		return EAsyncPackageState::Complete;
 	}
-#endif
 
 	LastObjectWorkWasPerformedOn	= LinkerRoot;
 	LastTypeOfWorkPerformed			= TEXT("loading imports");
@@ -6159,13 +6156,10 @@ EAsyncPackageState::Type FAsyncPackage::CreateImports()
 	SCOPED_LOADTIMER(CreateImportsTime);
 	SCOPE_CYCLE_COUNTER(STAT_FAsyncPackage_CreateImports);
 
-#if UE_WITH_OBJECT_HANDLE_LATE_RESOLVE
-	static bool bLazyLoadAllImports = FParse::Param(FCommandLine::Get(), TEXT("LazyLoadAllImports"));
-	if (bLazyLoadAllImports)
+	if (FLinkerLoad::IsImportLazyLoadEnabled())
 	{
 		return EAsyncPackageState::Complete;
 	}
-#endif
 
 	// GC can't run in here
 	FGCScopeGuard GCGuard;
diff --git a/Engine/Source/Runtime/CoreUObject/Private/UObject/Field.cpp b/Engine/Source/Runtime/CoreUObject/Private/UObject/Field.cpp
--- a/Engine/Source/Runtime/CoreUObject/Private/UObject/Field.cpp
+++ b/Engine/Source/Runtime/CoreUObject/Private/UObject/Field.cpp
@@ -31,6 +31,7 @@ Field.cpp: Defines FField property system fundamentals
 #include "UObject/ObjectResource.h"
 #include "UObject/LinkerSave.h"
 #include "UObject/Interface.h"
+#include "UObject/LinkerLoad.h"
 #include "UObject/LinkerPlaceholderClass.h"
 #include "UObject/LinkerPlaceholderFunction.h"
 #include "UObject/StructScriptLoader.h"
@@ -1044,7 +1045,14 @@ FField* FField::CreateFromUField(UField* InField)
 	}
 	else if (UFieldClass == UObjectProperty::StaticClass())
 	{
-		NewField = new FObjectProperty(InField);
+		if (FLinkerLoad::IsImportLazyLoadEnabled())
+		{
+			NewField = new FObjectPtrProperty(InField);
+		}
+		else
+		{
+			NewField = new FObjectProperty(InField);
+		}
 	}
 	else if (UFieldClass == UWeakObjectProperty::StaticClass())
 	{
diff --git a/Engine/Source/Runtime/CoreUObject/Private/UObject/LinkerLoad.cpp b/Engine/Source/Runtime/CoreUObject/Private/UObject/LinkerLoad.cpp
--- a/Engine/Source/Runtime/CoreUObject/Private/UObject/LinkerLoad.cpp
+++ b/Engine/Source/Runtime/CoreUObject/Private/UObject/LinkerLoad.cpp
@@ -1965,8 +1965,32 @@ FPackageIndex FLinkerLoad::FindOrCreateImport(const FName InObjectName, const FN
 }
 
 #if UE_WITH_OBJECT_HANDLE_LATE_RESOLVE
+FLinkerLoad::EImportLoadBehavior FLinkerLoad::ParseImportLoadBehavior(const FString* LoadBehaviorMeta)
+{
+	EImportLoadBehavior LoadBehavior = EImportLoadBehavior::Eager;
+	if (LoadBehaviorMeta)
+	{
+		if (*LoadBehaviorMeta == "Eager")
+		{
+			LoadBehavior = EImportLoadBehavior::Eager;
+		}
+		// @TODO: OBJPTR: we want to permit lazy background loading in the future
+		// else if (*LoadBehaviorMeta == "LazyBackground")
+		// {
+		// 	LoadBehavior = EImportLoadBehavior::LazyBackground;
+		// }
+		else if (*LoadBehaviorMeta == "LazyOnDemand")
+		{
+			LoadBehavior = EImportLoadBehavior::LazyOnDemand;
+		}
+	}
+
+	return LoadBehavior;
+}
+
 FLinkerLoad::EImportLoadBehavior FLinkerLoad::GetCurrentPropertyImportLoadBehavior(FPackageIndex ImportIndex)
 {
+	static const FName Name_LoadBehavior("LoadBehavior");
 	const FObjectImport& Import = Imp(ImportIndex);
 	if (Import.bImportSearchedFor)
 	{
@@ -1974,19 +1998,41 @@ FLinkerLoad::EImportLoadBehavior FLinkerLoad::GetCurrentPropertyImportLoadBehavi
 		return EImportLoadBehavior::Eager;
 	}
 
+	if (!IsImportLazyLoadEnabled() || !IsAllowingLazyLoading() || (LinkerRoot && LinkerRoot->HasAnyPackageFlags(PKG_PlayInEditor)))
+	{
+		return EImportLoadBehavior::Eager;
+	}
+
+	UObject* ClassPackage = FindObject<UPackage>(nullptr, *Import.ClassPackage.ToString());
+	UClass* FindClass = ClassPackage ? FindObject<UClass>(ClassPackage, *Import.ClassName.ToString()) : nullptr;
+
+	// Attempt to get the meta from the referenced class.  This only looks in already loaded classes.  May need to resolve the class in the future.
+	EImportLoadBehavior LoadBehavior = EImportLoadBehavior::Eager;
+	static const bool bDefaultLoadBehaviorTest = FParse::Param(FCommandLine::Get(), TEXT("DefaultLoadBehaviorTest"));
+	if (bDefaultLoadBehaviorTest)
+	{
+		LoadBehavior = EImportLoadBehavior::LazyOnDemand;
+	}
+
+	if (FindClass)
+	{
+		if (const FString* LoadBehaviorMeta = FindClass->FindMetaData(Name_LoadBehavior))
+		{
+			LoadBehavior = ParseImportLoadBehavior(LoadBehaviorMeta);
+		}
+	}
+
 	const FArchiveSerializedPropertyChain* PropChain = GetSerializedPropertyChain();
 	if (!PropChain)
 	{
-		// In the absence of property information, the safe choice is to not lazy load
-		return EImportLoadBehavior::Eager;
+		return LoadBehavior;
 	}
 
 	const int32 NumPropsInChain = PropChain->GetNumProperties();
 	const FProperty* CurrentProperty = NumPropsInChain > 0 ? PropChain->GetPropertyFromStack(0) : nullptr;
 	if (!CurrentProperty)
 	{
-		// In the absence of property information, the safe choice is to not lazy load
-		return EImportLoadBehavior::Eager;
+		return LoadBehavior;
 	}
 
 	const FProperty* ParentProperty = NumPropsInChain > 1 ? PropChain->GetPropertyFromStack(1) : nullptr;
@@ -2001,47 +2047,9 @@ FLinkerLoad::EImportLoadBehavior FLinkerLoad::GetCurrentPropertyImportLoadBehavi
 		return EImportLoadBehavior::Eager;
 	}
 
-	static bool bAllowLazyResolve = FParse::Param(FCommandLine::Get(), TEXT("AllowLazyResolve"));
-	if (!bAllowLazyResolve || !IsAllowingLazyLoading() || GIsPlayInEditorWorld)
+	if (const FString* LoadBehaviorMeta = CurrentProperty->FindMetaData(Name_LoadBehavior))
 	{
-		return EImportLoadBehavior::Eager;
-	}
-
-	static const bool bLazyResolveAllImports = FParse::Param(FCommandLine::Get(), TEXT("LazyResolveAllImports"));
-	if (bLazyResolveAllImports)
-	{
-		return EImportLoadBehavior::LazyOnDemand;
-	}
-
-	// If we're not blanket lazy resolving all imports, decide if this one has the metadata on the property or the referenced class to allow lazy resolve
-	EImportLoadBehavior LoadBehavior = EImportLoadBehavior::Eager;
-	static const FName Name_LoadBehavior("LoadBehavior");
-	const FString* LoadBehaviorMeta = CurrentProperty->FindMetaData(Name_LoadBehavior);
-	if (!LoadBehaviorMeta)
-	{
-		// Attempt to get the meta from the referenced class.  This only looks in already loaded classes.  May need to resolve the class in the future.
-		UObject* ClassPackage = FindObject<UPackage>(nullptr, *Import.ClassPackage.ToString());
-		if (UClass* FindClass = ClassPackage ? FindObject<UClass>(ClassPackage, *Import.ClassName.ToString()) : nullptr)
-		{
-			LoadBehaviorMeta = FindClass->FindMetaData(Name_LoadBehavior);
-		}
-
-	}
-	if (LoadBehaviorMeta)
-	{
-		if (*LoadBehaviorMeta == "Eager")
-		{
-			LoadBehavior = EImportLoadBehavior::Eager;
-		}
-		// @TODO: OBJPTR: we want to permit lazy background loading in the future
-		// else if (*LoadBehaviorMeta == "LazyBackground")
-		// {
-		// 	LoadBehavior = EImportLoadBehavior::LazyBackground;
-		// }
-		else if (*LoadBehaviorMeta == "LazyOnDemand")
-		{
-			LoadBehavior = EImportLoadBehavior::LazyOnDemand;
-		}
+		LoadBehavior = ParseImportLoadBehavior(LoadBehaviorMeta);
 	}
 
 	return LoadBehavior;
@@ -2674,15 +2682,7 @@ UObject* FLinkerLoad::FindExistingImport(int32 ImportIndex)
 
 void FLinkerLoad::Verify()
 {
-#if UE_WITH_OBJECT_HANDLE_LATE_RESOLVE
-	if (FParse::Param(FCommandLine::Get(), TEXT("DisableLoadingAllImports")) == true)
-	{
-		bHaveImportsBeenVerified = true;
-		return;
-	}
-#endif
-
-	if (!FApp::IsGame() || GIsEditor || IsRunningCommandlet())
+	if (!IsImportLazyLoadEnabled() && (!FApp::IsGame() || GIsEditor || IsRunningCommandlet()))
 	{
 		if (!bHaveImportsBeenVerified)
 		{
@@ -5861,6 +5861,14 @@ bool FLinkerLoad::RemoveKnownMissingPackage(FName PackageName)
 	return FCoreRedirects::RemoveKnownMissing(ECoreRedirectFlags::Type_Package, FCoreRedirectObjectName(NAME_None, NAME_None, PackageName));
 }
 
+#if UE_WITH_OBJECT_HANDLE_LATE_RESOLVE
+bool FLinkerLoad::IsImportLazyLoadEnabled()
+{
+	static const bool bLazyLoadImports = FParse::Param(FCommandLine::Get(), TEXT("LazyLoadImports"));
+	return bLazyLoadImports;
+}
+#endif
+
 void FLinkerLoad::OnNewFileAdded(const FString& Filename)
 {
 	FString PackageName;
diff --git a/Engine/Source/Runtime/CoreUObject/Private/UObject/ObjectHandle.cpp b/Engine/Source/Runtime/CoreUObject/Private/UObject/ObjectHandle.cpp
--- a/Engine/Source/Runtime/CoreUObject/Private/UObject/ObjectHandle.cpp
+++ b/Engine/Source/Runtime/CoreUObject/Private/UObject/ObjectHandle.cpp
@@ -14,6 +14,8 @@
 #include "UObject/Package.h"
 #include "UObject/UObjectHash.h"
 
+DEFINE_LOG_CATEGORY(LogObjectHandle);
+
 bool operator==(const FObjectHandleDataClassDescriptor& Lhs, const FObjectHandleDataClassDescriptor& Rhs)
 {
 	return (Lhs.PackageName == Rhs.PackageName) && (Lhs.ClassName == Rhs.ClassName);
@@ -434,6 +436,7 @@ UObject* ResolveObjectRef(const FObjectRef& ObjectRef, uint32 LoadFlags /*= LOAD
 			else
 			{
 				// Shunt the load request to happen on the game thread and block on its completion.  This is a deadlock risk!  The game thread may be blocked waiting on this thread.
+				UE_LOG(LogObjectHandle, Warning, TEXT("Resolve of object in package '%s' from a non-game thread was shunted to the game thread."), *ObjectRef.PackageName.ToString());
 				TGraphTask<FFullyLoadPackageOnHandleResolveTask>::CreateTask().ConstructAndDispatchWhenReady(TargetPackage)->Wait();
 			}
 
@@ -453,8 +456,16 @@ UObject* ResolveObjectRef(const FObjectRef& ObjectRef, uint32 LoadFlags /*= LOAD
 
 	if (CurrentObject->HasAnyFlags(RF_NeedLoad) && TargetPackage->LinkerLoad)
 	{
-		check(IsInGameThread());
-		TargetPackage->LinkerLoad->Preload(CurrentObject);
+		if (IsInAsyncLoadingThread() || IsInGameThread())
+		{
+			TargetPackage->LinkerLoad->LoadAllObjects(true);
+		}
+		else
+		{
+			// Shunt the load request to happen on the game thread and block on its completion.  This is a deadlock risk!  The game thread may be blocked waiting on this thread.
+			UE_LOG(LogObjectHandle, Warning, TEXT("Resolve of object in package '%s' from a non-game thread was shunted to the game thread."), *ObjectRef.PackageName.ToString());
+			TGraphTask<FFullyLoadPackageOnHandleResolveTask>::CreateTask().ConstructAndDispatchWhenReady(TargetPackage)->Wait();
+		}
 	}
 	ObjectHandle_Private::OnReferenceResolved(ObjectRef, TargetPackage, CurrentObject);
 
diff --git a/Engine/Source/Runtime/CoreUObject/Private/UObject/PropertyBaseObject.cpp b/Engine/Source/Runtime/CoreUObject/Private/UObject/PropertyBaseObject.cpp
--- a/Engine/Source/Runtime/CoreUObject/Private/UObject/PropertyBaseObject.cpp
+++ b/Engine/Source/Runtime/CoreUObject/Private/UObject/PropertyBaseObject.cpp
@@ -536,6 +536,11 @@ bool FObjectPropertyBase::AllowCrossLevel() const
 	return false;
 }
 
+bool FObjectPropertyBase::AllowObjectTypeReinterpretationTo(const FObjectPropertyBase* Other) const
+{
+	return false;
+}
+
 void FObjectPropertyBase::CheckValidObject(void* Value) const
 {
 	UObject *Object = GetObjectPropertyValue(Value);
@@ -586,7 +591,8 @@ void FObjectPropertyBase::CheckValidObject(void* Value) const
 
 bool FObjectPropertyBase::SameType(const FProperty* Other) const
 {
-	return Super::SameType(Other) && (PropertyClass == ((FObjectPropertyBase*)Other)->PropertyClass);
+	return (Super::SameType(Other) || (Other && Other->IsA<FObjectPropertyBase>() && ((FObjectPropertyBase*)Other)->AllowObjectTypeReinterpretationTo(this))) && 
+			 (PropertyClass == ((FObjectPropertyBase*)Other)->PropertyClass);
 }
 
 void FObjectPropertyBase::CopySingleValueToScriptVM( void* Dest, void const* Src ) const
diff --git a/Engine/Source/Runtime/CoreUObject/Private/UObject/PropertyObjectPtr.cpp b/Engine/Source/Runtime/CoreUObject/Private/UObject/PropertyObjectPtr.cpp
--- a/Engine/Source/Runtime/CoreUObject/Private/UObject/PropertyObjectPtr.cpp
+++ b/Engine/Source/Runtime/CoreUObject/Private/UObject/PropertyObjectPtr.cpp
@@ -4,13 +4,37 @@
 #include "UObject/ObjectMacros.h"
 #include "UObject/PropertyHelper.h"
 #include "UObject/UnrealType.h"
+#include "UObject/LinkerLoad.h"
 #include "UObject/LinkerPlaceholderExportObject.h"
 #include "UObject/LinkerPlaceholderClass.h"
 
 /*-----------------------------------------------------------------------------
 	FObjectPtrProperty.
 -----------------------------------------------------------------------------*/
-IMPLEMENT_FIELD(FObjectPtrProperty)
+struct FObjectPropertyFieldNameOverride
+{
+	FObjectPropertyFieldNameOverride(FFieldClass* FieldClass)
+	{
+		FFieldClass::GetNameToFieldClassMap().FindOrAdd(TEXT("ObjectPtrProperty")) = FieldClass;
+		if (!FLinkerLoad::IsImportLazyLoadEnabled())
+		{
+			FFieldClass::GetNameToFieldClassMap().FindOrAdd(TEXT("ObjectProperty")) = FObjectProperty::StaticClass();
+		}
+	}
+};
+
+FField* FObjectPtrProperty::Construct(const FFieldVariant& InOwner, const FName& InName, EObjectFlags InFlags)
+{
+	IMPLEMENT_FIELD_CONSTRUCT_IMPLEMENTATION(FObjectPtrProperty)
+}
+
+FFieldClass* FObjectPtrProperty::StaticClass()
+{
+	// This property type shares the same field name and can supplant FObjectProperty in the NameToFieldClassMap in some configurations because it is serialization compatible with it
+	static FFieldClass StaticFieldClass(TEXT("FObjectProperty"), FObjectPtrProperty::StaticClassCastFlagsPrivate(), FObjectPtrProperty::StaticClassCastFlags(), FObjectPtrProperty::Super::StaticClass(), &FObjectPtrProperty::Construct);
+	static FObjectPropertyFieldNameOverride ObjectPropertyFieldNameOverride(&StaticFieldClass);
+	return &StaticFieldClass;
+}
 
 FString FObjectPtrProperty::GetCPPType(FString* ExtendedTypeText/*=NULL*/, uint32 CPPExportFlags/*=0*/) const
 {
@@ -93,41 +117,11 @@ bool FObjectPtrProperty::Identical(const void* A, const void* B, uint32 PortFlag
 bool FObjectPtrProperty::StaticIdentical(const void* A, const void* B, uint32 PortFlags)
 {
 	FObjectPtr ObjectA = A ? *((FObjectPtr*)A) : FObjectPtr();
+	FObjectHandle ObjectAHandle = ObjectA.GetHandle();
 	FObjectPtr ObjectB = B ? *((FObjectPtr*)B) : FObjectPtr();
+	FObjectHandle ObjectBHandle = ObjectB.GetHandle();
 
-	if (ObjectA.IsNull() || ObjectB.IsNull())
-	{
-		return ObjectA.IsNull() == ObjectB.IsNull();
-	}
-
-	// Compare actual pointers. We don't do this during PIE because we want to be sure to serialize everything. An example is the LevelScriptActor being serialized against its CDO,
-	// which contains actor references. We want to serialize those references so they are fixed up.
-	const bool bDuplicatingForPIE = (PortFlags&PPF_DuplicateForPIE) != 0;
-	bool bResult = !bDuplicatingForPIE ? (ObjectA == ObjectB) : false;
-	// always serialize the cross level references, because they could be NULL
-	// @todo: okay, this is pretty hacky overall - we should have a PortFlag or something
-	// that is set during SavePackage. Other times, we don't want to immediately return false
-	// (instead of just this ExportDefProps case)
-	// instance testing
-	if (!bResult && ObjectA->GetClass() == ObjectB->GetClass())
-	{
-		bool bPerformDeepComparison = (PortFlags&PPF_DeepComparison) != 0;
-		if ((PortFlags&PPF_DeepCompareInstances) && !bPerformDeepComparison)
-		{
-			bPerformDeepComparison = ObjectA->IsTemplate() != ObjectB->IsTemplate();
-		}
-
-		if (!bResult && bPerformDeepComparison)
-		{
-			// In order for deep comparison to be match they both need to have the same name and that name needs to be included in the instancing table for the class
-			if (ObjectA->GetFName() == ObjectB->GetFName() && ObjectA->GetClass()->GetDefaultSubobjectByName(ObjectA->GetFName()))
-			{
-				checkSlow(ObjectA->IsDefaultSubobject() && ObjectB->IsDefaultSubobject() && ObjectA->GetClass()->GetDefaultSubobjectByName(ObjectA->GetFName()) == ObjectB->GetClass()->GetDefaultSubobjectByName(ObjectB->GetFName())); // equivalent
-				bResult = AreInstancedObjectsIdentical(ObjectA.Get(), ObjectB.Get(), PortFlags);
-			}
-		}
-	}
-	return bResult;
+	return ObjectAHandle == ObjectBHandle;
 }
 
 UObject* FObjectPtrProperty::GetObjectPropertyValue(const void* PropertyValueAddress) const
@@ -145,6 +139,11 @@ bool FObjectPtrProperty::AllowCrossLevel() const
 	return true;
 }
 
+bool FObjectPtrProperty::AllowObjectTypeReinterpretationTo(const FObjectPropertyBase* Other) const
+{
+	return Other && Other->IsA<FObjectProperty>();
+}
+
 uint32 FObjectPtrProperty::GetValueTypeHashInternal(const void* Src) const
 {
 	return GetTypeHash((FObjectPtr&)GetPropertyValue(Src));
diff --git a/Engine/Source/Runtime/CoreUObject/Public/UObject/LinkerLoad.h b/Engine/Source/Runtime/CoreUObject/Public/UObject/LinkerLoad.h
--- a/Engine/Source/Runtime/CoreUObject/Public/UObject/LinkerLoad.h
+++ b/Engine/Source/Runtime/CoreUObject/Public/UObject/LinkerLoad.h
@@ -10,6 +10,7 @@
 #include "UObject/SoftObjectPtr.h"
 #include "UObject/ObjectResource.h"
 #include "UObject/PackageResourceManager.h"
+#include "UObject/ObjectHandle.h"
 
 class FLinkerPlaceholderBase;
 class IPakFile;
@@ -282,6 +283,16 @@ public:
 	 */
 	COREUOBJECT_API static bool RemoveKnownMissingPackage(FName PackageName);
 
+	/**
+	 * Determines if imports can be lazily loaded.  This relies on compile-time enabling of UE_WITH_OBJECT_HANDLE_LATE_RESOLVE from ObjectHandle.h as well as other factors.
+	 * @return true if imports can be lazily loaded
+	 */
+	#if UE_WITH_OBJECT_HANDLE_LATE_RESOLVE
+	COREUOBJECT_API static bool IsImportLazyLoadEnabled();
+	#else
+	inline static bool IsImportLazyLoadEnabled() { return false; }
+	#endif // UE_WITH_OBJECT_HANDLE_LATE_RESOLVE
+
 	/**
 	 * 
 	 */
@@ -1054,6 +1065,7 @@ private:
 		//LazyBackground,
 		LazyOnDemand,
 	};
+	static EImportLoadBehavior ParseImportLoadBehavior(const FString* LoadBehaviorMeta);
 	EImportLoadBehavior GetCurrentPropertyImportLoadBehavior(FPackageIndex ImportIndex);
 #endif // UE_WITH_OBJECT_HANDLE_LATE_RESOLVE
 public:
diff --git a/Engine/Source/Runtime/CoreUObject/Public/UObject/ObjectHandle.h b/Engine/Source/Runtime/CoreUObject/Public/UObject/ObjectHandle.h
--- a/Engine/Source/Runtime/CoreUObject/Public/UObject/ObjectHandle.h
+++ b/Engine/Source/Runtime/CoreUObject/Public/UObject/ObjectHandle.h
@@ -5,6 +5,7 @@
 #include "Containers/ScriptArray.h"
 #include "HAL/CriticalSection.h"
 #include "HAL/Platform.h"
+#include "Logging/LogMacros.h"
 #include "UObject/NameTypes.h"
 #include "UObject/Object.h"
 #include "UObject/ObjectPathId.h"
@@ -12,6 +13,8 @@
 #define UE_WITH_OBJECT_HANDLE_LATE_RESOLVE WITH_EDITORONLY_DATA // @TODO: OBJPTR: Should be WITH_EDITORONLY_DATA when ready to tackle lazy load/resolve issues
 #define UE_WITH_OBJECT_HANDLE_TRACKING WITH_EDITORONLY_DATA
 
+DECLARE_LOG_CATEGORY_EXTERN(LogObjectHandle, Log, All);
+
 /**
  * FObjectRef represents a heavyweight reference that contains the specific pieces of information needed to reference an object
  * (or null) that may or may not be loaded yet.  The expectation is that given the imports of a package we have loaded, we should
diff --git a/Engine/Source/Runtime/CoreUObject/Public/UObject/UnrealType.h b/Engine/Source/Runtime/CoreUObject/Public/UObject/UnrealType.h
--- a/Engine/Source/Runtime/CoreUObject/Public/UObject/UnrealType.h
+++ b/Engine/Source/Runtime/CoreUObject/Public/UObject/UnrealType.h
@@ -2182,6 +2182,7 @@ public:
 
 protected:
 	virtual bool AllowCrossLevel() const;
+	virtual bool AllowObjectTypeReinterpretationTo(const FObjectPropertyBase* Other) const;
 	// End of FObjectPropertyBase interface
 };
 
@@ -2318,6 +2319,7 @@ class COREUOBJECT_API FObjectPtrProperty : public FObjectProperty
 	virtual UObject* GetObjectPropertyValue(const void* PropertyValueAddress) const override;
 	virtual void SetObjectPropertyValue(void* PropertyValueAddress, UObject* Value) const override;
 	virtual bool AllowCrossLevel() const override;
+	virtual bool AllowObjectTypeReinterpretationTo(const FObjectPropertyBase* Other) const override;
 private:
 	virtual uint32 GetValueTypeHashInternal(const void* Src) const override;
 public:
