PackagePath: Fix infinite recursion in GetDebugName when extension is unspecified.
#rnx

diff --git a/Engine/Source/Runtime/CoreUObject/Private/Misc/PackagePath.cpp b/Engine/Source/Runtime/CoreUObject/Private/Misc/PackagePath.cpp
--- a/Engine/Source/Runtime/CoreUObject/Private/Misc/PackagePath.cpp
+++ b/Engine/Source/Runtime/CoreUObject/Private/Misc/PackagePath.cpp
@@ -665,7 +665,7 @@ void FPackagePath::AppendLocalFullPath(FStringBuilderBase& Builder, EPackageSegm
 				UE_LOG(LogPackageName, Warning,
 					TEXT("GetLocalFullPath called on FPackagePath %s which has an unspecified header extension,")
 					TEXT(" and the path does not exist on disk. Assuming EPackageExtension::Asset."),
-					*GetDebugName());
+					Builder.ToString() + BeforeLen); // Can't use GetDebugName since that can call this function recursively
 				Extension = EPackageExtension::Asset;
 			}
 		}
