Change presave prototype to take an FObjectSaveContextRef to provide more information about the save parameters.

diff --git a/Engine/Source/Editor/WorldBrowser/Private/Tiles/WorldTileCollectionModel.h b/Engine/Source/Editor/WorldBrowser/Private/Tiles/WorldTileCollectionModel.h
--- a/Engine/Source/Editor/WorldBrowser/Private/Tiles/WorldTileCollectionModel.h
+++ b/Engine/Source/Editor/WorldBrowser/Private/Tiles/WorldTileCollectionModel.h
@@ -158,10 +158,10 @@ private:
 	void SetupParentChildLinks();
 
 	/** Called before saving world into package file */
-	void OnPreSaveWorld(uint32 SaveFlags, class UWorld* World);
+	void OnPreSaveWorld(class UWorld* World, FObjectPreSaveContext ObjectSaveContext);
 	
 	/** Called right after world was saved into package file */
-	void OnPostSaveWorld(uint32 SaveFlags, class UWorld* World, bool bSuccess);
+	void OnPostSaveWorld(class UWorld* World, FObjectPostSaveContext ObjectSaveContext);
 
 	/** Called when world has new current level */
 	void OnNewCurrentLevel();
diff --git a/Engine/Source/Runtime/AIModule/Classes/EnvironmentQuery/Tests/EnvQueryTest_GameplayTags.h b/Engine/Source/Runtime/AIModule/Classes/EnvironmentQuery/Tests/EnvQueryTest_GameplayTags.h
--- a/Engine/Source/Runtime/AIModule/Classes/EnvironmentQuery/Tests/EnvQueryTest_GameplayTags.h
+++ b/Engine/Source/Runtime/AIModule/Classes/EnvironmentQuery/Tests/EnvQueryTest_GameplayTags.h
@@ -31,13 +31,15 @@ protected:
 
 	bool SatisfiesTest(IGameplayTagAssetInterface* ItemGameplayTagAssetInterface) const;
 
+	UE_DEPRECATED(5.0, "Use version that takes FObjectPreSaveContext instead.")
+	virtual void PreSave(const class ITargetPlatform* TargetPlatform) override;
 	/**
 	 * Presave function. Gets called once before an object gets serialized for saving. This function is necessary
 	 * for save time computation as Serialize gets called three times per object from within SavePackage.
 	 *
 	 * @warning: Objects created from within PreSave will NOT have PreSave called on them!!!
 	 */
-	virtual void PreSave(const class ITargetPlatform* TargetPlatform) override;
+	virtual void PreSave(FObjectPreSaveContext ObjectSaveContext) override;
 
 	virtual void PostLoad() override;
 
diff --git a/Engine/Source/Runtime/AIModule/Private/EnvironmentQuery/Tests/EnvQueryTest_GameplayTags.cpp b/Engine/Source/Runtime/AIModule/Private/EnvironmentQuery/Tests/EnvQueryTest_GameplayTags.cpp
--- a/Engine/Source/Runtime/AIModule/Private/EnvironmentQuery/Tests/EnvQueryTest_GameplayTags.cpp
+++ b/Engine/Source/Runtime/AIModule/Private/EnvironmentQuery/Tests/EnvQueryTest_GameplayTags.cpp
@@ -1,10 +1,12 @@
 // Copyright Epic Games, Inc. All Rights Reserved.
 
 #include "EnvironmentQuery/Tests/EnvQueryTest_GameplayTags.h"
+
 #include "GameFramework/Actor.h"
 #include "EnvironmentQuery/Items/EnvQueryItemType_ActorBase.h"
 #include "BehaviorTree/BTNode.h"
 #include "GameplayTagAssetInterface.h"
+#include "UObject/ObjectSaveContext.h"
 
 UEnvQueryTest_GameplayTags::UEnvQueryTest_GameplayTags(const FObjectInitializer& ObjectInitializer) :
 	Super(ObjectInitializer)
@@ -19,13 +21,20 @@ UEnvQueryTest_GameplayTags::UEnvQueryTest_GameplayTags(const FObjectInitializer&
 }
 
 void UEnvQueryTest_GameplayTags::PreSave(const class ITargetPlatform* TargetPlatform)
+{
+	PRAGMA_DISABLE_DEPRECATION_WARNINGS;
+	Super::PreSave(TargetPlatform);
+	PRAGMA_ENABLE_DEPRECATION_WARNINGS;
+}
+
+void UEnvQueryTest_GameplayTags::PreSave(FObjectPreSaveContext ObjectSaveContext)
 {
 	// If this is a new object, it's already converted to new format.  If it was old, this should've already been set to
 	// true in PostLoad.  This prevents us from erroneously trying to import old data multiple times and so obliterating
 	// current data.
 	bUpdatedToUseQuery = true;
 
-	Super::PreSave(TargetPlatform);
+	Super::PreSave(ObjectSaveContext);
 }
 
 void UEnvQueryTest_GameplayTags::PostLoad()
diff --git a/Engine/Source/Runtime/CoreUObject/Private/Serialization/BulkData.cpp b/Engine/Source/Runtime/CoreUObject/Private/Serialization/BulkData.cpp
--- a/Engine/Source/Runtime/CoreUObject/Private/Serialization/BulkData.cpp
+++ b/Engine/Source/Runtime/CoreUObject/Private/Serialization/BulkData.cpp
@@ -1496,7 +1496,7 @@ void FUntypedBulkData::Serialize( FArchive& Ar, UObject* Owner, int32 Idx, bool
 #if WITH_EDITOR
 				// If we are overwriting the LoadedPath for the current package, set the location variables on *this equal to the new values
 				// that we are writing into the package on disk
-				if (LinkerSave && LinkerSave->bSavingNewLoadedPath)
+				if (LinkerSave && LinkerSave->bUpdatingLoadedPath)
 				{
 					SetFlagsFromDiskWrittenValues(LocalBulkDataFlags, LocalBulkDataOffsetInFile, LocalBulkDataSizeOnDisk,
 						INDEX_NONE /* LinkerSummaryBulkDataStartOffset, not applicable */);
diff --git a/Engine/Source/Runtime/CoreUObject/Private/UObject/Obj.cpp b/Engine/Source/Runtime/CoreUObject/Private/UObject/Obj.cpp
--- a/Engine/Source/Runtime/CoreUObject/Private/UObject/Obj.cpp
+++ b/Engine/Source/Runtime/CoreUObject/Private/UObject/Obj.cpp
@@ -34,6 +34,7 @@
 #include "UObject/PropertyPortFlags.h"
 #include "UObject/UnrealType.h"
 #include "UObject/ObjectRedirector.h"
+#include "UObject/ObjectSaveContext.h"
 #include "UObject/UObjectAnnotation.h"
 #include "UObject/ReferenceChainSearch.h"
 #include "Serialization/ArchiveCountMem.h"
@@ -1192,8 +1193,30 @@ void UObject::ConditionalPostLoadSubobjects( FObjectInstancingGraph* OuterInstan
 	CheckDefaultSubobjects();
 }
 
+bool UObject::PreSaveRoot(const TCHAR* Filename)
+{
+	return false;
+}
+
+void UObject::PreSaveRoot(FObjectPreSaveRootContext ObjectSaveContext)
+{
+}
+
+void UObject::PostSaveRoot(bool bCleanupIsRequired)
+{
+}
+
+void UObject::PostSaveRoot(FObjectPostSaveRootContext ObjectSaveContext)
+{
+}
+
 void UObject::PreSave(const class ITargetPlatform* TargetPlatform)
 {
+}
+
+void UObject::PreSave(FObjectPreSaveContext SaveContext)
+{
+	SaveContext.Data.bBaseClassCalled = true;
 #if WITH_EDITOR
 	FCoreUObjectDelegates::OnObjectSaved.Broadcast(this);
 #endif
diff --git a/Engine/Source/Runtime/CoreUObject/Private/UObject/ObjectRedirector.cpp b/Engine/Source/Runtime/CoreUObject/Private/UObject/ObjectRedirector.cpp
--- a/Engine/Source/Runtime/CoreUObject/Private/UObject/ObjectRedirector.cpp
+++ b/Engine/Source/Runtime/CoreUObject/Private/UObject/ObjectRedirector.cpp
@@ -5,6 +5,8 @@
 =============================================================================*/
 
 #include "UObject/ObjectRedirector.h"
+
+#include "UObject/ObjectSaveContext.h"
 #include "UObject/Package.h"
 #include "Templates/Casts.h"
 #include "UObject/PropertyPortFlags.h"
@@ -14,9 +16,16 @@
 -----------------------------------------------------------------------------*/
 
 
+void UObjectRedirector::PreSave(const class ITargetPlatform* TargetPlatform)
+{
+	PRAGMA_DISABLE_DEPRECATION_WARNINGS;
+	Super::PreSave(TargetPlatform);
+	PRAGMA_ENABLE_DEPRECATION_WARNINGS;
+}
+
 // If this object redirector is pointing to an object that won't be serialized anyway, set the RF_Transient flag
 // so that this redirector is also removed from the package.
-void UObjectRedirector::PreSave(const class ITargetPlatform* TargetPlatform)
+void UObjectRedirector::PreSave(FObjectPreSaveContext ObjectSaveContext)
 {
 	if (DestinationObject == NULL
 	||	DestinationObject->HasAnyFlags(RF_Transient)
@@ -31,6 +40,8 @@ void UObjectRedirector::PreSave(const class ITargetPlatform* TargetPlatform)
 			DestinationObject->SetFlags(RF_Transient);
 		}
 	}
+
+	Super::PreSave(ObjectSaveContext);
 }
 
 IMPLEMENT_FARCHIVE_SERIALIZER(UObjectRedirector);
diff --git a/Engine/Source/Runtime/CoreUObject/Private/UObject/Package.cpp b/Engine/Source/Runtime/CoreUObject/Private/UObject/Package.cpp
--- a/Engine/Source/Runtime/CoreUObject/Private/UObject/Package.cpp
+++ b/Engine/Source/Runtime/CoreUObject/Private/UObject/Package.cpp
@@ -17,11 +17,12 @@
 	UPackage.
 -----------------------------------------------------------------------------*/
 
-/** Delegate to notify subscribers when a package is about to be saved. */
+PRAGMA_DISABLE_DEPRECATION_WARNINGS;
 UPackage::FPreSavePackage UPackage::PreSavePackageEvent;
-/** Delegate to notify subscribers when a package has been saved. This is triggered when the package saving
- *  has completed and was successful. */
 UPackage::FOnPackageSaved UPackage::PackageSavedEvent;
+PRAGMA_ENABLE_DEPRECATION_WARNINGS;
+UPackage::FPreSavePackageWithContext UPackage::PreSavePackageWithContextEvent;
+UPackage::FOnPackageSavedWithContext UPackage::PackageSavedWithContextEvent;
 /** Delegate to notify subscribers when the dirty state of a package is changed.
  *  Allows the editor to register the modified package as one that should be prompted for source control checkout. 
  *  Use Package->IsDirty() to get the updated dirty state of the package */
diff --git a/Engine/Source/Runtime/CoreUObject/Private/UObject/SavePackage.cpp b/Engine/Source/Runtime/CoreUObject/Private/UObject/SavePackage.cpp
--- a/Engine/Source/Runtime/CoreUObject/Private/UObject/SavePackage.cpp
+++ b/Engine/Source/Runtime/CoreUObject/Private/UObject/SavePackage.cpp
@@ -26,6 +26,7 @@
 #include "UObject/UObjectGlobals.h"
 #include "UObject/UObjectHash.h"
 #include "UObject/Object.h"
+#include "UObject/ObjectSaveContext.h"
 #include "Serialization/ArchiveUObject.h"
 #include "UObject/Class.h"
 #include "UObject/UObjectIterator.h"
@@ -1641,7 +1642,7 @@ struct FPackageExportTagger
 	,	TargetPlatform(InTargetPlatform)
 	{}
 
-	void TagPackageExports( FArchiveSaveTagExports& ExportTagger, bool bRoutePresave )
+	void TagPackageExports( FArchiveSaveTagExports& ExportTagger, bool bRoutePresave, FObjectSaveContextData& ObjectSaveContext )
 	{
 		const bool bIsCooking = !!TargetPlatform;
 
@@ -1655,7 +1656,7 @@ struct FPackageExportTagger
 					FArchiveObjectCrc32NonEditorProperties CrcArchive;
 
 					int32 Before = CrcArchive.Crc32(Base);
-					Base->PreSave(TargetPlatform);
+					UE::SavePackageUtilities::CallPreSave(Base, ObjectSaveContext);
 					int32 After = CrcArchive.Crc32(Base);
 
 					if (Before != After)
@@ -1672,7 +1673,7 @@ struct FPackageExportTagger
 				}
 				else
 				{
-					Base->PreSave(TargetPlatform);
+					UE::SavePackageUtilities::CallPreSave(Base, ObjectSaveContext);
 				}
 			}
 
@@ -1715,7 +1716,7 @@ struct FPackageExportTagger
 						FArchiveObjectCrc32NonEditorProperties CrcArchive;
 
 						int32 Before = CrcArchive.Crc32(Obj);
-						Obj->PreSave(TargetPlatform);
+						UE::SavePackageUtilities::CallPreSave(Obj, ObjectSaveContext);
 						int32 After = CrcArchive.Crc32(Obj);
 
 						if (Before != After)
@@ -1732,7 +1733,7 @@ struct FPackageExportTagger
 					}
 					else
 					{
-						Obj->PreSave(TargetPlatform);
+						UE::SavePackageUtilities::CallPreSave(Obj, ObjectSaveContext);
 					}
 				}
 			}
@@ -2053,6 +2054,10 @@ FSavePackageResultStruct UPackage::Save(UPackage* InOuter, UObject* Base, EObjec
 	check(Filename);
 	const bool bIsCooking = TargetPlatform != nullptr;
 	FPackagePath TargetPackagePath = FPackagePath::FromLocalPath(Filename);
+	if (TargetPackagePath.GetHeaderExtension() == EPackageExtension::Unspecified)
+	{
+		TargetPackagePath.SetHeaderExtension(EPackageExtension::EmptyString);
+	}
 
 #if WITH_EDITOR
 	TMap<UObject*, UObject*> ReplacedImportOuters;
@@ -2070,8 +2075,8 @@ FSavePackageResultStruct UPackage::Save(UPackage* InOuter, UObject* Base, EObjec
 		InOuter->SetPersistentGuid(FGuid::NewGuid());
 	}
 #endif //WITH_EDITOR
-	const bool bSavingNewLoadedPath = SavePackageUtilities::IsSavingNewLoadedPath(bIsCooking, TargetPackagePath, SaveFlags);
 	const bool bSavingConcurrent = !!(SaveFlags & ESaveFlags::SAVE_Concurrent);
+	FObjectSaveContextData ObjectSaveContext(InOuter, TargetPlatform, TargetPackagePath, SaveFlags);
 
 	if (FPlatformProperties::HasEditorOnlyData())
 	{
@@ -2231,10 +2236,11 @@ FSavePackageResultStruct UPackage::Save(UPackage* InOuter, UObject* Base, EObjec
 
 		// Route PreSaveRoot to allow e.g. the world to attach components for the persistent level.
 		// If we are saving concurrently, this should have been called before UPackage::Save was called.
-		bool bCleanupIsRequired = false;
+		bool bCleanupRequired = false;
 		if (Base && !bSavingConcurrent)
 		{
-			bCleanupIsRequired = Base->PreSaveRoot(Filename);
+			UE::SavePackageUtilities::CallPreSaveRoot(Base, ObjectSaveContext);
+			bCleanupRequired = ObjectSaveContext.bCleanupRequired;
 		}
 
 		// Init.
@@ -2295,7 +2301,7 @@ FSavePackageResultStruct UPackage::Save(UPackage* InOuter, UObject* Base, EObjec
 				// Do not route presave if saving concurrently. This should have been done before the concurrent save started.
 				// Also if we're trying to diff the package and gathering callstacks, Presave has already been done
 				const bool bRoutePresave = !bSavingConcurrent && !(SaveFlags & SAVE_DiffCallstack);
-				PackageExportTagger.TagPackageExports(ExportTaggerArchive, bRoutePresave);
+				PackageExportTagger.TagPackageExports(ExportTaggerArchive, bRoutePresave, ObjectSaveContext);
 				ExportTaggerArchive.SetFilterEditorOnly(FilterEditorOnly);
 			}
 		
@@ -2326,7 +2332,7 @@ FSavePackageResultStruct UPackage::Save(UPackage* InOuter, UObject* Base, EObjec
 					UnMarkAllObjects();
 			
 					// We need to serialize objects yet again to tag objects that were created by PreSave as OBJECTMARK_TagExp.
-					PackageExportTagger.TagPackageExports( ExportTaggerArchive, false );
+					PackageExportTagger.TagPackageExports( ExportTaggerArchive, false, ObjectSaveContext);
 				}
 
 				// Kick off any Precaching required for the target platform to save these objects
@@ -2451,7 +2457,8 @@ FSavePackageResultStruct UPackage::Save(UPackage* InOuter, UObject* Base, EObjec
 							Linker = TUniquePtr<FLinkerSave>(new FLinkerSave(InOuter, *TempFilename.GetValue(), bForceByteSwapping, bSaveUnversioned));
 						}
 					}
-					Linker->bSavingNewLoadedPath = bSavingNewLoadedPath;
+					Linker->bProceduralSave = ObjectSaveContext.bProceduralSave;
+					Linker->bUpdatingLoadedPath = ObjectSaveContext.bUpdatingLoadedPath;
 
 #if WITH_TEXT_ARCHIVE_SUPPORT
 					if (bTextFormat)
@@ -4612,7 +4619,7 @@ FSavePackageResultStruct UPackage::Save(UPackage* InOuter, UObject* Base, EObjec
 			// Also if we're trying to diff the package and gathering callstacks, Presave has already been done
 			if( Base && !bSavingConcurrent && !(SaveFlags & SAVE_DiffCallstack) )
 			{
-				Base->PostSaveRoot( bCleanupIsRequired );
+				UE::SavePackageUtilities::CallPostSaveRoot(Base, ObjectSaveContext, bCleanupRequired);
 			}
 
 			SlowTask.EnterProgressFrame();
@@ -4634,7 +4641,10 @@ FSavePackageResultStruct UPackage::Save(UPackage* InOuter, UObject* Base, EObjec
 			InOuter->ClearPackageFlags(PKG_NewlyCreated);
 
 			// send a message that the package was saved
+			PRAGMA_DISABLE_DEPRECATION_WARNINGS;
 			UPackage::PackageSavedEvent.Broadcast(Filename, InOuter);
+			PRAGMA_ENABLE_DEPRECATION_WARNINGS;
+			UPackage::PackageSavedWithContextEvent.Broadcast(Filename, InOuter, FObjectPostSaveContext(ObjectSaveContext));
 		}
 
 		// We're done!
@@ -4646,7 +4656,7 @@ FSavePackageResultStruct UPackage::Save(UPackage* InOuter, UObject* Base, EObjec
 		{
 			// if the save was successful, update the internal package filename path if we aren't currently cooking
 #if WITH_EDITOR
-			if (bSavingNewLoadedPath)
+			if (ObjectSaveContext.bUpdatingLoadedPath)
 			{
 				InOuter->SetLoadedPath(TargetPackagePath);
 			}
diff --git a/Engine/Source/Runtime/CoreUObject/Private/UObject/SavePackage/SaveContext.h b/Engine/Source/Runtime/CoreUObject/Private/UObject/SavePackage/SaveContext.h
--- a/Engine/Source/Runtime/CoreUObject/Private/UObject/SavePackage/SaveContext.h
+++ b/Engine/Source/Runtime/CoreUObject/Private/UObject/SavePackage/SaveContext.h
@@ -17,6 +17,7 @@
 #include "UObject/LinkerSave.h"
 #include "UObject/NameTypes.h"
 #include "UObject/ObjectMacros.h"
+#include "UObject/ObjectSaveContext.h"
 #include "UObject/Package.h"
 #include "UObject/SavePackage.h"
 #include "UObject/SavePackage/SavePackageUtilities.h"
@@ -88,6 +89,10 @@ public:
 		}
 
 		TargetPackagePath = FPackagePath::FromLocalPath(InFilename);
+		if (TargetPackagePath.GetHeaderExtension() == EPackageExtension::Unspecified)
+		{
+			TargetPackagePath.SetHeaderExtension(EPackageExtension::EmptyString);
+		}
 
 		// if we aren't cooking and top level flags aren't empty, add RF_HasExternalPackage to them to catch external packages data
 		if (SaveArgs.TopLevelFlags != RF_NoFlags && !IsCooking())
@@ -103,7 +108,7 @@ public:
 			bIsProcessingPrestreamPackages = ProcessPrestreamingRequests->GetInt() > 0;
 		}
 
-		bSavingNewLoadedPath = SavePackageUtilities::IsSavingNewLoadedPath(IsCooking(), TargetPackagePath, SaveArgs.SaveFlags);
+		ObjectSaveContext.Set(InPackage, GetTargetPlatform(), TargetPackagePath, SaveArgs.SaveFlags);
 	}
 
 	~FSaveContext()
@@ -121,7 +126,7 @@ public:
 
 		if (bNeedPreSaveCleanup && Asset)
 		{
-			Asset->PostSaveRoot(bNeedPreSaveCleanup);
+			UE::SavePackageUtilities::CallPostSaveRoot(Asset, ObjectSaveContext, bNeedPreSaveCleanup);
 		}
 	}
 
@@ -190,9 +195,14 @@ public:
 		return SaveArgs.TargetPlatform != nullptr;
 	}
 
-	bool IsSavingNewLoadedPath() const
+	bool IsProceduralSave() const
+	{
+		return ObjectSaveContext.bProceduralSave;
+	}
+
+	bool IsUpdatingLoadedPath() const
 	{
-		return bSavingNewLoadedPath;
+		return ObjectSaveContext.bUpdatingLoadedPath;
 	}
 
 	bool IsFilterEditorOnly() const
@@ -536,6 +546,11 @@ public:
 		return FSavePackageResultStruct(FinalResult, TotalPackageSizeUncompressed, AsyncWriteAndHashSequence.Finalize(EAsyncExecution::TaskGraph, MoveTemp(HashCompletionFunc)), IsCompareLinker() ? MoveTemp(Linker) : nullptr);
 	}
 
+	FObjectSaveContextData& GetObjectSaveContext()
+	{
+		return ObjectSaveContext;
+	}
+
 public:
 	ESavePackageResult Result;
 
@@ -569,12 +584,12 @@ private:
 
 	// State context
 	FUObjectSerializeContext* SerializeContext = nullptr;
+	FObjectSaveContextData ObjectSaveContext;
 	bool bCanUseUnversionedPropertySerialization = false;
 	bool bTextFormat = false;
 	bool bIsProcessingPrestreamPackages = false;
 	bool bNeedPreSaveCleanup = false;
 	bool bGenerateFileStub = false;
-	bool bSavingNewLoadedPath = false;
 
 	// Config classes shared with the old Save
 	FSavePackageDiffSettings DiffSettings;
diff --git a/Engine/Source/Runtime/CoreUObject/Private/UObject/SavePackage/SavePackageUtilities.cpp b/Engine/Source/Runtime/CoreUObject/Private/UObject/SavePackage/SavePackageUtilities.cpp
--- a/Engine/Source/Runtime/CoreUObject/Private/UObject/SavePackage/SavePackageUtilities.cpp
+++ b/Engine/Source/Runtime/CoreUObject/Private/UObject/SavePackage/SavePackageUtilities.cpp
@@ -25,6 +25,7 @@
 #include "UObject/LinkerSave.h"
 #include "UObject/Object.h"
 #include "UObject/ObjectRedirector.h"
+#include "UObject/ObjectSaveContext.h"
 #include "UObject/Package.h"
 #include "UObject/SavePackage.h"
 #include "UObject/UnrealType.h"
@@ -690,19 +691,110 @@ void GetCDOSubobjects(UObject* CDO, TArray<UObject*>& Subobjects)
 		}
 	}
 }
+} // end namespace SavePackageUtilities
+
+namespace UE
+{
+namespace SavePackageUtilities
+{
 
-bool IsSavingNewLoadedPath(bool bIsCooking, const FPackagePath& TargetPackagePath, uint32 SaveFlags)
+bool IsUpdatingLoadedPath(bool bIsCooking, const FPackagePath& TargetPackagePath, uint32 SaveFlags)
 {
 #if WITH_EDITOR
-	return !bIsCooking &&												// Do not update the loadedpath if we're cooking
-		TargetPackagePath.IsMountedPath() &&							// Do not update the loadedpath if the new path is not a viable mounted path
-		!(SaveFlags & (SAVE_BulkDataByReference | SAVE_FromAutosave));	// Do not update the loadedpath if we're storing references to the old loadedpath or autosaving
+	return !bIsCooking &&							// Do not update the loadedpath if we're cooking
+		TargetPackagePath.IsMountedPath() &&		// Do not update the loadedpath if the new path is not a viable mounted path
+		!(SaveFlags & SAVE_BulkDataByReference) &&	// Do not update the loadedpath if it's an EditorDomainSave. TODO: Change the name of this flag.
+		!(SaveFlags & SAVE_FromAutosave);			// Do not update the loadedpath if it's an autosave.
 #else
 	return false; // Saving when not in editor never updates the LoadedPath
 #endif
 }
 
-} // end namespace SavePackageUtilities
+bool IsProceduralSave(bool bIsCooking, const FPackagePath& TargetPackagePath, uint32 SaveFlags)
+{
+#if WITH_EDITOR
+	return bIsCooking ||							// Cooking is a procedural save
+		(SaveFlags & SAVE_BulkDataByReference);		// EditorDomainSave is a procedural save. TODO: Change the name of this flag.
+#else
+	return false; // Saving when not in editor never has user changes
+#endif
+}
+
+void CallPreSave(UObject* Object, FObjectSaveContextData& ObjectSaveContext)
+{
+	PRAGMA_DISABLE_DEPRECATION_WARNINGS;
+	Object->PreSave(ObjectSaveContext.TargetPlatform);
+	PRAGMA_ENABLE_DEPRECATION_WARNINGS;
+
+	FObjectPreSaveContext ObjectPreSaveContext(ObjectSaveContext);
+	ObjectSaveContext.bBaseClassCalled = false;
+	ObjectSaveContext.NumRefPasses = 0;
+	Object->PreSave(ObjectPreSaveContext);
+	if (!ObjectSaveContext.bBaseClassCalled)
+	{
+		UE_LOG(LogSavePackage, Warning, TEXT("Class %s did not call Super::PreSave"), *Object->GetClass()->GetName());
+	}
+	// When we deprecate PreSave, and need to take different actions based on the PreSave, remove this bAllowPreSave variable
+	constexpr bool bAllowPreSave = true;
+	if (!bAllowPreSave && ObjectSaveContext.NumRefPasses > 1)
+	{
+		UE_LOG(LogSavePackage, Warning, TEXT("Class %s overrides the deprecated PreSave function"), *Object->GetClass()->GetName());
+	}
+}
+
+void CallPreSaveRoot(UObject* Object, FObjectSaveContextData& ObjectSaveContext)
+{
+	PRAGMA_DISABLE_DEPRECATION_WARNINGS;
+	bool bLegacyNeedsCleanup = Object->PreSaveRoot(*ObjectSaveContext.TargetFilename);
+	PRAGMA_ENABLE_DEPRECATION_WARNINGS;
+
+	ObjectSaveContext.bCleanupRequired = false;
+	Object->PreSaveRoot(FObjectPreSaveRootContext(ObjectSaveContext));
+	ObjectSaveContext.bCleanupRequired |= bLegacyNeedsCleanup;
+}
+
+void CallPostSaveRoot(UObject* Object, FObjectSaveContextData& ObjectSaveContext, bool bNeedsCleanup)
+{
+	PRAGMA_DISABLE_DEPRECATION_WARNINGS;
+	Object->PostSaveRoot(bNeedsCleanup);
+	PRAGMA_ENABLE_DEPRECATION_WARNINGS;
+
+	ObjectSaveContext.bCleanupRequired = bNeedsCleanup;
+	Object->PostSaveRoot(FObjectPostSaveRootContext(ObjectSaveContext));
+}
+
+}
+} // end namespace UE::SavePackageUtilities
+
+FObjectSaveContextData::FObjectSaveContextData(UPackage* Package, const ITargetPlatform* InTargetPlatform, const TCHAR* InTargetFilename, uint32 InSaveFlags)
+{
+	Set(Package, InTargetPlatform, InTargetFilename, InSaveFlags);
+}
+
+FObjectSaveContextData::FObjectSaveContextData(UPackage* Package, const ITargetPlatform* InTargetPlatform, const FPackagePath& TargetPath, uint32 InSaveFlags)
+{
+	Set(Package, InTargetPlatform, TargetPath, InSaveFlags);
+}
+
+void FObjectSaveContextData::Set(UPackage* Package, const ITargetPlatform* InTargetPlatform, const TCHAR* InTargetFilename, uint32 InSaveFlags)
+{
+	FPackagePath PackagePath(FPackagePath::FromLocalPath(TargetFilename));
+	if (PackagePath.GetHeaderExtension() == EPackageExtension::Unspecified)
+	{
+		PackagePath.SetHeaderExtension(EPackageExtension::EmptyString);
+	}
+	Set(Package, InTargetPlatform, PackagePath, InSaveFlags);
+}
+
+void FObjectSaveContextData::Set(UPackage* Package, const ITargetPlatform* InTargetPlatform, const FPackagePath& TargetPath, uint32 InSaveFlags)
+{
+	TargetFilename = TargetPath.GetLocalFullPath();
+	TargetPlatform = InTargetPlatform;
+	SaveFlags = InSaveFlags;
+	OriginalPackageFlags = Package ? Package->GetPackageFlags() : 0;
+	bProceduralSave = UE::SavePackageUtilities::IsProceduralSave(InTargetPlatform != nullptr, TargetPath, InSaveFlags);
+	bUpdatingLoadedPath = UE::SavePackageUtilities::IsUpdatingLoadedPath(InTargetPlatform != nullptr, TargetPath, InSaveFlags);
+}
 
 bool IsEditorOnlyObject(const UObject* InObject, bool bCheckRecursive, bool bCheckMarks)
 {
@@ -1879,9 +1971,9 @@ ESavePackageResult SaveBulkData(FLinkerSave* Linker, const UPackage* InOuter, co
 		{
 			FailureReason = TEXT("SAVE_BulkDataByReference is incompatible with BulkDataManifest");
 		}
-		else if (Linker->bSavingNewLoadedPath)
+		else if (Linker->bUpdatingLoadedPath)
 		{
-			FailureReason = TEXT("SAVE_BulkDataByReference is incompatible with bSavingNewLoadedPath");
+			FailureReason = TEXT("SAVE_BulkDataByReference is incompatible with bUpdatingLoadedPath");
 		}
 		else if (SavePackageContext != nullptr && SavePackageContext->bForceLegacyOffsets)
 		{
@@ -2084,7 +2176,7 @@ ESavePackageResult SaveBulkData(FLinkerSave* Linker, const UPackage* InOuter, co
 #if WITH_EDITOR
 		// If we are overwriting the LoadedPath for the current package, the bulk data flags and location need to be updated to match
 		// the values set in the package on disk.
-		if (Linker->bSavingNewLoadedPath)
+		if (Linker->bUpdatingLoadedPath)
 		{
 			BulkData->SetFlagsFromDiskWrittenValues(static_cast<EBulkDataFlags>(BulkDataFlags), BulkDataOffsetInFile,
 				BulkDataSizeOnDisk, StartOfBulkDataArea);
diff --git a/Engine/Source/Runtime/CoreUObject/Private/UObject/SavePackage/SavePackageUtilities.h b/Engine/Source/Runtime/CoreUObject/Private/UObject/SavePackage/SavePackageUtilities.h
--- a/Engine/Source/Runtime/CoreUObject/Private/UObject/SavePackage/SavePackageUtilities.h
+++ b/Engine/Source/Runtime/CoreUObject/Private/UObject/SavePackage/SavePackageUtilities.h
@@ -358,12 +358,6 @@ namespace SavePackageUtilities
 	void AsyncWriteFileWithSplitExports(TAsyncWorkSequence<FMD5>& AsyncWriteAndHashSequence, FLargeMemoryPtr Data, const int64 DataSize, const int64 HeaderSize, const TCHAR* Filename, EAsyncWriteOptions Options, TArrayView<const FFileRegion> InFileRegions);
 
 	void GetCDOSubobjects(UObject* CDO, TArray<UObject*>& Subobjects);
-
-	/**
-	 * Return whether the given save parameters indicate the LoadedPath of the package being saved should be updated.
-	 * This allows us to update the in-memory package when it is saved in editor to match its new save file.
-	 */
-	bool IsSavingNewLoadedPath(bool bIsCooking, const FPackagePath& TargetPackagePath, uint32 SaveFlags);
 }
 
 #if ENABLE_COOK_STATS
diff --git a/Engine/Source/Runtime/CoreUObject/Private/UObject/SavePackage2.cpp b/Engine/Source/Runtime/CoreUObject/Private/UObject/SavePackage2.cpp
--- a/Engine/Source/Runtime/CoreUObject/Private/UObject/SavePackage2.cpp
+++ b/Engine/Source/Runtime/CoreUObject/Private/UObject/SavePackage2.cpp
@@ -31,6 +31,7 @@
 #include "UObject/Object.h"
 #include "UObject/ObjectMacros.h"
 #include "UObject/ObjectRedirector.h"
+#include "UObject/ObjectSaveContext.h"
 #include "UObject/Package.h"
 #include "UObject/SavePackage/PackageHarvester.h"
 #include "UObject/SavePackage/SaveContext.h"
@@ -240,7 +241,7 @@ ESavePackageResult RoutePresave(FSaveContext& SaveContext)
 			{
 				FArchiveObjectCrc32NonEditorProperties CrcArchive;
 				int32 Before = CrcArchive.Crc32(Object);
-				Object->PreSave(SaveContext.GetTargetPlatform());
+				UE::SavePackageUtilities::CallPreSave(Object, SaveContext.GetObjectSaveContext());
 				int32 After = CrcArchive.Crc32(Object);
 
 				if (Before != After)
@@ -257,7 +258,7 @@ ESavePackageResult RoutePresave(FSaveContext& SaveContext)
 			}
 			else
 			{
-				Object->PreSave(SaveContext.GetTargetPlatform());
+				UE::SavePackageUtilities::CallPreSave(Object, SaveContext.GetObjectSaveContext());
 			}
 		}
 	}
@@ -700,7 +701,8 @@ ESavePackageResult CreateLinker(FSaveContext& SaveContext)
 				SaveContext.Linker = MakeUnique<FLinkerSave>(SaveContext.GetPackage(), *SaveContext.TempFilename.GetValue(), SaveContext.IsForceByteSwapping(), SaveContext.IsSaveUnversioned());
 			}
 		}
-		SaveContext.Linker->bSavingNewLoadedPath = SaveContext.IsSavingNewLoadedPath();
+		SaveContext.Linker->bUpdatingLoadedPath = SaveContext.IsUpdatingLoadedPath();
+		SaveContext.Linker->bProceduralSave = SaveContext.IsProceduralSave();
 
 #if WITH_TEXT_ARCHIVE_SUPPORT
 		if (SaveContext.IsTextFormat())
@@ -1919,12 +1921,15 @@ void PostSavePackage(FSaveContext& SaveContext)
 	Package->ClearPackageFlags(PKG_NewlyCreated);
 
 	// Send a message that the package was saved
+	PRAGMA_DISABLE_DEPRECATION_WARNINGS;
 	UPackage::PackageSavedEvent.Broadcast(SaveContext.GetFilename(), Package);
+	PRAGMA_ENABLE_DEPRECATION_WARNINGS;
+	UPackage::PackageSavedWithContextEvent.Broadcast(SaveContext.GetFilename(), Package, FObjectPostSaveContext(SaveContext.GetObjectSaveContext()));
 
 	// update the internal package filename path if we're saving to a valid mounted path and we aren't currently cooking
 #if WITH_EDITOR
 	const FPackagePath& PackagePath = SaveContext.GetTargetPackagePath();
-	if (SaveContext.IsSavingNewLoadedPath())
+	if (SaveContext.IsUpdatingLoadedPath())
 	{
 		Package->SetLoadedPath(PackagePath);
 	}
@@ -2167,7 +2172,9 @@ FSavePackageResultStruct UPackage::Save2(UPackage* InPackage, UObject* InAsset,
 	PreSavePackage(SaveContext);
 	if (InAsset)
 	{
-		SaveContext.SetPreSaveCleanup(InAsset->PreSaveRoot(InFilename));
+		FObjectSaveContextData& ObjectSaveContext = SaveContext.GetObjectSaveContext();
+		UE::SavePackageUtilities::CallPreSaveRoot(InAsset, ObjectSaveContext);
+		SaveContext.SetPreSaveCleanup(ObjectSaveContext.bCleanupRequired);
 	}
 
 	// Route Presave only if not calling concurrently or diffing callstack, in those case they should be handled separately already
@@ -2200,7 +2207,7 @@ FSavePackageResultStruct UPackage::Save2(UPackage* InPackage, UObject* InAsset,
 	SlowTask.EnterProgressFrame();
 	if (InAsset)
 	{
-		InAsset->PostSaveRoot(SaveContext.GetPreSaveCleanup());
+		UE::SavePackageUtilities::CallPostSaveRoot(InAsset, SaveContext.GetObjectSaveContext(), SaveContext.GetPreSaveCleanup());
 		SaveContext.SetPreSaveCleanup(false);
 	}
 
@@ -2246,7 +2253,9 @@ ESavePackageResult UPackage::SaveConcurrent(TArrayView<FPackageSaveInfo> InPacka
 			if (SaveContext.GetAsset())
 			{
 				SCOPED_SAVETIMER(UPackage_SaveConcurrent_PreSaveRoot);
-				SaveContext.SetPreSaveCleanup(SaveContext.GetAsset()->PreSaveRoot(SaveContext.GetFilename()));
+				FObjectSaveContextData& ObjectSaveContext = SaveContext.GetObjectSaveContext();
+				UE::SavePackageUtilities::CallPreSaveRoot(SaveContext.GetAsset(), ObjectSaveContext);
+				SaveContext.SetPreSaveCleanup(ObjectSaveContext.bCleanupRequired);
 			}
 
 			// Route Presave
@@ -2316,7 +2325,7 @@ ESavePackageResult UPackage::SaveConcurrent(TArrayView<FPackageSaveInfo> InPacka
 			// PostSave Asset
 			if (SaveContext.GetAsset())
 			{
-				SaveContext.GetAsset()->PostSaveRoot(SaveContext.GetPreSaveCleanup());
+				UE::SavePackageUtilities::CallPostSaveRoot(SaveContext.GetAsset(), SaveContext.GetObjectSaveContext(), SaveContext.GetPreSaveCleanup());
 				SaveContext.SetPreSaveCleanup(false);
 			}
 
diff --git a/Engine/Source/Runtime/CoreUObject/Public/UObject/LinkerSave.h b/Engine/Source/Runtime/CoreUObject/Public/UObject/LinkerSave.h
--- a/Engine/Source/Runtime/CoreUObject/Public/UObject/LinkerSave.h
+++ b/Engine/Source/Runtime/CoreUObject/Public/UObject/LinkerSave.h
@@ -78,10 +78,18 @@ public:
 	TArray<AdditionalDataCallback> AdditionalDataToAppend;
 
 	/**
-	 * Stores whether the target of the save is the path PackagePath that already is/will be stored in the Package's LoadedPath
+	 * Set to true when the package is being saved due to a procedural save.
+	 * Any save without the the possibility of user-generated edits to the package is a procedural save (Cooking, EditorDomain).
+	 * This allows us to execute transforms that only need to be executed in response to new user data.
+	 */
+	bool bProceduralSave = false;
+
+	/**
+	 * Set to true when the LoadedPath of the package being saved is being updated.
+	 * This allows us to update the in-memory package when it is saved in editor to match its new save file.
 	 * This is used e.g. to decide whether to update the in-memory file offsets for BulkData.
 	 */
-	bool bSavingNewLoadedPath = false;
+	bool bUpdatingLoadedPath = false;
 
 	/** A mapping of package name to generated script SHA keys */
 	COREUOBJECT_API static TMap<FString, TArray<uint8> > PackagesToScriptSHAMap;
diff --git a/Engine/Source/Runtime/CoreUObject/Public/UObject/Object.h b/Engine/Source/Runtime/CoreUObject/Public/UObject/Object.h
--- a/Engine/Source/Runtime/CoreUObject/Public/UObject/Object.h
+++ b/Engine/Source/Runtime/CoreUObject/Public/UObject/Object.h
@@ -16,6 +16,10 @@
 struct FAssetData;
 class FConfigCacheIni;
 class FEditPropertyChain;
+class FObjectPostSaveContext;
+class FObjectPostSaveRootContext;
+class FObjectPreSaveContext;
+class FObjectPreSaveRootContext;
 class ITargetPlatform;
 class ITransactionObjectAnnotation;
 class FTransactionObjectEvent;
@@ -207,25 +211,31 @@ public:
 	{
 	}
 
+	UE_DEPRECATED(5.0, "Use version that takes FObjectPreSaveContext instead.")
+	virtual bool PreSaveRoot(const TCHAR* Filename);
+
 	/**
 	 * Called from within SavePackage on the passed in base/root object. The return value of this function will be passed to PostSaveRoot. 
 	 * This is used to allow objects used as a base to perform required actions before saving and cleanup afterwards.
-	 * @param Filename: Name of the file being saved to (includes path)
-
-	 * @return	Whether PostSaveRoot needs to perform internal cleanup
+	 *
+	 * @param	ObjectSaveContext Context providing access to parameters of the save,
+	*							  Also allows storage of variables like bCleanupIsRequired for use in PostSaveRoot
 	 */
-	virtual bool PreSaveRoot(const TCHAR* Filename)
-	{
-		return false;
-	}
+	virtual void PreSaveRoot(FObjectPreSaveRootContext ObjectSaveContext);
+
+	UE_DEPRECATED(5.0, "Use version that takes FObjectPostSaveContext instead.")
+	virtual void PostSaveRoot(bool bCleanupIsRequired);
 
 	/**
 	 * Called from within SavePackage on the passed in base/root object. 
 	 * This function is called after the package has been saved and can perform cleanup.
 	 *
-	 * @param	bCleanupIsRequired	Whether PreSaveRoot dirtied state that needs to be cleaned up
+	 * @param	ObjectSaveContext Context providing access to parameters of the save and to values from PreSaveRoot
 	 */
-	virtual void PostSaveRoot( bool bCleanupIsRequired ) {}
+	virtual void PostSaveRoot(FObjectPostSaveRootContext ObjectSaveContext);
+
+	UE_DEPRECATED(5.0, "Use version that takes FObjectPreSaveContext instead.")
+	virtual void PreSave(const class ITargetPlatform* TargetPlatform);
 
 	/**
 	 * Presave function. Gets called once before an object gets serialized for saving. This function is necessary
@@ -233,7 +243,7 @@ public:
 	 *
 	 * @warning: Objects created from within PreSave will NOT have PreSave called on them!!!
 	 */
-	virtual void PreSave(const class ITargetPlatform* TargetPlatform);
+	virtual void PreSave(FObjectPreSaveContext SaveContext);
 
 	/**
 	 * Note that the object will be modified.  If we are currently recording into the 
diff --git a/Engine/Source/Runtime/CoreUObject/Public/UObject/Package.h b/Engine/Source/Runtime/CoreUObject/Public/UObject/Package.h
--- a/Engine/Source/Runtime/CoreUObject/Public/UObject/Package.h
+++ b/Engine/Source/Runtime/CoreUObject/Public/UObject/Package.h
@@ -114,16 +114,24 @@ public:
 	DECLARE_MULTICAST_DELEGATE_OneParam(FOnPackageDirtyStateChanged, class UPackage*);
 	/** delegate type for package saved events ( Params: const FString& PackageFileName, UObject* Outer ) */
 	DECLARE_MULTICAST_DELEGATE_TwoParams(FOnPackageSaved, const FString&, UObject*);					
+	/** delegate type for package saved events ( Params: const FString& PackageFileName, UObject* Outer, FObjectPostSaveContext ObjectSaveContext ) */
+	DECLARE_MULTICAST_DELEGATE_ThreeParams(FOnPackageSavedWithContext, const FString&, UPackage*, FObjectPostSaveContext);
 	/** delegate type for when a package is marked as dirty via UObjectBaseUtilty::MarkPackageDirty ( Params: UPackage* ModifiedPackage, bool bWasDirty ) */
 	DECLARE_MULTICAST_DELEGATE_TwoParams(FOnPackageMarkedDirty, class UPackage*, bool);
 	/** delegate type for when a package is about to be saved */
 	DECLARE_MULTICAST_DELEGATE_OneParam(FPreSavePackage, class UPackage*);
+	/** delegate type for when a package is about to be saved */
+	DECLARE_MULTICAST_DELEGATE_TwoParams(FPreSavePackageWithContext, class UPackage*, FObjectPreSaveContext);
 
-	/** Delegate to notify subscribers when a package is about to be saved. */
+	UE_DEPRECATED(5.0, "Use PreSavePackageWithContextEvent instead.")
 	static FPreSavePackage PreSavePackageEvent;
+	/** Delegate to notify subscribers when a package is about to be saved. */
+	static FPreSavePackageWithContext PreSavePackageWithContextEvent;
+	UE_DEPRECATED(5.0, "Use PackageSavedWithContextEvent instead.")
+	static FOnPackageSaved PackageSavedEvent;
 	/** Delegate to notify subscribers when a package has been saved. This is triggered when the package saving
 	*  has completed and was successful. */
-	static FOnPackageSaved PackageSavedEvent;
+	static FOnPackageSavedWithContext PackageSavedWithContextEvent;
 	/** Delegate to notify subscribers when the dirty state of a package is changed.
 	*  Allows the editor to register the modified package as one that should be prompted for source control checkout. 
 	*  Use Package->IsDirty() to get the updated dirty state of the package */
diff --git a/Engine/Source/Runtime/CoreUObject/Public/UObject/SavePackage.h b/Engine/Source/Runtime/CoreUObject/Public/UObject/SavePackage.h
--- a/Engine/Source/Runtime/CoreUObject/Public/UObject/SavePackage.h
+++ b/Engine/Source/Runtime/CoreUObject/Public/UObject/SavePackage.h
@@ -16,6 +16,8 @@
 
 class FArchive;
 class FIoBuffer;
+struct FObjectSaveContextData;
+class FPackagePath;
 class FPackageStoreBulkDataManifest;
 class FSavePackageContext;
 class FArchiveDiffMap;
@@ -98,3 +100,32 @@ public:
 	FPackageStoreBulkDataManifest* BulkDataManifest;
 	bool bForceLegacyOffsets;
 };
+
+namespace UE
+{
+namespace SavePackageUtilities
+{
+	/**
+	 * Return whether the given save parameters indicate the LoadedPath of the package being saved should be updated.
+	 * This allows us to update the in-memory package when it is saved in editor to match its new save file.
+	 */
+	COREUOBJECT_API bool IsUpdatingLoadedPath(bool bIsCooking, const FPackagePath& TargetPackagePath, uint32 SaveFlags);
+
+	/**
+	 * Return whether the given save parameters indicate the package is a procedural save.
+	 * Any save without the the possibility of user-generated edits to the package is a procedural save (Cooking, EditorDomain).
+	 * This allows us to execute transforms that only need to be executed in response to new user data.
+	 */
+	COREUOBJECT_API bool IsProceduralSave(bool bIsCooking, const FPackagePath& TargetPackagePath, uint32 SaveFlags);
+
+	/** Call the PreSave function on the given object and log a warning if there is an incorrect override. */
+	COREUOBJECT_API void CallPreSave(UObject* Object, FObjectSaveContextData& ObjectSaveContext);
+
+	/** Call the PreSaveRoot function on the given object. */
+	COREUOBJECT_API void CallPreSaveRoot(UObject* Object, FObjectSaveContextData& ObjectSaveContext);
+
+	/** Call the PostSaveRoot function on the given object. */
+	COREUOBJECT_API void CallPostSaveRoot(UObject* Object, FObjectSaveContextData& ObjectSaveContext, bool bCleanupRequired);
+
+}
+}
\ No newline at end of file
diff --git a/Engine/Source/Runtime/Engine/Classes/Engine/DataAsset.h b/Engine/Source/Runtime/Engine/Classes/Engine/DataAsset.h
--- a/Engine/Source/Runtime/Engine/Classes/Engine/DataAsset.h
+++ b/Engine/Source/Runtime/Engine/Classes/Engine/DataAsset.h
@@ -50,7 +50,9 @@ public:
 	ENGINE_API virtual void UpdateAssetBundleData();
 
 	/** Updates AssetBundleData */
+	UE_DEPRECATED(5.0, "Use version that takes FObjectPreSaveContext instead.")
 	ENGINE_API virtual void PreSave(const class ITargetPlatform* TargetPlatform) override;
+	ENGINE_API virtual void PreSave(FObjectPreSaveContext ObjectSaveContext) override;
 
 protected:
 	/** Asset Bundle data computed at save time. In cooked builds this is accessible from AssetRegistry */
diff --git a/Engine/Source/Runtime/Engine/Classes/Engine/Texture.h b/Engine/Source/Runtime/Engine/Classes/Engine/Texture.h
--- a/Engine/Source/Runtime/Engine/Classes/Engine/Texture.h
+++ b/Engine/Source/Runtime/Engine/Classes/Engine/Texture.h
@@ -1084,7 +1084,9 @@ public:
 	ENGINE_API virtual void Serialize(FArchive& Ar) override;
 	ENGINE_API virtual void PostInitProperties() override;
 	ENGINE_API virtual void PostLoad() override;
+	UE_DEPRECATED(5.0, "Use version that takes FObjectPreSaveContext instead.")
 	ENGINE_API virtual void PreSave(const class ITargetPlatform* TargetPlatform) override;
+	ENGINE_API virtual void PreSave(FObjectPreSaveContext ObjectSaveContext) override;
 	ENGINE_API virtual void BeginDestroy() override;
 	ENGINE_API virtual bool IsReadyForFinishDestroy() override;
 	ENGINE_API virtual void FinishDestroy() override;
diff --git a/Engine/Source/Runtime/Engine/Classes/Engine/World.h b/Engine/Source/Runtime/Engine/Classes/Engine/World.h
--- a/Engine/Source/Runtime/Engine/Classes/Engine/World.h
+++ b/Engine/Source/Runtime/Engine/Classes/Engine/World.h
@@ -2605,8 +2605,12 @@ public:
 	virtual bool IsReadyForFinishDestroy() override;
 	virtual void PostLoad() override;
 	virtual void PreDuplicate(FObjectDuplicationParameters& DupParams) override;
-	virtual bool PreSaveRoot(const TCHAR* Filename) override;
-	virtual void PostSaveRoot( bool bCleanupIsRequired ) override;
+	UE_DEPRECATED(5.0, "Use version that takes FObjectPreSaveContext instead.")
+	virtual bool PreSaveRoot(const TCHAR* InFilename) override;
+	virtual void PreSaveRoot(FObjectPreSaveRootContext ObjectSaveContext) override;
+	UE_DEPRECATED(5.0, "Use version that takes FObjectPostSaveContext instead.")
+	virtual void PostSaveRoot(bool bCleanupIsRequired) override;
+	virtual void PostSaveRoot(FObjectPostSaveRootContext ObjectSaveContext) override;
 	virtual UWorld* GetWorld() const override;
 	virtual FPrimaryAssetId GetPrimaryAssetId() const override;
 	static void AddReferencedObjects(UObject* InThis, FReferenceCollector& Collector);
diff --git a/Engine/Source/Runtime/Engine/Classes/GameFramework/Actor.h b/Engine/Source/Runtime/Engine/Classes/GameFramework/Actor.h
--- a/Engine/Source/Runtime/Engine/Classes/GameFramework/Actor.h
+++ b/Engine/Source/Runtime/Engine/Classes/GameFramework/Actor.h
@@ -1843,9 +1843,15 @@ public:
 	virtual bool IsEditorOnly() const override;
 	virtual bool IsAsset() const override;
 
+	UE_DEPRECATED(5.0, "Use version that takes FObjectPreSaveContext instead.")
 	virtual bool PreSaveRoot(const TCHAR* InFilename) override;
+	virtual void PreSaveRoot(FObjectPreSaveRootContext ObjectSaveContext) override;
+	UE_DEPRECATED(5.0, "Use version that takes FObjectPostSaveContext instead.")
 	virtual void PostSaveRoot(bool bCleanupIsRequired) override;
+	virtual void PostSaveRoot(FObjectPostSaveRootContext ObjectSaveContext) override;
+	UE_DEPRECATED(5.0, "Use version that takes FObjectPreSaveContext instead.")
 	virtual void PreSave(const ITargetPlatform* TargetPlatform) override;
+	virtual void PreSave(FObjectPreSaveContext ObjectSaveContext) override;
 
 #if WITH_EDITOR
 	virtual bool Modify(bool bAlwaysMarkDirty = true) override;
diff --git a/Engine/Source/Runtime/Engine/Private/Actor.cpp b/Engine/Source/Runtime/Engine/Private/Actor.cpp
--- a/Engine/Source/Runtime/Engine/Private/Actor.cpp
+++ b/Engine/Source/Runtime/Engine/Private/Actor.cpp
@@ -1,6 +1,7 @@
 // Copyright Epic Games, Inc. All Rights Reserved.
 
 #include "GameFramework/Actor.h"
+
 #include "EngineDefines.h"
 #include "EngineStats.h"
 #include "EngineGlobals.h"
@@ -12,6 +13,7 @@
 #include "AI/NavigationSystemBase.h"
 #include "Components/InputComponent.h"
 #include "Engine/Engine.h"
+#include "UObject/ObjectSaveContext.h"
 #include "UObject/UObjectHash.h"
 #include "UObject/PropertyPortFlags.h"
 #include "UObject/UE5MainStreamObjectVersion.h"
@@ -379,7 +381,14 @@ bool AActor::IsAsset() const
 
 bool AActor::PreSaveRoot(const TCHAR* InFilename)
 {
-	bool bResult = Super::PreSaveRoot(InFilename);
+	PRAGMA_DISABLE_DEPRECATION_WARNINGS;
+	return Super::PreSaveRoot(InFilename);
+	PRAGMA_ENABLE_DEPRECATION_WARNINGS;
+}
+
+void AActor::PreSaveRoot(FObjectPreSaveRootContext ObjectSaveContext)
+{
+	Super::PreSaveRoot(ObjectSaveContext);
 #if WITH_EDITOR
 	// if `PreSaveRoot` is called on an actor, it should be have its package overridden (external to the level)
 	// if this actor is not in the current persistent level then we might need to remove level streamin transform before saving it
@@ -392,19 +401,25 @@ bool AActor::PreSaveRoot(const TCHAR* InFilename)
 		{
 			FLevelUtils::RemoveEditorTransform(LevelStreaming, true, this);
 		}
-		bResult |= true;
+		ObjectSaveContext.SetCleanupRequired(true);
 	}
 #endif
-	return bResult;
 }
 
 void AActor::PostSaveRoot(bool bCleanupIsRequired)
 {
+	PRAGMA_DISABLE_DEPRECATION_WARNINGS;
 	Super::PostSaveRoot(bCleanupIsRequired);
+	PRAGMA_ENABLE_DEPRECATION_WARNINGS;
+}
+
+void AActor::PostSaveRoot(FObjectPostSaveRootContext  ObjectSaveContext)
+{
+	Super::PostSaveRoot(ObjectSaveContext);
 #if WITH_EDITOR
 	// if this actor is not in the current persistent level then we will need to clean up the removal of level streaming
 	ULevel* Level = GetLevel();
-	if (bCleanupIsRequired && IsPackageExternal() && Level && !Level->IsPersistentLevel())
+	if (ObjectSaveContext.IsCleanupRequired() && IsPackageExternal() && Level && !Level->IsPersistentLevel())
 	{
 		// If we can get the streaming level, we should remove the editor transform before saving
 		ULevelStreaming* LevelStreaming = FLevelUtils::FindStreamingLevel(Level);
@@ -418,7 +433,14 @@ void AActor::PostSaveRoot(bool bCleanupIsRequired)
 
 void AActor::PreSave(const ITargetPlatform* TargetPlatform)
 {
+	PRAGMA_DISABLE_DEPRECATION_WARNINGS;
 	Super::PreSave(TargetPlatform);
+	PRAGMA_ENABLE_DEPRECATION_WARNINGS;
+}
+
+void AActor::PreSave(FObjectPreSaveContext ObjectSaveContext)
+{
+	Super::PreSave(ObjectSaveContext);
 #if WITH_EDITOR
 	FixupDataLayers();
 #endif
diff --git a/Engine/Source/Runtime/Engine/Private/Animation/AnimMontage.cpp b/Engine/Source/Runtime/Engine/Private/Animation/AnimMontage.cpp
--- a/Engine/Source/Runtime/Engine/Private/Animation/AnimMontage.cpp
+++ b/Engine/Source/Runtime/Engine/Private/Animation/AnimMontage.cpp
@@ -7,6 +7,7 @@
 #include "Animation/AnimMontage.h"
 #include "UObject/LinkerLoad.h"
 #include "UObject/Package.h"
+#include "UObject/ObjectSaveContext.h"
 #include "UObject/UObjectThreadContext.h"
 #include "Animation/AssetMappingTable.h"
 #include "Animation/AnimSequence.h"
@@ -343,11 +344,18 @@ void UAnimMontage::UnregisterOnMontageChanged(void* Unregister)
 #endif	//WITH_EDITOR
 
 void UAnimMontage::PreSave(const class ITargetPlatform* TargetPlatform)
+{
+	PRAGMA_DISABLE_DEPRECATION_WARNINGS;
+	Super::PreSave(TargetPlatform);
+	PRAGMA_ENABLE_DEPRECATION_WARNINGS;
+}
+
+void UAnimMontage::PreSave(FObjectPreSaveContext ObjectSaveContext)
 {
 #if WITH_EDITOR
 	BakeTimeStretchCurve();
 #endif // WITH_EDITOR
-	Super::PreSave(TargetPlatform);
+	Super::PreSave(ObjectSaveContext);
 }
 
 void UAnimMontage::PostLoad()
diff --git a/Engine/Source/Runtime/Engine/Private/Animation/AnimNotify.cpp b/Engine/Source/Runtime/Engine/Private/Animation/AnimNotify.cpp
--- a/Engine/Source/Runtime/Engine/Private/Animation/AnimNotify.cpp
+++ b/Engine/Source/Runtime/Engine/Private/Animation/AnimNotify.cpp
@@ -1,8 +1,10 @@
 // Copyright Epic Games, Inc. All Rights Reserved.
 
 #include "Animation/AnimNotifies/AnimNotify.h"
+
 #include "Components/SkeletalMeshComponent.h"
 #include "Animation/AnimSequenceBase.h"
+#include "UObject/ObjectSaveContext.h"
 #include "UObject/Package.h"
 
 /////////////////////////////////////////////////////
@@ -81,11 +83,18 @@ void UAnimNotify::PostLoad()
 }
 
 void UAnimNotify::PreSave(const class ITargetPlatform* TargetPlatform)
+{
+	PRAGMA_DISABLE_DEPRECATION_WARNINGS;
+	Super::PreSave(TargetPlatform);
+	PRAGMA_ENABLE_DEPRECATION_WARNINGS;
+}
+
+void UAnimNotify::PreSave(FObjectPreSaveContext ObjectSaveContext)
 {
 #if WITH_EDITOR
 	ValidateAssociatedAssets();
 #endif
-	Super::PreSave(TargetPlatform);
+	Super::PreSave(ObjectSaveContext);
 }
 
 UObject* UAnimNotify::GetContainingAsset() const
diff --git a/Engine/Source/Runtime/Engine/Private/Animation/AnimSequence.cpp b/Engine/Source/Runtime/Engine/Private/Animation/AnimSequence.cpp
--- a/Engine/Source/Runtime/Engine/Private/Animation/AnimSequence.cpp
+++ b/Engine/Source/Runtime/Engine/Private/Animation/AnimSequence.cpp
@@ -5,9 +5,11 @@
 =============================================================================*/ 
 
 #include "Animation/AnimSequence.h"
+
 #include "Misc/MessageDialog.h"
 #include "Logging/LogScopedVerbosityOverride.h"
 #include "UObject/FrameworkObjectVersion.h"
+#include "UObject/ObjectSaveContext.h"
 #include "UObject/UE5MainStreamObjectVersion.h"
 #include "Serialization/MemoryReader.h"
 #include "UObject/UObjectIterator.h"
@@ -781,6 +783,13 @@ void UAnimSequence::GetPreloadDependencies(TArray<UObject*>& OutDeps)
 }
 
 void UAnimSequence::PreSave(const class ITargetPlatform* TargetPlatform)
+{
+	PRAGMA_DISABLE_DEPRECATION_WARNINGS;
+	Super::PreSave(TargetPlatform);
+	PRAGMA_ENABLE_DEPRECATION_WARNINGS;
+}
+
+void UAnimSequence::PreSave(FObjectPreSaveContext ObjectSaveContext)
 {
 #if WITH_EDITOR
 	// Could already be compressing
@@ -797,6 +806,7 @@ void UAnimSequence::PreSave(const class ITargetPlatform* TargetPlatform)
 		ensureAlwaysMsgf(!bUseRawDataOnly,  TEXT("Animation : %s failed to compress"), *GetName());
 	}
 
+	const ITargetPlatform* TargetPlatform = ObjectSaveContext.GetTargetPlatform();
 	if (TargetPlatform)
 	{
 		// Update compressed data for platform
@@ -807,14 +817,13 @@ void UAnimSequence::PreSave(const class ITargetPlatform* TargetPlatform)
 
 	WaitOnExistingCompression(); // Wait on updated data
 
-	const bool bIsCooking = (TargetPlatform != nullptr);
-	if (!bIsCooking)
+	if (!ObjectSaveContext.IsCooking())
 	{
 		UpdateRetargetSourceAsset();
 	}
 #endif
 
-	Super::PreSave(TargetPlatform);
+	Super::PreSave(ObjectSaveContext);
 }
 
 void UAnimSequence::PostLoad()
diff --git a/Engine/Source/Runtime/Engine/Private/Animation/AnimStreamable.cpp b/Engine/Source/Runtime/Engine/Private/Animation/AnimStreamable.cpp
--- a/Engine/Source/Runtime/Engine/Private/Animation/AnimStreamable.cpp
+++ b/Engine/Source/Runtime/Engine/Private/Animation/AnimStreamable.cpp
@@ -5,11 +5,13 @@
 =============================================================================*/ 
 
 #include "Animation/AnimStreamable.h"
+
 #include "Interfaces/ITargetPlatform.h"
 #include "Interfaces/ITargetPlatformManagerModule.h"
 #include "DeviceProfiles/DeviceProfileManager.h"
 #include "DeviceProfiles/DeviceProfile.h"
 #include "UObject/LinkerLoad.h"
+#include "UObject/ObjectSaveContext.h"
 #include "Animation/AnimCompressionDerivedData.h"
 #include "DerivedDataCacheInterface.h"
 #include "Animation/AnimBoneCompressionSettings.h"
@@ -132,8 +134,16 @@ UAnimStreamable::UAnimStreamable(const FObjectInitializer& ObjectInitializer)
 }
 
 void UAnimStreamable::PreSave(const class ITargetPlatform* TargetPlatform)
+{
+	PRAGMA_DISABLE_DEPRECATION_WARNINGS;
+	Super::PreSave(TargetPlatform);
+	PRAGMA_ENABLE_DEPRECATION_WARNINGS;
+}
+
+void UAnimStreamable::PreSave(FObjectPreSaveContext ObjectSaveContext)
 {
 #if WITH_EDITOR
+	const ITargetPlatform* TargetPlatform = ObjectSaveContext.GetTargetPlatform();
 	if (TargetPlatform)
 	{
 		RequestCompressedData(TargetPlatform); //Make sure target platform data is built
@@ -141,7 +151,7 @@ void UAnimStreamable::PreSave(const class ITargetPlatform* TargetPlatform)
 
 #endif
 
-	Super::PreSave(TargetPlatform);
+	Super::PreSave(ObjectSaveContext);
 }
 
 void UAnimStreamable::Serialize(FArchive& Ar)
diff --git a/Engine/Source/Runtime/Engine/Private/Animation/PoseAsset.cpp b/Engine/Source/Runtime/Engine/Private/Animation/PoseAsset.cpp
--- a/Engine/Source/Runtime/Engine/Private/Animation/PoseAsset.cpp
+++ b/Engine/Source/Runtime/Engine/Private/Animation/PoseAsset.cpp
@@ -1,6 +1,7 @@
 // Copyright Epic Games, Inc. All Rights Reserved.
 
 #include "Animation/PoseAsset.h"
+
 #include "UObject/FrameworkObjectVersion.h"
 #include "UObject/FortniteMainBranchObjectVersion.h"
 #include "AnimationRuntime.h"
@@ -10,6 +11,7 @@
 #include "Animation/AnimationPoseData.h"
 #include "Animation/AnimData/AnimDataModel.h"
 #include "Animation/AnimSequenceHelpers.h"
+#include "UObject/ObjectSaveContext.h"
 #include "UObject/UE5ReleaseStreamObjectVersion.h"
 
 #define LOCTEXT_NAMESPACE "PoseAsset"
@@ -890,15 +892,21 @@ void UPoseAsset::Serialize(FArchive& Ar)
 }
 
 void UPoseAsset::PreSave(const ITargetPlatform* TargetPlatform)
+{
+	PRAGMA_DISABLE_DEPRECATION_WARNINGS;
+	Super::PreSave(TargetPlatform);
+	PRAGMA_ENABLE_DEPRECATION_WARNINGS;
+}
+
+void UPoseAsset::PreSave(FObjectPreSaveContext ObjectSaveContext)
 {
 #if WITH_EDITOR
-	const bool bIsCooking = (TargetPlatform != nullptr);
-	if (!bIsCooking)
+	if (!ObjectSaveContext.IsCooking())
 	{
 		UpdateRetargetSourceAsset();
 	}
 #endif // WITH_EDITOR
-	Super::PreSave(TargetPlatform);
+	Super::PreSave(ObjectSaveContext);
 }
 
 void UPoseAsset::GetAssetRegistryTags(TArray<FAssetRegistryTag>& OutTags) const
diff --git a/Engine/Source/Runtime/Engine/Private/Blueprint.cpp b/Engine/Source/Runtime/Engine/Private/Blueprint.cpp
--- a/Engine/Source/Runtime/Engine/Private/Blueprint.cpp
+++ b/Engine/Source/Runtime/Engine/Private/Blueprint.cpp
@@ -1,10 +1,12 @@
 // Copyright Epic Games, Inc. All Rights Reserved.
 
 #include "Engine/Blueprint.h"
+
 #include "Misc/CoreMisc.h"
 #include "Misc/ConfigCacheIni.h"
 #include "UObject/BlueprintsObjectVersion.h"
 #include "UObject/FrameworkObjectVersion.h"
+#include "UObject/ObjectSaveContext.h"
 #include "UObject/UObjectHash.h"
 #include "Serialization/PropertyLocalizationDataGathering.h"
 #include "UObject/UnrealType.h"
@@ -368,10 +370,18 @@ static TAutoConsoleVariable<bool> CVarBPForceOldSearchDataFormatVersionOnSave(
 
 void UBlueprint::PreSave(const class ITargetPlatform* TargetPlatform)
 {
+	PRAGMA_DISABLE_DEPRECATION_WARNINGS;
 	Super::PreSave(TargetPlatform);
+	PRAGMA_ENABLE_DEPRECATION_WARNINGS;
+}
+
+void UBlueprint::PreSave(FObjectPreSaveContext ObjectSaveContext)
+{
+	Super::PreSave(ObjectSaveContext);
 
 	// Clear all upgrade notes, the user has saved and should not see them anymore
 	UpgradeNotesLog.Reset();
+	const ITargetPlatform* TargetPlatform = ObjectSaveContext.GetTargetPlatform();
 
 	if (!TargetPlatform || TargetPlatform->HasEditorOnlyData())
 	{
diff --git a/Engine/Source/Runtime/Engine/Private/Camera/CameraAnim.cpp b/Engine/Source/Runtime/Engine/Private/Camera/CameraAnim.cpp
--- a/Engine/Source/Runtime/Engine/Private/Camera/CameraAnim.cpp
+++ b/Engine/Source/Runtime/Engine/Private/Camera/CameraAnim.cpp
@@ -1,6 +1,7 @@
 // Copyright Epic Games, Inc. All Rights Reserved.
 
 #include "Camera/CameraAnim.h"
+
 #include "Serialization/ArchiveCountMem.h"
 #include "Camera/CameraActor.h"
 #include "Matinee/MatineeActor.h"
@@ -8,6 +9,7 @@
 #include "Matinee/InterpGroupInst.h"
 #include "Matinee/InterpGroupCamera.h"
 #include "Matinee/InterpTrackMove.h"
+#include "UObject/ObjectSaveContext.h"
 
 DEFINE_LOG_CATEGORY(LogCameraAnim);
 
@@ -87,13 +89,19 @@ FBox UCameraAnim::GetAABB(FVector const& BaseLoc, FRotator const& BaseRot, float
 	return ScaledLocalBox.TransformBy(BaseTM);
 }
 
-
 void UCameraAnim::PreSave(const class ITargetPlatform* TargetPlatform)
+{
+	PRAGMA_DISABLE_DEPRECATION_WARNINGS;
+	Super::PreSave(TargetPlatform);
+	PRAGMA_ENABLE_DEPRECATION_WARNINGS;
+}
+
+void UCameraAnim::PreSave(FObjectPreSaveContext ObjectSaveContext)
 {
 #if WITH_EDITORONLY_DATA
 	CalcLocalAABB();
 #endif // WITH_EDITORONLY_DATA
-	Super::PreSave(TargetPlatform);
+	Super::PreSave(ObjectSaveContext);
 }
 
 void UCameraAnim::PostLoad()
diff --git a/Engine/Source/Runtime/Engine/Private/Components/LightComponent.cpp b/Engine/Source/Runtime/Engine/Private/Components/LightComponent.cpp
--- a/Engine/Source/Runtime/Engine/Private/Components/LightComponent.cpp
+++ b/Engine/Source/Runtime/Engine/Private/Components/LightComponent.cpp
@@ -5,10 +5,12 @@
 =============================================================================*/
 
 #include "Components/LightComponent.h"
+
 #include "Misc/App.h"
 #include "RenderingThread.h"
 #include "Engine/MapBuildDataRegistry.h"
 #include "Materials/Material.h"
+#include "UObject/ObjectSaveContext.h"
 #include "UObject/RenderingObjectVersion.h"
 #include "UObject/UObjectHash.h"
 #include "UObject/UObjectIterator.h"
@@ -593,7 +595,14 @@ void ULightComponent::PostLoad()
 #if WITH_EDITOR
 void ULightComponent::PreSave(const class ITargetPlatform* TargetPlatform)
 {
+	PRAGMA_DISABLE_DEPRECATION_WARNINGS;
 	Super::PreSave(TargetPlatform);
+	PRAGMA_ENABLE_DEPRECATION_WARNINGS;
+}
+
+void ULightComponent::PreSave(FObjectPreSaveContext ObjectSaveContext)
+{
+	Super::PreSave(ObjectSaveContext);
 	ValidateLightGUIDs();
 }
 
diff --git a/Engine/Source/Runtime/Engine/Private/Components/StaticMeshComponent.cpp b/Engine/Source/Runtime/Engine/Private/Components/StaticMeshComponent.cpp
--- a/Engine/Source/Runtime/Engine/Private/Components/StaticMeshComponent.cpp
+++ b/Engine/Source/Runtime/Engine/Private/Components/StaticMeshComponent.cpp
@@ -1,12 +1,14 @@
 // Copyright Epic Games, Inc. All Rights Reserved.
 
 #include "Components/StaticMeshComponent.h"
+
 #include "Modules/ModuleManager.h"
 #include "RenderingThread.h"
 #include "Components.h"
 #include "Engine/MapBuildDataRegistry.h"
 #include "Materials/Material.h"
 #include "Misc/ConfigCacheIni.h"
+#include "UObject/ObjectSaveContext.h"
 #include "UObject/RenderingObjectVersion.h"
 #include "GameFramework/WorldSettings.h"
 #include "Engine/CollisionProfile.h"
@@ -381,7 +383,14 @@ bool UStaticMeshComponent::AreNativePropertiesIdenticalTo( UObject* Other ) cons
 #if WITH_EDITORONLY_DATA
 void UStaticMeshComponent::PreSave(const class ITargetPlatform* TargetPlatform)
 {
+	PRAGMA_DISABLE_DEPRECATION_WARNINGS;
 	Super::PreSave(TargetPlatform);
+	PRAGMA_ENABLE_DEPRECATION_WARNINGS;
+}
+
+void UStaticMeshComponent::PreSave(FObjectPreSaveContext ObjectSaveContext)
+{
+	Super::PreSave(ObjectSaveContext);
 
 	CachePaintedDataIfNecessary();
 }
diff --git a/Engine/Source/Runtime/Engine/Private/DataAsset.cpp b/Engine/Source/Runtime/Engine/Private/DataAsset.cpp
--- a/Engine/Source/Runtime/Engine/Private/DataAsset.cpp
+++ b/Engine/Source/Runtime/Engine/Private/DataAsset.cpp
@@ -1,9 +1,12 @@
 // Copyright Epic Games, Inc. All Rights Reserved.
 
 #include "Engine/DataAsset.h"
+
 #include "Misc/PackageName.h"
+#include "UObject/ObjectSaveContext.h"
 #include "UObject/Package.h"
 #include "Engine/AssetManager.h"
+
 UDataAsset::UDataAsset(const FObjectInitializer& ObjectInitializer)
 	: Super(ObjectInitializer)
 {
@@ -34,7 +37,14 @@ void UPrimaryDataAsset::UpdateAssetBundleData()
 
 void UPrimaryDataAsset::PreSave(const class ITargetPlatform* TargetPlatform)
 {
+	PRAGMA_DISABLE_DEPRECATION_WARNINGS;
 	Super::PreSave(TargetPlatform);
+	PRAGMA_ENABLE_DEPRECATION_WARNINGS;
+}
+
+void UPrimaryDataAsset::PreSave(FObjectPreSaveContext ObjectSaveContext)
+{
+	Super::PreSave(ObjectSaveContext);
 
 	UpdateAssetBundleData();
 
diff --git a/Engine/Source/Runtime/Engine/Private/EdGraph/EdGraphNode.cpp b/Engine/Source/Runtime/Engine/Private/EdGraph/EdGraphNode.cpp
--- a/Engine/Source/Runtime/Engine/Private/EdGraph/EdGraphNode.cpp
+++ b/Engine/Source/Runtime/Engine/Private/EdGraph/EdGraphNode.cpp
@@ -1,9 +1,11 @@
 // Copyright Epic Games, Inc. All Rights Reserved.
 
 #include "EdGraph/EdGraphNode.h"
+
 #include "Serialization/PropertyLocalizationDataGathering.h"
 #include "UObject/BlueprintsObjectVersion.h"
 #include "UObject/FrameworkObjectVersion.h"
+#include "UObject/ObjectSaveContext.h"
 #include "UObject/ReleaseObjectVersion.h"
 #include "EdGraph/EdGraphPin.h"
 #include "Textures/SlateIcon.h"
@@ -632,7 +634,14 @@ void UEdGraphNode::AddReferencedObjects(UObject* InThis, FReferenceCollector& Co
 
 void UEdGraphNode::PreSave(const class ITargetPlatform* TargetPlatform)
 {
+	PRAGMA_DISABLE_DEPRECATION_WARNINGS;
 	Super::PreSave(TargetPlatform);
+	PRAGMA_ENABLE_DEPRECATION_WARNINGS;
+}
+
+void UEdGraphNode::PreSave(FObjectPreSaveContext ObjectSaveContext)
+{
+	Super::PreSave(ObjectSaveContext);
 
 #if WITH_EDITORONLY_DATA
 	if (!NodeUpgradeMessage.IsEmpty())
diff --git a/Engine/Source/Runtime/Engine/Private/HLOD/HLODEngineSubsystem.cpp b/Engine/Source/Runtime/Engine/Private/HLOD/HLODEngineSubsystem.cpp
--- a/Engine/Source/Runtime/Engine/Private/HLOD/HLODEngineSubsystem.cpp
+++ b/Engine/Source/Runtime/Engine/Private/HLOD/HLODEngineSubsystem.cpp
@@ -15,6 +15,7 @@
 #include "IHierarchicalLODUtilities.h"
 #include "HierarchicalLODUtilitiesModule.h"
 #include "GameFramework/WorldSettings.h"
+#include "UObject/ObjectSaveContext.h"
 
 void UHLODEngineSubsystem::Initialize(FSubsystemCollectionBase& Collection)
 {
@@ -62,7 +63,7 @@ void UHLODEngineSubsystem::RegisterRecreateLODActorsDelegates()
 	{
 		OnPostWorldInitializationDelegateHandle = FWorldDelegates::OnPostWorldInitialization.AddUObject(this, &UHLODEngineSubsystem::RecreateLODActorsForWorld);
 		OnLevelAddedToWorldDelegateHandle = FWorldDelegates::LevelAddedToWorld.AddUObject(this, &UHLODEngineSubsystem::RecreateLODActorsForLevel);
-		OnPreSaveWorlDelegateHandle = FEditorDelegates::PreSaveWorld.AddUObject(this, &UHLODEngineSubsystem::OnPreSaveWorld);
+		OnPreSaveWorlDelegateHandle = FEditorDelegates::PreSaveWorldWithContext.AddUObject(this, &UHLODEngineSubsystem::OnPreSaveWorld);
 	}	
 }
 
@@ -165,7 +166,7 @@ bool UHLODEngineSubsystem::CleanupHLOD(ALODActor* InLODActor)
 	return bShouldDestroyActor;
 }
 
-void UHLODEngineSubsystem::OnPreSaveWorld(uint32 InSaveFlags, UWorld* InWorld)
+void UHLODEngineSubsystem::OnPreSaveWorld(UWorld* InWorld, FObjectPreSaveContext ObjectSaveContext)
 {
 	// When cooking, make sure that the LODActors are not transient
 	if (InWorld && InWorld->PersistentLevel && GIsCookerLoadingPackage)
diff --git a/Engine/Source/Runtime/Engine/Private/HLODProxy.cpp b/Engine/Source/Runtime/Engine/Private/HLODProxy.cpp
--- a/Engine/Source/Runtime/Engine/Private/HLODProxy.cpp
+++ b/Engine/Source/Runtime/Engine/Private/HLODProxy.cpp
@@ -1,6 +1,7 @@
 // Copyright Epic Games, Inc. All Rights Reserved.
 
 #include "Engine/HLODProxy.h"
+
 #include "Engine/LODActor.h"
 #include "GameFramework/WorldSettings.h"
 
@@ -19,6 +20,7 @@
 #include "LevelUtils.h"
 #include "Engine/LevelStreaming.h"
 #include "Math/UnrealMathUtility.h"
+#include "UObject/ObjectSaveContext.h"
 
 #if WITH_EDITOR
 
@@ -153,7 +155,14 @@ void UHLODProxy::DeletePackage()
 
 void UHLODProxy::PreSave(const class ITargetPlatform* TargetPlatform)
 {
+	PRAGMA_DISABLE_DEPRECATION_WARNINGS;
 	Super::PreSave(TargetPlatform);
+	PRAGMA_ENABLE_DEPRECATION_WARNINGS;
+}
+
+void UHLODProxy::PreSave(FObjectPreSaveContext ObjectSaveContext)
+{
+	Super::PreSave(ObjectSaveContext);
 
 	if (!OwningMap.IsValid())
 	{
diff --git a/Engine/Source/Runtime/Engine/Private/LODActor.cpp b/Engine/Source/Runtime/Engine/Private/LODActor.cpp
--- a/Engine/Source/Runtime/Engine/Private/LODActor.cpp
+++ b/Engine/Source/Runtime/Engine/Private/LODActor.cpp
@@ -5,6 +5,7 @@
 =============================================================================*/
 
 #include "Engine/LODActor.h"
+
 #include "UObject/UObjectIterator.h"
 #include "Engine/CollisionProfile.h"
 #include "Logging/TokenizedMessage.h"
@@ -26,6 +27,7 @@
 #include "IHierarchicalLODUtilities.h"
 #include "ObjectTools.h"
 #include "HierarchicalLOD.h"
+#include "UObject/ObjectSaveContext.h"
 #endif
 
 DEFINE_LOG_CATEGORY_STATIC(LogHLOD, Log, All);
@@ -1384,7 +1386,14 @@ void ALODActor::Serialize(FArchive& Ar)
 
 void ALODActor::PreSave(const class ITargetPlatform* TargetPlatform)
 {
-	Super::PreSave(TargetPlatform);	
+	PRAGMA_DISABLE_DEPRECATION_WARNINGS;
+	Super::PreSave(TargetPlatform);
+	PRAGMA_ENABLE_DEPRECATION_WARNINGS;
+}
+
+void ALODActor::PreSave(FObjectPreSaveContext ObjectSaveContext)
+{
+	Super::PreSave(ObjectSaveContext);
 
 	// Always rebuild key on save here.
 	// We don't do this while cooking as keys rely on platform derived data which is context-dependent during cook
diff --git a/Engine/Source/Runtime/Engine/Private/Level.cpp b/Engine/Source/Runtime/Engine/Private/Level.cpp
--- a/Engine/Source/Runtime/Engine/Private/Level.cpp
+++ b/Engine/Source/Runtime/Engine/Private/Level.cpp
@@ -5,6 +5,7 @@ Level.cpp: Level-related functions
 =============================================================================*/
 
 #include "Engine/Level.h"
+
 #include "Misc/ScopedSlowTask.h"
 #include "UObject/RenderingObjectVersion.h"
 #include "UObject/UE5MainStreamObjectVersion.h"
@@ -26,6 +27,7 @@ Level.cpp: Level-related functions
 #include "Engine/Brush.h"
 #include "Engine/Engine.h"
 #include "Containers/TransArray.h"
+#include "UObject/ObjectSaveContext.h"
 #include "UObject/UObjectHash.h"
 #include "UObject/UObjectIterator.h"
 #include "UObject/PropertyPortFlags.h"
@@ -732,10 +734,16 @@ void ULevel::SortActorList()
 	Actors = MoveTemp(NewActors);
 }
 
-
 void ULevel::PreSave(const class ITargetPlatform* TargetPlatform)
 {
+	PRAGMA_DISABLE_DEPRECATION_WARNINGS;
 	Super::PreSave(TargetPlatform);
+	PRAGMA_ENABLE_DEPRECATION_WARNINGS;
+}
+
+void ULevel::PreSave(FObjectPreSaveContext ObjectSaveContext)
+{
+	Super::PreSave(ObjectSaveContext);
 
 #if WITH_EDITOR
 	if( !IsTemplate() )
diff --git a/Engine/Source/Runtime/Engine/Private/Materials/Material.cpp b/Engine/Source/Runtime/Engine/Private/Materials/Material.cpp
--- a/Engine/Source/Runtime/Engine/Private/Materials/Material.cpp
+++ b/Engine/Source/Runtime/Engine/Private/Materials/Material.cpp
@@ -5,10 +5,12 @@
 =============================================================================*/
 
 #include "Materials/Material.h"
+
 #include "Stats/StatsMisc.h"
 #include "Misc/FeedbackContext.h"
 #include "UObject/RenderingObjectVersion.h"
 #include "Misc/App.h"
+#include "UObject/ObjectSaveContext.h"
 #include "UObject/UObjectHash.h"
 #include "UObject/UObjectIterator.h"
 #include "UObject/Package.h"
@@ -1027,7 +1029,14 @@ UMaterial::UMaterial(const FObjectInitializer& ObjectInitializer)
 
 void UMaterial::PreSave(const class ITargetPlatform* TargetPlatform)
 {
+	PRAGMA_DISABLE_DEPRECATION_WARNINGS;
 	Super::PreSave(TargetPlatform);
+	PRAGMA_ENABLE_DEPRECATION_WARNINGS;
+}
+
+void UMaterial::PreSave(FObjectPreSaveContext ObjectSaveContext)
+{
+	Super::PreSave(ObjectSaveContext);
 #if WITH_EDITOR
 	GMaterialsWithDirtyUsageFlags.RemoveAnnotation(this);
 #endif
diff --git a/Engine/Source/Runtime/Engine/Private/Materials/MaterialInstance.cpp b/Engine/Source/Runtime/Engine/Private/Materials/MaterialInstance.cpp
--- a/Engine/Source/Runtime/Engine/Private/Materials/MaterialInstance.cpp
+++ b/Engine/Source/Runtime/Engine/Private/Materials/MaterialInstance.cpp
@@ -1,6 +1,7 @@
 // Copyright Epic Games, Inc. All Rights Reserved.
 
 #include "Materials/MaterialInstance.h"
+
 #include "Stats/StatsMisc.h"
 #include "EngineGlobals.h"
 #include "EngineModule.h"
@@ -45,6 +46,7 @@
 #include "Materials/MaterialStaticParameterValueResolver.h"
 #include "ShaderPlatformQualitySettings.h"
 #include "MaterialShaderQualitySettings.h"
+#include "UObject/ObjectSaveContext.h"
 
 #if ENABLE_COOK_STATS
 #include "ProfilingDebugging/ScopedTimers.h"
@@ -4699,8 +4701,15 @@ void UMaterialInstance::GetLightingGuidChain(bool bIncludeTextures, TArray<FGuid
 
 void UMaterialInstance::PreSave(const class ITargetPlatform* TargetPlatform)
 {
-	// @TODO : Remove any duplicate data from parent? Aims at improving change propagation (if controlled by parent)
+	PRAGMA_DISABLE_DEPRECATION_WARNINGS;
 	Super::PreSave(TargetPlatform);
+	PRAGMA_ENABLE_DEPRECATION_WARNINGS;
+}
+
+void UMaterialInstance::PreSave(FObjectPreSaveContext ObjectSaveContext)
+{
+	// @TODO : Remove any duplicate data from parent? Aims at improving change propagation (if controlled by parent)
+	Super::PreSave(ObjectSaveContext);
 }
 
 float UMaterialInstance::GetTextureDensity(FName TextureName, const struct FMeshUVChannelInfo& UVChannelData) const
diff --git a/Engine/Source/Runtime/Engine/Private/Materials/MaterialInterface.cpp b/Engine/Source/Runtime/Engine/Private/Materials/MaterialInterface.cpp
--- a/Engine/Source/Runtime/Engine/Private/Materials/MaterialInterface.cpp
+++ b/Engine/Source/Runtime/Engine/Private/Materials/MaterialInterface.cpp
@@ -5,10 +5,12 @@
 =============================================================================*/
 
 #include "Materials/MaterialInterface.h"
+
 #include "RenderingThread.h"
 #include "PrimitiveViewRelevance.h"
 #include "MaterialShared.h"
 #include "Materials/Material.h"
+#include "UObject/ObjectSaveContext.h"
 #include "UObject/UObjectHash.h"
 #include "UObject/UObjectIterator.h"
 #include "EditorFramework/AssetImportData.h"
@@ -809,7 +811,15 @@ bool UMaterialInterface::UseAnyStreamingTexture() const
 
 void UMaterialInterface::PreSave(const class ITargetPlatform* TargetPlatform)
 {
+	PRAGMA_DISABLE_DEPRECATION_WARNINGS;
 	Super::PreSave(TargetPlatform);
+	PRAGMA_ENABLE_DEPRECATION_WARNINGS;
+}
+
+void UMaterialInterface::PreSave(FObjectPreSaveContext ObjectSaveContext)
+{
+	Super::PreSave(ObjectSaveContext);
+	const ITargetPlatform* TargetPlatform = ObjectSaveContext.GetTargetPlatform();
 	if (TargetPlatform && TargetPlatform->RequiresCookedData())
 	{
 		SortTextureStreamingData(true, true);
diff --git a/Engine/Source/Runtime/Engine/Private/Particles/ParticleComponents.cpp b/Engine/Source/Runtime/Engine/Private/Particles/ParticleComponents.cpp
--- a/Engine/Source/Runtime/Engine/Private/Particles/ParticleComponents.cpp
+++ b/Engine/Source/Runtime/Engine/Private/Particles/ParticleComponents.cpp
@@ -21,6 +21,7 @@
 #include "CollisionQueryParams.h"
 #include "WorldCollision.h"
 #include "Engine/CollisionProfile.h"
+#include "UObject/ObjectSaveContext.h"
 #include "UObject/UObjectIterator.h"
 #include "UObject/Package.h"
 #include "UObject/PropertyPortFlags.h"
@@ -2307,7 +2308,14 @@ void UParticleSystem::PostEditChangeProperty(FPropertyChangedEvent& PropertyChan
 
 void UParticleSystem::PreSave(const class ITargetPlatform* TargetPlatform)
 {
+	PRAGMA_DISABLE_DEPRECATION_WARNINGS;
 	Super::PreSave(TargetPlatform);
+	PRAGMA_ENABLE_DEPRECATION_WARNINGS;
+}
+
+void UParticleSystem::PreSave(FObjectPreSaveContext ObjectSaveContext)
+{
+	Super::PreSave(ObjectSaveContext);
 #if WITH_EDITORONLY_DATA
 	// Ensure that soloing is undone...
 	int32 NumEmitters = FMath::Min(Emitters.Num(),SoloTracking.Num());
diff --git a/Engine/Source/Runtime/Engine/Private/SkeletalMesh.cpp b/Engine/Source/Runtime/Engine/Private/SkeletalMesh.cpp
--- a/Engine/Source/Runtime/Engine/Private/SkeletalMesh.cpp
+++ b/Engine/Source/Runtime/Engine/Private/SkeletalMesh.cpp
@@ -5,6 +5,7 @@
 =============================================================================*/
 
 #include "Engine/SkeletalMesh.h"
+
 #include "Serialization/CustomVersion.h"
 #include "Misc/App.h"
 #include "Misc/PackageSegment.h"
@@ -27,6 +28,7 @@
 #include "UObject/EditorObjectVersion.h"
 #include "UObject/FrameworkObjectVersion.h"
 #include "UObject/NiagaraObjectVersion.h"
+#include "UObject/ObjectSaveContext.h"
 #include "UObject/PackageResourceManager.h"
 #include "UObject/PropertyPortFlags.h"
 #include "UObject/RenderingObjectVersion.h"
@@ -2088,11 +2090,18 @@ void USkeletalMesh::FinishBuildInternal(FSkeletalMeshBuildContext& Context)
 #endif // #if WITH_EDITOR
 
 void USkeletalMesh::PreSave(const class ITargetPlatform* TargetPlatform)
+{
+	PRAGMA_DISABLE_DEPRECATION_WARNINGS;
+	Super::PreSave(TargetPlatform);
+	PRAGMA_ENABLE_DEPRECATION_WARNINGS;
+}
+
+void USkeletalMesh::PreSave(FObjectPreSaveContext ObjectSaveContext)
 {
 	// check the parent index of the root bone is invalid
 	check((GetRefSkeleton().GetNum() == 0) || (GetRefSkeleton().GetRefBoneInfo()[0].ParentIndex == INDEX_NONE));
 
-	Super::PreSave(TargetPlatform);
+	Super::PreSave(ObjectSaveContext);
 }
 
 // Pre-calculate refpose-to-local transforms
diff --git a/Engine/Source/Runtime/Engine/Private/Texture.cpp b/Engine/Source/Runtime/Engine/Private/Texture.cpp
--- a/Engine/Source/Runtime/Engine/Private/Texture.cpp
+++ b/Engine/Source/Runtime/Engine/Private/Texture.cpp
@@ -1,6 +1,7 @@
 // Copyright Epic Games, Inc. All Rights Reserved.
 
 #include "Engine/Texture.h"
+
 #include "Misc/App.h"
 #include "Modules/ModuleManager.h"
 #include "Materials/MaterialInterface.h"
@@ -9,6 +10,7 @@
 #include "Misc/FeedbackContext.h"
 #include "HAL/LowLevelMemTracker.h"
 #include "UObject/UObjectIterator.h"
+#include "UObject/ObjectSaveContext.h"
 #include "TextureResource.h"
 #include "Engine/Texture2D.h"
 #include "ContentStreaming.h"
@@ -576,10 +578,17 @@ void UTexture::FinishDestroy()
 }
 
 void UTexture::PreSave(const class ITargetPlatform* TargetPlatform)
+{
+	PRAGMA_DISABLE_DEPRECATION_WARNINGS;
+	Super::PreSave(TargetPlatform);
+	PRAGMA_ENABLE_DEPRECATION_WARNINGS;
+}
+
+void UTexture::PreSave(FObjectPreSaveContext ObjectSaveContext)
 {
 	PreSaveEvent.Broadcast(this);
 
-	Super::PreSave(TargetPlatform);
+	Super::PreSave(ObjectSaveContext);
 
 #if WITH_EDITOR
 	if (DeferCompression)
@@ -589,8 +598,7 @@ void UTexture::PreSave(const class ITargetPlatform* TargetPlatform)
 		UpdateResource();
 	}
 
-	bool bIsCooking = TargetPlatform != nullptr;
-	if (!GEngine->IsAutosaving() && (!bIsCooking))
+	if (!GEngine->IsAutosaving() && !ObjectSaveContext.IsCooking())
 	{
 		GWarn->StatusUpdate(0, 0, FText::Format(NSLOCTEXT("UnrealEd", "SavingPackage_CompressingSourceArt", "Compressing source art for texture:  {0}"), FText::FromString(GetName())));
 		Source.Compress();
diff --git a/Engine/Source/Runtime/Engine/Private/Texture2D.cpp b/Engine/Source/Runtime/Engine/Private/Texture2D.cpp
--- a/Engine/Source/Runtime/Engine/Private/Texture2D.cpp
+++ b/Engine/Source/Runtime/Engine/Private/Texture2D.cpp
@@ -5,6 +5,7 @@
 =============================================================================*/
 
 #include "Engine/Texture2D.h"
+
 #include "Serialization/MemoryWriter.h"
 #include "Misc/App.h"
 #include "HAL/PlatformFileManager.h"
@@ -50,6 +51,7 @@
 #include "Engine/Public/ImageUtils.h"
 #include "TextureCompiler.h"
 #include "Misc/ScopedSlowTask.h"
+#include "UObject/ObjectSaveContext.h"
 #include "UObject/StrongObjectPtr.h"
 
 #define LOCTEXT_NAMESPACE "UTexture2D"
@@ -625,7 +627,14 @@ void UTexture2D::PostLoad()
 
 void UTexture2D::PreSave(const class ITargetPlatform* TargetPlatform)
 {
+	PRAGMA_DISABLE_DEPRECATION_WARNINGS;
 	Super::PreSave(TargetPlatform);
+	PRAGMA_ENABLE_DEPRECATION_WARNINGS;
+}
+
+void UTexture2D::PreSave(FObjectPreSaveContext ObjectSaveContext)
+{
+	Super::PreSave(ObjectSaveContext);
 #if WITH_EDITOR
 	if( bTemporarilyDisableStreaming )
 	{
diff --git a/Engine/Source/Runtime/Engine/Private/World.cpp b/Engine/Source/Runtime/Engine/Private/World.cpp
--- a/Engine/Source/Runtime/Engine/Private/World.cpp
+++ b/Engine/Source/Runtime/Engine/Private/World.cpp
@@ -5,6 +5,7 @@
 =============================================================================*/
 
 #include "Engine/World.h"
+
 #include "HAL/FileManager.h"
 #include "Misc/CommandLine.h"
 #include "Misc/Paths.h"
@@ -16,6 +17,7 @@
 #include "UObject/UObjectIterator.h"
 #include "UObject/Package.h"
 #include "UObject/ObjectRedirector.h"
+#include "UObject/ObjectSaveContext.h"
 #include "UObject/UObjectAnnotation.h"
 #include "UObject/ReferenceChainSearch.h"
 #include "Misc/PackageName.h"
@@ -1195,6 +1197,18 @@ void UWorld::PreDuplicate(FObjectDuplicationParameters& DupParams)
 
 bool UWorld::PreSaveRoot(const TCHAR* Filename)
 {
+	PRAGMA_DISABLE_DEPRECATION_WARNINGS;
+	return Super::PreSaveRoot(Filename);
+	PRAGMA_ENABLE_DEPRECATION_WARNINGS;
+}
+
+void UWorld::PreSaveRoot(FObjectPreSaveRootContext ObjectSaveContext)
+{
+	if (!ObjectSaveContext.IsFirstConcurrentSave())
+	{
+		return;
+	}
+
 #if WITH_EDITOR
 	// if we are cooking this world, convert its persistent level to internal actors before doing so
 	if (GIsCookerLoadingPackage)
@@ -1204,19 +1218,29 @@ bool UWorld::PreSaveRoot(const TCHAR* Filename)
 #endif
 
 	// Update components and keep track off whether we need to clean them up afterwards.
-	bool bCleanupIsRequired = false;
 	if(!PersistentLevel->bAreComponentsCurrentlyRegistered)
 	{
 		PersistentLevel->UpdateLevelComponents(true);
-		bCleanupIsRequired = true;
+		ObjectSaveContext.SetCleanupRequired(true);
 	}
+}
 
-	return bCleanupIsRequired;
+void UWorld::PostSaveRoot(bool bCleanupIsRequired)
+{
+	PRAGMA_DISABLE_DEPRECATION_WARNINGS;
+	Super::PostSaveRoot(bCleanupIsRequired);
+	PRAGMA_ENABLE_DEPRECATION_WARNINGS;
 }
 
-void UWorld::PostSaveRoot( bool bCleanupIsRequired )
+void UWorld::PostSaveRoot( FObjectPostSaveRootContext ObjectSaveContext )
 {
-	if( bCleanupIsRequired )
+	Super::PostSaveRoot(ObjectSaveContext);
+	if (!ObjectSaveContext.IsLastConcurrentSave())
+	{
+		return;
+	}
+
+	if( ObjectSaveContext.IsCleanupRequired() )
 	{
 		PersistentLevel->ClearLevelComponents();
 	}
diff --git a/Engine/Source/Runtime/Engine/Private/WorldPartition/WorldPartitionRuntimeSpatialHash.cpp b/Engine/Source/Runtime/Engine/Private/WorldPartition/WorldPartitionRuntimeSpatialHash.cpp
--- a/Engine/Source/Runtime/Engine/Private/WorldPartition/WorldPartitionRuntimeSpatialHash.cpp
+++ b/Engine/Source/Runtime/Engine/Private/WorldPartition/WorldPartitionRuntimeSpatialHash.cpp
@@ -4,6 +4,7 @@
 	WorldPartitionRuntimeSpatialHash.cpp: UWorldPartitionRuntimeSpatialHash implementation
 =============================================================================*/
 #include "WorldPartition/WorldPartitionRuntimeSpatialHash.h"
+
 #include "WorldPartition/WorldPartitionRuntimeSpatialHashCell.h"
 #include "WorldPartition/WorldPartitionActorDesc.h"
 #include "WorldPartition/WorldPartitionActorDescView.h"
@@ -29,6 +30,7 @@
 #include "DisplayDebugHelpers.h"
 #include "RenderUtils.h"
 #include "Algo/Transform.h"
+#include "UObject/ObjectSaveContext.h"
 
 #if WITH_EDITOR
 #include "Editor/EditorEngine.h"
@@ -440,7 +442,14 @@ UWorldPartitionRuntimeSpatialHash::UWorldPartitionRuntimeSpatialHash(const FObje
 
 void UWorldPartitionRuntimeSpatialHash::PreSave(const class ITargetPlatform* TargetPlatform)
 {
+	PRAGMA_DISABLE_DEPRECATION_WARNINGS;
 	Super::PreSave(TargetPlatform);
+	PRAGMA_ENABLE_DEPRECATION_WARNINGS;
+}
+
+void UWorldPartitionRuntimeSpatialHash::PreSave(FObjectPreSaveContext ObjectSaveContext)
+{
+	Super::PreSave(ObjectSaveContext);
 
 	if (!IsRunningCookCommandlet())
 	{
diff --git a/Engine/Source/Runtime/Engine/Public/HLOD/HLODEngineSubsystem.h b/Engine/Source/Runtime/Engine/Public/HLOD/HLODEngineSubsystem.h
--- a/Engine/Source/Runtime/Engine/Public/HLOD/HLODEngineSubsystem.h
+++ b/Engine/Source/Runtime/Engine/Public/HLOD/HLODEngineSubsystem.h
@@ -53,7 +53,7 @@ private:
 	// Recreate LOD actors for the given level.
 	void RecreateLODActorsForLevel(ULevel* InLevel, UWorld* InWorld);
 
-	void OnPreSaveWorld(uint32 InSaveFlags, UWorld* InWorld);
+	void OnPreSaveWorld(UWorld* InWorld, FObjectPreSaveContext ObjectSaveContext);
 
 	void UnregisterRecreateLODActorsDelegates();
 	void RegisterRecreateLODActorsDelegates();
diff --git a/Engine/Source/Runtime/Landscape/Private/Landscape.cpp b/Engine/Source/Runtime/Landscape/Private/Landscape.cpp
--- a/Engine/Source/Runtime/Landscape/Private/Landscape.cpp
+++ b/Engine/Source/Runtime/Landscape/Private/Landscape.cpp
@@ -5,6 +5,7 @@ Landscape.cpp: Terrain rendering
 =============================================================================*/
 
 #include "Landscape.h"
+
 #include "Serialization/MemoryWriter.h"
 #include "Serialization/BufferArchive.h"
 #include "Serialization/MemoryReader.h"
@@ -56,6 +57,7 @@ Landscape.cpp: Terrain rendering
 #include "LandscapeSubsystem.h"
 #include "Streaming/LandscapeMeshMobileUpdate.h"
 #include "ContentStreaming.h"
+#include "UObject/ObjectSaveContext.h"
 
 #if WITH_EDITOR
 #include "LandscapeEdit.h"
@@ -2047,7 +2049,14 @@ void ALandscapeProxy::OnFeatureLevelChanged(ERHIFeatureLevel::Type NewFeatureLev
 
 void ALandscapeProxy::PreSave(const class ITargetPlatform* TargetPlatform)
 {
+	PRAGMA_DISABLE_DEPRECATION_WARNINGS;
 	Super::PreSave(TargetPlatform);
+	PRAGMA_ENABLE_DEPRECATION_WARNINGS;
+}
+
+void ALandscapeProxy::PreSave(FObjectPreSaveContext ObjectSaveContext)
+{
+	Super::PreSave(ObjectSaveContext);
 
 #if WITH_EDITOR
 	// Work out whether we have grass or not for the next game run
@@ -2782,7 +2791,14 @@ UMaterialInterface* ALandscapeStreamingProxy::GetLandscapeHoleMaterial() const
 
 void ALandscape::PreSave(const class ITargetPlatform* TargetPlatform)
 {
+	PRAGMA_DISABLE_DEPRECATION_WARNINGS;
 	Super::PreSave(TargetPlatform);
+	PRAGMA_ENABLE_DEPRECATION_WARNINGS;
+}
+
+void ALandscape::PreSave(FObjectPreSaveContext ObjectSaveContext)
+{
+	Super::PreSave(ObjectSaveContext);
 	//ULandscapeInfo* Info = GetLandscapeInfo();
 	//if (GIsEditor && Info && !IsRunningCommandlet())
 	//{
diff --git a/Engine/Source/Runtime/Landscape/Private/LandscapeCollision.cpp b/Engine/Source/Runtime/Landscape/Private/LandscapeCollision.cpp
--- a/Engine/Source/Runtime/Landscape/Private/LandscapeCollision.cpp
+++ b/Engine/Source/Runtime/Landscape/Private/LandscapeCollision.cpp
@@ -6,6 +6,7 @@
 #include "Stats/Stats.h"
 #include "Serialization/BufferArchive.h"
 #include "Misc/FeedbackContext.h"
+#include "UObject/ObjectSaveContext.h"
 #include "UObject/PropertyPortFlags.h"
 #include "EngineDefines.h"
 #include "Engine/EngineTypes.h"
@@ -2480,7 +2481,14 @@ void ULandscapeHeightfieldCollisionComponent::PostLoad()
 
 void ULandscapeHeightfieldCollisionComponent::PreSave(const class ITargetPlatform* TargetPlatform)
 {
+	PRAGMA_DISABLE_DEPRECATION_WARNINGS;
 	Super::PreSave(TargetPlatform);
+	PRAGMA_ENABLE_DEPRECATION_WARNINGS;
+}
+
+void ULandscapeHeightfieldCollisionComponent::PreSave(FObjectPreSaveContext ObjectSaveContext)
+{
+	Super::PreSave(ObjectSaveContext);
 
 	if (!IsRunningCommandlet())
 	{
diff --git a/Engine/Source/Runtime/MediaAssets/Private/Assets/BaseMediaSource.cpp b/Engine/Source/Runtime/MediaAssets/Private/Assets/BaseMediaSource.cpp
--- a/Engine/Source/Runtime/MediaAssets/Private/Assets/BaseMediaSource.cpp
+++ b/Engine/Source/Runtime/MediaAssets/Private/Assets/BaseMediaSource.cpp
@@ -1,6 +1,8 @@
 // Copyright Epic Games, Inc. All Rights Reserved.
 
 #include "BaseMediaSource.h"
+
+#include "UObject/ObjectSaveContext.h"
 #include "UObject/SequencerObjectVersion.h"
 #include "UObject/MediaFrameWorkObjectVersion.h"
 #include "IMediaModule.h"
@@ -33,8 +35,16 @@ void UBaseMediaSource::GetAssetRegistryTagMetadata(TMap<FName, FAssetRegistryTag
 #endif
 
 void UBaseMediaSource::PreSave(const class ITargetPlatform* TargetPlatform)
+{
+	PRAGMA_DISABLE_DEPRECATION_WARNINGS;
+	Super::PreSave(TargetPlatform);
+	PRAGMA_ENABLE_DEPRECATION_WARNINGS;
+}
+
+void UBaseMediaSource::PreSave(FObjectPreSaveContext ObjectSaveContext)
 {
 #if WITH_EDITORONLY_DATA
+	const ITargetPlatform* TargetPlatform = ObjectSaveContext.GetTargetPlatform();
 	if (TargetPlatform)
 	{
 		// Make sure we setup the player name according to the currently selected platform on saves
@@ -42,6 +52,7 @@ void UBaseMediaSource::PreSave(const class ITargetPlatform* TargetPlatform)
 		PlayerName = (PlatformPlayerName != nullptr) ? *PlatformPlayerName : NAME_None;
 	}
 #endif
+	Super::PreSave(ObjectSaveContext);
 }
 
 void UBaseMediaSource::Serialize(FArchive& Ar)
diff --git a/Engine/Source/Runtime/MediaAssets/Private/Assets/PlatformMediaSource.cpp b/Engine/Source/Runtime/MediaAssets/Private/Assets/PlatformMediaSource.cpp
--- a/Engine/Source/Runtime/MediaAssets/Private/Assets/PlatformMediaSource.cpp
+++ b/Engine/Source/Runtime/MediaAssets/Private/Assets/PlatformMediaSource.cpp
@@ -1,6 +1,8 @@
 // Copyright Epic Games, Inc. All Rights Reserved.
 
 #include "PlatformMediaSource.h"
+
+#include "UObject/ObjectSaveContext.h"
 #include "UObject/SequencerObjectVersion.h"
 #include "UObject/MediaFrameWorkObjectVersion.h"
 #include "Modules/ModuleManager.h"
@@ -13,15 +15,24 @@
 
 
 void UPlatformMediaSource::PreSave(const class ITargetPlatform* TargetPlatform)
+{
+	PRAGMA_DISABLE_DEPRECATION_WARNINGS;
+	Super::PreSave(TargetPlatform);
+	PRAGMA_ENABLE_DEPRECATION_WARNINGS;
+}
+
+void UPlatformMediaSource::PreSave(FObjectPreSaveContext ObjectSaveContext)
 {
 #if WITH_EDITORONLY_DATA
 	// Do this only if we are cooking (aka: have a target platform)
+	const ITargetPlatform* TargetPlatform = ObjectSaveContext.GetTargetPlatform();
 	if (TargetPlatform)
 	{
 		UE_TRANSITIONAL_OBJECT_PTR(UMediaSource)* PlatformMediaSource = PlatformMediaSources.Find(TargetPlatform->IniPlatformName());
 		MediaSource = (PlatformMediaSource != nullptr) ? *PlatformMediaSource : nullptr;
 	}
 #endif
+	Super::PreSave(ObjectSaveContext);
 }
 
 /* UMediaSource interface
diff --git a/Engine/Source/Runtime/MediaAssets/Public/BaseMediaSource.h b/Engine/Source/Runtime/MediaAssets/Public/BaseMediaSource.h
--- a/Engine/Source/Runtime/MediaAssets/Public/BaseMediaSource.h
+++ b/Engine/Source/Runtime/MediaAssets/Public/BaseMediaSource.h
@@ -38,7 +38,9 @@ private:
 public:
 
 	//~ UObject interface
+	UE_DEPRECATED(5.0, "Use version that takes FObjectPreSaveContext instead.")
 	virtual void PreSave(const class ITargetPlatform* TargetPlatform);
+	virtual void PreSave(FObjectPreSaveContext ObjectSaveContext);
 	virtual void GetAssetRegistryTags(TArray<FAssetRegistryTag>& OutTags) const override;
 	virtual void Serialize(FArchive& Ar) override;
 
diff --git a/Engine/Source/Runtime/MediaAssets/Public/PlatformMediaSource.h b/Engine/Source/Runtime/MediaAssets/Public/PlatformMediaSource.h
--- a/Engine/Source/Runtime/MediaAssets/Public/PlatformMediaSource.h
+++ b/Engine/Source/Runtime/MediaAssets/Public/PlatformMediaSource.h
@@ -38,7 +38,9 @@ private:
 
 public:
 	//~ UObject interface
+	UE_DEPRECATED(5.0, "Use version that takes FObjectPreSaveContext instead.")
 	virtual void PreSave(const class ITargetPlatform* TargetPlatform);
+	virtual void PreSave(FObjectPreSaveContext ObjectSaveContext);
 	virtual void Serialize(FArchive& Ar) override;
 
 	//~ UMediaSource interface
diff --git a/Engine/Source/Runtime/MovieScene/Private/MovieScene.cpp b/Engine/Source/Runtime/MovieScene/Private/MovieScene.cpp
--- a/Engine/Source/Runtime/MovieScene/Private/MovieScene.cpp
+++ b/Engine/Source/Runtime/MovieScene/Private/MovieScene.cpp
@@ -1,6 +1,7 @@
 // Copyright Epic Games, Inc. All Rights Reserved.
 
 #include "MovieScene.h"
+
 #include "MovieSceneTrack.h"
 #include "MovieSceneFolder.h"
 #include "MovieSceneSequence.h"
@@ -10,6 +11,7 @@
 #include "Evaluation/IMovieSceneCustomClockSource.h"
 #include "CommonFrameRates.h"
 #include "EntitySystem/IMovieSceneEntityProvider.h"
+#include "UObject/ObjectSaveContext.h"
 #include "UObject/UObjectHash.h"
 
 #define LOCTEXT_NAMESPACE "MovieScene"
@@ -1227,7 +1229,14 @@ void UMovieScene::RemoveNullTracks()
 
 void UMovieScene::PreSave(const class ITargetPlatform* TargetPlatform)
 {
+	PRAGMA_DISABLE_DEPRECATION_WARNINGS;
 	Super::PreSave(TargetPlatform);
+	PRAGMA_ENABLE_DEPRECATION_WARNINGS;
+}
+
+void UMovieScene::PreSave(FObjectPreSaveContext ObjectSaveContext)
+{
+	Super::PreSave(ObjectSaveContext);
 
 #if WITH_EDITORONLY_DATA
 	// compress meta data mappings prior to saving
diff --git a/Engine/Source/Runtime/MovieScene/Private/MovieSceneSequence.cpp b/Engine/Source/Runtime/MovieScene/Private/MovieSceneSequence.cpp
--- a/Engine/Source/Runtime/MovieScene/Private/MovieSceneSequence.cpp
+++ b/Engine/Source/Runtime/MovieScene/Private/MovieSceneSequence.cpp
@@ -1,9 +1,11 @@
 // Copyright Epic Games, Inc. All Rights Reserved.
 
 #include "MovieSceneSequence.h"
+
 #include "Evaluation/MovieSceneEvaluationCustomVersion.h"
 #include "MovieScene.h"
 #include "UObject/EditorObjectVersion.h"
+#include "UObject/ObjectSaveContext.h"
 #include "Tracks/MovieSceneSubTrack.h"
 #include "Sections/MovieSceneSubSection.h"
 #include "Logging/MessageLog.h"
@@ -96,10 +98,18 @@ EMovieSceneServerClientMask UMovieSceneSequence::OverrideNetworkMask(EMovieScene
 }
 
 void UMovieSceneSequence::PreSave(const ITargetPlatform* TargetPlatform)
+{
+	PRAGMA_DISABLE_DEPRECATION_WARNINGS;
+	Super::PreSave(TargetPlatform);
+	PRAGMA_ENABLE_DEPRECATION_WARNINGS;
+}
+
+void UMovieSceneSequence::PreSave(FObjectPreSaveContext ObjectSaveContext)
 {
 #if WITH_EDITOR
 	if (!HasAnyFlags(RF_ClassDefaultObject|RF_ArchetypeObject))
 	{
+		const ITargetPlatform* TargetPlatform = ObjectSaveContext.GetTargetPlatform();
 		if (TargetPlatform && TargetPlatform->RequiresCookedData())
 		{
 			EMovieSceneServerClientMask NetworkMask = EMovieSceneServerClientMask::All;
@@ -122,7 +132,7 @@ void UMovieSceneSequence::PreSave(const ITargetPlatform* TargetPlatform)
 		}
 	}
 #endif
-	Super::PreSave(TargetPlatform);
+	Super::PreSave(ObjectSaveContext);
 }
 
 void UMovieSceneSequence::Serialize(FArchive& Ar)
diff --git a/Engine/Source/Runtime/MovieScene/Public/MovieScene.h b/Engine/Source/Runtime/MovieScene/Public/MovieScene.h
--- a/Engine/Source/Runtime/MovieScene/Public/MovieScene.h
+++ b/Engine/Source/Runtime/MovieScene/Public/MovieScene.h
@@ -1124,8 +1124,10 @@ protected:
 
 protected:
 
-	/** Called before this object is being deserialized. */
+	UE_DEPRECATED(5.0, "Use version that takes FObjectPreSaveContext instead.")
 	virtual void PreSave(const class ITargetPlatform* TargetPlatform) override;
+	/** Called before this object is being deserialized. */
+	virtual void PreSave(FObjectPreSaveContext ObjectSaveContext) override;
 
 	/** Perform legacy upgrade of time ranges */
 	void UpgradeTimeRanges();
diff --git a/Engine/Source/Runtime/MovieScene/Public/MovieSceneSequence.h b/Engine/Source/Runtime/MovieScene/Public/MovieSceneSequence.h
--- a/Engine/Source/Runtime/MovieScene/Public/MovieSceneSequence.h
+++ b/Engine/Source/Runtime/MovieScene/Public/MovieSceneSequence.h
@@ -209,7 +209,9 @@ public:
 public:
 
 	MOVIESCENE_API virtual void PostLoad() override;
+	UE_DEPRECATED(5.0, "Use version that takes FObjectPreSaveContext instead.")
 	MOVIESCENE_API virtual void PreSave(const ITargetPlatform* TargetPlatform) override;
+	MOVIESCENE_API virtual void PreSave(FObjectPreSaveContext ObjectSaveContext) override;
 	MOVIESCENE_API virtual void BeginDestroy() override;
 	MOVIESCENE_API virtual void PostDuplicate(bool bDuplicateForPIE) override;
 
diff --git a/Engine/Source/Runtime/UMG/Private/Components/Widget.cpp b/Engine/Source/Runtime/UMG/Private/Components/Widget.cpp
--- a/Engine/Source/Runtime/UMG/Private/Components/Widget.cpp
+++ b/Engine/Source/Runtime/UMG/Private/Components/Widget.cpp
@@ -1,6 +1,7 @@
 // Copyright Epic Games, Inc. All Rights Reserved.
 
 #include "Components/Widget.h"
+
 #include "Misc/ConfigCacheIni.h"
 #include "Misc/UObjectToken.h"
 #include "CoreGlobals.h"
@@ -9,6 +10,7 @@
 #include "Widgets/IToolTip.h"
 #include "Widgets/DeclarativeSyntaxSupport.h"
 #include "Widgets/SOverlay.h"
+#include "UObject/ObjectSaveContext.h"
 #include "UObject/UObjectHash.h"
 #include "UObject/UObjectIterator.h"
 #include "Engine/LocalPlayer.h"
@@ -1140,7 +1142,14 @@ void UWidget::DeselectByDesigner()
 
 void UWidget::PreSave(const class ITargetPlatform* TargetPlatform)
 {
+	PRAGMA_DISABLE_DEPRECATION_WARNINGS;
 	Super::PreSave(TargetPlatform);
+	PRAGMA_ENABLE_DEPRECATION_WARNINGS
+}
+
+void UWidget::PreSave(FObjectPreSaveContext ObjectSaveContext)
+{
+	Super::PreSave(ObjectSaveContext);
 
 	// This is a failsafe to make sure all the accessibility data is copied over in case
 	// some rare instance isn't handled by SynchronizeProperties. It might not be necessary.
diff --git a/Engine/Source/Runtime/UMG/Private/Slate/SlateVectorArtData.cpp b/Engine/Source/Runtime/UMG/Private/Slate/SlateVectorArtData.cpp
--- a/Engine/Source/Runtime/UMG/Private/Slate/SlateVectorArtData.cpp
+++ b/Engine/Source/Runtime/UMG/Private/Slate/SlateVectorArtData.cpp
@@ -1,11 +1,13 @@
 // Copyright Epic Games, Inc. All Rights Reserved.
 
 #include "Slate/SlateVectorArtData.h"
+
 #include "RawIndexBuffer.h"
 #include "Materials/MaterialInstanceDynamic.h"
 #include "UMGPrivate.h"
 #include "StaticMeshResources.h"
 #include "Engine/StaticMesh.h"
+#include "UObject/ObjectSaveContext.h"
 
 static void StaticMeshToSlateRenderData(const UStaticMesh& DataSource, TArray<FSlateMeshVertex>& OutSlateVerts, TArray<uint32>& OutIndexes, FVector2D& OutExtentMin, FVector2D& OutExtentMax )
 {
@@ -152,7 +154,14 @@ void USlateVectorArtData::EnsureValidData()
 
 void USlateVectorArtData::PreSave(const class ITargetPlatform* TargetPlatform)
 {
+	PRAGMA_DISABLE_DEPRECATION_WARNINGS;
 	Super::PreSave(TargetPlatform);
+	PRAGMA_ENABLE_DEPRECATION_WARNINGS;
+}
+
+void USlateVectorArtData::PreSave(FObjectPreSaveContext ObjectSaveContext)
+{
+	Super::PreSave(ObjectSaveContext);
 	EnsureValidData();
 }
 
diff --git a/Engine/Source/Runtime/UMG/Private/UserWidget.cpp b/Engine/Source/Runtime/UMG/Private/UserWidget.cpp
--- a/Engine/Source/Runtime/UMG/Private/UserWidget.cpp
+++ b/Engine/Source/Runtime/UMG/Private/UserWidget.cpp
@@ -1,6 +1,7 @@
 // Copyright Epic Games, Inc. All Rights Reserved.
 
 #include "Blueprint/UserWidget.h"
+
 #include "Rendering/DrawElements.h"
 #include "Sound/SoundBase.h"
 #include "Sound/SlateSound.h"
@@ -22,6 +23,7 @@
 #include "Blueprint/WidgetLayoutLibrary.h"
 #include "UObject/EditorObjectVersion.h"
 #include "UMGPrivate.h"
+#include "UObject/ObjectSaveContext.h"
 #include "UObject/UObjectHash.h"
 #include "UObject/PropertyPortFlags.h"
 #include "TimerManager.h"
@@ -1826,6 +1828,13 @@ bool UUserWidget::IsAsset() const
 }
 
 void UUserWidget::PreSave(const class ITargetPlatform* TargetPlatform)
+{
+	PRAGMA_DISABLE_DEPRECATION_WARNINGS;
+	Super::PreSave(TargetPlatform);
+	PRAGMA_ENABLE_DEPRECATION_WARNINGS;
+}
+
+void UUserWidget::PreSave(FObjectPreSaveContext ObjectSaveContext)
 {
 	// Remove bindings that are no longer contained in the class.
 	if ( UWidgetBlueprintGeneratedClass* BGClass = GetWidgetTreeOwningClass())
@@ -1833,7 +1842,7 @@ void UUserWidget::PreSave(const class ITargetPlatform* TargetPlatform)
 		RemoveObsoleteBindings(BGClass->NamedSlots);
 	}
 
-	Super::PreSave(TargetPlatform);
+	Super::PreSave(ObjectSaveContext);
 }
 
 void UUserWidget::PostLoad()
diff --git a/Engine/Source/Runtime/UMG/Private/WidgetTree.cpp b/Engine/Source/Runtime/UMG/Private/WidgetTree.cpp
--- a/Engine/Source/Runtime/UMG/Private/WidgetTree.cpp
+++ b/Engine/Source/Runtime/UMG/Private/WidgetTree.cpp
@@ -1,10 +1,12 @@
 // Copyright Epic Games, Inc. All Rights Reserved.
 
 #include "Blueprint/WidgetTree.h"
+
 #include "Components/Visual.h"
 #include "Components/Widget.h"
 #include "Blueprint/UserWidget.h"
 #include "UMGPrivate.h"
+#include "UObject/ObjectSaveContext.h"
 
 /////////////////////////////////////////////////////
 // UWidgetTree
@@ -240,7 +242,14 @@ void UWidgetTree::ForWidgetAndChildren(UWidget* Widget, TFunctionRef<void(UWidge
 
 void UWidgetTree::PreSave(const class ITargetPlatform* TargetPlatform)
 {
-	if (TargetPlatform == nullptr)
+	PRAGMA_DISABLE_DEPRECATION_WARNINGS;
+	Super::PreSave(TargetPlatform);
+	PRAGMA_ENABLE_DEPRECATION_WARNINGS;
+}
+
+void UWidgetTree::PreSave(FObjectPreSaveContext ObjectSaveContext)
+{
+	if (!ObjectSaveContext.IsCooking())
 	{
 #if WITH_EDITORONLY_DATA
 		AllWidgets.Empty();
@@ -262,7 +271,7 @@ void UWidgetTree::PreSave(const class ITargetPlatform* TargetPlatform)
 	});
 #endif
 
-	Super::PreSave(TargetPlatform);
+	Super::PreSave(ObjectSaveContext);
 }
 
 void UWidgetTree::PostLoad()
diff --git a/Engine/Source/Runtime/UMG/Public/Blueprint/UserWidget.h b/Engine/Source/Runtime/UMG/Public/Blueprint/UserWidget.h
--- a/Engine/Source/Runtime/UMG/Public/Blueprint/UserWidget.h
+++ b/Engine/Source/Runtime/UMG/Public/Blueprint/UserWidget.h
@@ -4,6 +4,7 @@
 
 #include "CoreMinimal.h"
 #include "UObject/ObjectMacros.h"
+#include "UObject/ObjectSaveContext.h"
 #include "Styling/SlateColor.h"
 #include "Layout/Geometry.h"
 #include "Input/CursorReply.h"
@@ -1091,7 +1092,9 @@ public:
 
 	//~ Begin UObject Interface
 	virtual bool IsAsset() const;
+	UE_DEPRECATED(5.0, "Use version that takes FObjectPreSaveContext instead.")
 	virtual void PreSave(const class ITargetPlatform* TargetPlatform) override;
+	virtual void PreSave(FObjectPreSaveContext ObjectSaveContext) override;
 	//~ End UObject Interface
 
 	/** Are we currently playing any animations? */
