Change presave prototype: temporarily remove deprecation to find the discrepancy with the farm build vs my local build.
#rnx

diff --git a/Engine/Source/Runtime/CoreUObject/Public/UObject/Object.h b/Engine/Source/Runtime/CoreUObject/Public/UObject/Object.h
--- a/Engine/Source/Runtime/CoreUObject/Public/UObject/Object.h
+++ b/Engine/Source/Runtime/CoreUObject/Public/UObject/Object.h
@@ -211,7 +211,7 @@ public:
 	{
 	}
 
-	UE_DEPRECATED(5.0, "Use version that takes FObjectPreSaveContext instead.")
+	//UE_DEPRECATED(5.0, "Use version that takes FObjectPreSaveContext instead.")
 	virtual bool PreSaveRoot(const TCHAR* Filename);
 
 	/**
@@ -223,7 +223,7 @@ public:
 	 */
 	virtual void PreSaveRoot(FObjectPreSaveRootContext ObjectSaveContext);
 
-	UE_DEPRECATED(5.0, "Use version that takes FObjectPostSaveContext instead.")
+	//UE_DEPRECATED(5.0, "Use version that takes FObjectPostSaveContext instead.")
 	virtual void PostSaveRoot(bool bCleanupIsRequired);
 
 	/**
@@ -234,7 +234,7 @@ public:
 	 */
 	virtual void PostSaveRoot(FObjectPostSaveRootContext ObjectSaveContext);
 
-	UE_DEPRECATED(5.0, "Use version that takes FObjectPreSaveContext instead.")
+	//UE_DEPRECATED(5.0, "Use version that takes FObjectPreSaveContext instead.")
 	virtual void PreSave(const class ITargetPlatform* TargetPlatform);
 
 	/**
diff --git a/Engine/Source/Runtime/Engine/Classes/Engine/Texture.h b/Engine/Source/Runtime/Engine/Classes/Engine/Texture.h
--- a/Engine/Source/Runtime/Engine/Classes/Engine/Texture.h
+++ b/Engine/Source/Runtime/Engine/Classes/Engine/Texture.h
@@ -1084,7 +1084,7 @@ public:
 	ENGINE_API virtual void Serialize(FArchive& Ar) override;
 	ENGINE_API virtual void PostInitProperties() override;
 	ENGINE_API virtual void PostLoad() override;
-	UE_DEPRECATED(5.0, "Use version that takes FObjectPreSaveContext instead.")
+	//UE_DEPRECATED(5.0, "Use version that takes FObjectPreSaveContext instead.")
 	ENGINE_API virtual void PreSave(const class ITargetPlatform* TargetPlatform) override;
 	ENGINE_API virtual void PreSave(FObjectPreSaveContext ObjectSaveContext) override;
 	ENGINE_API virtual void BeginDestroy() override;
diff --git a/Engine/Source/Runtime/Engine/Classes/Materials/MaterialInterface.h b/Engine/Source/Runtime/Engine/Classes/Materials/MaterialInterface.h
--- a/Engine/Source/Runtime/Engine/Classes/Materials/MaterialInterface.h
+++ b/Engine/Source/Runtime/Engine/Classes/Materials/MaterialInterface.h
@@ -943,7 +943,7 @@ public:
 	*/
 	ENGINE_API virtual float GetTextureDensity(FName TextureName, const struct FMeshUVChannelInfo& UVChannelData) const;
 
-	UE_DEPRECATED(5.0, "Use version that takes FObjectPreSaveContext instead.")
+	//UE_DEPRECATED(5.0, "Use version that takes FObjectPreSaveContext instead.")
 	ENGINE_API virtual void PreSave(const class ITargetPlatform* TargetPlatform) override;
 	ENGINE_API virtual void PreSave(FObjectPreSaveContext ObjectSaveContext) override;
 
