Moved UE4::ELoadConfigPropagationFlags to UE namespace and deprecated the old names

diff --git a/Engine/Source/Developer/Settings/Private/SettingsSection.cpp b/Engine/Source/Developer/Settings/Private/SettingsSection.cpp
--- a/Engine/Source/Developer/Settings/Private/SettingsSection.cpp
+++ b/Engine/Source/Developer/Settings/Private/SettingsSection.cpp
@@ -157,7 +157,7 @@ bool FSettingsSection::Import( const FString& Filename )
 
 	if (SettingsObject.IsValid())
 	{
-		SettingsObject->LoadConfig(SettingsObject->GetClass(), *Filename, UE4::LCPF_PropagateToInstances);
+		SettingsObject->LoadConfig(SettingsObject->GetClass(), *Filename, UE::LCPF_PropagateToInstances);
 
 		return true;
 	}
@@ -182,7 +182,7 @@ bool FSettingsSection::ResetDefaults()
 
 		FConfigCacheIni::LoadGlobalIniFile(ConfigName, *FPaths::GetBaseFilename(ConfigName), nullptr, true);
 
-		SettingsObject->ReloadConfig(nullptr, nullptr, UE4::LCPF_PropagateToInstances|UE4::LCPF_PropagateToChildDefaultObjects);
+		SettingsObject->ReloadConfig(nullptr, nullptr, UE::LCPF_PropagateToInstances|UE::LCPF_PropagateToChildDefaultObjects);
 
 		return true;
 	}
@@ -250,7 +250,7 @@ bool FSettingsSection::SaveDefaults()
 	if (SettingsObject.IsValid())
 	{
 		SettingsObject->UpdateDefaultConfigFile();
-		SettingsObject->ReloadConfig(nullptr, nullptr, UE4::LCPF_PropagateToInstances);
+		SettingsObject->ReloadConfig(nullptr, nullptr, UE::LCPF_PropagateToInstances);
 
 		return true;			
 	}
diff --git a/Engine/Source/Runtime/CoreUObject/Private/UObject/Obj.cpp b/Engine/Source/Runtime/CoreUObject/Private/UObject/Obj.cpp
--- a/Engine/Source/Runtime/CoreUObject/Private/UObject/Obj.cpp
+++ b/Engine/Source/Runtime/CoreUObject/Private/UObject/Obj.cpp
@@ -2049,14 +2049,14 @@ void UObject::ReloadConfig( UClass* ConfigClass/*=NULL*/, const TCHAR* InFilenam
 {
 	if (!GIsEditor)
 	{
-		LoadConfig(ConfigClass, InFilename, PropagationFlags | UE4::LCPF_ReloadingConfigData | UE4::LCPF_ReadParentSections, PropertyToLoad);
+		LoadConfig(ConfigClass, InFilename, PropagationFlags | UE::LCPF_ReloadingConfigData | UE::LCPF_ReadParentSections, PropertyToLoad);
 	}
 #if WITH_EDITOR
 	else
 	{
 		// When in the editor, raise change events so that the UI will update correctly when object configs are reloaded.
 		PreEditChange(NULL);
-		LoadConfig(ConfigClass, InFilename, PropagationFlags | UE4::LCPF_ReloadingConfigData | UE4::LCPF_ReadParentSections, PropertyToLoad);
+		LoadConfig(ConfigClass, InFilename, PropagationFlags | UE::LCPF_ReloadingConfigData | UE::LCPF_ReadParentSections, PropertyToLoad);
 		PostEditChange();
 	}
 #endif // WITH_EDITOR
@@ -2145,18 +2145,18 @@ void UObject::LoadConfig( UClass* ConfigClass/*=NULL*/, const TCHAR* InFilename/
 	{
 		if ( ParentClass->HasAnyClassFlags(CLASS_Config) )
 		{
-			if ( (PropagationFlags&UE4::LCPF_ReadParentSections) != 0 )
+			if ( (PropagationFlags&UE::LCPF_ReadParentSections) != 0 )
 			{
 				// call LoadConfig on the parent class
 				LoadConfig( ParentClass, NULL, PropagationFlags, PropertyToLoad );
 
 				// if we are also notifying child classes or instances, stop here as this object's properties will be imported as a result of notifying the others
-				if ( (PropagationFlags & (UE4::LCPF_PropagateToChildDefaultObjects|UE4::LCPF_PropagateToInstances)) != 0 )
+				if ( (PropagationFlags & (UE::LCPF_PropagateToChildDefaultObjects|UE::LCPF_PropagateToInstances)) != 0 )
 				{
 					return;
 				}
 			}
-			else if ( (PropagationFlags&UE4::LCPF_PropagateToChildDefaultObjects) != 0 )
+			else if ( (PropagationFlags&UE::LCPF_PropagateToChildDefaultObjects) != 0 )
 			{
 				// not propagating the call upwards, but we are propagating the call to all child classes
 				for (TObjectIterator<UClass> It; It; ++It)
@@ -2164,14 +2164,14 @@ void UObject::LoadConfig( UClass* ConfigClass/*=NULL*/, const TCHAR* InFilename/
 					if (It->IsChildOf(ConfigClass))
 					{
 						// mask out the PropgateToParent and PropagateToChildren values
-						It->GetDefaultObject()->LoadConfig(*It, NULL, (PropagationFlags&(UE4::LCPF_PersistentFlags|UE4::LCPF_PropagateToInstances)), PropertyToLoad);
+						It->GetDefaultObject()->LoadConfig(*It, NULL, (PropagationFlags&(UE::LCPF_PersistentFlags|UE::LCPF_PropagateToInstances)), PropertyToLoad);
 					}
 				}
 
 				// LoadConfig() was called on this object during iteration, so stop here 
 				return;
 			}
-			else if ( (PropagationFlags&UE4::LCPF_PropagateToInstances) != 0 )
+			else if ( (PropagationFlags&UE::LCPF_PropagateToInstances) != 0 )
 			{
 				// call LoadConfig() on all instances of this class (except the CDO)
 				// Do not propagate this call to parents, and do not propagate to children or instances (would be redundant) 
@@ -2182,7 +2182,7 @@ void UObject::LoadConfig( UClass* ConfigClass/*=NULL*/, const TCHAR* InFilename/
 						if ( !GIsEditor )
 						{
 							// make sure to pass in the class so that OriginalClass isn't reset
-							It->LoadConfig(It->GetClass(), NULL, (PropagationFlags&UE4::LCPF_PersistentFlags), PropertyToLoad);
+							It->LoadConfig(It->GetClass(), NULL, (PropagationFlags&UE::LCPF_PersistentFlags), PropertyToLoad);
 						}
 #if WITH_EDITOR
 						else
@@ -2190,7 +2190,7 @@ void UObject::LoadConfig( UClass* ConfigClass/*=NULL*/, const TCHAR* InFilename/
 							It->PreEditChange(NULL);
 
 							// make sure to pass in the class so that OriginalClass isn't reset
-							It->LoadConfig(It->GetClass(), NULL, (PropagationFlags&UE4::LCPF_PersistentFlags), PropertyToLoad);
+							It->LoadConfig(It->GetClass(), NULL, (PropagationFlags&UE::LCPF_PersistentFlags), PropertyToLoad);
 
 							It->PostEditChange();
 						}
@@ -2199,7 +2199,7 @@ void UObject::LoadConfig( UClass* ConfigClass/*=NULL*/, const TCHAR* InFilename/
 				}
 			}
 		}
-		else if ( (PropagationFlags&UE4::LCPF_PropagateToChildDefaultObjects) != 0 )
+		else if ( (PropagationFlags&UE::LCPF_PropagateToChildDefaultObjects) != 0 )
 		{
 			// we're at the base-most config class
 			for ( TObjectIterator<UClass> It; It; ++It )
@@ -2209,7 +2209,7 @@ void UObject::LoadConfig( UClass* ConfigClass/*=NULL*/, const TCHAR* InFilename/
 					if ( !GIsEditor )
 					{
 						// make sure to pass in the class so that OriginalClass isn't reset
-						It->GetDefaultObject()->LoadConfig( *It, NULL, (PropagationFlags&(UE4::LCPF_PersistentFlags|UE4::LCPF_PropagateToInstances)), PropertyToLoad );
+						It->GetDefaultObject()->LoadConfig( *It, NULL, (PropagationFlags&(UE::LCPF_PersistentFlags|UE::LCPF_PropagateToInstances)), PropertyToLoad );
 					}
 #if WITH_EDITOR
 					else
@@ -2217,7 +2217,7 @@ void UObject::LoadConfig( UClass* ConfigClass/*=NULL*/, const TCHAR* InFilename/
 						It->PreEditChange(NULL);
 
 						// make sure to pass in the class so that OriginalClass isn't reset
-						It->GetDefaultObject()->LoadConfig( *It, NULL, (PropagationFlags&(UE4::LCPF_PersistentFlags|UE4::LCPF_PropagateToInstances)), PropertyToLoad );
+						It->GetDefaultObject()->LoadConfig( *It, NULL, (PropagationFlags&(UE::LCPF_PersistentFlags|UE::LCPF_PropagateToInstances)), PropertyToLoad );
 
 						It->PostEditChange();
 					}
@@ -2227,7 +2227,7 @@ void UObject::LoadConfig( UClass* ConfigClass/*=NULL*/, const TCHAR* InFilename/
 
 			return;
 		}
-		else if ( (PropagationFlags&UE4::LCPF_PropagateToInstances) != 0 )
+		else if ( (PropagationFlags&UE::LCPF_PropagateToInstances) != 0 )
 		{
 			for ( TObjectIterator<UObject> It; It; ++It )
 			{
@@ -2236,7 +2236,7 @@ void UObject::LoadConfig( UClass* ConfigClass/*=NULL*/, const TCHAR* InFilename/
 					if ( !GIsEditor )
 					{
 						// make sure to pass in the class so that OriginalClass isn't reset
-						It->LoadConfig(It->GetClass(), NULL, (PropagationFlags&UE4::LCPF_PersistentFlags), PropertyToLoad);
+						It->LoadConfig(It->GetClass(), NULL, (PropagationFlags&UE::LCPF_PersistentFlags), PropertyToLoad);
 					}
 #if WITH_EDITOR
 					else
@@ -2244,7 +2244,7 @@ void UObject::LoadConfig( UClass* ConfigClass/*=NULL*/, const TCHAR* InFilename/
 						It->PreEditChange(NULL);
 
 						// make sure to pass in the class so that OriginalClass isn't reset
-						It->LoadConfig(It->GetClass(), NULL, (PropagationFlags&UE4::LCPF_PersistentFlags), PropertyToLoad);
+						It->LoadConfig(It->GetClass(), NULL, (PropagationFlags&UE::LCPF_PersistentFlags), PropertyToLoad);
 						It->PostEditChange();
 					}
 #endif // WITH_EDITOR
@@ -2479,7 +2479,7 @@ void UObject::LoadConfig( UClass* ConfigClass/*=NULL*/, const TCHAR* InFilename/
 	}
 
 	// if we are reloading config data after the initial class load, fire the callback now
-	if ( (PropagationFlags&UE4::LCPF_ReloadingConfigData) != 0 )
+	if ( (PropagationFlags&UE::LCPF_ReloadingConfigData) != 0 )
 	{
 		PostReloadConfig(PropertyToLoad);
 	}
@@ -2492,7 +2492,7 @@ void UObject::SaveConfig( uint64 Flags, const TCHAR* InFilename, FConfigCacheIni
 		return;
 	}
 
-	uint32 PropagationFlags = UE4::LCPF_None;
+	uint32 PropagationFlags = UE::LCPF_None;
 
 	const FString Filename
 	// if a filename was specified, always load from that file
@@ -2544,12 +2544,12 @@ void UObject::SaveConfig( uint64 Flags, const TCHAR* InFilename, FConfigCacheIni
 			if (Property->PropertyFlags & CPF_GlobalConfig)
 			{
 				// call LoadConfig() on child classes if any of the properties were global config
-				PropagationFlags |= UE4::LCPF_PropagateToChildDefaultObjects;
+				PropagationFlags |= UE::LCPF_PropagateToChildDefaultObjects;
 				BaseClass = Property->GetOwnerClass();
 				if ( BaseClass != GetClass() )
 				{
 					// call LoadConfig() on parent classes only if the global config property was declared in a parent class
-					PropagationFlags |= UE4::LCPF_ReadParentSections;
+					PropagationFlags |= UE::LCPF_ReadParentSections;
 				}
 			}
 
diff --git a/Engine/Source/Runtime/CoreUObject/Private/UObject/UObjectGlobals.cpp b/Engine/Source/Runtime/CoreUObject/Private/UObject/UObjectGlobals.cpp
--- a/Engine/Source/Runtime/CoreUObject/Private/UObject/UObjectGlobals.cpp
+++ b/Engine/Source/Runtime/CoreUObject/Private/UObject/UObjectGlobals.cpp
@@ -2891,7 +2891,7 @@ void FObjectInitializer::PostConstructInit()
 	{
 		if ((bIsCDO && !Class->HasAnyFlags(RF_Dynamic)) || Class->HasAnyClassFlags(CLASS_PerObjectConfig))
 		{
-			Obj->LoadConfig(NULL, NULL, bIsCDO ? UE4::LCPF_ReadParentSections : UE4::LCPF_None);
+			Obj->LoadConfig(NULL, NULL, bIsCDO ? UE::LCPF_ReadParentSections : UE::LCPF_None);
 		}
 		else if (bIsCDO && Class->HasAnyFlags(RF_Dynamic) && Class->HasAnyClassFlags(CLASS_Config))
 		{
diff --git a/Engine/Source/Runtime/CoreUObject/Public/UObject/Object.h b/Engine/Source/Runtime/CoreUObject/Public/UObject/Object.h
--- a/Engine/Source/Runtime/CoreUObject/Public/UObject/Object.h
+++ b/Engine/Source/Runtime/CoreUObject/Public/UObject/Object.h
@@ -1174,10 +1174,10 @@ public:
 	 *
 	 * @param	Class				the class to use for determining which section of the ini to retrieve text values from
 	 * @param	Filename			indicates the filename to load values from; if not specified, uses ConfigClass's ClassConfigName
-	 * @param	PropagationFlags	indicates how this call to LoadConfig should be propagated; expects a bitmask of UE4::ELoadConfigPropagationFlags values.
+	 * @param	PropagationFlags	indicates how this call to LoadConfig should be propagated; expects a bitmask of UE::ELoadConfigPropagationFlags values.
 	 * @param	PropertyToLoad		if specified, only the ini value for the specified property will be imported.
 	 */
-	void LoadConfig( UClass* ConfigClass=NULL, const TCHAR* Filename=NULL, uint32 PropagationFlags=UE4::LCPF_None, class FProperty* PropertyToLoad=NULL );
+	void LoadConfig( UClass* ConfigClass=NULL, const TCHAR* Filename=NULL, uint32 PropagationFlags=UE::LCPF_None, class FProperty* PropertyToLoad=NULL );
 
 	/**
 	 * Wrapper method for LoadConfig that is used when reloading the config data for objects at runtime which have already loaded their config data at least once.
@@ -1185,10 +1185,10 @@ public:
 	 *
 	 * @param	Class				the class to use for determining which section of the ini to retrieve text values from
 	 * @param	Filename			indicates the filename to load values from; if not specified, uses ConfigClass's ClassConfigName
-	 * @param	PropagationFlags	indicates how this call to LoadConfig should be propagated; expects a bitmask of UE4::ELoadConfigPropagationFlags values.
+	 * @param	PropagationFlags	indicates how this call to LoadConfig should be propagated; expects a bitmask of UE::ELoadConfigPropagationFlags values.
 	 * @param	PropertyToLoad		if specified, only the ini value for the specified property will be imported
 	 */
-	void ReloadConfig( UClass* ConfigClass=NULL, const TCHAR* Filename=NULL, uint32 PropagationFlags=UE4::LCPF_None, class FProperty* PropertyToLoad=NULL );
+	void ReloadConfig( UClass* ConfigClass=NULL, const TCHAR* Filename=NULL, uint32 PropagationFlags=UE::LCPF_None, class FProperty* PropertyToLoad=NULL );
 
 	/** Import an object from a file. */
 	void ParseParms( const TCHAR* Parms );
diff --git a/Engine/Source/Runtime/CoreUObject/Public/UObject/ObjectMacros.h b/Engine/Source/Runtime/CoreUObject/Public/UObject/ObjectMacros.h
--- a/Engine/Source/Runtime/CoreUObject/Public/UObject/ObjectMacros.h
+++ b/Engine/Source/Runtime/CoreUObject/Public/UObject/ObjectMacros.h
@@ -1879,7 +1879,7 @@ namespace GameplayTagsManager
 	UObject.
 -----------------------------------------------------------------------------*/
 
-namespace UE4
+namespace UE
 {
 	/**
 	 * Controls how calls to LoadConfig() should be propagated
@@ -1916,6 +1916,18 @@ namespace UE4
 	};
 }
 
+namespace UE4
+{
+	using ELoadConfigPropagationFlags UE_DEPRECATED(5.0, "Use UE namespace instead of UE4") = UE::ELoadConfigPropagationFlags;
+
+	UE_DEPRECATED(5.0, "Use UE namespace instead of UE4") static constexpr UE::ELoadConfigPropagationFlags LCPF_None							= UE::ELoadConfigPropagationFlags::LCPF_None;
+	UE_DEPRECATED(5.0, "Use UE namespace instead of UE4") static constexpr UE::ELoadConfigPropagationFlags LCPF_ReadParentSections				= UE::ELoadConfigPropagationFlags::LCPF_ReadParentSections;
+	UE_DEPRECATED(5.0, "Use UE namespace instead of UE4") static constexpr UE::ELoadConfigPropagationFlags LCPF_PropagateToChildDefaultObjects	= UE::ELoadConfigPropagationFlags::LCPF_PropagateToChildDefaultObjects;
+	UE_DEPRECATED(5.0, "Use UE namespace instead of UE4") static constexpr UE::ELoadConfigPropagationFlags LCPF_PropagateToInstances			= UE::ELoadConfigPropagationFlags::LCPF_PropagateToInstances;
+	UE_DEPRECATED(5.0, "Use UE namespace instead of UE4") static constexpr UE::ELoadConfigPropagationFlags LCPF_ReloadingConfigData				= UE::ELoadConfigPropagationFlags::LCPF_ReloadingConfigData;
+	UE_DEPRECATED(5.0, "Use UE namespace instead of UE4") static constexpr UE::ELoadConfigPropagationFlags LCPF_PersistentFlags					= UE::ELoadConfigPropagationFlags::LCPF_PersistentFlags;
+}
+
 
 /**
  * Helper class used to save and restore information across a StaticAllocateObject over the top of an existing object.
diff --git a/Engine/Source/Runtime/Engine/Private/DeviceProfiles/DeviceProfile.cpp b/Engine/Source/Runtime/Engine/Private/DeviceProfiles/DeviceProfile.cpp
--- a/Engine/Source/Runtime/Engine/Private/DeviceProfiles/DeviceProfile.cpp
+++ b/Engine/Source/Runtime/Engine/Private/DeviceProfiles/DeviceProfile.cpp
@@ -28,7 +28,7 @@ UDeviceProfile::UDeviceProfile(const FObjectInitializer& ObjectInitializer)
 	bVisible = true;
 
 	FString DeviceProfileFileName = FPaths::EngineConfigDir() + TEXT("Deviceprofiles.ini");
-//	LoadConfig(GetClass(), *DeviceProfileFileName, UE4::LCPF_ReadParentSections);
+//	LoadConfig(GetClass(), *DeviceProfileFileName, UE::LCPF_ReadParentSections);
 }
 
 
