Remove temporary local IsRunningCookCommandlet since UE5 has a global version.
#rnx

diff --git a/Engine/Source/Runtime/CoreUObject/Private/Misc/AssetRegistryInterface.cpp b/Engine/Source/Runtime/CoreUObject/Private/Misc/AssetRegistryInterface.cpp
--- a/Engine/Source/Runtime/CoreUObject/Private/Misc/AssetRegistryInterface.cpp
+++ b/Engine/Source/Runtime/CoreUObject/Private/Misc/AssetRegistryInterface.cpp
@@ -35,14 +35,6 @@ namespace Private
 	bool bInitializedSkipClasses = false;
 #endif //if WITH_ENGINE && WITH_EDITOR
 
-	// TODO: replace with a function in CoreGlobals
-	static bool IsRunningCookCommandlet()
-	{
-		FString Commandline = FCommandLine::Get();
-		const bool bIsCookCommandlet = IsRunningCommandlet() && Commandline.Contains(TEXT("run=cook"));
-		return bIsCookCommandlet;
-	}
-
 	bool FFiltering::ShouldSkipAsset(FName AssetClass, uint32 PackageFlags)
 	{
 #if WITH_ENGINE && WITH_EDITOR
