GarbageCollection: Treat cluster objects with EInternalObjectFlags::GarbageCollectionKeepFlags the same way as if they are in the root set to prevent them from being destroyed while being referenced by the async loader.

diff --git a/Engine/Source/Runtime/CoreUObject/Private/UObject/GarbageCollection.cpp b/Engine/Source/Runtime/CoreUObject/Private/UObject/GarbageCollection.cpp
--- a/Engine/Source/Runtime/CoreUObject/Private/UObject/GarbageCollection.cpp
+++ b/Engine/Source/Runtime/CoreUObject/Private/UObject/GarbageCollection.cpp
@@ -1234,6 +1234,16 @@ public:
 							ObjectItem->SetFlags(EInternalObjectFlags::Unreachable);
 						}
 					}
+					// Cluster objects 
+					else if (bWithClusters && ObjectItem->GetOwnerIndex() > 0)
+					{
+						// treat cluster objects with FastKeepFlags the same way as if they are in the root set
+						if (ObjectItem->HasAnyFlags(FastKeepFlags))
+						{
+							KeepClusterRefsList.Push(ObjectItem);
+							LocalObjectsToSerialize.Add(Object);
+						}
+					}
 				}
 			}
 
