Fix issue with deprecated warnings

#rnx

diff --git a/Engine/Source/Runtime/CoreUObject/Public/UObject/UObjectBase.h b/Engine/Source/Runtime/CoreUObject/Public/UObject/UObjectBase.h
--- a/Engine/Source/Runtime/CoreUObject/Public/UObject/UObjectBase.h
+++ b/Engine/Source/Runtime/CoreUObject/Public/UObject/UObjectBase.h
@@ -409,6 +409,7 @@ PRAGMA_ENABLE_DEPRECATION_WARNINGS
 /**
  * Specialized version of the deferred class registration structure.
  */
+PRAGMA_DISABLE_DEPRECATION_WARNINGS
 template <typename TClass>
 struct UE_DEPRECATED(5.0, "TClassCompiledInDefer has been deprecated, use FRegisterCompiledInInfo.") TClassCompiledInDefer : public FFieldCompiledInInfo
 {
@@ -427,6 +428,7 @@ struct UE_DEPRECATED(5.0, "TClassCompiledInDefer has been deprecated, use FRegis
 		return TClass::StaticPackage();
 	}
 };
+PRAGMA_ENABLE_DEPRECATION_WARNINGS
 
 /**
  * Stashes the singleton function that builds a compiled in class. Later, this is executed.
