Remove low value check statement that was causing high quantity of allocations due to FName to string conversion in the condition.

diff --git a/Engine/Source/Runtime/CoreUObject/Private/Misc/PackageAccessTracking.cpp b/Engine/Source/Runtime/CoreUObject/Private/Misc/PackageAccessTracking.cpp
--- a/Engine/Source/Runtime/CoreUObject/Private/Misc/PackageAccessTracking.cpp
+++ b/Engine/Source/Runtime/CoreUObject/Private/Misc/PackageAccessTracking.cpp
@@ -12,7 +12,6 @@ FPackageAccessRefScope::FPackageAccessRefScope(FName InPackageName, FName InOpNa
 	: PackageName(InPackageName)
 	, OpName(InOpName)
 {
-	checkf(FPackageName::IsValidLongPackageName(InPackageName.ToString(), true), TEXT("Invalid package name: %s"), *InPackageName.ToString());
 	Outer = CurrentThreadScope;
 	CurrentThreadScope = this;
 }
