Enable Save2 by defaut in the Editor:
Save2 is now the default save algorithm across both the editor and cooking. It has been used by default for cooking for a few months now and a full QA pass on its usage in the editor using the resave asset commandlet didn't come up with any issues

diff --git a/Engine/Source/Runtime/CoreUObject/Private/UObject/SavePackage/SavePackageUtilities.cpp b/Engine/Source/Runtime/CoreUObject/Private/UObject/SavePackage/SavePackageUtilities.cpp
--- a/Engine/Source/Runtime/CoreUObject/Private/UObject/SavePackage/SavePackageUtilities.cpp
+++ b/Engine/Source/Runtime/CoreUObject/Private/UObject/SavePackage/SavePackageUtilities.cpp
@@ -192,7 +192,7 @@ const FName NAME_Level("Level");
 const FName NAME_PrestreamPackage("PrestreamPackage");
 
 // Switch to SavePackage to Save2, 0:disabled, 1: enabled for cook, 2: enabled for uncooked, 3: enabled everywhere
-static TAutoConsoleVariable<int32> CVarEnablePackageNewSave(TEXT("SavePackage.EnableNewSave"), 1, TEXT("Enable new package save mechanism over the old one."));
+static TAutoConsoleVariable<int32> CVarEnablePackageNewSave(TEXT("SavePackage.EnableNewSave"), 3, TEXT("Enable new package save mechanism over the old one."));
 
 bool IsNewSaveEnabled(bool bForCooking)
 {
