Fix false-positive-redirect suppression code in SoftObjectPath. Package redirects are already builtin to FindObject, so the check for "Old object exists" was returning true even if the old object did not exist but was redirected to a new object through a package redirect. Narrow the check by requiring that FindObject succeeds and returns an object with path name equal to the old name.
#rnx

diff --git a/Engine/Source/Runtime/CoreUObject/Private/UObject/SoftObjectPath.cpp b/Engine/Source/Runtime/CoreUObject/Private/UObject/SoftObjectPath.cpp
--- a/Engine/Source/Runtime/CoreUObject/Private/UObject/SoftObjectPath.cpp
+++ b/Engine/Source/Runtime/CoreUObject/Private/UObject/SoftObjectPath.cpp
@@ -623,12 +623,13 @@ bool FSoftObjectPath::FixupCoreRedirects()
 
 	if (OldName != NewName)
 	{
-		// Only do the fixup if the old object isn't in memory, this avoids false positives
+		// Only do the fixup if the old object isn't in memory (or was redirected to new name), this avoids false positives
 		UObject* FoundOldObject = FindObjectSafe<UObject>(nullptr, *OldString);
+		FString NewString = NewName.ToString();
 
-		if (!FoundOldObject)
+		if (!FoundOldObject || FoundOldObject->GetPathName() == NewString)
 		{
-			SetPath(NewName.ToString());
+			SetPath(NewString);
 			return true;
 		}
 	}
