Change the Cast/ExactCast/CastChecked codepath for TObjectPtrs to avoid the act of casting causing a tracked access if the cast is returning nullptr.

diff --git a/Engine/Source/Runtime/CoreUObject/Public/Templates/Casts.h b/Engine/Source/Runtime/CoreUObject/Public/Templates/Casts.h
--- a/Engine/Source/Runtime/CoreUObject/Public/Templates/Casts.h
+++ b/Engine/Source/Runtime/CoreUObject/Public/Templates/Casts.h
@@ -5,6 +5,7 @@
 #include "CoreMinimal.h"
 #include "UObject/ObjectMacros.h"
 #include "UObject/Object.h"
+#include "UObject/ObjectPtr.h"
 #include "UObject/Class.h"
 #include "Templates/LosesQualifiersFromTo.h"
 #include "Traits/IsVoidType.h"
@@ -122,6 +123,22 @@ struct TCastImpl
 		return (To*)Src;
 	}
 
+	FORCEINLINE static To* DoCast( const FObjectPtr& Src )
+	{
+		UObject* SrcObj = ResolveObjectHandleNoRead(Src.GetHandleRef());
+		if (SrcObj && Src->GetClass()->HasAnyCastFlag(TCastFlags<To>::Value))
+		{
+			ObjectHandle_Private::OnHandleRead(SrcObj);
+			return (To*)SrcObj;
+		}
+		return nullptr;
+	}
+
+	FORCEINLINE static To* DoCastCheckedWithoutTypeCheck( const FObjectPtr& Src )
+	{
+		return (To*)Src.Get();
+	}
+
 	UE_DEPRECATED(4.25, "Cast<>() and CastChecked<>() should not be used with FProperties. Use CastField<>() or CastFieldChecked<>() instead.")
 	FORCEINLINE static To* DoCast( FField* Src )
 	{
@@ -148,6 +165,22 @@ struct TCastImpl<From, To, ECastType::UObjectToUObject>
 		return (To*)Src;
 	}
 
+	FORCEINLINE static To* DoCast( const FObjectPtr& Src )
+	{
+		UObject* SrcObj = ResolveObjectHandleNoRead(Src.GetHandleRef());
+		if (SrcObj && SrcObj->IsA<To>())
+		{
+			ObjectHandle_Private::OnHandleRead(SrcObj);
+			return (To*)SrcObj;
+		}
+		return nullptr;
+	}
+
+	FORCEINLINE static To* DoCastCheckedWithoutTypeCheck( const FObjectPtr& Src )
+	{
+		return (To*)Src.Get();
+	}
+
 	UE_DEPRECATED(4.25, "Cast<>() and CastChecked<>() should not be used with FProperties. Use CastField<>() or CastFieldChecked<>() instead.")
 	FORCEINLINE static To* DoCast( FField* Src )
 	{
@@ -196,6 +229,28 @@ struct TCastImpl<From, To, ECastType::UObjectToInterface>
 	{
 		return Src ? (To*)Src->GetInterfaceAddress(To::UClassType::StaticClass()) : nullptr;
 	}
+
+	FORCEINLINE static To* DoCast( const FObjectPtr& Src )
+	{
+		UObject* SrcObj = ResolveObjectHandleNoRead(Src.GetHandleRef());
+		if (SrcObj)
+		{
+			ObjectHandle_Private::OnHandleRead(SrcObj);
+			return (To*)Src.Get()->GetInterfaceAddress(To::UClassType::StaticClass());
+		}
+		return nullptr;
+	}
+
+	FORCEINLINE static To* DoCastCheckedWithoutTypeCheck( const FObjectPtr& Src )
+	{
+		UObject* SrcObj = ResolveObjectHandleNoRead(Src.GetHandleRef());
+		if (SrcObj)
+		{
+			ObjectHandle_Private::OnHandleRead(SrcObj);
+			return (To*)Src.Get()->GetInterfaceAddress(To::UClassType::StaticClass());
+		}
+		return nullptr;
+	}
 };
 
 template <typename From, typename To>
@@ -311,14 +366,41 @@ template< class T, class U > FORCEINLINE T* ExactCast  ( const TWeakObjectPtr<U>
 template< class T, class U > FORCEINLINE T* CastChecked( const TWeakObjectPtr<U>& Src, ECastCheckedType::Type CheckType = ECastCheckedType::NullChecked ) { return CastChecked<T>(Src.Get(), CheckType); }
 
 // object ptr versions
-template <class T, class U> FORCEINLINE T* Cast       ( const TObjectPtr<U>& Src                                                                   ) { return Cast       <T>(Src.Get()); }
-template <class T, class U> FORCEINLINE T* ExactCast  ( const TObjectPtr<U>& Src                                                                   ) { return ExactCast  <T>(Src.Get()); }
-template <class T, class U> FORCEINLINE T* CastChecked( const TObjectPtr<U>& Src, ECastCheckedType::Type CheckType = ECastCheckedType::NullChecked ) { return CastChecked<T>(Src.Get(), CheckType); }
-
-// const object ptr versions
-template <class T, class U> FORCEINLINE const T* Cast       ( const TObjectPtr<const U>& Src                                                                   ) { return Cast       <T>(const_cast<U*>(Src.Get())); }
-template <class T, class U> FORCEINLINE const T* ExactCast  ( const TObjectPtr<const U>& Src                                                                   ) { return ExactCast  <T>(const_cast<U*>(Src.Get())); }
-template <class T, class U> FORCEINLINE const T* CastChecked( const TObjectPtr<const U>& Src, ECastCheckedType::Type CheckType = ECastCheckedType::NullChecked ) { return CastChecked<T>(const_cast<U*>(Src.Get()), CheckType); }
+template <class T, class U> FORCEINLINE T* Cast       ( const TObjectPtr<U>& Src                                                                   ) { return TCastImpl<U, T>::DoCast((const FObjectPtr&)Src); }
+template <class T, class U> FORCEINLINE T* ExactCast  ( const TObjectPtr<U>& Src                                                                   )
+{
+	UObject* SrcObj = ResolveObjectHandleNoRead(((const FObjectPtr&)Src).GetHandleRef());
+	if (SrcObj && (SrcObj->GetClass() == T::StaticClass()))
+	{
+		ObjectHandle_Private::OnHandleRead(SrcObj);
+		return (T*)SrcObj;
+	}
+	return nullptr;
+}
+template <class T, class U> FORCEINLINE T* CastChecked( const TObjectPtr<U>& Src, ECastCheckedType::Type CheckType = ECastCheckedType::NullChecked )
+{
+#if DO_CHECK
+	if (!Src.IsNull())
+	{
+		T* Result = Cast<T>(Src);
+		if (!Result)
+		{
+			CastLogError(*GetFullNameForCastLogError(Src.Get()), *GetTypeName<T>());
+		}
+
+		return Result;
+	}
+
+	if (CheckType == ECastCheckedType::NullChecked)
+	{
+		CastLogError(TEXT("nullptr"), *GetTypeName<T>());
+	}
+
+	return nullptr;
+#else
+	return TCastImpl<U, T>::DoCastCheckedWithoutTypeCheck((const FObjectPtr&)Src);
+#endif
+}
 
 // TSubclassOf versions
 template< class T, class U > FORCEINLINE T* Cast       ( const TSubclassOf<U>& Src                                                                   ) { return Cast       <T>(*Src); }
