Reduce EDLBootNotificationManager allocations during early editor startup

diff --git a/Engine/Source/Runtime/CoreUObject/Private/Serialization/AsyncPackageLoader.cpp b/Engine/Source/Runtime/CoreUObject/Private/Serialization/AsyncPackageLoader.cpp
--- a/Engine/Source/Runtime/CoreUObject/Private/Serialization/AsyncPackageLoader.cpp
+++ b/Engine/Source/Runtime/CoreUObject/Private/Serialization/AsyncPackageLoader.cpp
@@ -151,9 +151,12 @@ struct FEDLBootNotificationManager
 			return; // We assume nothing in coreuobject ever loads assets in a constructor
 		}
 
-		FScopeLock Lock(&EDLBootNotificationManagerLock);
+		TStringBuilder<256> LongNameBuilder;
+		LongNameBuilder << PackageName;
+		FPathViews::Append(LongNameBuilder, Name);
+		FName LongFName(LongNameBuilder.ToView());
 
-		FName LongFName(*(FString(PackageName) / Name));
+		FScopeLock Lock(&EDLBootNotificationManagerLock);
 
 		//FPlatformMisc::LowLevelOutputDebugStringf(TEXT("NotifyRegistrationEvent %s %d %d\r\n"), *LongFName.ToString(), int32(NotifyRegistrationType), int32(NotifyRegistrationPhase));
 
