Make UE::AssetRegistry::Filtering work on provided containers so that the upcoming threadsafe AssetRegistry can use its own threadsafe data for the queries.
#rnx

diff --git a/Engine/Source/Runtime/CoreUObject/Private/Misc/AssetRegistryInterface.cpp b/Engine/Source/Runtime/CoreUObject/Private/Misc/AssetRegistryInterface.cpp
--- a/Engine/Source/Runtime/CoreUObject/Private/Misc/AssetRegistryInterface.cpp
+++ b/Engine/Source/Runtime/CoreUObject/Private/Misc/AssetRegistryInterface.cpp
@@ -20,9 +20,7 @@ IAssetRegistryInterface* IAssetRegistryInterface::GetPtr()
 	return Default;
 }
 
-namespace UE
-{
-namespace AssetRegistry
+namespace UE::AssetRegistry
 {
 namespace Private
 {
@@ -33,95 +31,130 @@ namespace Private
 	TSet<FName> SkipUncookedClasses;
 	TSet<FName> SkipCookedClasses;
 	bool bInitializedSkipClasses = false;
-#endif //if WITH_ENGINE && WITH_EDITOR
 
-	bool FFiltering::ShouldSkipAsset(FName AssetClass, uint32 PackageFlags)
+	void FFiltering::SetSkipClasses(const TSet<FName>& InSkipUncookedClasses, const TSet<FName>& InSkipCookedClasses)
 	{
+		bInitializedSkipClasses = true;
+		SkipUncookedClasses = InSkipUncookedClasses;
+		SkipCookedClasses = InSkipCookedClasses;
+	}
+#endif
+
 #if WITH_ENGINE && WITH_EDITOR
+namespace Utils
+{
 
-		TRACE_CPUPROFILER_EVENT_SCOPE(AssetRegistry::FFiltering::ShouldSkipAsset)
+	bool ShouldSkipAsset(FName AssetClass, uint32 PackageFlags, const TSet<FName>& InSkipUncookedClasses, const TSet<FName>& InSkipCookedClasses)
+	{
+		if (PackageFlags & PKG_ContainsNoAsset)
+		{
+			return true;
+		}
 
-		// We do not yet support having UBlueprintGeneratedClasses be assets when the UBlueprint is also
-		// an asset; the content browser does not handle the multiple assets correctly and displays this
-		// class asset as if it is in a separate package. Revisit when we have removed the UBlueprint as an asset
-		// or when we support multiple assets.
-		if (!bInitializedSkipClasses)
+		const bool bIsCooked = (PackageFlags & PKG_FilterEditorOnly);
+		if ((bIsCooked && SkipCookedClasses.Contains(AssetClass)) ||
+			(!bIsCooked && SkipUncookedClasses.Contains(AssetClass)))
 		{
-			// Since we only collect these the first on-demand time, it is possible we will miss subclasses
-			// from plugins that load later. This flaw is a rare edge case, though, and this solution will
-			// be replaced eventually, so leaving it for now.
-			if (GIsEditor && (!IsRunningCommandlet() || IsRunningCookCommandlet()))
-			{
-				static const FName NAME_EnginePackage("/Script/Engine");
-				UPackage* EnginePackage = Cast<UPackage>(StaticFindObjectFast(UPackage::StaticClass(), nullptr, NAME_EnginePackage));
+			return true;
+		}
+		return false;
+	}
 
-				{
-					SkipUncookedClasses.Reset();
+	bool ShouldSkipAsset(const UObject* InAsset, const TSet<FName>& InSkipUncookedClasses, const TSet<FName>& InSkipCookedClasses)
+	{
+		if (!InAsset)
+		{
+			return false;
+		}
+		UPackage* Package = InAsset->GetPackage();
+		if (!Package)
+		{
+			return false;
+		}
+		return ShouldSkipAsset(InAsset->GetClass()->GetFName(), Package->GetPackageFlags(), InSkipUncookedClasses, InSkipCookedClasses);
+	}
 
-					static const FName NAME_BlueprintGeneratedClass("BlueprintGeneratedClass");
-					UClass* BlueprintGeneratedClass = nullptr;
-					if (EnginePackage)
-					{
-						BlueprintGeneratedClass = Cast<UClass>(StaticFindObjectFast(UClass::StaticClass(), EnginePackage, NAME_BlueprintGeneratedClass));
-					}
-					if (!BlueprintGeneratedClass)
-					{
-						UE_LOG(LogCore, Warning, TEXT("Could not find BlueprintGeneratedClass; will not be able to filter uncooked BPGC"));
-					}
-					else
+	void PopulateSkipClasses(TSet<FName>& OutSkipUncookedClasses, TSet<FName>& OutSkipCookedClasses)
+	{
+		static const FName NAME_EnginePackage("/Script/Engine");
+		UPackage* EnginePackage = Cast<UPackage>(StaticFindObjectFast(UPackage::StaticClass(), nullptr, NAME_EnginePackage));
+		{
+			OutSkipUncookedClasses.Reset();
+
+			static const FName NAME_BlueprintGeneratedClass("BlueprintGeneratedClass");
+			UClass* BlueprintGeneratedClass = nullptr;
+			if (EnginePackage)
+			{
+				BlueprintGeneratedClass = Cast<UClass>(StaticFindObjectFast(UClass::StaticClass(), EnginePackage, NAME_BlueprintGeneratedClass));
+			}
+			if (!BlueprintGeneratedClass)
+			{
+				UE_LOG(LogCore, Warning, TEXT("Could not find BlueprintGeneratedClass; will not be able to filter uncooked BPGC"));
+			}
+			else
+			{
+				OutSkipUncookedClasses.Add(NAME_BlueprintGeneratedClass);
+				for (TObjectIterator<UClass> It; It; ++It)
+				{
+					if (It->IsChildOf(BlueprintGeneratedClass) && !It->HasAnyClassFlags(CLASS_Abstract))
 					{
-						SkipUncookedClasses.Add(NAME_BlueprintGeneratedClass);
-						for (TObjectIterator<UClass> It; It; ++It)
-						{
-							if (It->IsChildOf(BlueprintGeneratedClass) && !It->HasAnyClassFlags(CLASS_Abstract))
-							{
-								SkipUncookedClasses.Add(It->GetFName());
-							}
-						}
+						OutSkipUncookedClasses.Add(It->GetFName());
 					}
 				}
-				{
-					SkipCookedClasses.Reset();
+			}
+		}
+		{
+			OutSkipCookedClasses.Reset();
 
-					static const FName NAME_Blueprint("Blueprint");
-					UClass* BlueprintClass = nullptr;
-					if (EnginePackage)
-					{
-						BlueprintClass = Cast<UClass>(StaticFindObjectFast(UClass::StaticClass(), EnginePackage, NAME_Blueprint));
-					}
-					if (!BlueprintClass)
-					{
-						UE_LOG(LogCore, Warning, TEXT("Could not find BlueprintClass; will not be able to filter cooked BP"));
-					}
-					else
+			static const FName NAME_Blueprint("Blueprint");
+			UClass* BlueprintClass = nullptr;
+			if (EnginePackage)
+			{
+				BlueprintClass = Cast<UClass>(StaticFindObjectFast(UClass::StaticClass(), EnginePackage, NAME_Blueprint));
+			}
+			if (!BlueprintClass)
+			{
+				UE_LOG(LogCore, Warning, TEXT("Could not find BlueprintClass; will not be able to filter cooked BP"));
+			}
+			else
+			{
+				OutSkipCookedClasses.Add(NAME_Blueprint);
+				for (TObjectIterator<UClass> It; It; ++It)
+				{
+					if (It->IsChildOf(BlueprintClass) && !It->HasAnyClassFlags(CLASS_Abstract))
 					{
-						SkipCookedClasses.Add(NAME_Blueprint);
-						for (TObjectIterator<UClass> It; It; ++It)
-						{
-							if (It->IsChildOf(BlueprintClass) && !It->HasAnyClassFlags(CLASS_Abstract))
-							{
-								SkipCookedClasses.Add(It->GetFName());
-							}
-						}
+						OutSkipCookedClasses.Add(It->GetFName());
 					}
 				}
 			}
-
-			bInitializedSkipClasses = true;
 		}
+	}
 
-		if (PackageFlags & PKG_ContainsNoAsset)
-		{
-			return true;
-		}
+}
+#endif
 
-		const bool bIsCooked = (PackageFlags & PKG_FilterEditorOnly);
+	bool FFiltering::ShouldSkipAsset(FName AssetClass, uint32 PackageFlags)
+	{
+#if WITH_ENGINE && WITH_EDITOR
+		TRACE_CPUPROFILER_EVENT_SCOPE(AssetRegistry::FFiltering::ShouldSkipAsset)
 
-		if ((bIsCooked && SkipCookedClasses.Contains(AssetClass)) ||
-			(!bIsCooked && SkipUncookedClasses.Contains(AssetClass)))
+		// We do not yet support having UBlueprintGeneratedClasses be assets when the UBlueprint is also
+		// an asset; the content browser does not handle the multiple assets correctly and displays this
+		// class asset as if it is in a separate package. Revisit when we have removed the UBlueprint as an asset
+		// or when we support multiple assets.
+		if (!bInitializedSkipClasses)
 		{
-			return true;
+			// Since we only collect these the first on-demand time, it is possible we will miss subclasses
+			// from plugins that load later. This flaw is a rare edge case, though, and this solution will
+			// be replaced eventually, so leaving it for now.
+			if (GIsEditor && (!IsRunningCommandlet() || IsRunningCookCommandlet()))
+			{
+				Utils::PopulateSkipClasses(SkipUncookedClasses, SkipCookedClasses);
+			}
+
+			bInitializedSkipClasses = true;
 		}
+		return Utils::ShouldSkipAsset(AssetClass, PackageFlags, SkipUncookedClasses, SkipCookedClasses);
 #endif //if WITH_ENGINE && WITH_EDITOR
 
 		return false;
@@ -148,4 +181,3 @@ namespace Private
 #endif
 	}
 }
-}
diff --git a/Engine/Source/Runtime/CoreUObject/Public/Misc/AssetRegistryInterface.h b/Engine/Source/Runtime/CoreUObject/Public/Misc/AssetRegistryInterface.h
--- a/Engine/Source/Runtime/CoreUObject/Public/Misc/AssetRegistryInterface.h
+++ b/Engine/Source/Runtime/CoreUObject/Public/Misc/AssetRegistryInterface.h
@@ -3,6 +3,8 @@
 #pragma once
 
 #include "CoreMinimal.h"
+
+#include "Containers/ContainersFwd.h"
 #include "Modules/ModuleInterface.h"
 #include "Modules/ModuleManager.h"
 
@@ -37,9 +39,7 @@ namespace EAssetRegistryDependencyType
 class IAssetRegistry;
 class UAssetRegistryImpl;
 
-namespace UE
-{
-namespace AssetRegistry
+namespace UE::AssetRegistry
 {
 
 	/**
@@ -183,17 +183,35 @@ namespace AssetRegistry
 	class COREUOBJECT_API FFiltering
 	{
 	public:
-		/** Called to check whether we should filter out assets of the given class and package flags from the editor's asset registry */
+		/** Return whether to filter out assets of the given class and flags from the editor's asset registry */
 		static bool ShouldSkipAsset(FName AssetClass, uint32 PackageFlags);
 
-		/** Called to check whether we should filter out the given object (assumed to be an asset) from the editor's asset registry */
+		/** Return whether to filter out the given object (assumed to be an asset) from the editor's asset registry */
 		static bool ShouldSkipAsset(const UObject* InAsset);
 
 		/** Call to invalidate the list of skip assets and cause their next use to recreate them on demand */
 		static void MarkDirty();
+
+#if WITH_ENGINE && WITH_EDITOR
+		/** Copy the global skip classes set from the given external sets that were already populated. */
+		static void SetSkipClasses(const TSet<FName>& InSkipUncookedClasses, const TSet<FName>& InSkipCookedClasses);
+#endif
 	};
 
+#if WITH_ENGINE && WITH_EDITOR
+namespace Utils
+{
+	/** Return whether to filter out assets of the given class and flags based on the skip classes */
+	COREUOBJECT_API bool ShouldSkipAsset(FName AssetClass, uint32 PackageFlags,
+		const TSet<FName>& InSkipUncookedClasses, const TSet<FName>& InSkipCookedClasses);
+	/** Return whether to filter out the given object (assumed to be an asset) based on the skip classes */
+	COREUOBJECT_API bool ShouldSkipAsset(const UObject* InAsset,
+		const TSet<FName>& InSkipUncookedClasses, const TSet<FName>& InSkipCookedClasses);
+	/** Run the calculation of which classes to skip and store results in the given sets. */
+	COREUOBJECT_API void PopulateSkipClasses(TSet<FName>& OutSkipUncookedClasses, TSet<FName>& OutSkipCookedClasses);
 }
+#endif
+
 }
 
 // Enums used in public Engine headers
