Add IsRedirector static to FAssetData to prevent the need to create FAssetData from an object just to call IsRedirector on it.

diff --git a/Engine/Source/Runtime/CoreUObject/Public/AssetRegistry/AssetData.h b/Engine/Source/Runtime/CoreUObject/Public/AssetRegistry/AssetData.h
--- a/Engine/Source/Runtime/CoreUObject/Public/AssetRegistry/AssetData.h
+++ b/Engine/Source/Runtime/CoreUObject/Public/AssetRegistry/AssetData.h
@@ -207,8 +207,12 @@ public:
 	/** Returns true if the this asset is a redirector. */
 	bool IsRedirector() const
 	{
-		static const FName ObjectRedirectorClassName = UObjectRedirector::StaticClass()->GetFName();
-		return AssetClass == ObjectRedirectorClassName;
+		return IsRedirectorClassName(AssetClass);
+	}
+
+	static bool IsRedirector(UObject* Object)
+	{
+		return Object && IsRedirectorClassName(Object->GetClass()->GetFName());
 	}
 
 	/** Returns the class UClass if it is loaded. It is not possible to load the class if it is unloaded since we only have the short name. */
@@ -447,6 +451,12 @@ private:
 
 		return PackageBaseName.Equals(ObjectPathName, ESearchCase::IgnoreCase);
 	}
+
+	static bool IsRedirectorClassName(FName ClassName)
+	{
+		static const FName ObjectRedirectorClassName = UObjectRedirector::StaticClass()->GetFName();
+		return ClassName == ObjectRedirectorClassName;
+	}
 };
 
 FORCEINLINE uint32 GetTypeHash(const FAssetData& AssetData)
