Removal of UE4 references from code base

diff --git a/Engine/Source/Editor/UnrealEd/Private/Commandlets/ContentCommandlets.cpp b/Engine/Source/Editor/UnrealEd/Private/Commandlets/ContentCommandlets.cpp
--- a/Engine/Source/Editor/UnrealEd/Private/Commandlets/ContentCommandlets.cpp
+++ b/Engine/Source/Editor/UnrealEd/Private/Commandlets/ContentCommandlets.cpp
@@ -776,8 +776,8 @@ void UResavePackagesCommandlet::LoadAndSaveOnePackage(const FString& Filename)
 			{
 				if( bIsReadOnly == true && bVerifyContent == true && bAutoCheckOut == false )
 				{
-					UE_LOG(LogContentCommandlet, Warning, TEXT("Package [%s] is read-only but needs to be resaved (UE4 Version: %i, Licensee Version: %i  Current UE4 Version: %i, Current Licensee Version: %i)"),
-						*Filename, Linker->Summary.GetFileVersionUE4(), Linker->Summary.GetFileVersionLicenseeUE4(), GPackageFileUEVersion, VER_LATEST_ENGINE_LICENSEEUE4 );
+					UE_LOG(LogContentCommandlet, Warning, TEXT("Package [%s] is read-only but needs to be resaved (UE Version: %i, Licensee Version: %i  Current UE Version: %i, Current Licensee Version: %i)"),
+						*Filename, Linker->Summary.GetFileVersionUE(), Linker->Summary.GetFileVersionLicenseeUE(), GPackageFileUEVersion, VER_LATEST_ENGINE_LICENSEEUE4 );
 					if( SavePackageHelper(Package, FString(TEXT("Temp.temp"))) )
 					{
 						UE_LOG(LogContentCommandlet, Warning, TEXT("Correctly saved:  [Temp.temp].") );
@@ -842,8 +842,8 @@ void UResavePackagesCommandlet::LoadAndSaveOnePackage(const FString& Filename)
 
 					if (Verbosity != ONLY_ERRORS)
 					{
-						UE_LOG(LogContentCommandlet, Display, TEXT("Resaving package [%s] (UE4 Version: %i, Licensee Version: %i  Saved UE4 Version: %i, Saved Licensee Version: %i)"),
-							*Filename,Linker->Summary.GetFileVersionUE4(), Linker->Summary.GetFileVersionLicenseeUE4(), GPackageFileUEVersion, VER_LATEST_ENGINE_LICENSEEUE4 );
+						UE_LOG(LogContentCommandlet, Display, TEXT("Resaving package [%s] (UE Version: %i, Licensee Version: %i  Saved UE Version: %i, Saved Licensee Version: %i)"),
+							*Filename,Linker->Summary.GetFileVersionUE(), Linker->Summary.GetFileVersionLicenseeUE(), GPackageFileUEVersion, VER_LATEST_ENGINE_LICENSEEUE4 );
 					}
 
 					const static bool bKeepPackageGUIDOnSave = FParse::Param(FCommandLine::Get(), TEXT("KeepPackageGUIDOnSave"));
@@ -1363,12 +1363,12 @@ FText UResavePackagesCommandlet::GetChangelistDescription() const
 
 void UResavePackagesCommandlet::PerformPreloadOperations( FLinkerLoad* PackageLinker, bool& bSavePackage )
 {
-	const int32 UE4PackageVersion = PackageLinker->Summary.GetFileVersionUE4();
-	const int32 LicenseeUE4PackageVersion = PackageLinker->Summary.GetFileVersionLicenseeUE4();
+	const int32 UEPackageVersion = PackageLinker->Summary.GetFileVersionUE();
+	const int32 LicenseeUEPackageVersion = PackageLinker->Summary.GetFileVersionLicenseeUE();
 
 
 	// validate that this package meets the minimum requirement
-	if ( MinResaveUEVersion != IGNORE_PACKAGE_VERSION && UE4PackageVersion < MinResaveUEVersion )
+	if ( MinResaveUEVersion != IGNORE_PACKAGE_VERSION && UEPackageVersion < MinResaveUEVersion )
 	{
 		bSavePackage = false;
 		return;
@@ -1377,8 +1377,8 @@ void UResavePackagesCommandlet::PerformPreloadOperations( FLinkerLoad* PackageLi
 	// Check if this package meets the maximum requirements.
 	const bool bNoLimitation = MaxResaveUEVersion == IGNORE_PACKAGE_VERSION && MaxResaveLicenseeUEVersion == IGNORE_PACKAGE_VERSION;
 	const bool bAllowResave = bNoLimitation ||
-						 (MaxResaveUEVersion != IGNORE_PACKAGE_VERSION && UE4PackageVersion <= MaxResaveUEVersion) ||
-						 (MaxResaveLicenseeUEVersion != IGNORE_PACKAGE_VERSION && LicenseeUE4PackageVersion <= MaxResaveLicenseeUEVersion);
+						 (MaxResaveUEVersion != IGNORE_PACKAGE_VERSION && UEPackageVersion <= MaxResaveUEVersion) ||
+						 (MaxResaveLicenseeUEVersion != IGNORE_PACKAGE_VERSION && LicenseeUEPackageVersion <= MaxResaveLicenseeUEVersion);
 	// If not, don't resave it.
 	if ( !bAllowResave )
 	{
diff --git a/Engine/Source/Editor/UnrealEd/Private/Commandlets/GatherTextFromAssetsCommandlet.cpp b/Engine/Source/Editor/UnrealEd/Private/Commandlets/GatherTextFromAssetsCommandlet.cpp
--- a/Engine/Source/Editor/UnrealEd/Private/Commandlets/GatherTextFromAssetsCommandlet.cpp
+++ b/Engine/Source/Editor/UnrealEd/Private/Commandlets/GatherTextFromAssetsCommandlet.cpp
@@ -649,7 +649,7 @@ int32 UGatherTextFromAssetsCommandlet::Main(const FString& Params)
 		PackagePendingGather.PackageLocCacheState = EPackageLocCacheState::Cached;
 
 		// Have we been asked to skip the cache of text that exists in the header of newer packages?
-		if (bSkipGatherCache && PackageFileSummary.GetFileVersionUE4() >= VER_UE4_SERIALIZE_TEXT_IN_PACKAGES)
+		if (bSkipGatherCache && PackageFileSummary.GetFileVersionUE() >= VER_UE4_SERIALIZE_TEXT_IN_PACKAGES)
 		{
 			// Fallback on the old package flag check.
 			if (PackageFileSummary.GetPackageFlags() & PKG_RequiresLocalizationGather)
@@ -661,12 +661,12 @@ int32 UGatherTextFromAssetsCommandlet::Main(const FString& Params)
 		const FCustomVersion* const EditorVersion = PackageFileSummary.GetCustomVersionContainer().GetVersion(FEditorObjectVersion::GUID);
 
 		// Packages not resaved since localization gathering flagging was added to packages must be loaded.
-		if (PackageFileSummary.GetFileVersionUE4() < VER_UE4_PACKAGE_REQUIRES_LOCALIZATION_GATHER_FLAGGING)
+		if (PackageFileSummary.GetFileVersionUE() < VER_UE4_PACKAGE_REQUIRES_LOCALIZATION_GATHER_FLAGGING)
 		{
 			PackagePendingGather.PackageLocCacheState = EPackageLocCacheState::Uncached_TooOld;
 		}
 		// Package not resaved since gatherable text data was added to package headers must be loaded, since their package header won't contain pregathered text data.
-		else if (PackageFileSummary.GetFileVersionUE4() < VER_UE4_SERIALIZE_TEXT_IN_PACKAGES || (!EditorVersion || EditorVersion->Version < FEditorObjectVersion::GatheredTextEditorOnlyPackageLocId))
+		else if (PackageFileSummary.GetFileVersionUE() < VER_UE4_SERIALIZE_TEXT_IN_PACKAGES || (!EditorVersion || EditorVersion->Version < FEditorObjectVersion::GatheredTextEditorOnlyPackageLocId))
 		{
 			// Fallback on the old package flag check.
 			if (PackageFileSummary.GetPackageFlags() & PKG_RequiresLocalizationGather)
@@ -674,7 +674,7 @@ int32 UGatherTextFromAssetsCommandlet::Main(const FString& Params)
 				PackagePendingGather.PackageLocCacheState = EPackageLocCacheState::Uncached_TooOld;
 			}
 		}
-		else if (PackageFileSummary.GetFileVersionUE4() < VER_UE4_DIALOGUE_WAVE_NAMESPACE_AND_CONTEXT_CHANGES)
+		else if (PackageFileSummary.GetFileVersionUE() < VER_UE4_DIALOGUE_WAVE_NAMESPACE_AND_CONTEXT_CHANGES)
 		{
 			TArray<FAssetData> AllAssetDataInSamePackage;
 			AssetRegistry.GetAssetsByPackageName(PackagePendingGather.PackageName, AllAssetDataInSamePackage);
@@ -688,7 +688,7 @@ int32 UGatherTextFromAssetsCommandlet::Main(const FString& Params)
 		}
 
 		// If this package doesn't have any cached data, then we have to load it for gather
-		if (PackageFileSummary.GetFileVersionUE4() >= VER_UE4_SERIALIZE_TEXT_IN_PACKAGES && PackageFileSummary.GatherableTextDataOffset == 0 && (PackageFileSummary.GetPackageFlags() & PKG_RequiresLocalizationGather))
+		if (PackageFileSummary.GetFileVersionUE() >= VER_UE4_SERIALIZE_TEXT_IN_PACKAGES && PackageFileSummary.GatherableTextDataOffset == 0 && (PackageFileSummary.GetPackageFlags() & PKG_RequiresLocalizationGather))
 		{
 			PackagePendingGather.PackageLocCacheState = EPackageLocCacheState::Uncached_NoCache;
 		}
diff --git a/Engine/Source/Editor/UnrealEd/Private/Commandlets/PackageUtilities.cpp b/Engine/Source/Editor/UnrealEd/Private/Commandlets/PackageUtilities.cpp
--- a/Engine/Source/Editor/UnrealEd/Private/Commandlets/PackageUtilities.cpp
+++ b/Engine/Source/Editor/UnrealEd/Private/Commandlets/PackageUtilities.cpp
@@ -780,7 +780,7 @@ int32 ULoadPackageCommandlet::Main( const FString& Params )
 		if (bCheckForLegacyPackages)
 		{
 			FLinkerLoad* Linker = LoadPackageLinker(nullptr, PackagePath, LOAD_NoVerify);
-			MinVersion = FMath::Min<int32>(MinVersion, Linker->Summary.GetFileVersionUE4());
+			MinVersion = FMath::Min<int32>(MinVersion, Linker->Summary.GetFileVersionUE());
 		}
 		else
 		{
@@ -798,7 +798,7 @@ int32 ULoadPackageCommandlet::Main( const FString& Params )
 	GIsEditor = GIsServer = GIsClient = true;
 	if (bCheckForLegacyPackages)
 	{
-		UE_LOG(LogPackageUtilities, Log, TEXT("%d minimum UE4 version number."), MinVersion );
+		UE_LOG(LogPackageUtilities, Log, TEXT("%d minimum UE version number."), MinVersion );
 	}
 
 	return 0;
diff --git a/Engine/Source/Editor/UnrealEd/Private/EditorServer.cpp b/Engine/Source/Editor/UnrealEd/Private/EditorServer.cpp
--- a/Engine/Source/Editor/UnrealEd/Private/EditorServer.cpp
+++ b/Engine/Source/Editor/UnrealEd/Private/EditorServer.cpp
@@ -2416,7 +2416,7 @@ bool UEditorEngine::PackageIsAMapFile( const TCHAR* PackageFilename, FText& OutN
 			return false;
 		}
 
-		const int32 UEVersion = Summary.GetFileVersionUE4();
+		const int32 UEVersion = Summary.GetFileVersionUE();
 
 		// Validate the summary.
 		if( UEVersion < VER_UE4_OLDEST_LOADABLE_PACKAGE )
@@ -2425,12 +2425,12 @@ bool UEditorEngine::PackageIsAMapFile( const TCHAR* PackageFilename, FText& OutN
 			Arguments.Add(TEXT("File"), FText::FromString( FString( PackageFilename ) ));
 			Arguments.Add(TEXT("Version"), UEVersion);
 			Arguments.Add(TEXT("First"), VER_UE4_OLDEST_LOADABLE_PACKAGE);
-			OutNotMapReason = FText::Format( LOCTEXT( "UE4FileIsOlder", "{File} is an UE4 map [File:v{Version}], from an engine release no longer supported [Min:v{First}]." ), 
+			OutNotMapReason = FText::Format( LOCTEXT( "UEFileIsOlder", "{File} is an UE4 map [File:v{Version}], from an engine release no longer supported [Min:v{First}]." ), 
 				Arguments);
 			return false;
 		}
 
-		const int32 UE4LicenseeVersion = Summary.GetFileVersionLicenseeUE4();
+		const int32 UELicenseeVersion = Summary.GetFileVersionLicenseeUE();
 
 		// Don't load packages that were saved with an engine version newer than the current one.
 		if( UEVersion > GPackageFileUEVersion )
@@ -2439,17 +2439,17 @@ bool UEditorEngine::PackageIsAMapFile( const TCHAR* PackageFilename, FText& OutN
 			Arguments.Add(TEXT("File"), FText::FromString( FString( PackageFilename ) ));
 			Arguments.Add(TEXT("Version"), UEVersion);
 			Arguments.Add(TEXT("Last"), GPackageFileUEVersion);
-			OutNotMapReason = FText::Format( LOCTEXT( "UE4FileIsNewer", "{File} is a UE4 map [File:v{Version}], from an engine release newer than this [Cur:v{Last}]." ), 
+			OutNotMapReason = FText::Format( LOCTEXT( "UEFileIsNewer", "{File} is a UE map [File:v{Version}], from an engine release newer than this [Cur:v{Last}]." ), 
 				Arguments);
 			return false;
 		}
-		else if( UE4LicenseeVersion > GPackageFileLicenseeUEVersion )
+		else if( UELicenseeVersion > GPackageFileLicenseeUEVersion )
 		{
 			FFormatNamedArguments Arguments;
 			Arguments.Add(TEXT("File"), FText::FromString( FString( PackageFilename ) ));
-			Arguments.Add(TEXT("Version"), UE4LicenseeVersion);
+			Arguments.Add(TEXT("Version"), UELicenseeVersion);
 			Arguments.Add(TEXT("Last"), GPackageFileLicenseeUEVersion);
-			OutNotMapReason = FText::Format( LOCTEXT( "UE4FileIsNewer", "{File} is a UE4 map [File:v{Version}], from an engine release newer than this [Cur:v{Last}]." ), 
+			OutNotMapReason = FText::Format( LOCTEXT( "UEFileIsNewer", "{File} is a UE map [File:v{Version}], from an engine release newer than this [Cur:v{Last}]." ), 
 				Arguments);
 			return false;
 		}
diff --git a/Engine/Source/Runtime/AssetRegistry/Private/PackageReader.cpp b/Engine/Source/Runtime/AssetRegistry/Private/PackageReader.cpp
--- a/Engine/Source/Runtime/AssetRegistry/Private/PackageReader.cpp
+++ b/Engine/Source/Runtime/AssetRegistry/Private/PackageReader.cpp
@@ -80,7 +80,7 @@ bool FPackageReader::OpenPackageFile(EOpenPackageResult* OutErrorCode)
 	}
 
 	// Don't read packages that are too old
-	if( PackageFileSummary.GetFileVersionUE4() < VER_UE4_OLDEST_LOADABLE_PACKAGE )
+	if( PackageFileSummary.GetFileVersionUE() < VER_UE4_OLDEST_LOADABLE_PACKAGE )
 	{
 		UE_LOG(LogAssetRegistry, Error, TEXT("Package %s is too old"), *PackageFilename);
 		SetPackageErrorCode(EOpenPackageResult::VersionTooOld);
@@ -88,7 +88,7 @@ bool FPackageReader::OpenPackageFile(EOpenPackageResult* OutErrorCode)
 	}
 
 	// Don't read packages that were saved with an package version newer than the current one.
-	if( (PackageFileSummary.GetFileVersionUE4() > GPackageFileUEVersion) || (PackageFileSummary.GetFileVersionLicenseeUE4() > GPackageFileLicenseeUEVersion) )
+	if( (PackageFileSummary.GetFileVersionUE() > GPackageFileUEVersion) || (PackageFileSummary.GetFileVersionLicenseeUE() > GPackageFileLicenseeUEVersion) )
 	{
 		UE_LOG(LogAssetRegistry, Error, TEXT("Package %s is too new"), *PackageFilename);
 		SetPackageErrorCode(EOpenPackageResult::VersionTooNew);
@@ -119,8 +119,8 @@ bool FPackageReader::OpenPackageFile(EOpenPackageResult* OutErrorCode)
 	}
 
 	//make sure the filereader gets the correct version number (it defaults to latest version)
-	SetUEVer(PackageFileSummary.GetFileVersionUE4());
-	SetLicenseeUEVer(PackageFileSummary.GetFileVersionLicenseeUE4());
+	SetUEVer(PackageFileSummary.GetFileVersionUE());
+	SetLicenseeUEVer(PackageFileSummary.GetFileVersionLicenseeUE());
 	SetEngineVer(PackageFileSummary.SavedByEngineVersion);
 
 	const FCustomVersionContainer& PackageFileSummaryVersions = PackageFileSummary.GetCustomVersionContainer();
@@ -356,8 +356,8 @@ bool FPackageReader::ReadDependencyData(FPackageDependencyData& OutDependencyDat
 	PackageData.PackageGuid = PackageFileSummary.Guid;
 	PRAGMA_ENABLE_DEPRECATION_WARNINGS
 	PackageData.SetCustomVersions(PackageFileSummary.GetCustomVersionContainer().GetAllVersions());
-	PackageData.FileVersionUE4 = PackageFileSummary.GetFileVersionUE4();
-	PackageData.FileVersionLicenseeUE4 = PackageFileSummary.GetFileVersionLicenseeUE4();
+	PackageData.FileVersionUE = PackageFileSummary.GetFileVersionUE();
+	PackageData.FileVersionLicenseeUE = PackageFileSummary.GetFileVersionLicenseeUE();
 	PackageData.SetIsLicenseeVersion(PackageFileSummary.SavedByEngineVersion.IsLicenseeVersion());
 
 	if (!SerializeNameMap())
@@ -650,7 +650,7 @@ namespace AssetRegistry
 		const bool bIsMapPackage = (PackageFileSummary.GetPackageFlags() & PKG_ContainsMap) != 0;
 
 		// To avoid large patch sizes, we have frozen cooked package format at the format before VER_UE4_ASSETREGISTRY_DEPENDENCYFLAGS
-		bool bPreDependencyFormat = PackageFileSummary.GetFileVersionUE4() < VER_UE4_ASSETREGISTRY_DEPENDENCYFLAGS || !!(PackageFileSummary.GetPackageFlags() & PKG_FilterEditorOnly);
+		bool bPreDependencyFormat = PackageFileSummary.GetFileVersionUE() < VER_UE4_ASSETREGISTRY_DEPENDENCYFLAGS || !!(PackageFileSummary.GetPackageFlags() & PKG_FilterEditorOnly);
 
 		// Load offsets to optionally-read data
 		if (bPreDependencyFormat)
@@ -676,7 +676,7 @@ namespace AssetRegistry
 		// We will also do this for maps saved after they were marked public but no asset data was saved for some reason. A bug caused this to happen for some maps.
 		if (bIsMapPackage)
 		{
-			const bool bLegacyPackage = PackageFileSummary.GetFileVersionUE4() < VER_UE4_PUBLIC_WORLDS;
+			const bool bLegacyPackage = PackageFileSummary.GetFileVersionUE() < VER_UE4_PUBLIC_WORLDS;
 			const bool bNoMapAsset = (ObjectCount == 0);
 			if (bLegacyPackage || bNoMapAsset)
 			{
@@ -724,7 +724,7 @@ namespace AssetRegistry
 
 			// Before worlds were RF_Public, other non-public assets were added to the asset data table in map packages.
 			// Here we simply skip over them
-			if (bIsMapPackage && PackageFileSummary.GetFileVersionUE4() < VER_UE4_PUBLIC_WORLDS)
+			if (bIsMapPackage && PackageFileSummary.GetFileVersionUE() < VER_UE4_PUBLIC_WORLDS)
 			{
 				if (ObjectPath != FPackageName::GetLongPackageAssetName(PackageName))
 				{
@@ -747,7 +747,7 @@ namespace AssetRegistry
 				ObjectPath = PackageName + TEXT(".") + ObjectPath;
 			}
 			// Previously export couldn't have its outer as an import
-			else if (PackageFileSummary.GetFileVersionUE4() < VER_UE4_NON_OUTER_PACKAGE_IMPORT)
+			else if (PackageFileSummary.GetFileVersionUE() < VER_UE4_NON_OUTER_PACKAGE_IMPORT)
 			{
 				UE_ASSET_LOG(LogAssetRegistry, Warning, *PackageName, TEXT("Package has invalid export %s, resave source package!"), *ObjectPath);
 				continue;
diff --git a/Engine/Source/Runtime/CoreUObject/Private/Misc/WorldCompositionUtility.cpp b/Engine/Source/Runtime/CoreUObject/Private/Misc/WorldCompositionUtility.cpp
--- a/Engine/Source/Runtime/CoreUObject/Private/Misc/WorldCompositionUtility.cpp
+++ b/Engine/Source/Runtime/CoreUObject/Private/Misc/WorldCompositionUtility.cpp
@@ -173,9 +173,9 @@ bool FWorldTileInfo::Read(const FString& InPackageFileName, FWorldTileInfo& OutI
 		FileReader->Seek(FileSummary.WorldTileInfoDataOffset);
 
 		//make sure the filereader gets the correct version number (it defaults to latest version)
-		FileReader->SetUEVer(FileSummary.GetFileVersionUE4());
+		FileReader->SetUEVer(FileSummary.GetFileVersionUE());
 		FileReader->SetEngineVer(FileSummary.SavedByEngineVersion);
-		FileReader->SetLicenseeUEVer(FileSummary.GetFileVersionLicenseeUE4());
+		FileReader->SetLicenseeUEVer(FileSummary.GetFileVersionLicenseeUE());
 		FileReader->SetCustomVersions(FileSummary.GetCustomVersionContainer());
 		
 		// Load the structure
diff --git a/Engine/Source/Runtime/CoreUObject/Private/Serialization/AsyncLoading.cpp b/Engine/Source/Runtime/CoreUObject/Private/Serialization/AsyncLoading.cpp
--- a/Engine/Source/Runtime/CoreUObject/Private/Serialization/AsyncLoading.cpp
+++ b/Engine/Source/Runtime/CoreUObject/Private/Serialization/AsyncLoading.cpp
@@ -7248,7 +7248,7 @@ void FAsyncArchive::ReadCallback(bool bWasCancelled, IAsyncReadRequest* Request)
 			FBufferReader Ar(Mem, FMath::Min<int64>(FMaxPackageSummarySize::Value, FileSize),/*bInFreeOnClose=*/ false, /*bIsPersistent=*/ true);
 			FPackageFileSummary Sum;
 			Ar << Sum;
-			if (Ar.IsError() || Sum.TotalHeaderSize > FileSize || Sum.GetFileVersionUE4() < VER_UE4_OLDEST_LOADABLE_PACKAGE)
+			if (Ar.IsError() || Sum.TotalHeaderSize > FileSize || Sum.GetFileVersionUE() < VER_UE4_OLDEST_LOADABLE_PACKAGE)
 			{
 				SetError();
 			}
diff --git a/Engine/Source/Runtime/CoreUObject/Private/UObject/LinkerLoad.cpp b/Engine/Source/Runtime/CoreUObject/Private/UObject/LinkerLoad.cpp
--- a/Engine/Source/Runtime/CoreUObject/Private/UObject/LinkerLoad.cpp
+++ b/Engine/Source/Runtime/CoreUObject/Private/UObject/LinkerLoad.cpp
@@ -1228,10 +1228,10 @@ FLinkerLoad::ELinkerStatus FLinkerLoad::SerializePackageFileSummaryInternal()
 	}
 
 	// Validate the summary.
-	if (Summary.GetFileVersionUE4() < VER_UE4_OLDEST_LOADABLE_PACKAGE)
+	if (Summary.GetFileVersionUE() < VER_UE4_OLDEST_LOADABLE_PACKAGE)
 	{
 		UE_LOG(LogLinker, Warning, TEXT("The file %s was saved by a previous version which is not backwards compatible with this one. Min Required Version: %i  Package Version: %i"),
-			*GetDebugName(), (int32)VER_UE4_OLDEST_LOADABLE_PACKAGE, Summary.GetFileVersionUE4());
+			*GetDebugName(), (int32)VER_UE4_OLDEST_LOADABLE_PACKAGE, Summary.GetFileVersionUE());
 		return LINKER_Failed;
 	}
 
@@ -1290,10 +1290,10 @@ FLinkerLoad::ELinkerStatus FLinkerLoad::SerializePackageFileSummaryInternal()
 	}
 
 	// Don't load packages that were saved with package version newer than the current one.
-	if ((Summary.GetFileVersionUE4() > GPackageFileUEVersion) || (Summary.GetFileVersionLicenseeUE4() > GPackageFileLicenseeUEVersion))
+	if ((Summary.GetFileVersionUE() > GPackageFileUEVersion) || (Summary.GetFileVersionLicenseeUE() > GPackageFileLicenseeUEVersion))
 	{
 		UE_LOG(LogLinker, Warning, TEXT("Unable to load package (%s) PackageVersion %i, MaxExpected %i : LicenseePackageVersion %i, MaxExpected %i."),
-			*GetDebugName(), Summary.GetFileVersionUE4(), GPackageFileUEVersion, Summary.GetFileVersionLicenseeUE4(), GPackageFileLicenseeUEVersion);
+			*GetDebugName(), Summary.GetFileVersionUE(), GPackageFileUEVersion, Summary.GetFileVersionLicenseeUE(), GPackageFileLicenseeUEVersion);
 		return LINKER_Failed;
 	}
 
@@ -1443,8 +1443,8 @@ FLinkerLoad::ELinkerStatus FLinkerLoad::UpdateFromPackageFileSummary()
 
 	const FCustomVersionContainer& SummaryVersions = Summary.GetCustomVersionContainer();
 
-	SetUEVer(Summary.GetFileVersionUE4());
-	SetLicenseeUEVer(Summary.GetFileVersionLicenseeUE4());
+	SetUEVer(Summary.GetFileVersionUE());
+	SetLicenseeUEVer(Summary.GetFileVersionLicenseeUE());
 	SetEngineVer(Summary.SavedByEngineVersion);
 	SetCustomVersions(SummaryVersions);
 
@@ -1466,8 +1466,8 @@ FLinkerLoad::ELinkerStatus FLinkerLoad::UpdateFromPackageFileSummary()
 	// Loader needs to be the same version.
 	if (Loader)
 	{
-		Loader->SetUEVer(Summary.GetFileVersionUE4());
-		Loader->SetLicenseeUEVer(Summary.GetFileVersionLicenseeUE4());
+		Loader->SetUEVer(Summary.GetFileVersionUE());
+		Loader->SetLicenseeUEVer(Summary.GetFileVersionLicenseeUE());
 		Loader->SetEngineVer(Summary.SavedByEngineVersion);
 		Loader->SetCustomVersions(SummaryVersions);
 	}
@@ -1505,8 +1505,8 @@ FLinkerLoad::ELinkerStatus FLinkerLoad::UpdateFromPackageFileSummary()
 #endif
 
 		// Remember the linker versions
-		LinkerRootPackage->LinkerPackageVersion = Summary.GetFileVersionUE4();
-		LinkerRootPackage->LinkerLicenseeVersion = Summary.GetFileVersionLicenseeUE4();
+		LinkerRootPackage->LinkerPackageVersion = Summary.GetFileVersionUE();
+		LinkerRootPackage->LinkerLicenseeVersion = Summary.GetFileVersionLicenseeUE();
 
 		// Only set the custom version if it is not already latest.
 		// If it is latest, we will compare against latest in GetLinkerCustomVersion
diff --git a/Engine/Source/Runtime/CoreUObject/Private/UObject/PackageFileSummary.cpp b/Engine/Source/Runtime/CoreUObject/Private/UObject/PackageFileSummary.cpp
--- a/Engine/Source/Runtime/CoreUObject/Private/UObject/PackageFileSummary.cpp
+++ b/Engine/Source/Runtime/CoreUObject/Private/UObject/PackageFileSummary.cpp
@@ -111,8 +111,8 @@ void operator<<(FStructuredArchive::FSlot Slot, FPackageFileSummary& Sum)
 				{
 					// we can't safely load more than this because the legacy version code differs in ways we can not predict.
 					// Make sure that the linker will fail to load with it.
-					Sum.FileVersionUE4 = 0;
-					Sum.FileVersionLicenseeUE4 = 0;
+					Sum.FileVersionUE = 0;
+					Sum.FileVersionLicenseeUE = 0;
 					return;
 				}
 
@@ -121,15 +121,15 @@ void operator<<(FStructuredArchive::FSlot Slot, FPackageFileSummary& Sum)
 					int32 LegacyUE3Version = 0;
 					Record << SA_VALUE(TEXT("LegacyUE3Version"), LegacyUE3Version);
 				}
-				Record << SA_VALUE(TEXT("FileVersionUE4"), Sum.FileVersionUE4);
-				Record << SA_VALUE(TEXT("FileVersionLicenseeUE4"), Sum.FileVersionLicenseeUE4);
+				Record << SA_VALUE(TEXT("FileVersionUE4"), Sum.FileVersionUE);
+				Record << SA_VALUE(TEXT("FileVersionLicenseeUE4"), Sum.FileVersionLicenseeUE);
 
 				if (LegacyFileVersion <= -2)
 				{
 					Sum.CustomVersionContainer.Serialize(Record.EnterField(SA_FIELD_NAME(TEXT("CustomVersions"))), GetCustomVersionFormatForArchive(LegacyFileVersion));
 				}
 
-				if (!Sum.FileVersionUE4 && !Sum.FileVersionLicenseeUE4)
+				if (!Sum.FileVersionUE && !Sum.FileVersionLicenseeUE)
 				{
 #if WITH_EDITOR
 					if (!GAllowUnversionedContentInEditor)
@@ -141,16 +141,16 @@ void operator<<(FStructuredArchive::FSlot Slot, FPackageFileSummary& Sum)
 #endif
 					// this file is unversioned, remember that, then use current versions
 					Sum.bUnversioned = true;
-					Sum.FileVersionUE4 = GPackageFileUEVersion;
-					Sum.FileVersionLicenseeUE4 = GPackageFileLicenseeUEVersion;
+					Sum.FileVersionUE = GPackageFileUEVersion;
+					Sum.FileVersionLicenseeUE = GPackageFileLicenseeUEVersion;
 					Sum.CustomVersionContainer = FCurrentCustomVersions::GetAll();
 				}
 			}
 			else
 			{
 				// This is probably an old UE3 file, make sure that the linker will fail to load with it.
-				Sum.FileVersionUE4 = 0;
-				Sum.FileVersionLicenseeUE4 = 0;
+				Sum.FileVersionUE = 0;
+				Sum.FileVersionLicenseeUE = 0;
 			}
 		}
 		else
@@ -170,8 +170,8 @@ void operator<<(FStructuredArchive::FSlot Slot, FPackageFileSummary& Sum)
 				// Must write out the last UE3 engine version, so that older versions identify it as new
 				int32 LegacyUE3Version = 864;
 				Record << SA_VALUE(TEXT("LegacyUE3Version"), LegacyUE3Version);
-				Record << SA_VALUE(TEXT("FileVersionUE4"), Sum.FileVersionUE4);
-				Record << SA_VALUE(TEXT("FileVersionLicenseeUE4"), Sum.FileVersionLicenseeUE4);
+				Record << SA_VALUE(TEXT("FileVersionUE4"), Sum.FileVersionUE);
+				Record << SA_VALUE(TEXT("FileVersionLicenseeUE4"), Sum.FileVersionLicenseeUE);
 
 				// Serialise custom version map.
 				Sum.CustomVersionContainer.Serialize(Record.EnterField(SA_FIELD_NAME(TEXT("CustomVersions"))));
@@ -196,12 +196,12 @@ void operator<<(FStructuredArchive::FSlot Slot, FPackageFileSummary& Sum)
 		Record << SA_VALUE(TEXT("NameCount"), Sum.NameCount) << SA_VALUE(TEXT("NameOffset"), Sum.NameOffset);
 		if (!BaseArchive.IsFilterEditorOnly())
 		{
-			if (BaseArchive.IsSaving() || Sum.FileVersionUE4 >= VER_UE4_ADDED_PACKAGE_SUMMARY_LOCALIZATION_ID)
+			if (BaseArchive.IsSaving() || Sum.FileVersionUE >= VER_UE4_ADDED_PACKAGE_SUMMARY_LOCALIZATION_ID)
 			{
 				Record << SA_VALUE(TEXT("LocalizationId"), Sum.LocalizationId);
 			}
 		}
-		if (Sum.FileVersionUE4 >= VER_UE4_SERIALIZE_TEXT_IN_PACKAGES)
+		if (Sum.FileVersionUE >= VER_UE4_SERIALIZE_TEXT_IN_PACKAGES)
 		{
 			Record << SA_VALUE(TEXT("GatherableTextDataCount"), Sum.GatherableTextDataCount) << SA_VALUE(TEXT("GatherableTextDataOffset"), Sum.GatherableTextDataOffset);
 		}
@@ -209,17 +209,17 @@ void operator<<(FStructuredArchive::FSlot Slot, FPackageFileSummary& Sum)
 		Record << SA_VALUE(TEXT("ImportCount"), Sum.ImportCount) << SA_VALUE(TEXT("ImportOffset"), Sum.ImportOffset);
 		Record << SA_VALUE(TEXT("DependsOffset"), Sum.DependsOffset);
 
-		if (BaseArchive.IsLoading() && (Sum.FileVersionUE4 < VER_UE4_OLDEST_LOADABLE_PACKAGE || Sum.FileVersionUE4 > GPackageFileUEVersion))
+		if (BaseArchive.IsLoading() && (Sum.FileVersionUE < VER_UE4_OLDEST_LOADABLE_PACKAGE || Sum.FileVersionUE > GPackageFileUEVersion))
 		{
 			return; // we can't safely load more than this because the below was different in older files.
 		}
 
-		if (BaseArchive.IsSaving() || Sum.FileVersionUE4 >= VER_UE4_ADD_STRING_ASSET_REFERENCES_MAP)
+		if (BaseArchive.IsSaving() || Sum.FileVersionUE >= VER_UE4_ADD_STRING_ASSET_REFERENCES_MAP)
 		{
 			Record << SA_VALUE(TEXT("SoftPackageReferencesCount"), Sum.SoftPackageReferencesCount) << SA_VALUE(TEXT("SoftPackageReferencesOffset"), Sum.SoftPackageReferencesOffset);
 		}
 
-		if (BaseArchive.IsSaving() || Sum.FileVersionUE4 >= VER_UE4_ADDED_SEARCHABLE_NAMES)
+		if (BaseArchive.IsSaving() || Sum.FileVersionUE >= VER_UE4_ADDED_SEARCHABLE_NAMES)
 		{
 			Record << SA_VALUE(TEXT("SearchableNamesOffset"), Sum.SearchableNamesOffset);
 		}
@@ -233,7 +233,7 @@ void operator<<(FStructuredArchive::FSlot Slot, FPackageFileSummary& Sum)
 #if WITH_EDITORONLY_DATA
 		if (!BaseArchive.IsFilterEditorOnly())
 		{
-			if (BaseArchive.IsSaving() || Sum.FileVersionUE4 >= VER_UE4_ADDED_PACKAGE_OWNER)
+			if (BaseArchive.IsSaving() || Sum.FileVersionUE >= VER_UE4_ADDED_PACKAGE_OWNER)
 			{
 				Record << SA_VALUE(TEXT("PersistentGuid"), Sum.PersistentGuid);
 			}
@@ -246,7 +246,7 @@ void operator<<(FStructuredArchive::FSlot Slot, FPackageFileSummary& Sum)
 			}
 
 			// The owner persistent guid was added in VER_UE4_ADDED_PACKAGE_OWNER but removed in the next version VER_UE4_NON_OUTER_PACKAGE_IMPORT
-			if (BaseArchive.IsLoading() && Sum.FileVersionUE4 >= VER_UE4_ADDED_PACKAGE_OWNER && Sum.FileVersionUE4 < VER_UE4_NON_OUTER_PACKAGE_IMPORT)
+			if (BaseArchive.IsLoading() && Sum.FileVersionUE >= VER_UE4_ADDED_PACKAGE_OWNER && Sum.FileVersionUE < VER_UE4_NON_OUTER_PACKAGE_IMPORT)
 			{
 				FGuid OwnerPersistentGuid;
 				Record << SA_VALUE(TEXT("OwnerPersistentGuid"), OwnerPersistentGuid);
@@ -272,7 +272,7 @@ void operator<<(FStructuredArchive::FSlot Slot, FPackageFileSummary& Sum)
 			}
 		}
 
-		if (Sum.GetFileVersionUE4() >= VER_UE4_ENGINE_VERSION_OBJECT)
+		if (Sum.GetFileVersionUE() >= VER_UE4_ENGINE_VERSION_OBJECT)
 		{
 			if (BaseArchive.IsCooking() || (BaseArchive.IsSaving() && !FEngineVersion::Current().HasChangelist()))
 			{
@@ -282,7 +282,7 @@ void operator<<(FStructuredArchive::FSlot Slot, FPackageFileSummary& Sum)
 			else
 			{
 				Record << SA_VALUE(TEXT("SavedByEngineVersion"), Sum.SavedByEngineVersion);
-				FixCorruptEngineVersion(Sum.GetFileVersionUE4(), Sum.SavedByEngineVersion);
+				FixCorruptEngineVersion(Sum.GetFileVersionUE(), Sum.SavedByEngineVersion);
 			}
 		}
 		else
@@ -296,7 +296,7 @@ void operator<<(FStructuredArchive::FSlot Slot, FPackageFileSummary& Sum)
 			}
 		}
 
-		if (Sum.GetFileVersionUE4() >= VER_UE4_PACKAGE_SUMMARY_HAS_COMPATIBLE_ENGINE_VERSION)
+		if (Sum.GetFileVersionUE() >= VER_UE4_PACKAGE_SUMMARY_HAS_COMPATIBLE_ENGINE_VERSION)
 		{
 			if (BaseArchive.IsCooking() || (BaseArchive.IsSaving() && !FEngineVersion::Current().HasChangelist()))
 			{
@@ -306,7 +306,7 @@ void operator<<(FStructuredArchive::FSlot Slot, FPackageFileSummary& Sum)
 			else
 			{
 				Record << SA_VALUE(TEXT("CompatibleWithEngineVersion"), Sum.CompatibleWithEngineVersion);
-				FixCorruptEngineVersion(Sum.GetFileVersionUE4(), Sum.CompatibleWithEngineVersion);
+				FixCorruptEngineVersion(Sum.GetFileVersionUE(), Sum.CompatibleWithEngineVersion);
 			}
 		}
 		else
@@ -321,7 +321,7 @@ void operator<<(FStructuredArchive::FSlot Slot, FPackageFileSummary& Sum)
 		if (!FCompression::VerifyCompressionFlagsValid(Sum.CompressionFlags))
 		{
 			UE_LOG(LogLinker, Warning, TEXT("Failed to read package file summary, the file \"%s\" has invalid compression flags (%d)."), *BaseArchive.GetArchiveName(), Sum.CompressionFlags);
-			Sum.FileVersionUE4 = VER_UE4_OLDEST_LOADABLE_PACKAGE - 1;
+			Sum.FileVersionUE = VER_UE4_OLDEST_LOADABLE_PACKAGE - 1;
 			return;
 		}
 
@@ -332,7 +332,7 @@ void operator<<(FStructuredArchive::FSlot Slot, FPackageFileSummary& Sum)
 		{
 			// this file has package level compression, we won't load it.
 			UE_LOG(LogLinker, Warning, TEXT("Failed to read package file summary, the file \"%s\" is has package level compression (and is probably cooked). These old files cannot be loaded in the editor."), *BaseArchive.GetArchiveName());
-			Sum.FileVersionUE4 = VER_UE4_OLDEST_LOADABLE_PACKAGE - 1;
+			Sum.FileVersionUE = VER_UE4_OLDEST_LOADABLE_PACKAGE - 1;
 			return; // we can't safely load more than this because we just changed the version to something it is not.
 		}
 
@@ -354,16 +354,16 @@ void operator<<(FStructuredArchive::FSlot Slot, FPackageFileSummary& Sum)
 		Record << SA_VALUE(TEXT("AssetRegistryDataOffset"), Sum.AssetRegistryDataOffset);
 		Record << SA_VALUE(TEXT("BulkDataStartOffset"), Sum.BulkDataStartOffset);
 
-		if (Sum.GetFileVersionUE4() >= VER_UE4_WORLD_LEVEL_INFO)
+		if (Sum.GetFileVersionUE() >= VER_UE4_WORLD_LEVEL_INFO)
 		{
 			Record << SA_VALUE(TEXT("WorldTileInfoDataOffset"), Sum.WorldTileInfoDataOffset);
 		}
 
-		if (Sum.GetFileVersionUE4() >= VER_UE4_CHANGED_CHUNKID_TO_BE_AN_ARRAY_OF_CHUNKIDS)
+		if (Sum.GetFileVersionUE() >= VER_UE4_CHANGED_CHUNKID_TO_BE_AN_ARRAY_OF_CHUNKIDS)
 		{
 			Record << SA_VALUE(TEXT("ChunkIDs"), Sum.ChunkIDs);
 		}
-		else if (Sum.GetFileVersionUE4() >= VER_UE4_ADDED_CHUNKID_TO_ASSETDATA_AND_UPACKAGE)
+		else if (Sum.GetFileVersionUE() >= VER_UE4_ADDED_CHUNKID_TO_ASSETDATA_AND_UPACKAGE)
 		{
 			// handle conversion of single ChunkID to an array of ChunkIDs
 			if (BaseArchive.IsLoading())
@@ -378,7 +378,7 @@ void operator<<(FStructuredArchive::FSlot Slot, FPackageFileSummary& Sum)
 				}
 			}
 		}
-		if (BaseArchive.IsSaving() || Sum.FileVersionUE4 >= VER_UE4_PRELOAD_DEPENDENCIES_IN_COOKED_EXPORTS)
+		if (BaseArchive.IsSaving() || Sum.FileVersionUE >= VER_UE4_PRELOAD_DEPENDENCIES_IN_COOKED_EXPORTS)
 		{
 			Record << SA_VALUE(TEXT("PreloadDependencyCount"), Sum.PreloadDependencyCount) << SA_VALUE(TEXT("PreloadDependencyOffset"), Sum.PreloadDependencyOffset);
 		}
diff --git a/Engine/Source/Runtime/CoreUObject/Public/AssetRegistry/AssetData.h b/Engine/Source/Runtime/CoreUObject/Public/AssetRegistry/AssetData.h
--- a/Engine/Source/Runtime/CoreUObject/Public/AssetRegistry/AssetData.h
+++ b/Engine/Source/Runtime/CoreUObject/Public/AssetRegistry/AssetData.h
@@ -634,10 +634,10 @@ public:
 	int64 DiskSize;
 
 	/** UE4 file version that the package was saved with */
-	int32 FileVersionUE4;
+	int32 FileVersionUE;
 
 	/** Licensee file version that the package was saved with */
-	int32 FileVersionLicenseeUE4;
+	int32 FileVersionLicenseeUE;
 
 private:
 	UE::AssetRegistry::FPackageCustomVersionsHandle CustomVersions;
@@ -648,8 +648,8 @@ public:
 
 	FAssetPackageData()
 		: DiskSize(0)
-		, FileVersionUE4(-1)
-		, FileVersionLicenseeUE4(-1)
+		, FileVersionUE(-1)
+		, FileVersionLicenseeUE(-1)
 		, Flags(0)
 	{
 	}
@@ -686,7 +686,7 @@ public:
 		CustomVersions = UE::AssetRegistry::FPackageCustomVersionsHandle::FindOrAdd(MoveTemp(InCustomVersions));
 	}
 
-	/** Whether the package was saved from a licensee executable, used to tell whether non-matching FileVersionLicenseeUE4 requires a resave */
+	/** Whether the package was saved from a licensee executable, used to tell whether non-matching FileVersionLicenseeUE requires a resave */
 	bool IsLicenseeVersion() const { return (Flags & FLAG_LICENSEE_VERSION) != 0; }
 	void SetIsLicenseeVersion(bool bValue) { Flags = (Flags & ~FLAG_LICENSEE_VERSION) | (bValue ? FLAG_LICENSEE_VERSION : 0); }
 
@@ -701,8 +701,8 @@ public:
 		Ar << PackageGuid;
 		PRAGMA_ENABLE_DEPRECATION_WARNINGS
 		Ar << CookedHash;
-		Ar << FileVersionUE4;
-		Ar << FileVersionLicenseeUE4;
+		Ar << FileVersionUE;
+		Ar << FileVersionLicenseeUE;
 		Ar << Flags;
 		Ar << CustomVersions;
 	}
diff --git a/Engine/Source/Runtime/CoreUObject/Public/UObject/PackageFileSummary.h b/Engine/Source/Runtime/CoreUObject/Public/UObject/PackageFileSummary.h
--- a/Engine/Source/Runtime/CoreUObject/Public/UObject/PackageFileSummary.h
+++ b/Engine/Source/Runtime/CoreUObject/Public/UObject/PackageFileSummary.h
@@ -52,9 +52,9 @@ struct FPackageFileSummary
 
 private:
 	/* UE4 file version */
-	int32		FileVersionUE4;
+	int32		FileVersionUE;
 	/* Licensee file version */
-	int32		FileVersionLicenseeUE4;
+	int32		FileVersionLicenseeUE;
 	/* Custom version numbers. Keyed off a unique tag for each custom component. */
 	FCustomVersionContainer CustomVersionContainer;
 
@@ -218,14 +218,25 @@ public:
 	/** Constructor */
 	COREUOBJECT_API FPackageFileSummary();
 
+	int32 GetFileVersionUE() const
+	{
+		return FileVersionUE;
+	}
+
+	int32 GetFileVersionLicenseeUE() const
+	{
+		return FileVersionLicenseeUE;
+	}
+	UE_DEPRECATED(5.0, "Use GetFileVersion instead")
 	int32 GetFileVersionUE4() const
 	{
-		return FileVersionUE4;
+		return GetFileVersionUE();
 	}
 
+	UE_DEPRECATED(5.0, "Use GetFileVersionLicenseeUE instead")
 	int32 GetFileVersionLicenseeUE4() const
 	{
-		return FileVersionLicenseeUE4;
+		return GetFileVersionLicenseeUE();
 	}
 
 	const FCustomVersionContainer& GetCustomVersionContainer() const
@@ -235,10 +246,10 @@ public:
 
 	void SetCustomVersionContainer(const FCustomVersionContainer& InContainer);
 
-	void SetFileVersions(const int32 EpicUE4, const int32 LicenseeUE4, const bool bInSaveUnversioned = false)
+	void SetFileVersions(const int32 EpicUE, const int32 LicenseeUE, const bool bInSaveUnversioned = false)
 	{
-		FileVersionUE4 = EpicUE4;
-		FileVersionLicenseeUE4 = LicenseeUE4;
+		FileVersionUE = EpicUE;
+		FileVersionLicenseeUE = LicenseeUE;
 		bUnversioned = bInSaveUnversioned;
 	}
 
