Removed more references to UE5

diff --git a/Engine/Source/Runtime/CoreUObject/Private/UObject/Class.cpp b/Engine/Source/Runtime/CoreUObject/Private/UObject/Class.cpp
--- a/Engine/Source/Runtime/CoreUObject/Private/UObject/Class.cpp
+++ b/Engine/Source/Runtime/CoreUObject/Private/UObject/Class.cpp
@@ -1033,7 +1033,7 @@ void UStruct::InitializeStruct(void* InDest, int32 ArrayDim/* = 1*/) const
 
 	int32 Stride = GetStructureSize();
 
-	//@todo UE4 optimize
+	//@todo UE optimize
 	FMemory::Memzero(Dest, 1 * Stride);
 
 	for (FProperty* Property = PropertyLink; Property; Property = Property->PropertyLinkNext)
@@ -3087,7 +3087,7 @@ void UScriptStruct::InitializeStruct(void* InDest, int32 ArrayDim) const
 
 	int32 Stride = GetStructureSize();
 
-	//@todo UE4 optimize
+	//@todo UE optimize
 	FMemory::Memzero(Dest, ArrayDim * Stride);
 
 	int32 InitializedSize = 0;
diff --git a/Engine/Source/Runtime/CoreUObject/Private/UObject/GarbageCollection.cpp b/Engine/Source/Runtime/CoreUObject/Private/UObject/GarbageCollection.cpp
--- a/Engine/Source/Runtime/CoreUObject/Private/UObject/GarbageCollection.cpp
+++ b/Engine/Source/Runtime/CoreUObject/Private/UObject/GarbageCollection.cpp
@@ -1578,7 +1578,7 @@ bool IncrementalDestroyGarbage(bool bUseTimeLimit, float TimeLimit)
 			FUObjectItem* ObjectItem = GUnreachableObjects[GObjCurrentPurgeObjectIndex];
 			checkSlow(ObjectItem);
 
-			//@todo UE4 - A prefetch was removed here. Re-add it. It wasn't right anyway, since it was ten items ahead and the consoles on have 8 prefetch slots
+			//@todo UE - A prefetch was removed here. Re-add it. It wasn't right anyway, since it was ten items ahead and the consoles on have 8 prefetch slots
 
 			check(ObjectItem->IsUnreachable());
 			{
@@ -2124,7 +2124,7 @@ bool UnhashUnreachableObjects(bool bUseTimeLimit, float TimeLimit)
 
 	while (GUnrechableObjectIndex < GUnreachableObjects.Num())
 	{
-		//@todo UE4 - A prefetch was removed here. Re-add it. It wasn't right anyway, since it was ten items ahead and the consoles on have 8 prefetch slots
+		//@todo UE - A prefetch was removed here. Re-add it. It wasn't right anyway, since it was ten items ahead and the consoles on have 8 prefetch slots
 
 		FUObjectItem* ObjectItem = GUnreachableObjects[GUnrechableObjectIndex++];
 		{
@@ -2221,7 +2221,7 @@ void UObject::CallAddReferencedObjects(FReferenceCollector& Collector)
 void UObject::AddReferencedObjects(UObject* This, FReferenceCollector& Collector)
 {
 #if WITH_EDITOR
-	//@todo UE4 - This seems to be required and it should not be. Seems to be related to the texture streamer.
+	//@todo UE - This seems to be required and it should not be. Seems to be related to the texture streamer.
 	FLinkerLoad* LinkerLoad = This->GetLinker();
 	if (LinkerLoad)
 	{
diff --git a/Engine/Source/Runtime/CoreUObject/Private/UObject/PropertyArray.cpp b/Engine/Source/Runtime/CoreUObject/Private/UObject/PropertyArray.cpp
--- a/Engine/Source/Runtime/CoreUObject/Private/UObject/PropertyArray.cpp
+++ b/Engine/Source/Runtime/CoreUObject/Private/UObject/PropertyArray.cpp
@@ -656,7 +656,7 @@ void FArrayProperty::DestroyValueInternal( void* Dest ) const
 	FScriptArrayHelper ArrayHelper(this, Dest);
 	ArrayHelper.EmptyValues();
 
-	//@todo UE4 potential double destroy later from this...would be ok for a script array, but still
+	//@todo UE potential double destroy later from this...would be ok for a script array, but still
 	ArrayHelper.DestroyContainer_Unsafe();
 }
 bool FArrayProperty::PassCPPArgsByRef() const
diff --git a/Engine/Source/Runtime/CoreUObject/Private/UObject/PropertyMap.cpp b/Engine/Source/Runtime/CoreUObject/Private/UObject/PropertyMap.cpp
--- a/Engine/Source/Runtime/CoreUObject/Private/UObject/PropertyMap.cpp
+++ b/Engine/Source/Runtime/CoreUObject/Private/UObject/PropertyMap.cpp
@@ -15,7 +15,7 @@
 // WARNING: This should always be the last include in any file that needs it (except .generated.h)
 #include "UObject/UndefineUPropertyMacros.h"
 
-namespace UE4MapProperty_Private
+namespace UEMapProperty_Private
 {
 	/**
 	 * Checks if any of the pairs in the map compare equal to the one passed.
@@ -283,7 +283,7 @@ bool FMapProperty::Identical(const void* A, const void* B, uint32 PortFlags) con
 		return false;
 	}
 
-	return UE4MapProperty_Private::IsPermutation(MapHelperA, MapHelperB, PortFlags);
+	return UEMapProperty_Private::IsPermutation(MapHelperA, MapHelperB, PortFlags);
 }
 
 void FMapProperty::GetPreloadDependencies(TArray<UObject*>& OutDeps)
@@ -888,7 +888,7 @@ void FMapProperty::DestroyValueInternal(void* Data) const
 	FScriptMapHelper MapHelper(this, Data);
 	MapHelper.EmptyValues();
 
-	//@todo UE4 potential double destroy later from this...would be ok for a script map, but still
+	//@todo UE potential double destroy later from this...would be ok for a script map, but still
 	((FScriptMap*)Data)->~FScriptMap();
 }
 
diff --git a/Engine/Source/Runtime/CoreUObject/Private/UObject/PropertySet.cpp b/Engine/Source/Runtime/CoreUObject/Private/UObject/PropertySet.cpp
--- a/Engine/Source/Runtime/CoreUObject/Private/UObject/PropertySet.cpp
+++ b/Engine/Source/Runtime/CoreUObject/Private/UObject/PropertySet.cpp
@@ -13,7 +13,7 @@
 // WARNING: This should always be the last include in any file that needs it (except .generated.h)
 #include "UObject/UndefineUPropertyMacros.h"
 
-namespace UE4SetProperty_Private
+namespace UESetProperty_Private
 {
 	/**
 	 * Checks if any of the elements in the set compare equal to the one passed.
@@ -242,7 +242,7 @@ bool FSetProperty::Identical(const void* A, const void* B, uint32 PortFlags) con
 		return false;
 	}
 
-	return UE4SetProperty_Private::IsPermutation(SetHelperA, SetHelperB, PortFlags);
+	return UESetProperty_Private::IsPermutation(SetHelperA, SetHelperB, PortFlags);
 }
 
 void FSetProperty::GetPreloadDependencies(TArray<UObject*>& OutDeps)
@@ -763,7 +763,7 @@ void FSetProperty::DestroyValueInternal(void* Data) const
 	FScriptSetHelper SetHelper(this, Data);
 	SetHelper.EmptyElements();
 
-	//@todo UE4 potential double destroy later from this...would be ok for a script set, but still
+	//@todo UE potential double destroy later from this...would be ok for a script set, but still
 	((FScriptSet*)Data)->~FScriptSet();
 }
 
diff --git a/Engine/Source/Runtime/CoreUObject/Private/UObject/UObjectGlobals.cpp b/Engine/Source/Runtime/CoreUObject/Private/UObject/UObjectGlobals.cpp
--- a/Engine/Source/Runtime/CoreUObject/Private/UObject/UObjectGlobals.cpp
+++ b/Engine/Source/Runtime/CoreUObject/Private/UObject/UObjectGlobals.cpp
@@ -1356,7 +1356,7 @@ UPackage* LoadPackageInternal(UPackage* InOuter, const FPackagePath& PackagePath
 			Linker->FlushCache();
 		}
 
-		// With UE4 and single asset per package, we load so many packages that some platforms will run out
+		// With UE and single asset per package, we load so many packages that some platforms will run out
 		// of file handles. So, this will close the package, but just things like bulk data loading will
 		// fail, so we only currently do this when loading on consoles.
 		// The only exception here is when we're in the middle of async loading where we can't reset loaders yet. This should only happen when
@@ -2346,7 +2346,7 @@ void NotifyConstructedDuringAsyncLoading(UObject* Object, bool bSubObject);
 * For object overwrites, the class may want to persist some info over the re-initialize
 * this is only used for classes in the script compiler
 **/
-//@todo UE4 this is clunky
+//@todo UE this is clunky
 static FRestoreForUObjectOverwrite* ObjectRestoreAfterInitProps = NULL;  
 
 COREUOBJECT_API bool GOutputCookingWarnings = false;
@@ -2607,7 +2607,7 @@ UObject* StaticAllocateObject
 	return Obj;
 }
 
-//@todo UE4 - move this stuff to UnObj.cpp or something
+//@todo UE - move this stuff to UnObj.cpp or something
 
 void UObject::PostReinitProperties()
 {
diff --git a/Engine/Source/Runtime/CoreUObject/Private/UObject/UObjectLinker.cpp b/Engine/Source/Runtime/CoreUObject/Private/UObject/UObjectLinker.cpp
--- a/Engine/Source/Runtime/CoreUObject/Private/UObject/UObjectLinker.cpp
+++ b/Engine/Source/Runtime/CoreUObject/Private/UObject/UObjectLinker.cpp
@@ -13,7 +13,7 @@
 
 DEFINE_LOG_CATEGORY_STATIC(LogUObjectLinker, Log, All);
 
-//@todo UE4 Console - Check that the mapping of UObjects to linkers is sparse and that we aren't spending a ton of time with these lookups.
+//@todo UE Console - Check that the mapping of UObjects to linkers is sparse and that we aren't spending a ton of time with these lookups.
 
 struct FLinkerIndexPair
 {
diff --git a/Engine/Source/Runtime/CoreUObject/Public/Templates/Casts.h b/Engine/Source/Runtime/CoreUObject/Public/Templates/Casts.h
--- a/Engine/Source/Runtime/CoreUObject/Public/Templates/Casts.h
+++ b/Engine/Source/Runtime/CoreUObject/Public/Templates/Casts.h
@@ -514,7 +514,7 @@ private:
 
 
 
-namespace UE4Casts_Private
+namespace UECasts_Private
 {
 	template <typename To, typename From>
 	FORCEINLINE typename TEnableIf<TAnd<TIsPointer<To>, TAnd<TIsCastableToPointer<typename TRemovePointer<To>::Type>, TIsCastable<From>>>::Value, To>::Type DynamicCast(From* Arg)
@@ -559,4 +559,4 @@ namespace UE4Casts_Private
 	}
 }
 
-#define dynamic_cast UE4Casts_Private::DynamicCast
+#define dynamic_cast UECasts_Private::DynamicCast
diff --git a/Engine/Source/Runtime/CoreUObject/Public/UObject/NoExportTypes.h b/Engine/Source/Runtime/CoreUObject/Public/UObject/NoExportTypes.h
--- a/Engine/Source/Runtime/CoreUObject/Public/UObject/NoExportTypes.h
+++ b/Engine/Source/Runtime/CoreUObject/Public/UObject/NoExportTypes.h
@@ -1646,7 +1646,7 @@ struct FTestUninitializedScriptStructMembersTest
 };
 
 /**
- * Direct base class for all UE4 objects
+ * Direct base class for all UE objects
  * @note The full C++ class is located here: Engine\Source\Runtime\CoreUObject\Public\UObject\Object.h
  */
 UCLASS(abstract, noexport)
diff --git a/Engine/Source/Runtime/CoreUObject/Public/UObject/Object.h b/Engine/Source/Runtime/CoreUObject/Public/UObject/Object.h
--- a/Engine/Source/Runtime/CoreUObject/Public/UObject/Object.h
+++ b/Engine/Source/Runtime/CoreUObject/Public/UObject/Object.h
@@ -3,7 +3,7 @@
 #pragma once
 
 /*=============================================================================
-	Object.h: Direct base class for all UE4 objects
+	Object.h: Direct base class for all UE objects
 =============================================================================*/
 
 #include "CoreMinimal.h"
@@ -43,7 +43,7 @@ namespace ECastCheckedType
 };
 
 /** 
- * The base class of all UE4 objects. The type of an object is defined by its UClass.
+ * The base class of all UE objects. The type of an object is defined by its UClass.
  * This provides support functions for creating and using objects, and virtual functions that should be overridden in child classes.
  * 
  * @see https://docs.unrealengine.com/en-us/Programming/UnrealArchitecture/Objects
@@ -425,7 +425,7 @@ private:
 	 * Test the selection state of a UObject
 	 *
 	 * @return		true if the object is selected, false otherwise.
-	 * @todo UE4 this doesn't belong here, but it doesn't belong anywhere else any better
+	 * @todo UE	 this doesn't belong here, but it doesn't belong anywhere else any better
 	 */
 	virtual bool IsSelectedInEditor() const;
 
@@ -903,7 +903,7 @@ public:
 	 * Test the selection state of a UObject
 	 *
 	 * @return		true if the object is selected, false otherwise.
-	 * @todo UE4 this doesn't belong here, but it doesn't belong anywhere else any better
+	 * @todo UE this doesn't belong here, but it doesn't belong anywhere else any better
 	 */
 	bool IsSelected() const;
 
diff --git a/Engine/Source/Runtime/CoreUObject/Public/UObject/ObjectMacros.h b/Engine/Source/Runtime/CoreUObject/Public/UObject/ObjectMacros.h
--- a/Engine/Source/Runtime/CoreUObject/Public/UObject/ObjectMacros.h
+++ b/Engine/Source/Runtime/CoreUObject/Public/UObject/ObjectMacros.h
@@ -19,7 +19,7 @@ typedef	uint64 ScriptPointerType;
 
 
 #if PLATFORM_VTABLE_AT_END_OF_CLASS
-#error "not supported in UE4"
+#error "not supported in UE"
 #endif
 
 #if HACK_HEADER_GENERATOR 
diff --git a/Engine/Source/Runtime/CoreUObject/Public/UObject/StrongObjectPtr.h b/Engine/Source/Runtime/CoreUObject/Public/UObject/StrongObjectPtr.h
--- a/Engine/Source/Runtime/CoreUObject/Public/UObject/StrongObjectPtr.h
+++ b/Engine/Source/Runtime/CoreUObject/Public/UObject/StrongObjectPtr.h
@@ -7,7 +7,7 @@
 #include "Templates/PointerIsConvertibleFromTo.h"
 #include "Templates/UniquePtr.h"
 
-namespace UE4StrongObjectPtr_Private
+namespace UEStrongObjectPtr_Private
 {
 	class FInternalReferenceCollector : public FGCObject
 	{
@@ -47,7 +47,7 @@ namespace UE4StrongObjectPtr_Private
 
 		virtual FString GetReferencerName() const override
 		{
-			return "UE4StrongObjectPtr_Private::FInternalReferenceCollector";
+			return "UEStrongObjectPtr_Private::FInternalReferenceCollector";
 		}
 
 	private:
@@ -149,7 +149,7 @@ public:
 			else
 			{
 				// Lazily create the ReferenceCollector to allow TStrongObjectPtr to be used during static initialization
-				ReferenceCollector = MakeUnique<UE4StrongObjectPtr_Private::FInternalReferenceCollector>(InNewObject);
+				ReferenceCollector = MakeUnique<UEStrongObjectPtr_Private::FInternalReferenceCollector>(InNewObject);
 			}
 		}
 		else
@@ -166,7 +166,7 @@ public:
 	}
 
 private:
-	TUniquePtr<UE4StrongObjectPtr_Private::FInternalReferenceCollector> ReferenceCollector;
+	TUniquePtr<UEStrongObjectPtr_Private::FInternalReferenceCollector> ReferenceCollector;
 };
 
 template <typename LHSObjectType, typename RHSObjectType>
diff --git a/Engine/Source/Runtime/CoreUObject/Public/UObject/UObjectAnnotation.h b/Engine/Source/Runtime/CoreUObject/Public/UObject/UObjectAnnotation.h
--- a/Engine/Source/Runtime/CoreUObject/Public/UObject/UObjectAnnotation.h
+++ b/Engine/Source/Runtime/CoreUObject/Public/UObject/UObjectAnnotation.h
@@ -435,7 +435,7 @@ template <> struct TIsPODType<FBoolAnnotation> { enum { Value = true }; };
 * FUObjectAnnotationSparseBool is a specialization of FUObjectAnnotationSparse for bools, slow, temporary, editor only, external 
 * or other low priority bools about UObjects.
 *
-* @todo UE4 this should probably be reimplemented from scratch as a TSet instead of essentially a map to a value that is always true anyway.
+* @todo UE this should probably be reimplemented from scratch as a TSet instead of essentially a map to a value that is always true anyway.
 **/
 class FUObjectAnnotationSparseBool : private FUObjectAnnotationSparse<FBoolAnnotation,true>
 {
diff --git a/Engine/Source/Runtime/CoreUObject/Public/UObject/UObjectArray.h b/Engine/Source/Runtime/CoreUObject/Public/UObject/UObjectArray.h
--- a/Engine/Source/Runtime/CoreUObject/Public/UObject/UObjectArray.h
+++ b/Engine/Source/Runtime/CoreUObject/Public/UObject/UObjectArray.h
@@ -1014,7 +1014,7 @@ public:
 		 */
 		FORCEINLINE bool Advance()
 		{
-			//@todo UE4 check this for LHS on Index on consoles
+			//@todo UE check this for LHS on Index on consoles
 			FUObjectItem* NextObject = nullptr;
 			CurrentObject = nullptr;
 			while(++Index < Array.GetObjectArrayNum())
diff --git a/Engine/Source/Runtime/CoreUObject/Public/UObject/UObjectIterator.h b/Engine/Source/Runtime/CoreUObject/Public/UObject/UObjectIterator.h
--- a/Engine/Source/Runtime/CoreUObject/Public/UObject/UObjectIterator.h
+++ b/Engine/Source/Runtime/CoreUObject/Public/UObject/UObjectIterator.h
@@ -335,7 +335,7 @@ protected:
 	 */
 	FORCEINLINE bool Advance()
 	{
-		//@todo UE4 check this for LHS on Index on consoles
+		//@todo UE check this for LHS on Index on consoles
 		while(++Index < ObjectArray.Num())
 		{
 			if (GetObject())
diff --git a/Engine/Source/Runtime/CoreUObject/Public/UObject/UnrealType.h b/Engine/Source/Runtime/CoreUObject/Public/UObject/UnrealType.h
--- a/Engine/Source/Runtime/CoreUObject/Public/UObject/UnrealType.h
+++ b/Engine/Source/Runtime/CoreUObject/Public/UObject/UnrealType.h
@@ -94,7 +94,7 @@ enum class EPropertyObjectReferenceType : uint32
 };
 ENUM_CLASS_FLAGS(EPropertyObjectReferenceType);
 
-namespace UE4Property_Private { class FProperty_DoNotUse; }
+namespace UEProperty_Private { class FProperty_DoNotUse; }
 
 /**
  * Helper class that calculates the maximum stack entry size required by Garbage Collector Token Stream for each class 
@@ -254,7 +254,7 @@ private:
 
 protected:
 	friend class FMapProperty;
-	friend class UE4Property_Private::FProperty_DoNotUse;
+	friend class UEProperty_Private::FProperty_DoNotUse;
 
 	/** Set the alignment offset for this property - added for FMapProperty */
 	void SetOffset_Internal(int32 NewOffset);
@@ -945,7 +945,7 @@ public:
 	static const TCHAR* ReadToken( const TCHAR* Buffer, FStringBuilderBase& Out, bool DottedNames = false);
 };
 
-namespace UE4Property_Private
+namespace UEProperty_Private
 {
 	/** FProperty methods FOR INTERNAL USE ONLY -- only authorized users should be making use of this. -- DO NOT USE! */
 	class FProperty_DoNotUse
@@ -1871,7 +1871,7 @@ class COREUOBJECT_API FUInt64Property : public TProperty_Numeric<uint64>
 	Aliases for implicitly-sized integer properties.
 -----------------------------------------------------------------------------*/
 
-namespace UE4Types_Private
+namespace UETypes_Private
 {
 	template <typename IntType> struct TIntegerPropertyMapping;
 
@@ -1885,8 +1885,8 @@ namespace UE4Types_Private
 	template <> struct TIntegerPropertyMapping<uint64> { typedef FUInt64Property Type; };
 }
 
-typedef UE4Types_Private::TIntegerPropertyMapping<signed int>::Type UUnsizedIntProperty;
-typedef UE4Types_Private::TIntegerPropertyMapping<unsigned int>::Type UUnsizedFIntProperty;
+typedef UETypes_Private::TIntegerPropertyMapping<signed int>::Type UUnsizedIntProperty;
+typedef UETypes_Private::TIntegerPropertyMapping<unsigned int>::Type UUnsizedFIntProperty;
 
 
 /*-----------------------------------------------------------------------------
