- Merge CL's 15977926, 15978323, 15978393, 15978401 and 15981176 from //UE5/RES to //UE5/Main
- Moves CoreOnline headers into their own module

#tests built EdTest, QAGame, and one internal project in DevEd Win64

diff --git a/Engine/Plugins/NetcodeUnitTest/NetcodeUnitTest/Source/NetcodeUnitTest/Private/Net/NUTUtilNet.cpp b/Engine/Plugins/NetcodeUnitTest/NetcodeUnitTest/Source/NetcodeUnitTest/Private/Net/NUTUtilNet.cpp
--- a/Engine/Plugins/NetcodeUnitTest/NetcodeUnitTest/Source/NetcodeUnitTest/Private/Net/NUTUtilNet.cpp
+++ b/Engine/Plugins/NetcodeUnitTest/NetcodeUnitTest/Source/NetcodeUnitTest/Private/Net/NUTUtilNet.cpp
@@ -2,7 +2,7 @@
 
 #include "Net/NUTUtilNet.h"
 
-#include "UObject/CoreOnline.h"
+#include "Online/CoreOnline.h"
 #include "GameFramework/OnlineReplStructs.h"
 #include "Engine/Engine.h"
 #include "PacketHandlers/StatelessConnectHandlerComponent.h"
diff --git a/Engine/Plugins/Online/OnlineFramework/Source/Lobby/Lobby.Build.cs b/Engine/Plugins/Online/OnlineFramework/Source/Lobby/Lobby.Build.cs
--- a/Engine/Plugins/Online/OnlineFramework/Source/Lobby/Lobby.Build.cs
+++ b/Engine/Plugins/Online/OnlineFramework/Source/Lobby/Lobby.Build.cs
@@ -15,8 +15,9 @@ public class Lobby : ModuleRules
 			);
 
 		PrivateDependencyModuleNames.AddRange(
-			new string[] { 
+			new string[] {
 				"Core",
+				"CoreOnline",
 				"CoreUObject",
 				"Engine",
 				"OnlineSubsystem",
diff --git a/Engine/Plugins/Online/OnlineFramework/Source/LoginFlow/LoginFlow.Build.cs b/Engine/Plugins/Online/OnlineFramework/Source/LoginFlow/LoginFlow.Build.cs
--- a/Engine/Plugins/Online/OnlineFramework/Source/LoginFlow/LoginFlow.Build.cs
+++ b/Engine/Plugins/Online/OnlineFramework/Source/LoginFlow/LoginFlow.Build.cs
@@ -12,6 +12,7 @@ public class LoginFlow : ModuleRules
 			new string[]
 			{
 				"Core",
+				"CoreOnline",
 				"CoreUObject",
 				"InputCore",
 				"Slate",
diff --git a/Engine/Plugins/Online/OnlineFramework/Source/Party/Public/Chat/SocialGroupChannel.h b/Engine/Plugins/Online/OnlineFramework/Source/Party/Public/Chat/SocialGroupChannel.h
--- a/Engine/Plugins/Online/OnlineFramework/Source/Party/Public/Chat/SocialGroupChannel.h
+++ b/Engine/Plugins/Online/OnlineFramework/Source/Party/Public/Chat/SocialGroupChannel.h
@@ -3,7 +3,7 @@
 #pragma once
 
 #include "Chat/SocialChatChannel.h"
-#include "UObject/CoreOnline.h"
+#include "Online/CoreOnline.h"
 #include "Interfaces/OnlineGroupsInterface.h"
 #include "SocialGroupChannel.generated.h"
 
diff --git a/Engine/Plugins/Online/OnlineFramework/Source/PlayTimeLimit/PlayTimeLimit.Build.cs b/Engine/Plugins/Online/OnlineFramework/Source/PlayTimeLimit/PlayTimeLimit.Build.cs
--- a/Engine/Plugins/Online/OnlineFramework/Source/PlayTimeLimit/PlayTimeLimit.Build.cs
+++ b/Engine/Plugins/Online/OnlineFramework/Source/PlayTimeLimit/PlayTimeLimit.Build.cs
@@ -12,6 +12,7 @@ public class PlayTimeLimit : ModuleRules
 			new string[]
 			{
 				"Core",
+				"CoreOnline",
 				"CoreUObject",
 			}
 		);
diff --git a/Engine/Plugins/Online/OnlineFramework/Source/PlayTimeLimit/Private/PlayTimeLimitImpl.cpp b/Engine/Plugins/Online/OnlineFramework/Source/PlayTimeLimit/Private/PlayTimeLimitImpl.cpp
--- a/Engine/Plugins/Online/OnlineFramework/Source/PlayTimeLimit/Private/PlayTimeLimitImpl.cpp
+++ b/Engine/Plugins/Online/OnlineFramework/Source/PlayTimeLimit/Private/PlayTimeLimitImpl.cpp
@@ -10,7 +10,7 @@
 #include "PlayTimeLimitUser.h"
 #include "PlayTimeLimitUserMock.h"
 
-#include "UObject/CoreOnline.h"
+#include "Online/CoreOnline.h"
 
 // FPlayTimeLimitImpl
 FPlayTimeLimitImpl::FPlayTimeLimitImpl()
diff --git a/Engine/Plugins/Online/OnlineFramework/Source/PlayTimeLimit/Public/PlayTimeLimitUser.h b/Engine/Plugins/Online/OnlineFramework/Source/PlayTimeLimit/Public/PlayTimeLimitUser.h
--- a/Engine/Plugins/Online/OnlineFramework/Source/PlayTimeLimit/Public/PlayTimeLimitUser.h
+++ b/Engine/Plugins/Online/OnlineFramework/Source/PlayTimeLimit/Public/PlayTimeLimitUser.h
@@ -3,7 +3,7 @@
 #pragma once
 
 #include "CoreMinimal.h"
-#include "UObject/CoreOnline.h"
+#include "Online/CoreOnline.h"
 
 /**
  * Information about a user we are observing the play time and instituting limits for
diff --git a/Engine/Plugins/Online/OnlineFramework/Source/Qos/Private/QosStats.h b/Engine/Plugins/Online/OnlineFramework/Source/Qos/Private/QosStats.h
--- a/Engine/Plugins/Online/OnlineFramework/Source/Qos/Private/QosStats.h
+++ b/Engine/Plugins/Online/OnlineFramework/Source/Qos/Private/QosStats.h
@@ -3,7 +3,7 @@
 #pragma once
 
 #include "CoreMinimal.h"
-#include "UObject/CoreOnline.h"
+#include "Online/CoreOnline.h"
 #include "Serialization/JsonSerializerMacros.h"
 
 #include "QosRegionManager.h"
diff --git a/Engine/Plugins/Online/OnlineSubsystem/Source/OnlineSubsystem.Build.cs b/Engine/Plugins/Online/OnlineSubsystem/Source/OnlineSubsystem.Build.cs
--- a/Engine/Plugins/Online/OnlineSubsystem/Source/OnlineSubsystem.Build.cs
+++ b/Engine/Plugins/Online/OnlineSubsystem/Source/OnlineSubsystem.Build.cs
@@ -11,6 +11,7 @@ public class OnlineSubsystem : ModuleRules
 		PublicDependencyModuleNames.AddRange(
 			new string[] {
 				"Json",
+				"CoreOnline",
 				"SignalProcessing"
 			}
 		);
@@ -24,13 +25,10 @@ public class OnlineSubsystem : ModuleRules
 		// OnlineSubsystem cannot depend on Engine!
 		PrivateDependencyModuleNames.AddRange(
 			new string[] { 
-				"Core", 
+				"Core",
 				"CoreUObject",
-				"ImageCore",
 				"Sockets",
 				"JsonUtilities",
-				"AudioMixerCore",
-				"SignalProcessing",
 			}
 		);
 	}
diff --git a/Engine/Plugins/Online/OnlineSubsystem/Source/Public/Interfaces/AvatarProvider.h b/Engine/Plugins/Online/OnlineSubsystem/Source/Public/Interfaces/AvatarProvider.h
--- a/Engine/Plugins/Online/OnlineSubsystem/Source/Public/Interfaces/AvatarProvider.h
+++ b/Engine/Plugins/Online/OnlineSubsystem/Source/Public/Interfaces/AvatarProvider.h
@@ -4,7 +4,7 @@
 
 #include "CoreMinimal.h"
 #include "Features/IModularFeature.h"
-#include "UObject/CoreOnlineFwd.h"
+#include "Online/CoreOnlineFwd.h"
 
 
 struct FAvatarInfo : public TSharedFromThis<FAvatarInfo>
diff --git a/Engine/Plugins/Online/OnlineSubsystem/Source/Public/Interfaces/OnlineChatInterface.h b/Engine/Plugins/Online/OnlineSubsystem/Source/Public/Interfaces/OnlineChatInterface.h
--- a/Engine/Plugins/Online/OnlineSubsystem/Source/Public/Interfaces/OnlineChatInterface.h
+++ b/Engine/Plugins/Online/OnlineSubsystem/Source/Public/Interfaces/OnlineChatInterface.h
@@ -3,7 +3,7 @@
 #pragma once
 
 #include "CoreMinimal.h"
-#include "UObject/CoreOnline.h"
+#include "Online/CoreOnline.h"
 #include "OnlineDelegateMacros.h"
 
 ONLINESUBSYSTEM_API DECLARE_LOG_CATEGORY_EXTERN(LogOnlineChat, Log, All);
diff --git a/Engine/Plugins/Online/OnlineSubsystem/Source/Public/Interfaces/OnlineExternalUIInterface.h b/Engine/Plugins/Online/OnlineSubsystem/Source/Public/Interfaces/OnlineExternalUIInterface.h
--- a/Engine/Plugins/Online/OnlineSubsystem/Source/Public/Interfaces/OnlineExternalUIInterface.h
+++ b/Engine/Plugins/Online/OnlineSubsystem/Source/Public/Interfaces/OnlineExternalUIInterface.h
@@ -3,7 +3,7 @@
 #pragma once
 
 #include "CoreMinimal.h"
-#include "UObject/CoreOnline.h"
+#include "Online/CoreOnlineFwd.h"
 #include "OnlineDelegateMacros.h"
 #include "Interfaces/OnlineMessageInterface.h"
 #include "OnlineError.h"
diff --git a/Engine/Plugins/Online/OnlineSubsystem/Source/Public/Interfaces/OnlineGameMatchesInterface.h b/Engine/Plugins/Online/OnlineSubsystem/Source/Public/Interfaces/OnlineGameMatchesInterface.h
--- a/Engine/Plugins/Online/OnlineSubsystem/Source/Public/Interfaces/OnlineGameMatchesInterface.h
+++ b/Engine/Plugins/Online/OnlineSubsystem/Source/Public/Interfaces/OnlineGameMatchesInterface.h
@@ -3,7 +3,7 @@
 #pragma once
 
 #include "CoreMinimal.h"
-#include "UObject/CoreOnlineFwd.h"
+#include "Online/CoreOnlineFwd.h"
 #include "OnlineDelegateMacros.h"
 #include "OnlineStatsInterface.h"
 
diff --git a/Engine/Plugins/Online/OnlineSubsystem/Source/Public/Interfaces/OnlineGroupsInterface.h b/Engine/Plugins/Online/OnlineSubsystem/Source/Public/Interfaces/OnlineGroupsInterface.h
--- a/Engine/Plugins/Online/OnlineSubsystem/Source/Public/Interfaces/OnlineGroupsInterface.h
+++ b/Engine/Plugins/Online/OnlineSubsystem/Source/Public/Interfaces/OnlineGroupsInterface.h
@@ -2,7 +2,7 @@
 #pragma once
 
 #include "CoreMinimal.h"
-#include "UObject/CoreOnline.h"
+#include "Online/CoreOnline.h"
 #include "OnlineSubsystemTypes.h"
 
 class Error;
diff --git a/Engine/Plugins/Online/OnlineSubsystem/Source/Public/Interfaces/OnlineIdentityInterface.h b/Engine/Plugins/Online/OnlineSubsystem/Source/Public/Interfaces/OnlineIdentityInterface.h
--- a/Engine/Plugins/Online/OnlineSubsystem/Source/Public/Interfaces/OnlineIdentityInterface.h
+++ b/Engine/Plugins/Online/OnlineSubsystem/Source/Public/Interfaces/OnlineIdentityInterface.h
@@ -3,7 +3,7 @@
 #pragma once
 
 #include "CoreMinimal.h"
-#include "UObject/CoreOnline.h"
+#include "Online/CoreOnline.h"
 #include "OnlineSubsystemTypes.h"
 #include "OnlineDelegateMacros.h"
 #include "OnlineIdentityErrors.h"
diff --git a/Engine/Plugins/Online/OnlineSubsystem/Source/Public/Interfaces/OnlineLeaderboardInterface.h b/Engine/Plugins/Online/OnlineSubsystem/Source/Public/Interfaces/OnlineLeaderboardInterface.h
--- a/Engine/Plugins/Online/OnlineSubsystem/Source/Public/Interfaces/OnlineLeaderboardInterface.h
+++ b/Engine/Plugins/Online/OnlineSubsystem/Source/Public/Interfaces/OnlineLeaderboardInterface.h
@@ -3,7 +3,7 @@
 #pragma once
 
 #include "CoreMinimal.h"
-#include "UObject/CoreOnline.h"
+#include "Online/CoreOnline.h"
 #include "OnlineDelegateMacros.h"
 #include "OnlineStats.h"
 
diff --git a/Engine/Plugins/Online/OnlineSubsystem/Source/Public/Interfaces/OnlineMessageInterface.h b/Engine/Plugins/Online/OnlineSubsystem/Source/Public/Interfaces/OnlineMessageInterface.h
--- a/Engine/Plugins/Online/OnlineSubsystem/Source/Public/Interfaces/OnlineMessageInterface.h
+++ b/Engine/Plugins/Online/OnlineSubsystem/Source/Public/Interfaces/OnlineMessageInterface.h
@@ -4,7 +4,7 @@
 
 #include "CoreMinimal.h"
 #include "OnlineKeyValuePair.h"
-#include "UObject/CoreOnlineFwd.h"
+#include "Online/CoreOnlineFwd.h"
 #include "OnlineSubsystemTypes.h"
 #include "OnlineDelegateMacros.h"
 
diff --git a/Engine/Plugins/Online/OnlineSubsystem/Source/Public/Interfaces/OnlinePartyInterface.h b/Engine/Plugins/Online/OnlineSubsystem/Source/Public/Interfaces/OnlinePartyInterface.h
--- a/Engine/Plugins/Online/OnlineSubsystem/Source/Public/Interfaces/OnlinePartyInterface.h
+++ b/Engine/Plugins/Online/OnlineSubsystem/Source/Public/Interfaces/OnlinePartyInterface.h
@@ -4,7 +4,7 @@
 
 #include "CoreMinimal.h"
 #include "OnlineKeyValuePair.h"
-#include "UObject/CoreOnline.h"
+#include "Online/CoreOnline.h"
 #include "OnlineSubsystemTypes.h"
 #include "OnlineDelegateMacros.h"
 #include "OnlineError.h"
diff --git a/Engine/Plugins/Online/OnlineSubsystem/Source/Public/Interfaces/OnlinePresenceInterface.h b/Engine/Plugins/Online/OnlineSubsystem/Source/Public/Interfaces/OnlinePresenceInterface.h
--- a/Engine/Plugins/Online/OnlineSubsystem/Source/Public/Interfaces/OnlinePresenceInterface.h
+++ b/Engine/Plugins/Online/OnlineSubsystem/Source/Public/Interfaces/OnlinePresenceInterface.h
@@ -3,7 +3,7 @@
 #pragma once
 
 #include "CoreMinimal.h"
-#include "UObject/CoreOnline.h"
+#include "Online/CoreOnline.h"
 #include "OnlineSubsystemTypes.h"
 #include "OnlineDelegateMacros.h"
 #include "OnlineKeyValuePair.h"
diff --git a/Engine/Plugins/Online/OnlineSubsystem/Source/Public/Interfaces/OnlineSessionInterface.h b/Engine/Plugins/Online/OnlineSubsystem/Source/Public/Interfaces/OnlineSessionInterface.h
--- a/Engine/Plugins/Online/OnlineSubsystem/Source/Public/Interfaces/OnlineSessionInterface.h
+++ b/Engine/Plugins/Online/OnlineSubsystem/Source/Public/Interfaces/OnlineSessionInterface.h
@@ -3,7 +3,7 @@
 #pragma once
 
 #include "CoreMinimal.h"
-#include "UObject/CoreOnline.h"
+#include "Online/CoreOnline.h"
 #include "OnlineSubsystemTypes.h"
 #include "OnlineDelegateMacros.h"
 #include "OnlineKeyValuePair.h"
diff --git a/Engine/Plugins/Online/OnlineSubsystem/Source/Public/Interfaces/OnlineSharingInterface.h b/Engine/Plugins/Online/OnlineSubsystem/Source/Public/Interfaces/OnlineSharingInterface.h
--- a/Engine/Plugins/Online/OnlineSubsystem/Source/Public/Interfaces/OnlineSharingInterface.h
+++ b/Engine/Plugins/Online/OnlineSubsystem/Source/Public/Interfaces/OnlineSharingInterface.h
@@ -3,7 +3,7 @@
 #pragma once
 
 #include "CoreMinimal.h"
-#include "UObject/CoreOnline.h"
+#include "Online/CoreOnline.h"
 #include "OnlineSubsystemTypes.h"
 #include "OnlineDelegateMacros.h"
 
diff --git a/Engine/Plugins/Online/OnlineSubsystem/Source/Public/Interfaces/OnlineTournamentInterface.h b/Engine/Plugins/Online/OnlineSubsystem/Source/Public/Interfaces/OnlineTournamentInterface.h
--- a/Engine/Plugins/Online/OnlineSubsystem/Source/Public/Interfaces/OnlineTournamentInterface.h
+++ b/Engine/Plugins/Online/OnlineSubsystem/Source/Public/Interfaces/OnlineTournamentInterface.h
@@ -3,7 +3,7 @@
 #pragma once
 
 #include "CoreMinimal.h"
-#include "UObject/CoreOnlineFwd.h"
+#include "Online/CoreOnlineFwd.h"
 #include "OnlineSubsystemTypes.h"
 #include "OnlineKeyValuePair.h"
 
diff --git a/Engine/Plugins/Online/OnlineSubsystem/Source/Public/Interfaces/OnlineUserInterface.h b/Engine/Plugins/Online/OnlineSubsystem/Source/Public/Interfaces/OnlineUserInterface.h
--- a/Engine/Plugins/Online/OnlineSubsystem/Source/Public/Interfaces/OnlineUserInterface.h
+++ b/Engine/Plugins/Online/OnlineSubsystem/Source/Public/Interfaces/OnlineUserInterface.h
@@ -3,7 +3,7 @@
 #pragma once
 
 #include "CoreMinimal.h"
-#include "UObject/CoreOnline.h"
+#include "Online/CoreOnline.h"
 #include "OnlineSubsystemTypes.h"
 #include "OnlineDelegateMacros.h"
 
diff --git a/Engine/Plugins/Online/OnlineSubsystem/Source/Public/Interfaces/VoiceInterface.h b/Engine/Plugins/Online/OnlineSubsystem/Source/Public/Interfaces/VoiceInterface.h
--- a/Engine/Plugins/Online/OnlineSubsystem/Source/Public/Interfaces/VoiceInterface.h
+++ b/Engine/Plugins/Online/OnlineSubsystem/Source/Public/Interfaces/VoiceInterface.h
@@ -3,7 +3,7 @@
 #pragma once
 
 #include "CoreMinimal.h"
-#include "UObject/CoreOnline.h"
+#include "Online/CoreOnline.h"
 #include "OnlineDelegateMacros.h"
 #include "OnlineSubsystemPackage.h"
 #include "DSP/MultithreadedPatching.h"
diff --git a/Engine/Plugins/Online/OnlineSubsystem/Source/Public/OnlineNotification.h b/Engine/Plugins/Online/OnlineSubsystem/Source/Public/OnlineNotification.h
--- a/Engine/Plugins/Online/OnlineSubsystem/Source/Public/OnlineNotification.h
+++ b/Engine/Plugins/Online/OnlineSubsystem/Source/Public/OnlineNotification.h
@@ -3,7 +3,7 @@
 #pragma once
 
 #include "CoreMinimal.h"
-#include "UObject/CoreOnline.h"
+#include "Online/CoreOnline.h"
 #include "JsonObjectWrapper.h"
 
 class FJsonValue;
diff --git a/Engine/Plugins/Online/OnlineSubsystem/Source/Public/OnlineSessionSettings.h b/Engine/Plugins/Online/OnlineSubsystem/Source/Public/OnlineSessionSettings.h
--- a/Engine/Plugins/Online/OnlineSubsystem/Source/Public/OnlineSessionSettings.h
+++ b/Engine/Plugins/Online/OnlineSubsystem/Source/Public/OnlineSessionSettings.h
@@ -3,7 +3,7 @@
 #pragma once
 
 #include "CoreMinimal.h"
-#include "UObject/CoreOnline.h"
+#include "Online/CoreOnline.h"
 #include "OnlineSubsystemTypes.h"
 #include "OnlineKeyValuePair.h"
 #include "OnlineSubsystemPackage.h"
diff --git a/Engine/Plugins/Online/OnlineSubsystem/Source/Public/OnlineStats.h b/Engine/Plugins/Online/OnlineSubsystem/Source/Public/OnlineStats.h
--- a/Engine/Plugins/Online/OnlineSubsystem/Source/Public/OnlineStats.h
+++ b/Engine/Plugins/Online/OnlineSubsystem/Source/Public/OnlineStats.h
@@ -3,7 +3,7 @@
 #pragma once
 
 #include "CoreMinimal.h"
-#include "UObject/CoreOnline.h"
+#include "Online/CoreOnline.h"
 #include "OnlineSubsystemTypes.h"
 #include "OnlineKeyValuePair.h"
 
diff --git a/Engine/Plugins/Online/OnlineSubsystem/Source/Public/OnlineSubsystem.h b/Engine/Plugins/Online/OnlineSubsystem/Source/Public/OnlineSubsystem.h
--- a/Engine/Plugins/Online/OnlineSubsystem/Source/Public/OnlineSubsystem.h
+++ b/Engine/Plugins/Online/OnlineSubsystem/Source/Public/OnlineSubsystem.h
@@ -3,7 +3,7 @@
 #pragma once
 
 #include "CoreMinimal.h"
-#include "UObject/CoreOnline.h"
+#include "Online/CoreOnlineFwd.h"
 #include "Stats/Stats.h"
 #include "Modules/ModuleManager.h"
 #include "OnlineSubsystemModule.h"
diff --git a/Engine/Plugins/Online/OnlineSubsystem/Source/Public/OnlineSubsystemTypes.h b/Engine/Plugins/Online/OnlineSubsystem/Source/Public/OnlineSubsystemTypes.h
--- a/Engine/Plugins/Online/OnlineSubsystem/Source/Public/OnlineSubsystemTypes.h
+++ b/Engine/Plugins/Online/OnlineSubsystem/Source/Public/OnlineSubsystemTypes.h
@@ -3,7 +3,7 @@
 #pragma once
 
 #include "CoreMinimal.h"
-#include "UObject/CoreOnline.h"
+#include "Online/CoreOnline.h"
 #include "OnlineSubsystemNames.h"  // can be removed once we have no more temporary FUniqueNetId subtypes
 #include "OnlineSubsystemPackage.h"
 
diff --git a/Engine/Plugins/Online/OnlineSubsystemAmazon/Source/Private/OnlineIdentityAmazon.h b/Engine/Plugins/Online/OnlineSubsystemAmazon/Source/Private/OnlineIdentityAmazon.h
--- a/Engine/Plugins/Online/OnlineSubsystemAmazon/Source/Private/OnlineIdentityAmazon.h
+++ b/Engine/Plugins/Online/OnlineSubsystemAmazon/Source/Private/OnlineIdentityAmazon.h
@@ -3,7 +3,7 @@
 #pragma once
  
 #include "CoreMinimal.h"
-#include "UObject/CoreOnline.h"
+#include "Online/CoreOnline.h"
 #include "OnlineSubsystemTypes.h"
 #include "Interfaces/OnlineIdentityInterface.h"
 #include "OnlineJsonSerializer.h"
diff --git a/Engine/Plugins/Online/OnlineSubsystemEOS/Source/OnlineSubsystemEOS/OnlineSubsystemEOS.Build.cs b/Engine/Plugins/Online/OnlineSubsystemEOS/Source/OnlineSubsystemEOS/OnlineSubsystemEOS.Build.cs
--- a/Engine/Plugins/Online/OnlineSubsystemEOS/Source/OnlineSubsystemEOS/OnlineSubsystemEOS.Build.cs
+++ b/Engine/Plugins/Online/OnlineSubsystemEOS/Source/OnlineSubsystemEOS/OnlineSubsystemEOS.Build.cs
@@ -18,6 +18,7 @@ public class OnlineSubsystemEOS : ModuleRules
 		PrivateDependencyModuleNames.AddRange(
 			new string[] {
 				"Core",
+				"CoreOnline",
 				"CoreUObject",
 				"Engine",
 				"Sockets",
diff --git a/Engine/Plugins/Online/OnlineSubsystemEOS/Source/OnlineSubsystemEOS/Private/InternetAddrEOS.cpp b/Engine/Plugins/Online/OnlineSubsystemEOS/Source/OnlineSubsystemEOS/Private/InternetAddrEOS.cpp
--- a/Engine/Plugins/Online/OnlineSubsystemEOS/Source/OnlineSubsystemEOS/Private/InternetAddrEOS.cpp
+++ b/Engine/Plugins/Online/OnlineSubsystemEOS/Source/OnlineSubsystemEOS/Private/InternetAddrEOS.cpp
@@ -1,7 +1,7 @@
 // Copyright Epic Games, Inc. All Rights Reserved.
 
 #include "InternetAddrEOS.h"
-#include "UObject/CoreOnline.h"
+#include "Online/CoreOnline.h"
 #include "OnlineSubsystemEOSTypes.h"
 
 DEFINE_LOG_CATEGORY(LogSocketSubsystemEOS);
diff --git a/Engine/Plugins/Online/OnlineSubsystemEOS/Source/OnlineSubsystemEOS/Private/OnlineAchievementsEOS.h b/Engine/Plugins/Online/OnlineSubsystemEOS/Source/OnlineSubsystemEOS/Private/OnlineAchievementsEOS.h
--- a/Engine/Plugins/Online/OnlineSubsystemEOS/Source/OnlineSubsystemEOS/Private/OnlineAchievementsEOS.h
+++ b/Engine/Plugins/Online/OnlineSubsystemEOS/Source/OnlineSubsystemEOS/Private/OnlineAchievementsEOS.h
@@ -3,7 +3,7 @@
 #pragma once
 
 #include "CoreMinimal.h"
-#include "UObject/CoreOnline.h"
+#include "Online/CoreOnline.h"
 #include "Interfaces/OnlineAchievementsInterface.h"
 #include "OnlineSubsystemEOSPackage.h"
 #include "OnlineSubsystemEOSTypes.h"
diff --git a/Engine/Plugins/Online/OnlineSubsystemEOS/Source/OnlineSubsystemEOS/Private/OnlineStoreEOS.h b/Engine/Plugins/Online/OnlineSubsystemEOS/Source/OnlineSubsystemEOS/Private/OnlineStoreEOS.h
--- a/Engine/Plugins/Online/OnlineSubsystemEOS/Source/OnlineSubsystemEOS/Private/OnlineStoreEOS.h
+++ b/Engine/Plugins/Online/OnlineSubsystemEOS/Source/OnlineSubsystemEOS/Private/OnlineStoreEOS.h
@@ -3,7 +3,7 @@
 #pragma once
 
 #include "CoreMinimal.h"
-#include "UObject/CoreOnline.h"
+#include "Online/CoreOnline.h"
 #include "Interfaces/OnlineStoreInterfaceV2.h"
 #include "Interfaces/OnlinePurchaseInterface.h"
 #include "OnlineSubsystemEOSPackage.h"
diff --git a/Engine/Plugins/Online/OnlineSubsystemEOS/Source/OnlineSubsystemEOS/Private/OnlineTitleFileEOS.h b/Engine/Plugins/Online/OnlineSubsystemEOS/Source/OnlineSubsystemEOS/Private/OnlineTitleFileEOS.h
--- a/Engine/Plugins/Online/OnlineSubsystemEOS/Source/OnlineSubsystemEOS/Private/OnlineTitleFileEOS.h
+++ b/Engine/Plugins/Online/OnlineSubsystemEOS/Source/OnlineSubsystemEOS/Private/OnlineTitleFileEOS.h
@@ -3,7 +3,7 @@
 #pragma once
 
 #include "CoreMinimal.h"
-#include "UObject/CoreOnline.h"
+#include "Online/CoreOnline.h"
 #include "Interfaces/OnlineTitleFileInterface.h"
 #include "OnlineSubsystemEOSPackage.h"
 #include "OnlineSubsystemEOSTypes.h"
diff --git a/Engine/Plugins/Online/OnlineSubsystemEOS/Source/OnlineSubsystemEOS/Private/OnlineUserCloudEOS.h b/Engine/Plugins/Online/OnlineSubsystemEOS/Source/OnlineSubsystemEOS/Private/OnlineUserCloudEOS.h
--- a/Engine/Plugins/Online/OnlineSubsystemEOS/Source/OnlineSubsystemEOS/Private/OnlineUserCloudEOS.h
+++ b/Engine/Plugins/Online/OnlineSubsystemEOS/Source/OnlineSubsystemEOS/Private/OnlineUserCloudEOS.h
@@ -3,7 +3,7 @@
 #pragma once
 
 #include "CoreMinimal.h"
-#include "UObject/CoreOnline.h"
+#include "Online/CoreOnline.h"
 #include "Interfaces/OnlineUserCloudInterface.h"
 #include "OnlineSubsystemEOSPackage.h"
 #include "OnlineSubsystemEOSTypes.h"
diff --git a/Engine/Plugins/Online/OnlineSubsystemEOS/Source/OnlineSubsystemEOSPlus/Private/OnlineAchievementsEOSPlus.h b/Engine/Plugins/Online/OnlineSubsystemEOS/Source/OnlineSubsystemEOSPlus/Private/OnlineAchievementsEOSPlus.h
--- a/Engine/Plugins/Online/OnlineSubsystemEOS/Source/OnlineSubsystemEOSPlus/Private/OnlineAchievementsEOSPlus.h
+++ b/Engine/Plugins/Online/OnlineSubsystemEOS/Source/OnlineSubsystemEOSPlus/Private/OnlineAchievementsEOSPlus.h
@@ -3,7 +3,7 @@
 #pragma once
 
 #include "CoreMinimal.h"
-#include "UObject/CoreOnline.h"
+#include "Online/CoreOnline.h"
 #include "OnlineSubsystem.h"
 #include "Interfaces/OnlineAchievementsInterface.h"
 #include "OnlineUserEOSPlus.h"
diff --git a/Engine/Plugins/Online/OnlineSubsystemEOS/Source/OnlineSubsystemEOSPlus/Private/OnlineSessionEOSPlus.h b/Engine/Plugins/Online/OnlineSubsystemEOS/Source/OnlineSubsystemEOSPlus/Private/OnlineSessionEOSPlus.h
--- a/Engine/Plugins/Online/OnlineSubsystemEOS/Source/OnlineSubsystemEOSPlus/Private/OnlineSessionEOSPlus.h
+++ b/Engine/Plugins/Online/OnlineSubsystemEOS/Source/OnlineSubsystemEOSPlus/Private/OnlineSessionEOSPlus.h
@@ -3,7 +3,7 @@
 #pragma once
 
 #include "CoreMinimal.h"
-#include "UObject/CoreOnline.h"
+#include "Online/CoreOnline.h"
 #include "OnlineSessionSettings.h"
 #include "Interfaces/OnlineSessionInterface.h"
 #include "OnlineUserEOSPlus.h"
diff --git a/Engine/Plugins/Online/OnlineSubsystemEOS/Source/OnlineSubsystemEOSPlus/Private/OnlineStatsEOSPlus.h b/Engine/Plugins/Online/OnlineSubsystemEOS/Source/OnlineSubsystemEOSPlus/Private/OnlineStatsEOSPlus.h
--- a/Engine/Plugins/Online/OnlineSubsystemEOS/Source/OnlineSubsystemEOSPlus/Private/OnlineStatsEOSPlus.h
+++ b/Engine/Plugins/Online/OnlineSubsystemEOS/Source/OnlineSubsystemEOSPlus/Private/OnlineStatsEOSPlus.h
@@ -3,7 +3,7 @@
 #pragma once
 
 #include "CoreMinimal.h"
-#include "UObject/CoreOnline.h"
+#include "Online/CoreOnline.h"
 #include "Interfaces/OnlineStatsInterface.h"
 
 class FOnlineSubsystemEOSPlus;
diff --git a/Engine/Plugins/Online/OnlineSubsystemEOS/Source/OnlineSubsystemEOSPlus/Private/OnlineUserEOSPlus.h b/Engine/Plugins/Online/OnlineSubsystemEOS/Source/OnlineSubsystemEOSPlus/Private/OnlineUserEOSPlus.h
--- a/Engine/Plugins/Online/OnlineSubsystemEOS/Source/OnlineSubsystemEOSPlus/Private/OnlineUserEOSPlus.h
+++ b/Engine/Plugins/Online/OnlineSubsystemEOS/Source/OnlineSubsystemEOSPlus/Private/OnlineUserEOSPlus.h
@@ -3,7 +3,7 @@
 #pragma once
 
 #include "CoreMinimal.h"
-#include "UObject/CoreOnline.h"
+#include "Online/CoreOnline.h"
 #include "OnlineSubsystem.h"
 #include "Interfaces/OnlineIdentityInterface.h"
 #include "Interfaces/OnlineFriendsInterface.h"
diff --git a/Engine/Plugins/Online/OnlineSubsystemFacebook/Source/OnlineSubsystemFacebook.Build.cs b/Engine/Plugins/Online/OnlineSubsystemFacebook/Source/OnlineSubsystemFacebook.Build.cs
--- a/Engine/Plugins/Online/OnlineSubsystemFacebook/Source/OnlineSubsystemFacebook.Build.cs
+++ b/Engine/Plugins/Online/OnlineSubsystemFacebook/Source/OnlineSubsystemFacebook.Build.cs
@@ -16,6 +16,7 @@ public class OnlineSubsystemFacebook : ModuleRules
 		PrivateDependencyModuleNames.AddRange(
 			new string[] { 
 				"Core",
+				"CoreOnline",
 				"CoreUObject",
 				"ApplicationCore",
 				"HTTP",
diff --git a/Engine/Plugins/Online/OnlineSubsystemGoogle/Source/OnlineSubsystemGoogle.Build.cs b/Engine/Plugins/Online/OnlineSubsystemGoogle/Source/OnlineSubsystemGoogle.Build.cs
--- a/Engine/Plugins/Online/OnlineSubsystemGoogle/Source/OnlineSubsystemGoogle.Build.cs
+++ b/Engine/Plugins/Online/OnlineSubsystemGoogle/Source/OnlineSubsystemGoogle.Build.cs
@@ -16,6 +16,7 @@ public class OnlineSubsystemGoogle : ModuleRules
 		PrivateDependencyModuleNames.AddRange(
 			new string[] { 
 				"Core",
+				"CoreOnline",
 				"CoreUObject",
 				"ApplicationCore",
 				"HTTP",
diff --git a/Engine/Plugins/Online/OnlineSubsystemNull/Source/Private/OnlineIdentityNull.h b/Engine/Plugins/Online/OnlineSubsystemNull/Source/Private/OnlineIdentityNull.h
--- a/Engine/Plugins/Online/OnlineSubsystemNull/Source/Private/OnlineIdentityNull.h
+++ b/Engine/Plugins/Online/OnlineSubsystemNull/Source/Private/OnlineIdentityNull.h
@@ -3,7 +3,7 @@
 #pragma once
 
 #include "CoreMinimal.h"
-#include "UObject/CoreOnline.h"
+#include "Online/CoreOnline.h"
 #include "OnlineSubsystemTypes.h"
 #include "Interfaces/OnlineIdentityInterface.h"
 #include "OnlineSubsystemNullTypes.h"
diff --git a/Engine/Plugins/Online/OnlineSubsystemNull/Source/Private/OnlineLeaderboardInterfaceNull.h b/Engine/Plugins/Online/OnlineSubsystemNull/Source/Private/OnlineLeaderboardInterfaceNull.h
--- a/Engine/Plugins/Online/OnlineSubsystemNull/Source/Private/OnlineLeaderboardInterfaceNull.h
+++ b/Engine/Plugins/Online/OnlineSubsystemNull/Source/Private/OnlineLeaderboardInterfaceNull.h
@@ -3,7 +3,7 @@
 #pragma once
 
 #include "CoreMinimal.h"
-#include "UObject/CoreOnline.h"
+#include "Online/CoreOnline.h"
 #include "OnlineSubsystemTypes.h"
 #include "OnlineStats.h"
 #include "Interfaces/OnlineLeaderboardInterface.h"
diff --git a/Engine/Plugins/Online/OnlineSubsystemNull/Source/Private/OnlineSessionInterfaceNull.h b/Engine/Plugins/Online/OnlineSubsystemNull/Source/Private/OnlineSessionInterfaceNull.h
--- a/Engine/Plugins/Online/OnlineSubsystemNull/Source/Private/OnlineSessionInterfaceNull.h
+++ b/Engine/Plugins/Online/OnlineSubsystemNull/Source/Private/OnlineSessionInterfaceNull.h
@@ -3,7 +3,7 @@
 #pragma once
 
 #include "CoreMinimal.h"
-#include "UObject/CoreOnline.h"
+#include "Online/CoreOnline.h"
 #include "Misc/ScopeLock.h"
 #include "OnlineSessionSettings.h"
 #include "Interfaces/OnlineSessionInterface.h"
diff --git a/Engine/Plugins/Online/OnlineSubsystemNull/Source/Private/VoiceInterfaceNull.h b/Engine/Plugins/Online/OnlineSubsystemNull/Source/Private/VoiceInterfaceNull.h
--- a/Engine/Plugins/Online/OnlineSubsystemNull/Source/Private/VoiceInterfaceNull.h
+++ b/Engine/Plugins/Online/OnlineSubsystemNull/Source/Private/VoiceInterfaceNull.h
@@ -3,7 +3,7 @@
 #pragma once
 
 #include "VoiceInterfaceImpl.h"
-#include "UObject/CoreOnline.h"
+#include "Online/CoreOnline.h"
 #include "OnlineSubsystemTypes.h"
 #include "Interfaces/VoiceInterface.h"
 #include "Net/VoiceDataCommon.h"
diff --git a/Engine/Plugins/Online/OnlineSubsystemOculus/Source/Private/OculusCreateSessionCallbackProxy.cpp b/Engine/Plugins/Online/OnlineSubsystemOculus/Source/Private/OculusCreateSessionCallbackProxy.cpp
--- a/Engine/Plugins/Online/OnlineSubsystemOculus/Source/Private/OculusCreateSessionCallbackProxy.cpp
+++ b/Engine/Plugins/Online/OnlineSubsystemOculus/Source/Private/OculusCreateSessionCallbackProxy.cpp
@@ -2,7 +2,7 @@
 
 #include "OculusCreateSessionCallbackProxy.h"
 #include "OnlineSubsystemOculusPrivate.h"
-#include "UObject/CoreOnline.h"
+#include "Online/CoreOnline.h"
 #include "Online.h"
 #include "OnlineSessionInterfaceOculus.h"
 
diff --git a/Engine/Plugins/Online/OnlineSubsystemOculus/Source/Private/OculusUpdateSessionCallbackProxy.cpp b/Engine/Plugins/Online/OnlineSubsystemOculus/Source/Private/OculusUpdateSessionCallbackProxy.cpp
--- a/Engine/Plugins/Online/OnlineSubsystemOculus/Source/Private/OculusUpdateSessionCallbackProxy.cpp
+++ b/Engine/Plugins/Online/OnlineSubsystemOculus/Source/Private/OculusUpdateSessionCallbackProxy.cpp
@@ -2,7 +2,7 @@
 
 #include "OculusUpdateSessionCallbackProxy.h"
 #include "OnlineSubsystemOculusPrivate.h"
-#include "UObject/CoreOnline.h"
+#include "Online/CoreOnline.h"
 #include "Online.h"
 #include "OnlineSessionInterfaceOculus.h"
 
diff --git a/Engine/Plugins/Online/OnlineSubsystemSteam/Source/Private/OnlineFriendsInterfaceSteam.h b/Engine/Plugins/Online/OnlineSubsystemSteam/Source/Private/OnlineFriendsInterfaceSteam.h
--- a/Engine/Plugins/Online/OnlineSubsystemSteam/Source/Private/OnlineFriendsInterfaceSteam.h
+++ b/Engine/Plugins/Online/OnlineSubsystemSteam/Source/Private/OnlineFriendsInterfaceSteam.h
@@ -3,7 +3,7 @@
 #pragma once
 
 #include "CoreMinimal.h"
-#include "UObject/CoreOnline.h"
+#include "Online/CoreOnline.h"
 #include "OnlineSubsystemTypes.h"
 #include "OnlineSubsystemSteamPrivate.h"
 #include "OnlineAsyncTaskManager.h"
diff --git a/Engine/Plugins/Online/OnlineSubsystemSteam/Source/Private/OnlineIdentityInterfaceSteam.h b/Engine/Plugins/Online/OnlineSubsystemSteam/Source/Private/OnlineIdentityInterfaceSteam.h
--- a/Engine/Plugins/Online/OnlineSubsystemSteam/Source/Private/OnlineIdentityInterfaceSteam.h
+++ b/Engine/Plugins/Online/OnlineSubsystemSteam/Source/Private/OnlineIdentityInterfaceSteam.h
@@ -3,7 +3,7 @@
 #pragma once
 
 #include "CoreMinimal.h"
-#include "UObject/CoreOnline.h"
+#include "Online/CoreOnline.h"
 #include "Interfaces/OnlineIdentityInterface.h"
 #include "OnlineSubsystemSteamPackage.h"
 
diff --git a/Engine/Plugins/Online/OnlineSubsystemSteam/Source/Private/OnlineLeaderboardInterfaceSteam.h b/Engine/Plugins/Online/OnlineSubsystemSteam/Source/Private/OnlineLeaderboardInterfaceSteam.h
--- a/Engine/Plugins/Online/OnlineSubsystemSteam/Source/Private/OnlineLeaderboardInterfaceSteam.h
+++ b/Engine/Plugins/Online/OnlineSubsystemSteam/Source/Private/OnlineLeaderboardInterfaceSteam.h
@@ -3,7 +3,7 @@
 #pragma once
 
 #include "CoreMinimal.h"
-#include "UObject/CoreOnline.h"
+#include "Online/CoreOnline.h"
 #include "OnlineSubsystemSteamTypes.h"
 #include "OnlineStats.h"
 #include "Interfaces/OnlineLeaderboardInterface.h"
diff --git a/Engine/Plugins/Online/OnlineSubsystemSteam/Source/Private/OnlineSessionAsyncLobbySteam.h b/Engine/Plugins/Online/OnlineSubsystemSteam/Source/Private/OnlineSessionAsyncLobbySteam.h
--- a/Engine/Plugins/Online/OnlineSubsystemSteam/Source/Private/OnlineSessionAsyncLobbySteam.h
+++ b/Engine/Plugins/Online/OnlineSubsystemSteam/Source/Private/OnlineSessionAsyncLobbySteam.h
@@ -3,7 +3,7 @@
 #pragma once
 
 #include "CoreMinimal.h"
-#include "UObject/CoreOnline.h"
+#include "Online/CoreOnline.h"
 #include "OnlineSubsystemSteamTypes.h"
 #include "OnlineSessionSettings.h"
 #include "Interfaces/OnlineSessionInterface.h"
diff --git a/Engine/Plugins/Online/OnlineSubsystemSteam/Source/Private/OnlineSessionAsyncServerSteam.h b/Engine/Plugins/Online/OnlineSubsystemSteam/Source/Private/OnlineSessionAsyncServerSteam.h
--- a/Engine/Plugins/Online/OnlineSubsystemSteam/Source/Private/OnlineSessionAsyncServerSteam.h
+++ b/Engine/Plugins/Online/OnlineSubsystemSteam/Source/Private/OnlineSessionAsyncServerSteam.h
@@ -4,7 +4,7 @@
 
 #include "CoreMinimal.h"
 #include "Containers/IndirectArray.h"
-#include "UObject/CoreOnline.h"
+#include "Online/CoreOnline.h"
 #include "OnlineSubsystemSteamTypes.h"
 #include "OnlineSessionSettings.h"
 #include "OnlineSessionInterfaceSteam.h"
diff --git a/Engine/Plugins/Online/OnlineSubsystemSteam/Source/Private/OnlineSessionInterfaceSteam.h b/Engine/Plugins/Online/OnlineSubsystemSteam/Source/Private/OnlineSessionInterfaceSteam.h
--- a/Engine/Plugins/Online/OnlineSubsystemSteam/Source/Private/OnlineSessionInterfaceSteam.h
+++ b/Engine/Plugins/Online/OnlineSubsystemSteam/Source/Private/OnlineSessionInterfaceSteam.h
@@ -3,7 +3,7 @@
 #pragma once
 
 #include "CoreMinimal.h"
-#include "UObject/CoreOnline.h"
+#include "Online/CoreOnline.h"
 #include "OnlineSubsystemSteamTypes.h"
 #include "Misc/ScopeLock.h"
 #include "OnlineKeyValuePair.h"
diff --git a/Engine/Plugins/Online/OnlineSubsystemSteam/Source/Private/OnlineSubsystemSteamTypes.h b/Engine/Plugins/Online/OnlineSubsystemSteam/Source/Private/OnlineSubsystemSteamTypes.h
--- a/Engine/Plugins/Online/OnlineSubsystemSteam/Source/Private/OnlineSubsystemSteamTypes.h
+++ b/Engine/Plugins/Online/OnlineSubsystemSteam/Source/Private/OnlineSubsystemSteamTypes.h
@@ -3,7 +3,7 @@
 #pragma once
 
 #include "CoreMinimal.h"
-#include "UObject/CoreOnline.h"
+#include "Online/CoreOnline.h"
 #include "OnlineSubsystemTypes.h"
 #include "OnlineSubsystemSteamPrivate.h"
 #include "IPAddress.h"
diff --git a/Engine/Plugins/Online/OnlineSubsystemSteam/Source/Private/VoicePacketSteam.h b/Engine/Plugins/Online/OnlineSubsystemSteam/Source/Private/VoicePacketSteam.h
--- a/Engine/Plugins/Online/OnlineSubsystemSteam/Source/Private/VoicePacketSteam.h
+++ b/Engine/Plugins/Online/OnlineSubsystemSteam/Source/Private/VoicePacketSteam.h
@@ -3,7 +3,7 @@
 #pragma once
 
 #include "CoreMinimal.h"
-#include "UObject/CoreOnline.h"
+#include "Online/CoreOnline.h"
 #include "Net/VoiceDataCommon.h"
 #include "OnlineSubsystemSteamPackage.h"
 
diff --git a/Engine/Plugins/Online/OnlineSubsystemUtils/Source/OnlineSubsystemUtils/Classes/OnlineSessionClient.h b/Engine/Plugins/Online/OnlineSubsystemUtils/Source/OnlineSubsystemUtils/Classes/OnlineSessionClient.h
--- a/Engine/Plugins/Online/OnlineSubsystemUtils/Source/OnlineSubsystemUtils/Classes/OnlineSessionClient.h
+++ b/Engine/Plugins/Online/OnlineSubsystemUtils/Source/OnlineSubsystemUtils/Classes/OnlineSessionClient.h
@@ -9,7 +9,7 @@
 
 #include "CoreMinimal.h"
 #include "UObject/ObjectMacros.h"
-#include "UObject/CoreOnline.h"
+#include "Online/CoreOnline.h"
 #include "OnlineSubsystem.h"
 #include "OnlineSessionSettings.h"
 #include "Interfaces/OnlineSessionInterface.h"
diff --git a/Engine/Plugins/Online/OnlineSubsystemUtils/Source/OnlineSubsystemUtils/Classes/ShowLoginUICallbackProxy.h b/Engine/Plugins/Online/OnlineSubsystemUtils/Source/OnlineSubsystemUtils/Classes/ShowLoginUICallbackProxy.h
--- a/Engine/Plugins/Online/OnlineSubsystemUtils/Source/OnlineSubsystemUtils/Classes/ShowLoginUICallbackProxy.h
+++ b/Engine/Plugins/Online/OnlineSubsystemUtils/Source/OnlineSubsystemUtils/Classes/ShowLoginUICallbackProxy.h
@@ -4,7 +4,7 @@
 
 #include "CoreMinimal.h"
 #include "UObject/ObjectMacros.h"
-#include "UObject/CoreOnline.h"
+#include "Online/CoreOnline.h"
 #include "Kismet/BlueprintAsyncActionBase.h"
 #include "ShowLoginUICallbackProxy.generated.h"
 
diff --git a/Engine/Plugins/Online/OnlineSubsystemUtils/Source/OnlineSubsystemUtils/Private/LeaderboardBlueprintLibrary.cpp b/Engine/Plugins/Online/OnlineSubsystemUtils/Source/OnlineSubsystemUtils/Private/LeaderboardBlueprintLibrary.cpp
--- a/Engine/Plugins/Online/OnlineSubsystemUtils/Source/OnlineSubsystemUtils/Private/LeaderboardBlueprintLibrary.cpp
+++ b/Engine/Plugins/Online/OnlineSubsystemUtils/Source/OnlineSubsystemUtils/Private/LeaderboardBlueprintLibrary.cpp
@@ -1,7 +1,7 @@
 // Copyright Epic Games, Inc. All Rights Reserved.
 
 #include "LeaderboardBlueprintLibrary.h"
-#include "UObject/CoreOnline.h"
+#include "Online/CoreOnline.h"
 #include "GameFramework/PlayerController.h"
 #include "GameFramework/PlayerState.h"
 #include "OnlineSubsystem.h"
diff --git a/Engine/Plugins/Online/OnlineSubsystemUtils/Source/OnlineSubsystemUtils/Private/OnlineEngineInterfaceImpl.h b/Engine/Plugins/Online/OnlineSubsystemUtils/Source/OnlineSubsystemUtils/Private/OnlineEngineInterfaceImpl.h
--- a/Engine/Plugins/Online/OnlineSubsystemUtils/Source/OnlineSubsystemUtils/Private/OnlineEngineInterfaceImpl.h
+++ b/Engine/Plugins/Online/OnlineSubsystemUtils/Source/OnlineSubsystemUtils/Private/OnlineEngineInterfaceImpl.h
@@ -4,7 +4,7 @@
 
 #include "CoreMinimal.h"
 #include "UObject/ObjectMacros.h"
-#include "UObject/CoreOnline.h"
+#include "Online/CoreOnline.h"
 #include "Net/OnlineEngineInterface.h"
 #include "OnlineEngineInterfaceImpl.generated.h"
 
diff --git a/Engine/Plugins/Online/OnlineSubsystemUtils/Source/OnlineSubsystemUtils/Private/OnlineSubsystemBPCallHelper.h b/Engine/Plugins/Online/OnlineSubsystemUtils/Source/OnlineSubsystemUtils/Private/OnlineSubsystemBPCallHelper.h
--- a/Engine/Plugins/Online/OnlineSubsystemUtils/Source/OnlineSubsystemUtils/Private/OnlineSubsystemBPCallHelper.h
+++ b/Engine/Plugins/Online/OnlineSubsystemUtils/Source/OnlineSubsystemUtils/Private/OnlineSubsystemBPCallHelper.h
@@ -3,7 +3,7 @@
 #pragma once
 
 #include "CoreMinimal.h"
-#include "UObject/CoreOnline.h"
+#include "Online/CoreOnline.h"
 
 class APlayerController;
 class IOnlineSubsystem;
diff --git a/Engine/Plugins/Online/OnlineSubsystemUtils/Source/OnlineSubsystemUtils/Private/Tests/TestAchievementsInterface.h b/Engine/Plugins/Online/OnlineSubsystemUtils/Source/OnlineSubsystemUtils/Private/Tests/TestAchievementsInterface.h
--- a/Engine/Plugins/Online/OnlineSubsystemUtils/Source/OnlineSubsystemUtils/Private/Tests/TestAchievementsInterface.h
+++ b/Engine/Plugins/Online/OnlineSubsystemUtils/Source/OnlineSubsystemUtils/Private/Tests/TestAchievementsInterface.h
@@ -5,7 +5,7 @@
 
 
 #include "CoreMinimal.h"
-#include "UObject/CoreOnline.h"
+#include "Online/CoreOnline.h"
 #include "OnlineSubsystem.h"
 #include "Interfaces/OnlineAchievementsInterface.h"
 
diff --git a/Engine/Plugins/Online/OnlineSubsystemUtils/Source/OnlineSubsystemUtils/Private/Tests/TestCloudInterface.h b/Engine/Plugins/Online/OnlineSubsystemUtils/Source/OnlineSubsystemUtils/Private/Tests/TestCloudInterface.h
--- a/Engine/Plugins/Online/OnlineSubsystemUtils/Source/OnlineSubsystemUtils/Private/Tests/TestCloudInterface.h
+++ b/Engine/Plugins/Online/OnlineSubsystemUtils/Source/OnlineSubsystemUtils/Private/Tests/TestCloudInterface.h
@@ -4,7 +4,7 @@
 
 #include "CoreMinimal.h"
 #include "Containers/Ticker.h"
-#include "UObject/CoreOnline.h"
+#include "Online/CoreOnline.h"
 #include "OnlineSubsystemTypes.h"
 #include "Interfaces/OnlineUserCloudInterface.h"
 #include "Interfaces/OnlineSharedCloudInterface.h"
diff --git a/Engine/Plugins/Online/OnlineSubsystemUtils/Source/OnlineSubsystemUtils/Private/Tests/TestEntitlementsInterface.h b/Engine/Plugins/Online/OnlineSubsystemUtils/Source/OnlineSubsystemUtils/Private/Tests/TestEntitlementsInterface.h
--- a/Engine/Plugins/Online/OnlineSubsystemUtils/Source/OnlineSubsystemUtils/Private/Tests/TestEntitlementsInterface.h
+++ b/Engine/Plugins/Online/OnlineSubsystemUtils/Source/OnlineSubsystemUtils/Private/Tests/TestEntitlementsInterface.h
@@ -3,7 +3,7 @@
 #pragma once
 
 #include "CoreMinimal.h"
-#include "UObject/CoreOnline.h"
+#include "Online/CoreOnline.h"
 #include "OnlineSubsystem.h"
 #include "Interfaces/OnlineEntitlementsInterface.h"
 
diff --git a/Engine/Plugins/Online/OnlineSubsystemUtils/Source/OnlineSubsystemUtils/Private/Tests/TestExternalUIInterface.h b/Engine/Plugins/Online/OnlineSubsystemUtils/Source/OnlineSubsystemUtils/Private/Tests/TestExternalUIInterface.h
--- a/Engine/Plugins/Online/OnlineSubsystemUtils/Source/OnlineSubsystemUtils/Private/Tests/TestExternalUIInterface.h
+++ b/Engine/Plugins/Online/OnlineSubsystemUtils/Source/OnlineSubsystemUtils/Private/Tests/TestExternalUIInterface.h
@@ -3,7 +3,7 @@
 #pragma once
 
 #include "CoreMinimal.h"
-#include "UObject/CoreOnline.h"
+#include "Online/CoreOnline.h"
 #include "Interfaces/OnlineExternalUIInterface.h"
 
 class IOnlineSubsystem;
diff --git a/Engine/Plugins/Online/OnlineSubsystemUtils/Source/OnlineSubsystemUtils/Private/Tests/TestFriendsInterface.h b/Engine/Plugins/Online/OnlineSubsystemUtils/Source/OnlineSubsystemUtils/Private/Tests/TestFriendsInterface.h
--- a/Engine/Plugins/Online/OnlineSubsystemUtils/Source/OnlineSubsystemUtils/Private/Tests/TestFriendsInterface.h
+++ b/Engine/Plugins/Online/OnlineSubsystemUtils/Source/OnlineSubsystemUtils/Private/Tests/TestFriendsInterface.h
@@ -3,7 +3,7 @@
 #pragma once
 
 #include "CoreMinimal.h"
-#include "UObject/CoreOnline.h"
+#include "Online/CoreOnline.h"
 #include "Interfaces/OnlineFriendsInterface.h"
 
 class IOnlineSubsystem;
diff --git a/Engine/Plugins/Online/OnlineSubsystemUtils/Source/OnlineSubsystemUtils/Private/Tests/TestLeaderboardInterface.h b/Engine/Plugins/Online/OnlineSubsystemUtils/Source/OnlineSubsystemUtils/Private/Tests/TestLeaderboardInterface.h
--- a/Engine/Plugins/Online/OnlineSubsystemUtils/Source/OnlineSubsystemUtils/Private/Tests/TestLeaderboardInterface.h
+++ b/Engine/Plugins/Online/OnlineSubsystemUtils/Source/OnlineSubsystemUtils/Private/Tests/TestLeaderboardInterface.h
@@ -4,7 +4,7 @@
 
 #include "CoreMinimal.h"
 #include "Containers/Ticker.h"
-#include "UObject/CoreOnline.h"
+#include "Online/CoreOnline.h"
 #include "Interfaces/OnlineLeaderboardInterface.h"
 
 #if WITH_DEV_AUTOMATION_TESTS
diff --git a/Engine/Plugins/Online/OnlineSubsystemUtils/Source/OnlineSubsystemUtils/Private/Tests/TestMessageInterface.h b/Engine/Plugins/Online/OnlineSubsystemUtils/Source/OnlineSubsystemUtils/Private/Tests/TestMessageInterface.h
--- a/Engine/Plugins/Online/OnlineSubsystemUtils/Source/OnlineSubsystemUtils/Private/Tests/TestMessageInterface.h
+++ b/Engine/Plugins/Online/OnlineSubsystemUtils/Source/OnlineSubsystemUtils/Private/Tests/TestMessageInterface.h
@@ -3,7 +3,7 @@
 #pragma once
 
 #include "CoreMinimal.h"
-#include "UObject/CoreOnline.h"
+#include "Online/CoreOnline.h"
 #include "Interfaces/OnlineMessageInterface.h"
 
 class IOnlineSubsystem;
diff --git a/Engine/Plugins/Online/OnlineSubsystemUtils/Source/OnlineSubsystemUtils/Private/Tests/TestPurchaseInterface.h b/Engine/Plugins/Online/OnlineSubsystemUtils/Source/OnlineSubsystemUtils/Private/Tests/TestPurchaseInterface.h
--- a/Engine/Plugins/Online/OnlineSubsystemUtils/Source/OnlineSubsystemUtils/Private/Tests/TestPurchaseInterface.h
+++ b/Engine/Plugins/Online/OnlineSubsystemUtils/Source/OnlineSubsystemUtils/Private/Tests/TestPurchaseInterface.h
@@ -5,7 +5,7 @@
 // UE4 includes
 
 #include "CoreMinimal.h"
-#include "UObject/CoreOnline.h"
+#include "Online/CoreOnline.h"
 
 // Module includes
 #include "Interfaces/OnlinePurchaseInterface.h"
diff --git a/Engine/Plugins/Online/OnlineSubsystemUtils/Source/OnlineSubsystemUtils/Private/Tests/TestSessionInterface.cpp b/Engine/Plugins/Online/OnlineSubsystemUtils/Source/OnlineSubsystemUtils/Private/Tests/TestSessionInterface.cpp
--- a/Engine/Plugins/Online/OnlineSubsystemUtils/Source/OnlineSubsystemUtils/Private/Tests/TestSessionInterface.cpp
+++ b/Engine/Plugins/Online/OnlineSubsystemUtils/Source/OnlineSubsystemUtils/Private/Tests/TestSessionInterface.cpp
@@ -4,7 +4,7 @@
 #include "GameFramework/GameModeBase.h"
 #include "GameFramework/GameMode.h"
 #include "Kismet/GameplayStatics.h"
-#include "UObject/CoreOnlineFwd.h"
+#include "Online/CoreOnlineFwd.h"
 #include "OnlineSubsystemUtils.h"
 #include "OnlineSubsystemSessionSettings.h"
 
diff --git a/Engine/Plugins/Online/OnlineSubsystemUtils/Source/OnlineSubsystemUtils/Private/Tests/TestSessionInterface.h b/Engine/Plugins/Online/OnlineSubsystemUtils/Source/OnlineSubsystemUtils/Private/Tests/TestSessionInterface.h
--- a/Engine/Plugins/Online/OnlineSubsystemUtils/Source/OnlineSubsystemUtils/Private/Tests/TestSessionInterface.h
+++ b/Engine/Plugins/Online/OnlineSubsystemUtils/Source/OnlineSubsystemUtils/Private/Tests/TestSessionInterface.h
@@ -5,7 +5,7 @@
 #include "CoreMinimal.h"
 #include "Misc/CoreMisc.h"
 #include "Containers/Ticker.h"
-#include "UObject/CoreOnline.h"
+#include "Online/CoreOnline.h"
 #include "OnlineSubsystemTypes.h"
 #include "OnlineSubsystem.h"
 #include "OnlineSessionSettings.h"
diff --git a/Engine/Plugins/Online/OnlineSubsystemUtils/Source/OnlineSubsystemUtils/Private/Tests/TestStatsInterface.h b/Engine/Plugins/Online/OnlineSubsystemUtils/Source/OnlineSubsystemUtils/Private/Tests/TestStatsInterface.h
--- a/Engine/Plugins/Online/OnlineSubsystemUtils/Source/OnlineSubsystemUtils/Private/Tests/TestStatsInterface.h
+++ b/Engine/Plugins/Online/OnlineSubsystemUtils/Source/OnlineSubsystemUtils/Private/Tests/TestStatsInterface.h
@@ -4,7 +4,7 @@
 
 #include "CoreMinimal.h"
 #include "Containers/Ticker.h"
-#include "UObject/CoreOnline.h"
+#include "Online/CoreOnline.h"
 #include "Interfaces/OnlineStatsInterface.h"
 #include "OnlineSubsystem.h"
 
diff --git a/Engine/Plugins/Online/OnlineSubsystemUtils/Source/OnlineSubsystemUtils/Private/Tests/TestStoreInterface.h b/Engine/Plugins/Online/OnlineSubsystemUtils/Source/OnlineSubsystemUtils/Private/Tests/TestStoreInterface.h
--- a/Engine/Plugins/Online/OnlineSubsystemUtils/Source/OnlineSubsystemUtils/Private/Tests/TestStoreInterface.h
+++ b/Engine/Plugins/Online/OnlineSubsystemUtils/Source/OnlineSubsystemUtils/Private/Tests/TestStoreInterface.h
@@ -5,7 +5,7 @@
 // UE4 includes
 
 #include "CoreMinimal.h"
-#include "UObject/CoreOnline.h"
+#include "Online/CoreOnline.h"
 
 // Module includes
 #include "Interfaces/OnlineStoreInterfaceV2.h"
diff --git a/Engine/Plugins/Online/OnlineSubsystemUtils/Source/OnlineSubsystemUtils/Private/Tests/TestUserInterface.h b/Engine/Plugins/Online/OnlineSubsystemUtils/Source/OnlineSubsystemUtils/Private/Tests/TestUserInterface.h
--- a/Engine/Plugins/Online/OnlineSubsystemUtils/Source/OnlineSubsystemUtils/Private/Tests/TestUserInterface.h
+++ b/Engine/Plugins/Online/OnlineSubsystemUtils/Source/OnlineSubsystemUtils/Private/Tests/TestUserInterface.h
@@ -5,7 +5,7 @@
 // UE4 includes
 
 #include "CoreMinimal.h"
-#include "UObject/CoreOnline.h"
+#include "Online/CoreOnline.h"
 
 // Module includes
 #include "Interfaces/OnlineUserInterface.h"
diff --git a/Engine/Plugins/Online/OnlineSubsystemUtils/Source/OnlineSubsystemUtils/Public/PartyBeaconState.h b/Engine/Plugins/Online/OnlineSubsystemUtils/Source/OnlineSubsystemUtils/Public/PartyBeaconState.h
--- a/Engine/Plugins/Online/OnlineSubsystemUtils/Source/OnlineSubsystemUtils/Public/PartyBeaconState.h
+++ b/Engine/Plugins/Online/OnlineSubsystemUtils/Source/OnlineSubsystemUtils/Public/PartyBeaconState.h
@@ -5,7 +5,7 @@
 #include "CoreMinimal.h"
 #include "UObject/ObjectMacros.h"
 #include "UObject/Object.h"
-#include "UObject/CoreOnline.h"
+#include "Online/CoreOnline.h"
 #include "GameFramework/OnlineReplStructs.h"
 #include "OnlineBeaconReservation.h"
 #include "PartyBeaconState.generated.h"
diff --git a/Engine/Plugins/Online/OnlineSubsystemUtils/Source/OnlineSubsystemUtils/Public/SpectatorBeaconState.h b/Engine/Plugins/Online/OnlineSubsystemUtils/Source/OnlineSubsystemUtils/Public/SpectatorBeaconState.h
--- a/Engine/Plugins/Online/OnlineSubsystemUtils/Source/OnlineSubsystemUtils/Public/SpectatorBeaconState.h
+++ b/Engine/Plugins/Online/OnlineSubsystemUtils/Source/OnlineSubsystemUtils/Public/SpectatorBeaconState.h
@@ -5,7 +5,7 @@
 #include "CoreMinimal.h"
 #include "UObject/ObjectMacros.h"
 #include "UObject/Object.h"
-#include "UObject/CoreOnline.h"
+#include "Online/CoreOnline.h"
 #include "GameFramework/OnlineReplStructs.h"
 #include "OnlineBeaconReservation.h"
 #include "SpectatorBeaconState.generated.h"
diff --git a/Engine/Plugins/Online/OnlineSubsystemUtils/Source/OnlineSubsystemUtils/Public/VoiceEngineImpl.h b/Engine/Plugins/Online/OnlineSubsystemUtils/Source/OnlineSubsystemUtils/Public/VoiceEngineImpl.h
--- a/Engine/Plugins/Online/OnlineSubsystemUtils/Source/OnlineSubsystemUtils/Public/VoiceEngineImpl.h
+++ b/Engine/Plugins/Online/OnlineSubsystemUtils/Source/OnlineSubsystemUtils/Public/VoiceEngineImpl.h
@@ -18,7 +18,7 @@
 #include "DSP/MultithreadedPatching.h"
 
 #include "VoicePacketImpl.h"
-#include "UObject/CoreOnline.h"
+#include "Online/CoreOnline.h"
 
 class IOnlineSubsystem;
 class IVoiceDecoder;
diff --git a/Engine/Plugins/Online/OnlineSubsystemUtils/Source/OnlineSubsystemUtils/Public/VoiceInterfaceImpl.h b/Engine/Plugins/Online/OnlineSubsystemUtils/Source/OnlineSubsystemUtils/Public/VoiceInterfaceImpl.h
--- a/Engine/Plugins/Online/OnlineSubsystemUtils/Source/OnlineSubsystemUtils/Public/VoiceInterfaceImpl.h
+++ b/Engine/Plugins/Online/OnlineSubsystemUtils/Source/OnlineSubsystemUtils/Public/VoiceInterfaceImpl.h
@@ -3,7 +3,7 @@
 #pragma once
 
 #include "CoreMinimal.h"
-#include "UObject/CoreOnline.h"
+#include "Online/CoreOnline.h"
 #include "OnlineSubsystemTypes.h"
 #include "Interfaces/VoiceInterface.h"
 #include "Net/VoiceDataCommon.h"
diff --git a/Engine/Plugins/Online/OnlineSubsystemUtils/Source/OnlineSubsystemUtils/Public/VoicePacketImpl.h b/Engine/Plugins/Online/OnlineSubsystemUtils/Source/OnlineSubsystemUtils/Public/VoicePacketImpl.h
--- a/Engine/Plugins/Online/OnlineSubsystemUtils/Source/OnlineSubsystemUtils/Public/VoicePacketImpl.h
+++ b/Engine/Plugins/Online/OnlineSubsystemUtils/Source/OnlineSubsystemUtils/Public/VoicePacketImpl.h
@@ -3,7 +3,7 @@
 #pragma once
 
 #include "CoreMinimal.h"
-#include "UObject/CoreOnline.h"
+#include "Online/CoreOnline.h"
 #include "Net/VoiceDataCommon.h"
 #include "OnlineSubsystemUtilsPackage.h"
 #include "Net/VoiceConfig.h"
diff --git a/Engine/Source/Editor/GameProjectGeneration/Private/GameProjectUtils.cpp b/Engine/Source/Editor/GameProjectGeneration/Private/GameProjectUtils.cpp
--- a/Engine/Source/Editor/GameProjectGeneration/Private/GameProjectUtils.cpp
+++ b/Engine/Source/Editor/GameProjectGeneration/Private/GameProjectUtils.cpp
@@ -2179,7 +2179,7 @@ bool GameProjectUtils::GenerateConfigFiles(const FProjectInformation& InProjectI
 		}
 
 		FileContents += LINE_TERMINATOR;
-		FileContents += TEXT("[CoreUObject]") LINE_TERMINATOR;
+		FileContents += TEXT("[CoreOnline]") LINE_TERMINATOR;
 		FileContents += TEXT("bEnableThreadSafeUniqueNetIds=True") LINE_TERMINATOR;
 
 		if (WriteOutputFile(DefaultEngineIniFilename, FileContents, OutFailReason))
diff --git a/Engine/Source/Runtime/CoreOnline/CoreOnline.Build.cs b/Engine/Source/Runtime/CoreOnline/CoreOnline.Build.cs
new file mode 100644
--- /dev/null
+++ b/Engine/Source/Runtime/CoreOnline/CoreOnline.Build.cs
@@ -0,0 +1,27 @@
+// Copyright Epic Games, Inc. All Rights Reserved.
+
+using UnrealBuildTool;
+using EpicGames.Core;
+
+public class CoreOnline : ModuleRules
+{
+	[ConfigFile(ConfigHierarchyType.Engine, "CoreOnline")]
+	bool bEnableThreadSafeUniqueNetIds = false;
+
+	public CoreOnline(ReadOnlyTargetRules Target) : base(Target)
+	{
+		PublicDependencyModuleNames.AddRange(
+			new string[] {
+				"Core",
+				"CoreUObject"
+			}
+		);
+
+		ConfigCache.ReadSettings(DirectoryReference.FromFile(Target.ProjectFile), Target.Platform, this);
+		PublicDefinitions.Add("UNIQUENETID_ESPMODE=ESPMode::" + (bEnableThreadSafeUniqueNetIds ? "ThreadSafe" : "Fast"));
+		// This is to ease migration to ESPMode::ThreadSafe. We have deprecated public FUniqueNetId constructors, by including it in the
+		// ESPMode::Fast deprecation mechanism. The constructors are public when the ESPMode is Fast, and protected when it is ThreadSafe.
+		PublicDefinitions.Add("UNIQUENETID_CONSTRUCTORVIS=" + (bEnableThreadSafeUniqueNetIds ? "protected" : "public"));
+	}
+
+}
diff --git a/Engine/Source/Runtime/CoreOnline/Private/CoreOnlinePrivate.h b/Engine/Source/Runtime/CoreOnline/Private/CoreOnlinePrivate.h
new file mode 100644
--- /dev/null
+++ b/Engine/Source/Runtime/CoreOnline/Private/CoreOnlinePrivate.h
@@ -0,0 +1,12 @@
+// Copyright Epic Games, Inc. All Rights Reserved.
+
+#pragma once
+
+#include "UObject/ObjectMacros.h"
+
+// This exists purely to guarantee a package is created, or the engine will not boot.
+UENUM(BlueprintType)
+enum class ECoreOnlineDummy : uint8
+{
+	Dummy
+};
\ No newline at end of file
diff --git a/Engine/Source/Runtime/CoreOnline/Public/Online/CoreOnline.h b/Engine/Source/Runtime/CoreOnline/Public/Online/CoreOnline.h
new file mode 100644
--- /dev/null
+++ b/Engine/Source/Runtime/CoreOnline/Public/Online/CoreOnline.h
@@ -0,0 +1,350 @@
+// Copyright Epic Games, Inc. All Rights Reserved.
+
+#pragma once
+
+#include "CoreMinimal.h"
+#include "Hash/CityHash.h"
+#include "Online/CoreOnlineFwd.h"
+#include "UObject/ObjectMacros.h"
+
+#define GameSessionName NAME_GameSession
+#define PartySessionName NAME_PartySession
+#define GamePort NAME_GamePort
+#define BeaconPort NAME_BeaconPort
+
+USTRUCT(noexport)
+struct FJoinabilitySettings
+{
+	//GENERATED_BODY()
+
+	/** Name of session these settings affect */
+	UPROPERTY()
+	FName SessionName;
+	/** Is this session now publicly searchable */
+	UPROPERTY()
+	bool bPublicSearchable;
+	/** Does this session allow invites */
+	UPROPERTY()
+	bool bAllowInvites;
+	/** Does this session allow public join via presence */
+	UPROPERTY()
+	bool bJoinViaPresence;
+	/** Does this session allow friends to join via presence */
+	UPROPERTY()
+	bool bJoinViaPresenceFriendsOnly;
+	/** Current max players in this session */
+	UPROPERTY()
+	int32 MaxPlayers;
+	/** Current max party size in this session */
+	UPROPERTY()
+	int32 MaxPartySize;
+
+	FJoinabilitySettings() :
+		SessionName(NAME_None),
+		bPublicSearchable(false),
+		bAllowInvites(false),
+		bJoinViaPresence(false),
+		bJoinViaPresenceFriendsOnly(false),
+		MaxPlayers(0),
+		MaxPartySize(0)
+	{
+	}
+
+	bool operator==(const FJoinabilitySettings& Other) const
+	{
+		return SessionName == Other.SessionName &&
+			bPublicSearchable == Other.bPublicSearchable &&
+			bAllowInvites == Other.bAllowInvites &&
+			bJoinViaPresence == Other.bJoinViaPresence &&
+			bJoinViaPresenceFriendsOnly == Other.bJoinViaPresenceFriendsOnly &&
+			MaxPlayers == Other.MaxPlayers &&
+			MaxPartySize == Other.MaxPartySize;
+	}
+
+	bool operator!=(const FJoinabilitySettings& Other) const
+	{
+		return !(FJoinabilitySettings::operator==(Other));
+	}
+};
+
+/**
+ * Abstraction of a profile service online Id
+ * The class is meant to be opaque
+ */
+class FUniqueNetId : public TSharedFromThis<FUniqueNetId, UNIQUENETID_ESPMODE>
+{
+protected:
+
+	/** Only constructible by derived type */
+	FUniqueNetId() = default;
+
+	FUniqueNetId(const FUniqueNetId& Src) = default;
+	FUniqueNetId& operator=(const FUniqueNetId& Src) = default;
+
+	virtual bool Compare(const FUniqueNetId& Other) const
+	{
+		return (GetSize() == Other.GetSize()) &&
+			(FMemory::Memcmp(GetBytes(), Other.GetBytes(), GetSize()) == 0);
+	}
+
+public:
+
+	virtual ~FUniqueNetId() = default;
+
+	/**
+	 *	Comparison operator
+	 */
+	friend bool operator==(const FUniqueNetId& Lhs, const FUniqueNetId& Rhs)
+	{
+		return Lhs.Compare(Rhs);
+	}
+
+	friend bool operator!=(const FUniqueNetId& Lhs, const FUniqueNetId& Rhs)
+	{
+		return !Lhs.Compare(Rhs);
+	}
+
+	/**
+	 * Get the type token for this opaque data
+	 * This is useful for inferring UniqueId subclasses and knowing which OSS it "goes with"
+	 *
+	 * @return FName representing the Type
+	 */
+	virtual FName GetType() const { return NAME_None; /* This should be pure virtual, however, older versions of the OSS plugins cannot handle that */ }
+
+	/** 
+	 * Get the raw byte representation of this opaque data
+	 * This data is platform dependent and shouldn't be manipulated directly
+	 *
+	 * @return byte array of size GetSize()
+	 */
+	virtual const uint8* GetBytes() const = 0;
+
+	/** 
+	 * Get the size of the opaque data
+	 *
+	 * @return size in bytes of the data representation
+	 */
+	virtual int32 GetSize() const = 0;
+
+	/**
+	 * Check the validity of the opaque data
+	 *
+	 * @return true if this is well formed data, false otherwise
+	 */
+	virtual bool IsValid() const = 0;
+
+	/**
+	 * Platform specific conversion to string representation of data
+	 *
+	 * @return data in string form
+	 */
+	virtual FString ToString() const = 0;
+
+	/**
+	 * Get a human readable representation of the opaque data
+	 * Shouldn't be used for anything other than logging/debugging
+	 *
+	 * @return data in string form
+	 */
+	virtual FString ToDebugString() const = 0;
+
+	/**
+	 * @return hex encoded string representation of unique id
+	 */
+	FString GetHexEncodedString() const
+	{
+		if (GetSize() > 0 && GetBytes() != NULL)
+		{
+			return BytesToHex(GetBytes(), GetSize());
+		}
+		return FString();
+	}
+
+	friend inline uint32 GetTypeHash(const FUniqueNetId& Value)
+	{
+		return CityHash32(reinterpret_cast<const char*>(Value.GetBytes()), Value.GetSize());
+	}
+};
+
+USTRUCT(noexport)
+struct FUniqueNetIdWrapper
+{
+	//GENERATED_BODY()
+
+	FUniqueNetIdWrapper() = default;
+	virtual ~FUniqueNetIdWrapper() = default;
+
+	// copy operators generated by compiler
+
+	FUniqueNetIdWrapper(const FUniqueNetIdRef& InUniqueNetId)
+		: UniqueNetId(InUniqueNetId)
+	{
+	}
+
+	FUniqueNetIdWrapper(const FUniqueNetIdPtr& InUniqueNetId)
+		: UniqueNetId(InUniqueNetId)
+	{
+	}
+
+	// temporarily restored implicit conversion from FUniqueNetId
+	FUniqueNetIdWrapper(const FUniqueNetId& InUniqueNetId)
+		: UniqueNetId(InUniqueNetId.AsShared())
+	{
+	}
+
+	FName GetType() const
+	{
+		return IsValid() ? UniqueNetId->GetType() : NAME_None;
+	}
+	
+	/** Convert this value to a string */
+	FString ToString() const
+	{
+		return IsValid() ? UniqueNetId->ToString() : TEXT("INVALID");
+	}
+
+	/** Convert this value to a string with additional information */
+	FString ToDebugString() const
+	{
+		return IsValid() ? FString::Printf(TEXT("%s:%s"), *UniqueNetId->GetType().ToString(), *UniqueNetId->ToDebugString()) : TEXT("INVALID");
+	}
+
+	/** Is the FUniqueNetId wrapped in this object valid */
+	bool IsValid() const
+	{
+		return UniqueNetId.IsValid() && UniqueNetId->IsValid();
+	}
+	
+	/** 
+	 * Assign a unique id to this wrapper object
+	 *
+	 * @param InUniqueNetId id to associate
+	 */
+	virtual void SetUniqueNetId(const FUniqueNetIdPtr& InUniqueNetId)
+	{
+		UniqueNetId = InUniqueNetId;
+	}
+
+	/** @return unique id associated with this wrapper object */
+	const FUniqueNetIdPtr& GetUniqueNetId() const
+	{
+		return UniqueNetId;
+	}
+
+	/**
+	 * Dereference operator returns a reference to the FUniqueNetId
+	 */
+	const FUniqueNetId& operator*() const
+	{
+		return *UniqueNetId;
+	}
+
+	/**
+	 * Arrow operator returns a pointer to this FUniqueNetId
+	 */
+	const FUniqueNetId* operator->() const
+	{
+		return UniqueNetId.Get();
+	}
+
+	/**
+	* Friend function for using FUniqueNetIdWrapper as a hashable key
+	*/
+	friend inline uint32 GetTypeHash(FUniqueNetIdWrapper const& Value)
+	{
+		if (Value.IsValid())
+		{
+			return GetTypeHash(*Value);
+		}
+		else
+		{
+			// If we hit this, something went wrong and we have received an unhashable wrapper.
+			return INDEX_NONE;
+		}
+	}
+
+	static FUniqueNetIdWrapper Invalid()
+	{
+		static FUniqueNetIdWrapper InvalidId(nullptr);
+		return InvalidId;
+	}
+
+	friend bool operator==(const FUniqueNetIdWrapper& Lhs, const FUniqueNetIdWrapper& Rhs)
+	{
+		bool bLhsValid = Lhs.IsValid();
+		return bLhsValid == Rhs.IsValid() && (!bLhsValid || *Lhs == *Rhs);
+	}
+
+	friend bool operator!=(const FUniqueNetIdWrapper& Lhs, const FUniqueNetIdWrapper& Rhs)
+	{
+		return !(Lhs == Rhs);
+	}
+
+	friend bool operator==(const FUniqueNetIdWrapper& Lhs, const FUniqueNetId& Rhs)
+	{
+		bool bLhsValid = Lhs.IsValid();
+		return bLhsValid == Rhs.IsValid() && (!bLhsValid || *Lhs == Rhs);
+	}
+
+	friend bool operator!=(const FUniqueNetIdWrapper& Lhs, const FUniqueNetId& Rhs)
+	{
+		return !(Lhs == Rhs);
+	}
+
+	friend bool operator==(const FUniqueNetId& Lhs, const FUniqueNetIdWrapper& Rhs)
+	{
+		return Rhs == Lhs;
+	}
+
+	friend bool operator!=(const FUniqueNetId& Lhs, const FUniqueNetIdWrapper& Rhs)
+	{
+		return Rhs != Lhs;
+	}
+
+	// comparison with nullptr (alternative to IsValid)
+	friend bool operator==(const FUniqueNetIdWrapper& NetIdWrapper, TYPE_OF_NULLPTR)
+	{
+		return !NetIdWrapper.IsValid();
+	}
+
+	friend bool operator!=(const FUniqueNetIdWrapper& NetIdWrapper, TYPE_OF_NULLPTR)
+	{
+		return NetIdWrapper.IsValid();
+	}
+
+	friend bool operator==(TYPE_OF_NULLPTR, const FUniqueNetIdWrapper& NetIdWrapper)
+	{
+		return !NetIdWrapper.IsValid();
+	}
+
+	friend bool operator!=(TYPE_OF_NULLPTR, const FUniqueNetIdWrapper& NetIdWrapper)
+	{
+		return NetIdWrapper.IsValid();
+	}
+
+protected:
+
+	// Actual unique id
+	FUniqueNetIdPtr UniqueNetId;
+};
+
+template <typename ValueType>
+struct TUniqueNetIdMapKeyFuncs : public TDefaultMapKeyFuncs<FUniqueNetIdRef, ValueType, false>
+{
+	static FORCEINLINE FUniqueNetIdRef GetSetKey(TPair<FUniqueNetIdRef, ValueType> const& Element) { return Element.Key; }
+	static FORCEINLINE uint32 GetKeyHash(FUniqueNetIdRef const& Key) { return GetTypeHash(*Key); }
+	static FORCEINLINE bool Matches(FUniqueNetIdRef const& A, FUniqueNetIdRef const& B) { return (A == B) || (*A == *B); }
+};
+
+template <typename ValueType>
+using TUniqueNetIdMap = TMap<FUniqueNetIdRef, ValueType, FDefaultSetAllocator, TUniqueNetIdMapKeyFuncs<ValueType>>;
+
+struct FUniqueNetIdKeyFuncs : public DefaultKeyFuncs<FUniqueNetIdRef>
+{
+	static FORCEINLINE FUniqueNetIdRef GetSetKey(FUniqueNetIdRef const& Element) { return Element; }
+	static FORCEINLINE uint32 GetKeyHash(FUniqueNetIdRef const& Key) { return GetTypeHash(*Key); }
+	static FORCEINLINE bool Matches(FUniqueNetIdRef const& A, FUniqueNetIdRef const& B) { return (A == B) || (*A == *B); }
+};
+
+using FUniqueNetIdSet = TSet<FUniqueNetIdRef, FUniqueNetIdKeyFuncs>;
diff --git a/Engine/Source/Runtime/CoreOnline/Public/Online/CoreOnlineFwd.h b/Engine/Source/Runtime/CoreOnline/Public/Online/CoreOnlineFwd.h
new file mode 100644
--- /dev/null
+++ b/Engine/Source/Runtime/CoreOnline/Public/Online/CoreOnlineFwd.h
@@ -0,0 +1,12 @@
+// Copyright Epic Games, Inc. All Rights Reserved.
+
+#pragma once
+
+#include "Templates/SharedPointer.h"
+
+class FUniqueNetId;
+struct FUniqueNetIdWrapper;
+
+using FUniqueNetIdPtr = TSharedPtr<const FUniqueNetId, UNIQUENETID_ESPMODE>;
+using FUniqueNetIdRef = TSharedRef<const FUniqueNetId, UNIQUENETID_ESPMODE>;
+using FUniqueNetIdWeakPtr = TWeakPtr<const FUniqueNetId, UNIQUENETID_ESPMODE>;
\ No newline at end of file
diff --git a/Engine/Source/Runtime/CoreUObject/CoreUObject.Build.cs b/Engine/Source/Runtime/CoreUObject/CoreUObject.Build.cs
--- a/Engine/Source/Runtime/CoreUObject/CoreUObject.Build.cs
+++ b/Engine/Source/Runtime/CoreUObject/CoreUObject.Build.cs
@@ -1,13 +1,9 @@
 // Copyright Epic Games, Inc. All Rights Reserved.
 
 using UnrealBuildTool;
-using EpicGames.Core;
 
 public class CoreUObject : ModuleRules
 {
-	[ConfigFile(ConfigHierarchyType.Engine, "CoreUObject")]
-	bool bEnableThreadSafeUniqueNetIds = false;
-
 	public CoreUObject(ReadOnlyTargetRules Target) : base(Target)
 	{
 		PrivatePCHHeaderFile = "Private/CoreUObjectPrivatePCH.h";
@@ -28,13 +24,5 @@ public class CoreUObject : ModuleRules
 
 		PrivateDependencyModuleNames.Add("Projects");
         PrivateDependencyModuleNames.Add("Json");
-
-		ConfigCache.ReadSettings(DirectoryReference.FromFile(Target.ProjectFile), Target.Platform, this);
-		bool bUseThreadSafeUniqueNetIds = bEnableThreadSafeUniqueNetIds && Target.Type != TargetType.Program;
-		PublicDefinitions.Add("UNIQUENETID_ESPMODE=ESPMode::" + (bUseThreadSafeUniqueNetIds ? "ThreadSafe" : "Fast"));
-		// This is to ease migration to ESPMode::ThreadSafe. We have deprecated public FUniqueNetId constructors, by including it in the
-		// ESPMode::Fast deprecation mechanism. The constructors are public when the ESPMode is Fast, and protected when it is ThreadSafe.
-		PublicDefinitions.Add("UNIQUENETID_CONSTRUCTORVIS=" + (bUseThreadSafeUniqueNetIds ? "protected" : "public"));
 	}
-
 }
diff --git a/Engine/Source/Runtime/CoreUObject/Public/UObject/CoreOnline.h b/Engine/Source/Runtime/CoreUObject/Public/UObject/CoreOnline.h
--- a/Engine/Source/Runtime/CoreUObject/Public/UObject/CoreOnline.h
+++ b/Engine/Source/Runtime/CoreUObject/Public/UObject/CoreOnline.h
@@ -2,376 +2,4 @@
 
 #pragma once
 
-#include "CoreMinimal.h"
-#include "Hash/CityHash.h"
-#include "UObject/CoreOnlineFwd.h"
-#include "UObject/ObjectMacros.h"
-
-#define GameSessionName NAME_GameSession
-#define PartySessionName NAME_PartySession
-#define GamePort NAME_GamePort
-#define BeaconPort NAME_BeaconPort
-
-#if !CPP
-// Circular dependency on Core vs UHT means we have to noexport these structs so tools can build
-USTRUCT(noexport)
-struct FJoinabilitySettings
-{
-	//GENERATED_BODY()
-
-	/** Name of session these settings affect */
-	UPROPERTY()
-	FName SessionName;
-	/** Is this session now publicly searchable */
-	UPROPERTY()
-	bool bPublicSearchable;
-	/** Does this session allow invites */
-	UPROPERTY()
-	bool bAllowInvites;
-	/** Does this session allow public join via presence */
-	UPROPERTY()
-	bool bJoinViaPresence;
-	/** Does this session allow friends to join via presence */
-	UPROPERTY()
-	bool bJoinViaPresenceFriendsOnly;
-	/** Current max players in this session */
-	UPROPERTY()
-	int32 MaxPlayers;
-	/** Current max party size in this session */
-	UPROPERTY()
-	int32 MaxPartySize;
-};
-
-USTRUCT(noexport)
-struct FUniqueNetIdWrapper
-{
-	//GENERATED_BODY()
-};
-
-#endif
-
-struct FJoinabilitySettings
-{
-	/** Name of session these settings affect */
-	FName SessionName;
-	/** Is this session now publicly searchable */
-	bool bPublicSearchable;
-	/** Does this session allow invites */
-	bool bAllowInvites;
-	/** Does this session allow public join via presence */
-	bool bJoinViaPresence;
-	/** Does this session allow friends to join via presence */
-	bool bJoinViaPresenceFriendsOnly;
-	/** Current max players in this session */
-	int32 MaxPlayers;
-	/** Current max party size in this session */
-	int32 MaxPartySize;
-
-	FJoinabilitySettings() :
-		SessionName(NAME_None),
-		bPublicSearchable(false),
-		bAllowInvites(false),
-		bJoinViaPresence(false),
-		bJoinViaPresenceFriendsOnly(false),
-		MaxPlayers(0),
-		MaxPartySize(0)
-	{
-	}
-
-	bool operator==(const FJoinabilitySettings& Other) const
-	{
-		return SessionName == Other.SessionName &&
-			bPublicSearchable == Other.bPublicSearchable &&
-			bAllowInvites == Other.bAllowInvites &&
-			bJoinViaPresence == Other.bJoinViaPresence &&
-			bJoinViaPresenceFriendsOnly == Other.bJoinViaPresenceFriendsOnly &&
-			MaxPlayers == Other.MaxPlayers &&
-			MaxPartySize == Other.MaxPartySize;
-	}
-
-	bool operator!=(const FJoinabilitySettings& Other) const
-	{
-		return !(FJoinabilitySettings::operator==(Other));
-	}
-};
-
-/**
- * Abstraction of a profile service online Id
- * The class is meant to be opaque
- */
-class FUniqueNetId : public TSharedFromThis<FUniqueNetId, UNIQUENETID_ESPMODE>
-{
-protected:
-
-	/** Only constructible by derived type */
-	FUniqueNetId() = default;
-
-	FUniqueNetId(const FUniqueNetId& Src) = default;
-	FUniqueNetId& operator=(const FUniqueNetId& Src) = default;
-
-	virtual bool Compare(const FUniqueNetId& Other) const
-	{
-		return (GetSize() == Other.GetSize()) &&
-			(FMemory::Memcmp(GetBytes(), Other.GetBytes(), GetSize()) == 0);
-	}
-
-public:
-
-	virtual ~FUniqueNetId() = default;
-
-	/**
-	 *	Comparison operator
-	 */
-	friend bool operator==(const FUniqueNetId& Lhs, const FUniqueNetId& Rhs)
-	{
-		return Lhs.Compare(Rhs);
-	}
-
-	friend bool operator!=(const FUniqueNetId& Lhs, const FUniqueNetId& Rhs)
-	{
-		return !Lhs.Compare(Rhs);
-	}
-
-	/**
-	 * Get the type token for this opaque data
-	 * This is useful for inferring UniqueId subclasses and knowing which OSS it "goes with"
-	 *
-	 * @return FName representing the Type
-	 */
-	virtual FName GetType() const { return NAME_None; /* This should be pure virtual, however, older versions of the OSS plugins cannot handle that */ }
-
-	/** 
-	 * Get the raw byte representation of this opaque data
-	 * This data is platform dependent and shouldn't be manipulated directly
-	 *
-	 * @return byte array of size GetSize()
-	 */
-	virtual const uint8* GetBytes() const = 0;
-
-	/** 
-	 * Get the size of the opaque data
-	 *
-	 * @return size in bytes of the data representation
-	 */
-	virtual int32 GetSize() const = 0;
-
-	/**
-	 * Check the validity of the opaque data
-	 *
-	 * @return true if this is well formed data, false otherwise
-	 */
-	virtual bool IsValid() const = 0;
-
-	/**
-	 * Platform specific conversion to string representation of data
-	 *
-	 * @return data in string form
-	 */
-	virtual FString ToString() const = 0;
-
-	/**
-	 * Get a human readable representation of the opaque data
-	 * Shouldn't be used for anything other than logging/debugging
-	 *
-	 * @return data in string form
-	 */
-	virtual FString ToDebugString() const = 0;
-
-	/**
-	 * @return hex encoded string representation of unique id
-	 */
-	FString GetHexEncodedString() const
-	{
-		if (GetSize() > 0 && GetBytes() != NULL)
-		{
-			return BytesToHex(GetBytes(), GetSize());
-		}
-		return FString();
-	}
-
-	friend inline uint32 GetTypeHash(const FUniqueNetId& Value)
-	{
-		return CityHash32(reinterpret_cast<const char*>(Value.GetBytes()), Value.GetSize());
-	}
-};
-
-struct FUniqueNetIdWrapper
-{
-public:
-
-	FUniqueNetIdWrapper() = default;
-	virtual ~FUniqueNetIdWrapper() = default;
-
-	// copy operators generated by compiler
-
-	FUniqueNetIdWrapper(const FUniqueNetIdRef& InUniqueNetId)
-		: UniqueNetId(InUniqueNetId)
-	{
-	}
-
-	FUniqueNetIdWrapper(const FUniqueNetIdPtr& InUniqueNetId)
-		: UniqueNetId(InUniqueNetId)
-	{
-	}
-
-	// temporarily restored implicit conversion from FUniqueNetId
-	FUniqueNetIdWrapper(const FUniqueNetId& InUniqueNetId)
-		: UniqueNetId(InUniqueNetId.AsShared())
-	{
-	}
-
-	FName GetType() const
-	{
-		return IsValid() ? UniqueNetId->GetType() : NAME_None;
-	}
-	
-	/** Convert this value to a string */
-	FString ToString() const
-	{
-		return IsValid() ? UniqueNetId->ToString() : TEXT("INVALID");
-	}
-
-	/** Convert this value to a string with additional information */
-	FString ToDebugString() const
-	{
-		return IsValid() ? FString::Printf(TEXT("%s:%s"), *UniqueNetId->GetType().ToString(), *UniqueNetId->ToDebugString()) : TEXT("INVALID");
-	}
-
-	/** Is the FUniqueNetId wrapped in this object valid */
-	bool IsValid() const
-	{
-		return UniqueNetId.IsValid() && UniqueNetId->IsValid();
-	}
-	
-	/** 
-	 * Assign a unique id to this wrapper object
-	 *
-	 * @param InUniqueNetId id to associate
-	 */
-	virtual void SetUniqueNetId(const FUniqueNetIdPtr& InUniqueNetId)
-	{
-		UniqueNetId = InUniqueNetId;
-	}
-
-	/** @return unique id associated with this wrapper object */
-	const FUniqueNetIdPtr& GetUniqueNetId() const
-	{
-		return UniqueNetId;
-	}
-
-	/**
-	 * Dereference operator returns a reference to the FUniqueNetId
-	 */
-	const FUniqueNetId& operator*() const
-	{
-		return *UniqueNetId;
-	}
-
-	/**
-	 * Arrow operator returns a pointer to this FUniqueNetId
-	 */
-	const FUniqueNetId* operator->() const
-	{
-		return UniqueNetId.Get();
-	}
-
-	/**
-	* Friend function for using FUniqueNetIdWrapper as a hashable key
-	*/
-	friend inline uint32 GetTypeHash(FUniqueNetIdWrapper const& Value)
-	{
-		if (Value.IsValid())
-		{
-			return GetTypeHash(*Value);
-		}
-		else
-		{
-			// If we hit this, something went wrong and we have received an unhashable wrapper.
-			return INDEX_NONE;
-		}
-	}
-
-	static FUniqueNetIdWrapper Invalid()
-	{
-		static FUniqueNetIdWrapper InvalidId(nullptr);
-		return InvalidId;
-	}
-
-	friend bool operator==(const FUniqueNetIdWrapper& Lhs, const FUniqueNetIdWrapper& Rhs)
-	{
-		bool bLhsValid = Lhs.IsValid();
-		return bLhsValid == Rhs.IsValid() && (!bLhsValid || *Lhs == *Rhs);
-	}
-
-	friend bool operator!=(const FUniqueNetIdWrapper& Lhs, const FUniqueNetIdWrapper& Rhs)
-	{
-		return !(Lhs == Rhs);
-	}
-
-	friend bool operator==(const FUniqueNetIdWrapper& Lhs, const FUniqueNetId& Rhs)
-	{
-		bool bLhsValid = Lhs.IsValid();
-		return bLhsValid == Rhs.IsValid() && (!bLhsValid || *Lhs == Rhs);
-	}
-
-	friend bool operator!=(const FUniqueNetIdWrapper& Lhs, const FUniqueNetId& Rhs)
-	{
-		return !(Lhs == Rhs);
-	}
-
-	friend bool operator==(const FUniqueNetId& Lhs, const FUniqueNetIdWrapper& Rhs)
-	{
-		return Rhs == Lhs;
-	}
-
-	friend bool operator!=(const FUniqueNetId& Lhs, const FUniqueNetIdWrapper& Rhs)
-	{
-		return Rhs != Lhs;
-	}
-
-	// comparison with nullptr (alternative to IsValid)
-	friend bool operator==(const FUniqueNetIdWrapper& NetIdWrapper, TYPE_OF_NULLPTR)
-	{
-		return !NetIdWrapper.IsValid();
-	}
-
-	friend bool operator!=(const FUniqueNetIdWrapper& NetIdWrapper, TYPE_OF_NULLPTR)
-	{
-		return NetIdWrapper.IsValid();
-	}
-
-	friend bool operator==(TYPE_OF_NULLPTR, const FUniqueNetIdWrapper& NetIdWrapper)
-	{
-		return !NetIdWrapper.IsValid();
-	}
-
-	friend bool operator!=(TYPE_OF_NULLPTR, const FUniqueNetIdWrapper& NetIdWrapper)
-	{
-		return NetIdWrapper.IsValid();
-	}
-
-protected:
-
-	// Actual unique id
-	FUniqueNetIdPtr UniqueNetId;
-};
-
-template <typename ValueType>
-struct TUniqueNetIdMapKeyFuncs : public TDefaultMapKeyFuncs<FUniqueNetIdRef, ValueType, false>
-{
-	static FORCEINLINE FUniqueNetIdRef GetSetKey(TPair<FUniqueNetIdRef, ValueType> const& Element) { return Element.Key; }
-	static FORCEINLINE uint32 GetKeyHash(FUniqueNetIdRef const& Key) { return GetTypeHash(*Key); }
-	static FORCEINLINE bool Matches(FUniqueNetIdRef const& A, FUniqueNetIdRef const& B) { return (A == B) || (*A == *B); }
-};
-
-template <typename ValueType>
-using TUniqueNetIdMap = TMap<FUniqueNetIdRef, ValueType, FDefaultSetAllocator, TUniqueNetIdMapKeyFuncs<ValueType>>;
-
-struct FUniqueNetIdKeyFuncs : public DefaultKeyFuncs<FUniqueNetIdRef>
-{
-	static FORCEINLINE FUniqueNetIdRef GetSetKey(FUniqueNetIdRef const& Element) { return Element; }
-	static FORCEINLINE uint32 GetKeyHash(FUniqueNetIdRef const& Key) { return GetTypeHash(*Key); }
-	static FORCEINLINE bool Matches(FUniqueNetIdRef const& A, FUniqueNetIdRef const& B) { return (A == B) || (*A == *B); }
-};
-
-using FUniqueNetIdSet = TSet<FUniqueNetIdRef, FUniqueNetIdKeyFuncs>;
+static_assert(false, "UObject/CoreOnline.h has moved to Online/CoreOnline.h in the CoreOnline module. Please update #includes and module dependencies.");
\ No newline at end of file
diff --git a/Engine/Source/Runtime/CoreUObject/Public/UObject/CoreOnlineFwd.h b/Engine/Source/Runtime/CoreUObject/Public/UObject/CoreOnlineFwd.h
--- a/Engine/Source/Runtime/CoreUObject/Public/UObject/CoreOnlineFwd.h
+++ b/Engine/Source/Runtime/CoreUObject/Public/UObject/CoreOnlineFwd.h
@@ -2,11 +2,4 @@
 
 #pragma once
 
-#include "Templates/SharedPointer.h"
-
-class FUniqueNetId;
-struct FUniqueNetIdWrapper;
-
-using FUniqueNetIdPtr = TSharedPtr<const FUniqueNetId, UNIQUENETID_ESPMODE>;
-using FUniqueNetIdRef = TSharedRef<const FUniqueNetId, UNIQUENETID_ESPMODE>;
-using FUniqueNetIdWeakPtr = TWeakPtr<const FUniqueNetId, UNIQUENETID_ESPMODE>;
\ No newline at end of file
+static_assert(false, "UObject/CoreOnlineFwd.h has moved to Online/CoreOnlineFwd.h in the CoreOnline module. Please update #includes and module dependencies.");
\ No newline at end of file
diff --git a/Engine/Source/Runtime/Engine/Classes/Engine/GameInstance.h b/Engine/Source/Runtime/Engine/Classes/Engine/GameInstance.h
--- a/Engine/Source/Runtime/Engine/Classes/Engine/GameInstance.h
+++ b/Engine/Source/Runtime/Engine/Classes/Engine/GameInstance.h
@@ -3,7 +3,7 @@
 #pragma once
 
 #include "CoreMinimal.h"
-#include "UObject/CoreOnlineFwd.h"
+#include "Online/CoreOnlineFwd.h"
 #include "UObject/ObjectMacros.h"
 #include "UObject/ScriptMacros.h"
 #include "UObject/Object.h"
diff --git a/Engine/Source/Runtime/Engine/Classes/Engine/LocalPlayer.h b/Engine/Source/Runtime/Engine/Classes/Engine/LocalPlayer.h
--- a/Engine/Source/Runtime/Engine/Classes/Engine/LocalPlayer.h
+++ b/Engine/Source/Runtime/Engine/Classes/Engine/LocalPlayer.h
@@ -12,7 +12,7 @@
 #include "Engine/EngineTypes.h"
 #include "Input/Reply.h"
 #include "Engine/GameViewportClient.h"
-#include "UObject/CoreOnline.h"
+#include "Online/CoreOnline.h"
 #include "SceneTypes.h"
 #include "Engine/Player.h"
 #include "GameFramework/OnlineReplStructs.h"
diff --git a/Engine/Source/Runtime/Engine/Classes/Engine/World.h b/Engine/Source/Runtime/Engine/Classes/Engine/World.h
--- a/Engine/Source/Runtime/Engine/Classes/Engine/World.h
+++ b/Engine/Source/Runtime/Engine/Classes/Engine/World.h
@@ -4,7 +4,7 @@
 
 #include "CoreMinimal.h"
 #include "HAL/ThreadSafeCounter.h"
-#include "UObject/CoreOnlineFwd.h"
+#include "Online/CoreOnlineFwd.h"
 #include "UObject/ObjectMacros.h"
 #include "UObject/UObjectGlobals.h"
 #include "UObject/Object.h"
diff --git a/Engine/Source/Runtime/Engine/Classes/GameFramework/GameModeBase.h b/Engine/Source/Runtime/Engine/Classes/GameFramework/GameModeBase.h
--- a/Engine/Source/Runtime/Engine/Classes/GameFramework/GameModeBase.h
+++ b/Engine/Source/Runtime/Engine/Classes/GameFramework/GameModeBase.h
@@ -10,7 +10,7 @@
 #include "GameFramework/Pawn.h"
 #include "GameFramework/Info.h"
 #include "Engine/ServerStatReplicator.h"
-#include "UObject/CoreOnline.h"
+#include "Online/CoreOnline.h"
 #include "GameFramework/PlayerController.h"
 #include "GameModeBase.generated.h"
 
diff --git a/Engine/Source/Runtime/Engine/Classes/GameFramework/GameSession.h b/Engine/Source/Runtime/Engine/Classes/GameFramework/GameSession.h
--- a/Engine/Source/Runtime/Engine/Classes/GameFramework/GameSession.h
+++ b/Engine/Source/Runtime/Engine/Classes/GameFramework/GameSession.h
@@ -9,7 +9,7 @@
 #include "CoreMinimal.h"
 #include "UObject/ObjectMacros.h"
 #include "GameFramework/Actor.h"
-#include "UObject/CoreOnline.h"
+#include "Online/CoreOnline.h"
 #include "GameFramework/Info.h"
 #include "GameSession.generated.h"
 
diff --git a/Engine/Source/Runtime/Engine/Classes/GameFramework/OnlineReplStructs.h b/Engine/Source/Runtime/Engine/Classes/GameFramework/OnlineReplStructs.h
--- a/Engine/Source/Runtime/Engine/Classes/GameFramework/OnlineReplStructs.h
+++ b/Engine/Source/Runtime/Engine/Classes/GameFramework/OnlineReplStructs.h
@@ -9,7 +9,7 @@
 #include "CoreMinimal.h"
 #include "UObject/ObjectMacros.h"
 #include "UObject/Class.h"
-#include "UObject/CoreOnline.h"
+#include "Online/CoreOnline.h"
 #include "OnlineReplStructs.generated.h"
 
 class FJsonValue;
diff --git a/Engine/Source/Runtime/Engine/Classes/GameFramework/OnlineSession.h b/Engine/Source/Runtime/Engine/Classes/GameFramework/OnlineSession.h
--- a/Engine/Source/Runtime/Engine/Classes/GameFramework/OnlineSession.h
+++ b/Engine/Source/Runtime/Engine/Classes/GameFramework/OnlineSession.h
@@ -9,7 +9,7 @@
 #include "CoreMinimal.h"
 #include "UObject/ObjectMacros.h"
 #include "UObject/Object.h"
-#include "UObject/CoreOnline.h"
+#include "Online/CoreOnline.h"
 #include "OnlineSession.generated.h"
 
 class FOnlineSessionSearchResult;
diff --git a/Engine/Source/Runtime/Engine/Classes/GameFramework/PlayerMuteList.h b/Engine/Source/Runtime/Engine/Classes/GameFramework/PlayerMuteList.h
--- a/Engine/Source/Runtime/Engine/Classes/GameFramework/PlayerMuteList.h
+++ b/Engine/Source/Runtime/Engine/Classes/GameFramework/PlayerMuteList.h
@@ -4,7 +4,7 @@
 
 #include "CoreMinimal.h"
 #include "UObject/ObjectMacros.h"
-#include "UObject/CoreOnline.h"
+#include "Online/CoreOnline.h"
 #include "PlayerMuteList.generated.h"
 
 struct FUniqueNetIdRepl;
diff --git a/Engine/Source/Runtime/Engine/Classes/GameFramework/PlayerState.h b/Engine/Source/Runtime/Engine/Classes/GameFramework/PlayerState.h
--- a/Engine/Source/Runtime/Engine/Classes/GameFramework/PlayerState.h
+++ b/Engine/Source/Runtime/Engine/Classes/GameFramework/PlayerState.h
@@ -6,7 +6,7 @@
 #include "UObject/ObjectMacros.h"
 #include "Templates/SubclassOf.h"
 #include "GameFramework/Actor.h"
-#include "UObject/CoreOnline.h"
+#include "Online/CoreOnline.h"
 #include "GameFramework/OnlineReplStructs.h"
 #include "GameFramework/Info.h"
 #include "PlayerState.generated.h"
diff --git a/Engine/Source/Runtime/Engine/Engine.Build.cs b/Engine/Source/Runtime/Engine/Engine.Build.cs
--- a/Engine/Source/Runtime/Engine/Engine.Build.cs
+++ b/Engine/Source/Runtime/Engine/Engine.Build.cs
@@ -66,6 +66,7 @@ public class Engine : ModuleRules
 		PublicDependencyModuleNames.AddRange(
 			new string[] {
 				"Core",
+				"CoreOnline",
 				"CoreUObject",
 				"NetCore",
 				"ApplicationCore",
diff --git a/Engine/Source/Runtime/Engine/Private/EnginePrivatePCH.h b/Engine/Source/Runtime/Engine/Private/EnginePrivatePCH.h
--- a/Engine/Source/Runtime/Engine/Private/EnginePrivatePCH.h
+++ b/Engine/Source/Runtime/Engine/Private/EnginePrivatePCH.h
@@ -278,7 +278,7 @@
 #include "UObject/Package.h"
 #include "UObject/UObjectHash.h"
 #include "UObject/UObjectIterator.h"
-#include "UObject/CoreOnline.h"
+#include "Online/CoreOnline.h"
 #include "UObject/TextProperty.h"
 #include "Misc/PackageName.h"
 #include "UObject/ObjectResource.h"
diff --git a/Engine/Source/Runtime/Engine/Private/GameInstance.cpp b/Engine/Source/Runtime/Engine/Private/GameInstance.cpp
--- a/Engine/Source/Runtime/Engine/Private/GameInstance.cpp
+++ b/Engine/Source/Runtime/Engine/Private/GameInstance.cpp
@@ -11,7 +11,7 @@
 #include "Engine/World.h"
 #include "AI/NavigationSystemBase.h"
 #include "Misc/Paths.h"
-#include "UObject/CoreOnline.h"
+#include "Online/CoreOnline.h"
 #include "GameFramework/PlayerController.h"
 #include "Engine/Engine.h"
 #include "Engine/Console.h"
diff --git a/Engine/Source/Runtime/Engine/Private/LevelActor.cpp b/Engine/Source/Runtime/Engine/Private/LevelActor.cpp
--- a/Engine/Source/Runtime/Engine/Private/LevelActor.cpp
+++ b/Engine/Source/Runtime/Engine/Private/LevelActor.cpp
@@ -25,7 +25,7 @@
 #include "AI/NavigationSystemBase.h"
 #include "Engine/Brush.h"
 #include "UObject/LinkerLoad.h"
-#include "UObject/CoreOnline.h"
+#include "Online/CoreOnline.h"
 #include "GameFramework/OnlineReplStructs.h"
 #include "Engine/Engine.h"
 #include "Engine/LevelStreaming.h"
diff --git a/Engine/Source/Runtime/Engine/Public/EngineSharedPCH.h b/Engine/Source/Runtime/Engine/Public/EngineSharedPCH.h
--- a/Engine/Source/Runtime/Engine/Public/EngineSharedPCH.h
+++ b/Engine/Source/Runtime/Engine/Public/EngineSharedPCH.h
@@ -304,7 +304,7 @@
 #include "UObject/GCObject.h"
 #include "Serialization/BulkData.h"
 #include "UObject/Package.h"
-#include "UObject/CoreOnline.h"
+#include "Online/CoreOnline.h"
 #include "UObject/LinkerLoad.h"
 #include "UObject/ObjectResource.h"
 #include "UObject/Linker.h"
diff --git a/Engine/Source/Runtime/Engine/Public/Net/OnlineEngineInterface.h b/Engine/Source/Runtime/Engine/Public/Net/OnlineEngineInterface.h
--- a/Engine/Source/Runtime/Engine/Public/Net/OnlineEngineInterface.h
+++ b/Engine/Source/Runtime/Engine/Public/Net/OnlineEngineInterface.h
@@ -5,7 +5,7 @@
 #include "CoreMinimal.h"
 #include "UObject/ObjectMacros.h"
 #include "UObject/Object.h"
-#include "UObject/CoreOnline.h"
+#include "Online/CoreOnline.h"
 #include "OnlineEngineInterface.generated.h"
 
 class FVoicePacket;
diff --git a/Engine/Source/Runtime/Engine/Public/Net/VoiceDataCommon.h b/Engine/Source/Runtime/Engine/Public/Net/VoiceDataCommon.h
--- a/Engine/Source/Runtime/Engine/Public/Net/VoiceDataCommon.h
+++ b/Engine/Source/Runtime/Engine/Public/Net/VoiceDataCommon.h
@@ -3,7 +3,7 @@
 #pragma once
 
 #include "CoreMinimal.h"
-#include "UObject/CoreOnline.h"
+#include "Online/CoreOnline.h"
 #include "Net/VoiceConfig.h"
 #ifndef MAX_VOICE_DATA_SIZE
 	#define MAX_VOICE_DATA_SIZE 8*1024
diff --git a/Engine/Source/Runtime/NetworkReplayStreaming/LocalFileNetworkReplayStreaming/Private/LocalFileNetworkReplayStreaming.cpp b/Engine/Source/Runtime/NetworkReplayStreaming/LocalFileNetworkReplayStreaming/Private/LocalFileNetworkReplayStreaming.cpp
--- a/Engine/Source/Runtime/NetworkReplayStreaming/LocalFileNetworkReplayStreaming/Private/LocalFileNetworkReplayStreaming.cpp
+++ b/Engine/Source/Runtime/NetworkReplayStreaming/LocalFileNetworkReplayStreaming/Private/LocalFileNetworkReplayStreaming.cpp
@@ -10,7 +10,7 @@
 #include "Engine/Engine.h"
 #include "Engine/LocalPlayer.h"
 #include "Engine/World.h"
-#include "UObject/CoreOnline.h"
+#include "Online/CoreOnline.h"
 #include "Serialization/LargeMemoryReader.h"
 #include "Misc/ConfigCacheIni.h"
 
