Sparse class data changes to allow compiler-generated structs to be used

Changes to serialization to add a serialized reference to the struct and properly specify dependencies.
Use TPropertyValueIterator to recover nested UObject* for GC.
Make sure to use struct wrappers rather than CppStructOps for init and destruction.
Fix cleanup call from SetSparseClassDataStruct to only set the new struct AFTER destrUctors etc. have been called.

diff --git a/Engine/Source/Runtime/Core/Public/UObject/UE5MainStreamObjectVersion.h b/Engine/Source/Runtime/Core/Public/UObject/UE5MainStreamObjectVersion.h
--- a/Engine/Source/Runtime/Core/Public/UObject/UE5MainStreamObjectVersion.h
+++ b/Engine/Source/Runtime/Core/Public/UObject/UE5MainStreamObjectVersion.h
@@ -109,6 +109,9 @@ struct CORE_API FUE5MainStreamObjectVersion
 		// Removing Tessellation parameters from Materials
 		RemovingTessellationParameters,
 
+		// Sparse class data serializes its associated structure to allow for BP types to be used
+		SparseClassDataStructSerialization,
+
         // -----<new versions can be added above this line>-------------------------------------------------
 		VersionPlusOne,
 		LatestVersion = VersionPlusOne - 1
diff --git a/Engine/Source/Runtime/CoreUObject/Private/UObject/Class.cpp b/Engine/Source/Runtime/CoreUObject/Private/UObject/Class.cpp
--- a/Engine/Source/Runtime/CoreUObject/Private/UObject/Class.cpp
+++ b/Engine/Source/Runtime/CoreUObject/Private/UObject/Class.cpp
@@ -54,6 +54,7 @@
 #include "UObject/FieldPath.h"
 #include "HAL/ThreadSafeCounter.h"
 #include "Math/InterpCurvePoint.h"
+#include "UObject/UE5MainStreamObjectVersion.h"
 
 // WARNING: This should always be the last include in any file that needs it (except .generated.h)
 #include "UObject/UndefineUPropertyMacros.h"
@@ -3630,6 +3631,26 @@ void UClass::AddReferencedObjects(UObject* InThis, FReferenceCollector& Collecto
 		This->CallAddReferencedObjects(This->ClassDefaultObject, Collector);
 	}
 
+	Collector.AddReferencedObject( This->SparseClassDataStruct, This );
+
+	// Add sparse class data
+	if (This->SparseClassDataStruct && This->SparseClassData)
+	{
+		if (This->SparseClassDataStruct->StructFlags & STRUCT_AddStructReferencedObjects)
+		{
+			This->SparseClassDataStruct->GetCppStructOps()->AddStructReferencedObjects()(This->SparseClassData, Collector);
+		}
+		else
+		{
+			// The iterator will recursively loop through all structs in structs too.
+			for (TPropertyValueIterator<FObjectProperty> It(This->SparseClassDataStruct, This->SparseClassData); It; ++It)
+			{
+				UObject** ObjectPtr = static_cast<UObject**>(const_cast<void*>(It.Value()));
+				Collector.AddReferencedObject(*ObjectPtr);
+			}
+		}
+	}
+	
 	Super::AddReferencedObjects( This, Collector );
 }
 
@@ -4105,6 +4126,16 @@ void UClass::GetAssetRegistryTags(TArray<FAssetRegistryTag>& OutTags) const
 #endif
 }
 
+void UClass::GetPreloadDependencies(TArray<UObject*>& OutDeps)
+{
+	Super::GetPreloadDependencies(OutDeps);
+
+	if(SparseClassDataStruct)
+	{
+		OutDeps.Add(SparseClassDataStruct);
+	}
+}
+
 void UClass::Link(FArchive& Ar, bool bRelinkExistingProperties)
 {
 	check(!bRelinkExistingProperties || !(ClassFlags & CLASS_Intrinsic));
@@ -4618,6 +4649,13 @@ void UClass::Serialize( FArchive& Ar )
 
 	if (!Ar.IsLoading() && !Ar.IsSaving())
 	{
+		Ar.UsingCustomVersion(FUE5MainStreamObjectVersion::GUID);
+
+		if(Ar.CustomVer(FUE5MainStreamObjectVersion::GUID) >= FUE5MainStreamObjectVersion::SparseClassDataStructSerialization)
+		{
+			Ar << SparseClassDataStruct;
+		}
+
 		if (GetSparseClassDataStruct() != nullptr)
 		{
 			SerializeSparseClassData(FStructuredArchiveFromArchive(Ar).GetSlot());
@@ -4709,6 +4747,12 @@ void UClass::SerializeSparseClassData(FStructuredArchive::FSlot Slot)
 	// tell the archive that it's allowed to load data for transient properties
 	FArchive& UnderlyingArchive = Slot.GetUnderlyingArchive();
 
+	// Preload sparse class data struct if required
+	if(SparseClassDataStruct->HasAnyFlags(RF_NeedLoad))
+	{
+		UnderlyingArchive.Preload(SparseClassDataStruct);
+	}
+	
 	// make sure we always have sparse class a sparse class data struct to read from/write to
 	GetOrCreateSparseClassData();
 
@@ -5016,7 +5060,7 @@ void* UClass::CreateSparseClassData()
 	if (SparseClassDataStruct)
 	{
 		SparseClassData = FMemory::Malloc(SparseClassDataStruct->GetStructureSize(), SparseClassDataStruct->GetMinAlignment());
-		SparseClassDataStruct->GetCppStructOps()->Construct(SparseClassData);
+		SparseClassDataStruct->InitializeStruct(SparseClassData);
 	}
 	if (SparseClassData)
 	{
@@ -5040,7 +5084,7 @@ void UClass::CleanupSparseClassData()
 {
 	if (SparseClassData)
 	{
-		SparseClassDataStruct->GetCppStructOps()->Destruct(SparseClassData);
+		SparseClassDataStruct->DestroyStruct(SparseClassData);
 		FMemory::Free(SparseClassData);
 		SparseClassData = nullptr;
 	}
@@ -5056,10 +5100,10 @@ void UClass::SetSparseClassDataStruct(UScriptStruct* InSparseClassDataStruct)
 { 
 	if (SparseClassDataStruct != InSparseClassDataStruct)
 	{
-		SparseClassDataStruct = InSparseClassDataStruct;
-
 		// the old type and new type may not match when we do a reload so get rid of the old data
 		CleanupSparseClassData();
+
+		SparseClassDataStruct = InSparseClassDataStruct;
 	}
 }
 
@@ -5834,6 +5878,17 @@ bool FStructUtils::TheSameLayout(const UStruct* StructA, const UStruct* StructB,
 			PropertyA = PropertyA ? PropertyA->PropertyLinkNext : NULL;
 			PropertyB = PropertyB ? PropertyB->PropertyLinkNext : NULL;
 		}
+
+		if(bResult)
+		{
+			// If structs are actually classes, their 'layout' is affected by their sparse class data too 
+			const UClass* ClassA = Cast<UClass>(StructA);
+			const UClass* ClassB = Cast<UClass>(StructB);
+			if(ClassA && ClassB)
+			{
+				bResult = TheSameLayout(ClassA->GetSparseClassDataStruct(), ClassB->GetSparseClassDataStruct(), bCheckPropertiesNames);
+			}
+		}
 	}
 	return bResult;
 }
diff --git a/Engine/Source/Runtime/CoreUObject/Public/UObject/Class.h b/Engine/Source/Runtime/CoreUObject/Public/UObject/Class.h
--- a/Engine/Source/Runtime/CoreUObject/Public/UObject/Class.h
+++ b/Engine/Source/Runtime/CoreUObject/Public/UObject/Class.h
@@ -2897,6 +2897,7 @@ public:
 	virtual void GetAssetRegistryTags(TArray<FAssetRegistryTag>& OutTags) const override;
 	virtual bool IsAsset() const override { return false; }	
 	virtual bool IsNameStableForNetworking() const override { return true; } // For now, assume all classes have stable net names
+	virtual void GetPreloadDependencies(TArray<UObject*>& OutDeps) override;
 	// End of UObject interface.
 
 	// UField interface.
diff --git a/Engine/Source/Runtime/Engine/Private/BlueprintGeneratedClass.cpp b/Engine/Source/Runtime/Engine/Private/BlueprintGeneratedClass.cpp
--- a/Engine/Source/Runtime/Engine/Private/BlueprintGeneratedClass.cpp
+++ b/Engine/Source/Runtime/Engine/Private/BlueprintGeneratedClass.cpp
@@ -23,6 +23,7 @@
 #include "UObject/CoreObjectVersion.h"
 #include "Net/Core/PushModel/PushModel.h"
 #include "UObject/CoreObjectVersion.h"
+#include "UObject/UE5MainStreamObjectVersion.h"
 
 #if WITH_EDITOR
 #include "Kismet2/BlueprintEditorUtils.h"
@@ -401,6 +402,13 @@ void UBlueprintGeneratedClass::SerializeDefaultObject(UObject* Object, FStructur
 		}
 	}
 
+	UnderlyingArchive.UsingCustomVersion(FUE5MainStreamObjectVersion::GUID);
+
+	if(UnderlyingArchive.CustomVer(FUE5MainStreamObjectVersion::GUID) >= FUE5MainStreamObjectVersion::SparseClassDataStructSerialization)
+	{
+		UnderlyingArchive << SparseClassDataStruct;
+	}
+
 #if WITH_EDITORONLY_DATA
 	if (bIsSparseClassDataSerializable)
 #endif
