AsyncLoading2: Improve overlap of ALT and GT work for data sets with expensive PostLoad to reduce wall times when flushing async loading.
Interleave polling of packages in external read queue with processing of normal work queue nodes.

#rnx

diff --git a/Engine/Source/Runtime/CoreUObject/Private/Serialization/AsyncLoading2.cpp b/Engine/Source/Runtime/CoreUObject/Private/Serialization/AsyncLoading2.cpp
--- a/Engine/Source/Runtime/CoreUObject/Private/Serialization/AsyncLoading2.cpp
+++ b/Engine/Source/Runtime/CoreUObject/Private/Serialization/AsyncLoading2.cpp
@@ -4088,7 +4088,7 @@ EAsyncPackageState::Type FAsyncLoadingThread2::ProcessAsyncLoadingFromGameThread
 
 			if (!ExternalReadQueue.IsEmpty())
 			{
-				TRACE_CPUPROFILER_EVENT_SCOPE(ProcessExternalReads);
+				TRACE_CPUPROFILER_EVENT_SCOPE(WaitingForExternalReads);
 
 				FAsyncPackage2* Package = nullptr;
 				ExternalReadQueue.Dequeue(Package);
@@ -4661,6 +4661,21 @@ uint32 FAsyncLoadingThread2::Run()
 					do 
 					{
 						bPopped = false;
+						{
+							FAsyncPackage2* Package = nullptr;
+							if (ExternalReadQueue.Peek(Package))
+							{
+								TRACE_CPUPROFILER_EVENT_SCOPE(PollExternalReads);
+								EAsyncPackageState::Type Result = Package->ProcessExternalReads(FAsyncPackage2::ExternalReadAction_Poll);
+								if (Result == EAsyncPackageState::Complete)
+								{
+									ExternalReadQueue.Pop();
+									bPopped = true;
+									bDidSomething = true;
+								}
+							}
+						}
+
 						for (FAsyncLoadEventQueue2* Queue : AltEventQueues)
 						{
 							if (Queue->PopAndExecute(ThreadState))
@@ -4687,29 +4702,6 @@ uint32 FAsyncLoadingThread2::Run()
 						break;
 					}
 
-					{
-						bool bDidExternalRead = false;
-						do
-						{
-							bDidExternalRead = false;
-							FAsyncPackage2* Package = nullptr;
-							if (ExternalReadQueue.Peek(Package))
-							{
-								TRACE_CPUPROFILER_EVENT_SCOPE(ProcessExternalReads);
-
-								FAsyncPackage2::EExternalReadAction Action = FAsyncPackage2::ExternalReadAction_Poll;
-
-								EAsyncPackageState::Type Result = Package->ProcessExternalReads(Action);
-								if (Result == EAsyncPackageState::Complete)
-								{
-									ExternalReadQueue.Pop();
-									bDidExternalRead = true;
-									bDidSomething = true;
-								}
-							}
-						} while (bDidExternalRead);
-					}
-
 				} while (bDidSomething);
 			}
 
@@ -4748,7 +4740,7 @@ uint32 FAsyncLoadingThread2::Run()
 				else if (ExternalReadQueue.Peek(Package))
 				{
 					TRACE_CPUPROFILER_EVENT_SCOPE(AsyncLoadingTime);
-					TRACE_CPUPROFILER_EVENT_SCOPE(ProcessExternalReads);
+					TRACE_CPUPROFILER_EVENT_SCOPE(WaitingForExternalReads);
 
 					EAsyncPackageState::Type Result = Package->ProcessExternalReads(FAsyncPackage2::ExternalReadAction_Wait);
 					check(Result == EAsyncPackageState::Complete);
