Always load all imports in editor builds, including -game.

diff --git a/Engine/Source/Runtime/CoreUObject/Private/UObject/LinkerLoad.cpp b/Engine/Source/Runtime/CoreUObject/Private/UObject/LinkerLoad.cpp
--- a/Engine/Source/Runtime/CoreUObject/Private/UObject/LinkerLoad.cpp
+++ b/Engine/Source/Runtime/CoreUObject/Private/UObject/LinkerLoad.cpp
@@ -2684,9 +2684,17 @@ UObject* FLinkerLoad::FindExistingImport(int32 ImportIndex)
 
 void FLinkerLoad::Verify()
 {
-	if (!IsImportLazyLoadEnabled() && (!FApp::IsGame() || GIsEditor || IsRunningCommandlet()))
+	if (!bHaveImportsBeenVerified)
 	{
-		if (!bHaveImportsBeenVerified)
+		bool bShouldVerifiedAllImports = IsRunningCommandlet();
+
+#if WITH_EDITOR
+		// In editor builds using OFPA, we need to resolve imports for BP classes referenced by the level script in order to be able
+		// to properly reinstanced them. We could filter out imports to resolve here, but we resolve all of them instead.
+		bShouldVerifiedAllImports = true;
+#endif
+
+		if (!IsImportLazyLoadEnabled() && bShouldVerifiedAllImports)
 		{
 #if WITH_EDITOR
 			TOptional<FScopedSlowTask> SlowTask;
@@ -2711,9 +2719,9 @@ void FLinkerLoad::Verify()
 				VerifyImport( ImportIndex );
 			}
 		}
-	}
 
-	bHaveImportsBeenVerified = true;
+		bHaveImportsBeenVerified = true;
+	}
 }
 
 FName FLinkerLoad::GetExportClassPackage( int32 i )
