Fix AsyncLoadPackage's error handling to call the completion delegate with the requested package name instead of NAME_None when the packagename is invalid.
#rnx

diff --git a/Engine/Source/Runtime/CoreUObject/Private/Serialization/AsyncLoading.cpp b/Engine/Source/Runtime/CoreUObject/Private/Serialization/AsyncLoading.cpp
--- a/Engine/Source/Runtime/CoreUObject/Private/Serialization/AsyncLoading.cpp
+++ b/Engine/Source/Runtime/CoreUObject/Private/Serialization/AsyncLoading.cpp
@@ -6925,7 +6925,7 @@ int32 FAsyncLoadingThread::LoadPackage(const FPackagePath& InPackagePath, FName
 	if (!InPackagePath.IsMountedPath())
 	{
 		UE_LOG(LogStreaming, Warning, TEXT("Async loading code requires MountedPath PackagePaths (%s)."), *InPackagePath.GetDebugName());
-		InCompletionDelegate.ExecuteIfBound(NAME_None, nullptr, EAsyncLoadingResult::Failed);
+		InCompletionDelegate.ExecuteIfBound(InPackagePath.GetPackageFName(), nullptr, EAsyncLoadingResult::Failed);
 		return INDEX_NONE;
 	}
 
diff --git a/Engine/Source/Runtime/CoreUObject/Private/Serialization/AsyncPackageLoader.cpp b/Engine/Source/Runtime/CoreUObject/Private/Serialization/AsyncPackageLoader.cpp
--- a/Engine/Source/Runtime/CoreUObject/Private/Serialization/AsyncPackageLoader.cpp
+++ b/Engine/Source/Runtime/CoreUObject/Private/Serialization/AsyncPackageLoader.cpp
@@ -711,6 +711,17 @@ static FPackagePath GetLoadPackageAsyncPackagePath(FStringView InPackageNameOrFi
 			}
 		}
 	}
+
+	// If PackagePath is empty at this point, then we are going to fail because its not mounted. But pass the input name 
+	// into PackagePath (preferably as a packagename) to provide to the caller's CompletionCallback
+	if (PackagePath.IsEmpty())
+	{
+		if (!FPackagePath::TryFromPackageName(InPackageNameOrFilePath, PackagePath))
+		{
+			PackagePath = FPackagePath::FromLocalPath(InPackageNameOrFilePath);
+		}
+	}
+
 	return PackagePath;
 }
 
