Fixing ClassPrivate nullptr deref bug

#rnx

diff --git a/Engine/Source/Runtime/CoreUObject/Private/UObject/UObjectBase.cpp b/Engine/Source/Runtime/CoreUObject/Private/UObject/UObjectBase.cpp
--- a/Engine/Source/Runtime/CoreUObject/Private/UObject/UObjectBase.cpp
+++ b/Engine/Source/Runtime/CoreUObject/Private/UObject/UObjectBase.cpp
@@ -342,7 +342,7 @@ bool UObjectBase::IsValidLowLevelFast(bool bRecursive /*= true*/) const
 #if WITH_LIVE_CODING
 		// When live coding is re-instancing blueprint generated classes, we have to clear out the default object so it can get 
 		// GC'ed and deleted prior to live coding completing the patching process (of the destructor specifically)
-		if (!ClassPrivate->HasAnyClassFlags(CLASS_NewerVersionExists))
+		if (ClassPrivate == nullptr || !ClassPrivate->HasAnyClassFlags(CLASS_NewerVersionExists))
 #endif
 		{
 			UE_LOG(LogUObjectBase, Error, TEXT("Class pointer is invalid or CDO is invalid."));
