Reduce FName memory usage when saving thumbnails.
#rnx

diff --git a/Engine/Source/Runtime/CoreUObject/Private/UObject/SavePackage/SavePackageUtilities.cpp b/Engine/Source/Runtime/CoreUObject/Private/UObject/SavePackage/SavePackageUtilities.cpp
--- a/Engine/Source/Runtime/CoreUObject/Private/UObject/SavePackage/SavePackageUtilities.cpp
+++ b/Engine/Source/Runtime/CoreUObject/Private/UObject/SavePackage/SavePackageUtilities.cpp
@@ -1812,16 +1812,24 @@ void SaveThumbnails(UPackage* InOuter, FLinkerSave* Linker, FStructuredArchive::
 			FObjectExport& Export = Linker->ExportMap[i];
 			if( Export.Object )
 			{
-				const FName ObjectFullName( *Export.Object->GetFullName() );
-				const FObjectThumbnail* ObjectThumbnail = PackageThumbnailMap.Find( ObjectFullName );
+				const FName ObjectFullName( *Export.Object->GetFullName(), FNAME_Find );
+				const FObjectThumbnail* ObjectThumbnail = nullptr;
+				// If the FName does not exist, then we know it is not in the map and do not need to search
+				if (!ObjectFullName.IsNone())
+				{
+					ObjectThumbnail = PackageThumbnailMap.Find(ObjectFullName);
+				}
 		
 				// if we didn't find the object via full name, try again with ??? as the class name, to support having
 				// loaded old packages without going through the editor (ie cooking old packages)
 				if (ObjectThumbnail == nullptr)
 				{
 					// can't overwrite ObjectFullName, so that we add it properly to the map
-					FName OldPackageStyleObjectFullName = FName(*FString::Printf(TEXT("??? %s"), *Export.Object->GetPathName()));
-					ObjectThumbnail = PackageThumbnailMap.Find(OldPackageStyleObjectFullName);
+					FName OldPackageStyleObjectFullName = FName(*FString::Printf(TEXT("??? %s"), *Export.Object->GetPathName()), FNAME_Find);
+					if (!OldPackageStyleObjectFullName.IsNone())
+					{
+						ObjectThumbnail = PackageThumbnailMap.Find(OldPackageStyleObjectFullName);
+					}
 				}
 				if( ObjectThumbnail != nullptr )
 				{
