Add a way to set an output device to record error during linker save serialization
Error out when there is an error on the linker save during export serialization
LinkerSave will now set an error on the archive when a name is unable to be mapped and during serialization
Still run save cleanup for a package if errors are encountered

diff --git a/Engine/Source/Runtime/CoreUObject/Private/UObject/LinkerSave.cpp b/Engine/Source/Runtime/CoreUObject/Private/UObject/LinkerSave.cpp
--- a/Engine/Source/Runtime/CoreUObject/Private/UObject/LinkerSave.cpp
+++ b/Engine/Source/Runtime/CoreUObject/Private/UObject/LinkerSave.cpp
@@ -305,11 +305,23 @@ FArchive& FLinkerSave::operator<<( FName& InName )
 	int32 Save = MapName(InName.GetDisplayIndex());
 
 	check(GetSerializeContext());
-	ensureMsgf(Save != INDEX_NONE, TEXT("Name \"%s\" is not mapped when saving %s (object: %s, property: %s)"), 
-		*InName.ToString(),
-		*GetArchiveName(),
-		*GetSerializeContext()->SerializedObject->GetFullName(),
-		*GetFullNameSafe(GetSerializedProperty()));
+
+	bool bNameMapped = Save != INDEX_NONE;
+	if (!bNameMapped)
+	{
+		// Set an error on the archive and record the error on the log output if one is set.
+		SetCriticalError();
+		FString ErrorMessage = FString::Printf(TEXT("Name \"%s\" is not mapped when saving %s (object: %s, property: %s)."),
+			*InName.ToString(),
+			*GetArchiveName(),
+			*GetSerializeContext()->SerializedObject->GetFullName(),
+			*GetFullNameSafe(GetSerializedProperty()));
+		ensureMsgf(false, TEXT("%s"), *ErrorMessage);
+		if (LogOutput)
+		{
+			LogOutput->Logf(ELogVerbosity::Error, TEXT("%s"), *ErrorMessage);
+		}
+	}
 
 	int32 Number = InName.GetNumber();
 	FArchive& Ar = *this;
diff --git a/Engine/Source/Runtime/CoreUObject/Private/UObject/SavePackage2.cpp b/Engine/Source/Runtime/CoreUObject/Private/UObject/SavePackage2.cpp
--- a/Engine/Source/Runtime/CoreUObject/Private/UObject/SavePackage2.cpp
+++ b/Engine/Source/Runtime/CoreUObject/Private/UObject/SavePackage2.cpp
@@ -701,6 +701,10 @@ ESavePackageResult CreateLinker(FSaveContext& SaveContext)
 				SaveContext.Linker = MakeUnique<FLinkerSave>(SaveContext.GetPackage(), *SaveContext.TempFilename.GetValue(), SaveContext.IsForceByteSwapping(), SaveContext.IsSaveUnversioned());
 			}
 		}
+		if (SaveContext.IsGenerateSaveError())
+		{
+			SaveContext.Linker->SetOutputDevice(SaveContext.GetError());
+		}
 		SaveContext.Linker->bUpdatingLoadedPath = SaveContext.IsUpdatingLoadedPath();
 		SaveContext.Linker->bProceduralSave = SaveContext.IsProceduralSave();
 
@@ -1598,7 +1602,8 @@ ESavePackageResult WriteExports(FStructuredArchive::FRecord& StructuredArchiveRo
 			Export.SerialSize = Linker->Tell() - Export.SerialOffset;
 		}
 	}
-	return ReturnSuccessOrCancel();
+	// if an error occurred on the linker while serializing exports, return an error
+	return Linker->IsError() ? ESavePackageResult::Error : ReturnSuccessOrCancel();
 }
 
 ESavePackageResult WriteAdditionalExportFiles(FSaveContext& SaveContext)
@@ -2201,10 +2206,8 @@ FSavePackageResultStruct UPackage::Save2(UPackage* InPackage, UObject* InAsset,
 		// @todo: Once GIsSavingPackage is reworked we should reinstore the saving flag here for the gc lock
 		//FScopedSavingFlag IsSavingFlag(SaveContext.IsConcurrent());
 		SaveContext.Result = InnerSave(SaveContext);
-		if (SaveContext.Result != ESavePackageResult::Success)
-		{
-			return SaveContext.Result;
-		}
+
+		// in case of failure or cancellation, do not exit here, still run cleanup (PostSaveRoot/ClearCachedPlatformCookedData)
 	}
 
 	// PostSave Asset
@@ -2217,8 +2220,9 @@ FSavePackageResultStruct UPackage::Save2(UPackage* InPackage, UObject* InAsset,
 
 	ClearCachedPlatformCookedData(SaveContext);
 
-	// PostSave Package - edit in memory package and send events
+	// PostSave Package - edit in memory package and send events if save was successful
 	SlowTask.EnterProgressFrame();
+	if (SaveContext.Result == ESavePackageResult::Success)
 	{
 		PostSavePackage(SaveContext);
 	}
diff --git a/Engine/Source/Runtime/CoreUObject/Public/UObject/LinkerSave.h b/Engine/Source/Runtime/CoreUObject/Public/UObject/LinkerSave.h
--- a/Engine/Source/Runtime/CoreUObject/Public/UObject/LinkerSave.h
+++ b/Engine/Source/Runtime/CoreUObject/Public/UObject/LinkerSave.h
@@ -159,7 +159,17 @@ public:
 	/** Get the filename being saved to */
 	const FString& GetFilename() const;
 
+	/* Set the output device used to log errors, if any. */
+	void SetOutputDevice(FOutputDevice* InOutputDevice)
+	{
+		LogOutput = InOutputDevice;
+	}
+
 protected:
 	/** Set the filename being saved to */
 	void SetFilename(FStringView InFilename);
+
+private:
+	/** Optional log output to bubble errors back up. */
+	FOutputDevice* LogOutput = nullptr;
 };
