Clang compile fix for 15735615 to avoid deprecation warning for FPackageFileSummary::Guid in implicit constructors

#rnx

diff --git a/Engine/Source/Runtime/CoreUObject/Public/UObject/PackageFileSummary.h b/Engine/Source/Runtime/CoreUObject/Public/UObject/PackageFileSummary.h
--- a/Engine/Source/Runtime/CoreUObject/Public/UObject/PackageFileSummary.h
+++ b/Engine/Source/Runtime/CoreUObject/Public/UObject/PackageFileSummary.h
@@ -218,6 +218,14 @@ public:
 	/** Constructor */
 	COREUOBJECT_API FPackageFileSummary();
 
+	// Workaround for clang deprecation warnings for deprecated members in implicit constructors
+	PRAGMA_DISABLE_DEPRECATION_WARNINGS
+	FPackageFileSummary(FPackageFileSummary&&) = default;
+	FPackageFileSummary(const FPackageFileSummary&) = default;
+	FPackageFileSummary& operator=(FPackageFileSummary&&) = default;
+	FPackageFileSummary& operator=(const FPackageFileSummary&) = default;
+	PRAGMA_ENABLE_DEPRECATION_WARNINGS
+
 	int32 GetFileVersionUE() const
 	{
 		return FileVersionUE;
