Cherrypicking Mirage code base (disabled) from Dev-Cooker

diff --git a/Engine/Source/Runtime/CoreUObject/Private/Misc/PackagePath.cpp b/Engine/Source/Runtime/CoreUObject/Private/Misc/PackagePath.cpp
--- a/Engine/Source/Runtime/CoreUObject/Private/Misc/PackagePath.cpp
+++ b/Engine/Source/Runtime/CoreUObject/Private/Misc/PackagePath.cpp
@@ -27,6 +27,8 @@ const TCHAR* LexToString(EPackageSegment PackageSegment)
 		return TEXT("BulkDataOptional");
 	case EPackageSegment::BulkDataMemoryMapped:
 		return TEXT("BulkDataMemoryMapped");
+	case EPackageSegment::PayloadSidecar:
+		return TEXT("PayloadSidecar");
 	default:
 		checkNoEntry();
 		return TEXT("");
@@ -59,6 +61,8 @@ EPackageSegment ExtensionToSegment(EPackageExtension PackageExtension)
 		return EPackageSegment::BulkDataDefault;
 	case EPackageExtension::BulkDataMemoryMapped:
 		return EPackageSegment::BulkDataDefault;
+	case EPackageExtension::PayloadSidecar:
+		return EPackageSegment::PayloadSidecar;
 	default:
 		checkNoEntry();
 		return EPackageSegment::Header;
@@ -79,6 +83,8 @@ EPackageExtension SegmentToExtension(EPackageSegment PackageSegment)
 		return EPackageExtension::BulkDataOptional;
 	case EPackageSegment::BulkDataMemoryMapped:
 		return EPackageExtension::BulkDataMemoryMapped;
+	case EPackageSegment::PayloadSidecar:
+		return EPackageExtension::PayloadSidecar;
 	default:
 		checkNoEntry();
 		return EPackageExtension::Unspecified;
@@ -111,6 +117,8 @@ const TCHAR* LexToString(EPackageExtension PackageExtension)
 		return TEXT(".uptnl");
 	case EPackageExtension::BulkDataMemoryMapped:
 		return TEXT(".m.ubulk");
+	case EPackageExtension::PayloadSidecar:
+		return TEXT(".upayload");
 	default:
 		checkNoEntry();
 		return TEXT("");
@@ -171,6 +179,10 @@ EPackageExtension FPackagePath::ParseExtension(FStringView Filename, int32* OutE
 	{
 		return EPackageExtension::BulkDataOptional;
 	}
+	if (Extension.Equals(TEXT(".upayload"_SV), ESearchCase::IgnoreCase))
+	{
+		return EPackageExtension::PayloadSidecar;
+	}
 	return EPackageExtension::Custom;
 }
 
@@ -758,6 +770,7 @@ EPackageExtension AllExtensions[] =
 	EPackageExtension::BulkDataDefault,
 	EPackageExtension::BulkDataOptional,
 	EPackageExtension::BulkDataMemoryMapped,
+	EPackageExtension::PayloadSidecar,
 };
 static_assert(static_cast<int>(EPackageExtension::Unspecified) == 0 && EPackageExtensionCount == UE_ARRAY_COUNT(AllExtensions), "Need to add new extensions to AllExtensions array");
 }
diff --git a/Engine/Source/Runtime/CoreUObject/Private/Misc/PackagePathTest.cpp b/Engine/Source/Runtime/CoreUObject/Private/Misc/PackagePathTest.cpp
--- a/Engine/Source/Runtime/CoreUObject/Private/Misc/PackagePathTest.cpp
+++ b/Engine/Source/Runtime/CoreUObject/Private/Misc/PackagePathTest.cpp
@@ -48,7 +48,7 @@ public:
 			TestEqual(Errorf(TEXT("ParsedExtension returns correct extension for %s"), LexToString(Extension)), ParsedExtension, Extension);
 			TestEqual(Errorf(TEXT("ParsedExtension sets correct ExtensionStart for %s"), LexToString(Extension)), ExtensionStart, PathWithoutExtension.Len());
 		};
-		static_assert(EPackageExtensionCount == 11, "Need to add extensions to the test");
+		static_assert(EPackageExtensionCount == 12, "Need to add extensions to the test");
 		TestExtension(EPackageExtension::Unspecified);
 		TestExtension(EPackageExtension::Asset);
 		TestExtension(EPackageExtension::Map);
@@ -60,6 +60,7 @@ public:
 		TestExtension(EPackageExtension::BulkDataDefault);
 		TestExtension(EPackageExtension::BulkDataOptional);
 		TestExtension(EPackageExtension::BulkDataMemoryMapped);
+		TestExtension(EPackageExtension::PayloadSidecar);
 
 		// Test Custom
 		{
diff --git a/Engine/Source/Runtime/CoreUObject/Private/Tests/Virtualization/VirtualizedBulkDataTests.cpp b/Engine/Source/Runtime/CoreUObject/Private/Tests/Virtualization/VirtualizedBulkDataTests.cpp
--- a/Engine/Source/Runtime/CoreUObject/Private/Tests/Virtualization/VirtualizedBulkDataTests.cpp
+++ b/Engine/Source/Runtime/CoreUObject/Private/Tests/Virtualization/VirtualizedBulkDataTests.cpp
@@ -1,15 +1,20 @@
 // Copyright Epic Games, Inc. All Rights Reserved.
 
-#include "Virtualization/VirtualizedBulkData.h"
-#include "Virtualization/VirtualizationUtilities.h"
+#include "Algo/AllOf.h"
 #include "Memory/SharedBuffer.h"
+#include "Serialization/MemoryReader.h"
+#include "Serialization/MemoryWriter.h"
 #include "Templates/UniquePtr.h"
-#include "Algo/AllOf.h"
+#include "Virtualization/VirtualizationUtilities.h"
+#include "Virtualization/VirtualizedBulkData.h"
 
 #include "Misc/AutomationTest.h"
 
 #if WITH_DEV_AUTOMATION_TESTS && WITH_EDITORONLY_DATA
 
+namespace UE::Virtualization
+{
+	
 constexpr const uint32 TestFlags = EAutomationTestFlags::ApplicationContextMask | EAutomationTestFlags::EngineFilter;
 
 /**
@@ -17,12 +22,12 @@ constexpr const uint32 TestFlags = EAutomationTestFlags::ApplicationContextMask
  * payload via both TFuture and callback methods. It then creates copies of the object and makes sure that we can get the payload from
  * the copies, even when the original source object has been reset.
  */
-IMPLEMENT_SIMPLE_AUTOMATION_TEST(FVirtualizationWrapperTestBasic, TEXT("System.Core.VirtualizedBulkData.Basic"), TestFlags)
+IMPLEMENT_SIMPLE_AUTOMATION_TEST(FVirtualizationWrapperTestBasic, TEXT("System.Core.Virtualization.BulkData.Basic"), TestFlags)
 bool FVirtualizationWrapperTestBasic::RunTest(const FString& Parameters)
 {
 	const int64 BufferSize = 1024;
 	TUniquePtr<uint8[]> SourceBuffer = MakeUnique<uint8[]>(BufferSize);
-	
+
 	// Write some values to SourceBuffer
 	for (int64 Index = 0; Index < BufferSize; ++Index)
 	{
@@ -31,7 +36,7 @@ bool FVirtualizationWrapperTestBasic::RunTest(const FString& Parameters)
 
 	auto ValidateBulkData = [this, &SourceBuffer, BufferSize](const FVirtualizedUntypedBulkData& BulkDataToValidate, const TCHAR* Label)
 		{
-			FSharedBuffer RetrievedBuffer = BulkDataToValidate.GetData().Get();
+			FSharedBuffer RetrievedBuffer = BulkDataToValidate.GetPayload().Get();
 			TestEqual(FString::Printf(TEXT("%s buffer length"),Label), (int64)RetrievedBuffer.GetSize(), BufferSize);
 			TestTrue(FString::Printf(TEXT("SourceBuffer values == %s values"),Label), FMemory::Memcmp(SourceBuffer.Get(), RetrievedBuffer.GetData(), BufferSize) == 0);
 		};
@@ -40,13 +45,6 @@ bool FVirtualizationWrapperTestBasic::RunTest(const FString& Parameters)
 	FByteVirtualizedBulkData BulkData;
 	BulkData.UpdatePayload(FSharedBuffer::MakeView(SourceBuffer.Get(), BufferSize));
 
-	// Test access via callback
-	BulkData.GetData([this,BufferSize,&SourceBuffer](FSharedBuffer Payload)
-		{
-			TestEqual(TEXT("Callback buffer length"), (int64)Payload.GetSize(), BufferSize);
-			TestTrue(TEXT("SourceBuffer values == Callback values"), FMemory::Memcmp(SourceBuffer.Get(), Payload.GetData(), BufferSize) == 0);
-		});
-
 	// Test accessing the data from the bulkdata object
 	ValidateBulkData(BulkData, TEXT("Retrieved"));
 
@@ -76,24 +74,19 @@ bool FVirtualizationWrapperTestBasic::RunTest(const FString& Parameters)
  * This test will validate how VirtualizedBulkData behaves when it has no associated payload and make sure 
  * that our assumptions are correct.
  */
-IMPLEMENT_SIMPLE_AUTOMATION_TEST(FVirtualizationWrapperTestEmpty, TEXT("System.Core.VirtualizedBulkData.Empty"), TestFlags)
+IMPLEMENT_SIMPLE_AUTOMATION_TEST(FVirtualizationWrapperTestEmpty, TEXT("System.Core.Virtualization.BulkData.Empty"), TestFlags)
 bool FVirtualizationWrapperTestEmpty::RunTest(const FString& Parameters)
 {
 	FByteVirtualizedBulkData BulkData;
 
 	// Validate the general accessors
-	TestEqual(TEXT("Return value of ::GetBulkDataSize()"), BulkData.GetBulkDataSize(), (int64)0);
-	TestTrue(TEXT("Payload key is invalid"), !BulkData.GetKey().IsValid());
+	TestEqual(TEXT("Return value of ::GetBulkDataSize()"), BulkData.GetPayloadSize(), (int64)0);
+	TestTrue(TEXT("Payload key is invalid"), !BulkData.GetPayloadId().IsValid());
 	TestFalse(TEXT("Return value of ::IsDataLoaded()"), BulkData.IsDataLoaded());
 
 	// Validate the payload accessors
-	FSharedBuffer Payload = BulkData.GetData().Get();
-	TestTrue(TEXT("The payload from the GetData TFuture is null"), Payload.IsNull());
-
-	BulkData.GetData([this](FSharedBuffer Payload) 
-		{
-			TestTrue(TEXT("The payload via the GetData callback is null"), Payload.IsNull());
-		});
+	FSharedBuffer Payload = BulkData.GetPayload().Get();
+	TestTrue(TEXT("The payload from the GetPayload TFuture is null"), Payload.IsNull());
 
 	return true;
 }
@@ -101,7 +94,7 @@ bool FVirtualizationWrapperTestEmpty::RunTest(const FString& Parameters)
 /**
  * Test the various methods for updating the payload that a FVirtualizedUntypedBulkData owns
  */
-IMPLEMENT_SIMPLE_AUTOMATION_TEST(FVirtualizationWrapperTestUpdatePayload, TEXT("System.Core.VirtualizedBulkData.UpdatePayload"), TestFlags)
+IMPLEMENT_SIMPLE_AUTOMATION_TEST(FVirtualizationWrapperTestUpdatePayload, TEXT("System.Core.Virtualization.BulkData.UpdatePayload"), TestFlags)
 bool FVirtualizationWrapperTestUpdatePayload::RunTest(const FString& Parameters)
 {
 	// Create a memory buffer of all zeros
@@ -118,7 +111,7 @@ bool FVirtualizationWrapperTestUpdatePayload::RunTest(const FString& Parameters)
 		// The payload should be the same size and same contents as the original buffer but a different
 		// memory address since we retained ownership in the TUniquePtr, so the bulkdata object should 
 		// have created it's own copy.
-		FSharedBuffer Payload = BulkData.GetData().Get();
+		FSharedBuffer Payload = BulkData.GetPayload().Get();
 		TestEqual(TEXT("Payload length"), (int64)Payload.GetSize(), BufferSize);
 		TestNotEqual(TEXT("OriginalData and the payload should have different memory addresses"), (uint8*)OriginalData.Get(), (uint8*)Payload.GetData());
 		TestTrue(TEXT("Orginal buffer == Payload data"), FMemory::Memcmp(OriginalData.Get(), Payload.GetData(), Payload.GetSize()) == 0);
@@ -135,7 +128,7 @@ bool FVirtualizationWrapperTestUpdatePayload::RunTest(const FString& Parameters)
 		// Update the bulkdata object with the new edited payload
 		BulkData.UpdatePayload(EditedPayload);
 
-		Payload = BulkData.GetData().Get();
+		Payload = BulkData.GetPayload().Get();
 		TestEqual(TEXT("Updated payload length"), (int64)Payload.GetSize(), BufferSize);
 		TestEqual(TEXT("Payload and EditablePayload should have the same memory addresses"), (uint8*)Payload.GetData(), (uint8*)EditedPayload.GetData());
 
@@ -156,7 +149,7 @@ bool FVirtualizationWrapperTestUpdatePayload::RunTest(const FString& Parameters)
 		// the bulkdata object.
 		BulkData.UpdatePayload(FSharedBuffer::TakeOwnership(OriginalData.Release(), BufferSize, FMemory::Free));
 
-		FSharedBuffer Payload = BulkData.GetData().Get();
+		FSharedBuffer Payload = BulkData.GetPayload().Get();
 		TestEqual(TEXT("Updated payload length"), (int64)Payload.GetSize(), BufferSize);
 		TestEqual(TEXT("Payload and OriginalDataPtr should have the same memory addresses"), (uint8*)Payload.GetData(), OriginalDataPtr);
 
@@ -177,7 +170,7 @@ bool FVirtualizationWrapperTestUpdatePayload::RunTest(const FString& Parameters)
  * Then we will serialize the VirtualizedBulkData object back to a second buffer and compare the results.
  * If the reader and writers are working then the ReplicatedBuffer should be the same as the original SourceBuffer.
  */
-IMPLEMENT_SIMPLE_AUTOMATION_TEST(FVirtualizationWrapperTestReaderWriter, TEXT("System.Core.VirtualizedBulkData.Reader/Writer"), TestFlags)
+IMPLEMENT_SIMPLE_AUTOMATION_TEST(FVirtualizationWrapperTestReaderWriter, TEXT("System.Core.Virtualization.BulkData.Reader/Writer"), TestFlags)
 bool FVirtualizationWrapperTestReaderWriter::RunTest(const FString& Parameters)
 {
 	const int64 BufferSize = 1024;
@@ -219,4 +212,197 @@ bool FVirtualizationWrapperTestReaderWriter::RunTest(const FString& Parameters)
 	return true;
 }
 
+/**
+ * This set of tests validate that the BulkData's identifier works how we expect it too. It should remain unique in all cases except
+ * move semantics.
+ */
+IMPLEMENT_SIMPLE_AUTOMATION_TEST(FVirtualizationWrapperTestIdentifiers, TEXT("System.Core.Virtualization.BulkData.Identifiers"), TestFlags)
+bool FVirtualizationWrapperTestIdentifiers::RunTest(const FString& Parameters)
+{
+	// Some basic tests with an invalid id
+	{
+		FByteVirtualizedBulkData BulkData;
+		TestFalse(TEXT("BulkData with no payload should return an invalid identifier"), BulkData.GetIdentifier().IsValid());
+
+		FByteVirtualizedBulkData CopiedBulkData(BulkData);
+		TestFalse(TEXT("Copying a bulkdata with an invalid id should result in an invalid id"), CopiedBulkData.GetIdentifier().IsValid());
+
+		FByteVirtualizedBulkData AssignedBulkData; 
+		AssignedBulkData = BulkData;
+		TestFalse(TEXT("Assigning a a bulkdata with an invalid id should result in an invalid id"), AssignedBulkData.GetIdentifier().IsValid());
+		
+		// Check that we did not change the initial object at any point
+		TestFalse(TEXT("Being copied and assigned to other objects should not affect the identifier"), BulkData.GetIdentifier().IsValid());
+	}
+
+	// Some basic tests with a valid id
+	{
+		FByteVirtualizedBulkData BulkData;
+		BulkData.UpdatePayload(FUniqueBuffer::Alloc(32).MoveToShared()); // Assigning this payload should cause BulkData to gain an identifier
+		TestTrue(TEXT("BulkData with a payload should returns a valid identifier"), BulkData.GetIdentifier().IsValid());
+
+		const FGuid OriginalGuid = BulkData.GetIdentifier();
+
+		FByteVirtualizedBulkData CopiedBulkData(BulkData);
+		TestNotEqual(TEXT("Copying a bulkdata with a valid id should result in a unique identifier"), BulkData.GetIdentifier(), CopiedBulkData.GetIdentifier());
+
+		FByteVirtualizedBulkData AssignedBulkData;
+		AssignedBulkData = BulkData;
+		TestNotEqual(TEXT("Assignment operator creates different identifiers"), BulkData.GetIdentifier(), AssignedBulkData.GetIdentifier());
+
+		// Check that we did not change the initial object at any point
+		TestEqual(TEXT("Being copied and assigned to other objects should not affect the identifier"), BulkData.GetIdentifier(), OriginalGuid);
+
+		// Now that AssignedBulkData has a valid identifier, make sure that it is not changed if we assign something else to it.
+		const FGuid OriginalAssignedGuid = AssignedBulkData.GetIdentifier();
+		AssignedBulkData = CopiedBulkData;
+		TestEqual(TEXT("Being copied and assigned to other objects should not affect the identifier"), AssignedBulkData.GetIdentifier(), OriginalAssignedGuid);
+	}
+
+	// Test move constructor
+	{	
+		FVirtualizedUntypedBulkData BulkData;
+		BulkData.UpdatePayload(FUniqueBuffer::Alloc(32).MoveToShared());
+
+		FGuid OriginalGuid = BulkData.GetIdentifier();
+
+		FVirtualizedUntypedBulkData MovedBulkData = MoveTemp(BulkData);
+
+		TestEqual(TEXT("Move constructor should preserve the identifier"), MovedBulkData.GetIdentifier(), OriginalGuid);
+	}
+
+	// Test move assignment
+	{
+		FVirtualizedUntypedBulkData BulkData;
+		BulkData.UpdatePayload(FUniqueBuffer::Alloc(32).MoveToShared());
+
+		FGuid OriginalGuid = BulkData.GetIdentifier();
+
+		FVirtualizedUntypedBulkData MovedBulkData;
+		MovedBulkData = MoveTemp(BulkData);
+
+		TestEqual(TEXT("Move assignment should preserve the identifier"), MovedBulkData.GetIdentifier(), OriginalGuid);
+	}
+
+	// Check that resizing an array will not change the internals
+	{
+		const uint32 NumToTest = 10;
+
+		TArray<FByteVirtualizedBulkData> BulkDataArray;
+		TArray<FGuid> GuidArray;
+
+		for (uint32 Index = 0; Index < NumToTest; ++Index)
+		{
+			BulkDataArray.Add(FByteVirtualizedBulkData());
+
+			// Leave some with invalid ids and some with valid ones
+			if (Index % 2 == 0)
+			{
+				BulkDataArray[Index].UpdatePayload(FUniqueBuffer::Alloc(32).MoveToShared());
+			}
+
+			GuidArray.Add(BulkDataArray[Index].GetIdentifier());
+		}
+
+		// Force an internal reallocation and make sure that the identifiers are unchanged.
+		// Note that it is possible that the allocation is just resized and not reallocated.
+		BulkDataArray.Reserve(BulkDataArray.Max() * 4);
+
+		for (uint32 Index = 0; Index < NumToTest; ++Index)
+		{
+			TestEqual(TEXT(""), BulkDataArray[Index].GetIdentifier(), GuidArray[Index]);
+		}
+
+		// Now insert a new item, moving all of the existing entries and make sure that
+		// the identifiers are unchanged.
+		BulkDataArray.Insert(FByteVirtualizedBulkData(), 0);
+
+		for (uint32 Index = 0; Index < NumToTest; ++Index)
+		{
+			TestEqual(TEXT(""), BulkDataArray[Index+1].GetIdentifier(), GuidArray[Index]);
+		}
+	}
+
+	// Test that adding a payload to a reset bulkdata object or one that has had a zero length payload applied
+	// will correctly show the original id once it has a valid payload.
+	{
+		FByteVirtualizedBulkData BulkData;
+		BulkData.UpdatePayload(FUniqueBuffer::Alloc(32).MoveToShared());
+		
+		const FGuid OriginalGuid = BulkData.GetIdentifier();
+
+		BulkData.Reset();
+		TestFalse(TEXT("BulkData with no payload should return an invalid identifier"), BulkData.GetIdentifier().IsValid());
+
+		BulkData.UpdatePayload(FUniqueBuffer::Alloc(32).MoveToShared());
+		TestEqual(TEXT("Removing a payload then adding a new one should return the original identifier"), BulkData.GetIdentifier(), OriginalGuid);
+
+		BulkData.UpdatePayload(FUniqueBuffer::Alloc(0).MoveToShared());
+		TestFalse(TEXT("Setting a zero length payload should return an invalid identifier"), BulkData.GetIdentifier().IsValid());
+		
+		BulkData.UpdatePayload(FUniqueBuffer::Alloc(32).MoveToShared());
+		TestEqual(TEXT("Restoring a payload should return the original identifier"), BulkData.GetIdentifier(), OriginalGuid);
+	}
+
+	// Test that serialization does not change the identifier (in this case serializing to and from a memory buffer) 
+	{
+		FByteVirtualizedBulkData SrcData;
+		SrcData.UpdatePayload(FUniqueBuffer::Alloc(32).MoveToShared());
+
+		TArray<uint8> MemoryBuffer;
+		const bool bIsArPersistent = true;
+
+		FByteVirtualizedBulkData DstData;
+
+		// Serialize the SourceBuffer to BulkData 
+		{
+			FMemoryWriter WriterAr(MemoryBuffer, bIsArPersistent);
+			SrcData.Serialize(WriterAr, nullptr);
+		}
+
+		// Serialize BulkData back to ReplicatedBuffer
+		{
+			FMemoryReader ReaderAr(MemoryBuffer, bIsArPersistent);
+			DstData.Serialize(ReaderAr, nullptr);
+		}
+
+		TestEqual(TEXT("Serialization should preserve the identifier"), SrcData.GetIdentifier(), DstData.GetIdentifier());
+	}
+
+	// Test that serialization does not change the identifier (in this case serializing to and from a memory buffer) 
+	{
+		FByteVirtualizedBulkData SrcData;
+		SrcData.UpdatePayload(FUniqueBuffer::Alloc(32).MoveToShared());
+
+		const FGuid OriginalIdentifier = SrcData.GetIdentifier();
+		SrcData.UpdatePayload(FUniqueBuffer::Alloc(0).MoveToShared());
+
+		TArray<uint8> MemoryBuffer;
+		const bool bIsArPersistent = true;
+
+		FByteVirtualizedBulkData DstData;
+
+		// Serialize the SourceBuffer to BulkData 
+		{
+			FMemoryWriter WriterAr(MemoryBuffer, bIsArPersistent);
+			SrcData.Serialize(WriterAr, nullptr);
+		}
+
+		// Serialize BulkData back to ReplicatedBuffer
+		{
+			FMemoryReader ReaderAr(MemoryBuffer, bIsArPersistent);
+			DstData.Serialize(ReaderAr, nullptr);
+		}
+
+		TestFalse(TEXT("After serialization the identifier should be invalid"), DstData.GetIdentifier().IsValid());
+
+		DstData.UpdatePayload(FUniqueBuffer::Alloc(32).MoveToShared());
+		TestEqual(TEXT("After adding a new payload the object should have the original identifier"), DstData.GetIdentifier(), OriginalIdentifier);
+	}
+
+	return true;
+}
+
+} // namespace UE::Virtualization
+
 #endif //WITH_DEV_AUTOMATION_TESTS && WITH_EDITORONLY_DATA
diff --git a/Engine/Source/Runtime/CoreUObject/Private/UObject/LinkerSave.cpp b/Engine/Source/Runtime/CoreUObject/Private/UObject/LinkerSave.cpp
--- a/Engine/Source/Runtime/CoreUObject/Private/UObject/LinkerSave.cpp
+++ b/Engine/Source/Runtime/CoreUObject/Private/UObject/LinkerSave.cpp
@@ -273,6 +273,16 @@ void FLinkerSave::Serialize( void* V, int64 Length )
 #endif
 	Saver->Serialize( V, Length );
 }
+
+void FLinkerSave::OnPostSave(const FPackagePath& PackagePath)
+{
+	for (TUniqueFunction<void(const FPackagePath&)>& Callback : PostSaveCallbacks)
+	{
+		Callback(PackagePath);
+	}
+
+	PostSaveCallbacks.Empty();
+}
 	
 FString FLinkerSave::GetDebugName() const
 {
diff --git a/Engine/Source/Runtime/CoreUObject/Private/UObject/SavePackage.cpp b/Engine/Source/Runtime/CoreUObject/Private/UObject/SavePackage.cpp
--- a/Engine/Source/Runtime/CoreUObject/Private/UObject/SavePackage.cpp
+++ b/Engine/Source/Runtime/CoreUObject/Private/UObject/SavePackage.cpp
@@ -18,7 +18,6 @@
 #include "Serialization/LargeMemoryReader.h"
 #include "Serialization/BufferArchive.h"
 #include "Misc/AssetRegistryInterface.h"
-#include "Misc/ConfigCacheIni.h"
 #include "Misc/FeedbackContext.h"
 #include "Misc/ScopedSlowTask.h"
 #include "Misc/ObjectThumbnail.h"
@@ -81,18 +80,6 @@ static FCriticalSection InitializeCoreClassesCritSec;
 // bring the UObectGlobal declaration visible to non editor
 bool IsEditorOnlyObject(const UObject* InObject, bool bCheckRecursive, bool bCheckMarks);
 
-static void AppendAdditionalData(FLinkerSave& Linker)
-{
-	FArchive& Ar = Linker;
-
-	for (FLinkerSave::AdditionalDataCallback& Callback : Linker.AdditionalDataToAppend)
-	{
-		Callback(Ar);
-	}
-
-	Linker.AdditionalDataToAppend.Empty();
-}
-
 static bool EndSavingIfCancelled() 
 {
 	if ( GWarn->ReceivedUserCancel() )
@@ -2051,6 +2038,7 @@ FSavePackageResultStruct UPackage::Save(UPackage* InOuter, UObject* Base, EObjec
 	// Sanity checks
 	check(InOuter);
 	check(Filename);
+
 	const bool bIsCooking = TargetPlatform != nullptr;
 	FPackagePath TargetPackagePath = FPackagePath::FromLocalPath(Filename);
 	if (TargetPackagePath.GetHeaderExtension() == EPackageExtension::Unspecified)
@@ -2252,7 +2240,7 @@ FSavePackageResultStruct UPackage::Save(UPackage* InOuter, UObject* Base, EObjec
 
 		SlowTask.EnterProgressFrame();
 
-		bool Success = true;
+		bool bSuccess = true;
 		bool bRequestStub = false;
 		{
 			// FullyLoad the package's Loader, so that anything we need to serialize (bulkdata, thumbnails) is available
@@ -3461,7 +3449,7 @@ FSavePackageResultStruct UPackage::Save(UPackage* InOuter, UObject* Base, EObjec
 						{
 							Error->Logf(ELogVerbosity::Warning, TEXT("%s"), *FText::Format( NSLOCTEXT( "Core", "SavePackageNoMap", "Attempting to save a map asset '{0}' that does not contain a map object" ), FText::FromString( FString( Filename ) ) ).ToString() );
 						}
-						Success = false;
+						bSuccess = false;
 					}
 				}
 
@@ -4171,7 +4159,18 @@ FSavePackageResultStruct UPackage::Save(UPackage* InOuter, UObject* Base, EObjec
 					return SaveResult;
 				}
 
-				AppendAdditionalData(*Linker);
+				ESavePackageResult AppendAdditionalDataResult = SavePackageUtilities::AppendAdditionalData(*Linker.Get());
+				if (AppendAdditionalDataResult != ESavePackageResult::Success)
+				{
+					return AppendAdditionalDataResult;
+				}
+
+				FSavePackageOutputFileArray AdditionalOutputFiles;
+				ESavePackageResult CreateSidecarResult = SavePackageUtilities::CreatePayloadSidecarFile(*Linker, TargetPackagePath, bSaveAsync, !bDiffing, AdditionalOutputFiles);
+				if (CreateSidecarResult != ESavePackageResult::Success)
+				{
+					return CreateSidecarResult;
+				}
 
 #if WITH_EDITOR
 				if (bIsCooking && AdditionalFilesFromExports.Num() > 0)
@@ -4423,7 +4422,7 @@ FSavePackageResultStruct UPackage::Save(UPackage* InOuter, UObject* Base, EObjec
 				}
 				SlowTask.EnterProgressFrame();
 
-				if( Success == true )
+				if( bSuccess == true )
 				{
 					{
 						// If we're writing to the existing file call ResetLoaders on the Package so that we drop the handle to the file on disk and can write to it
@@ -4433,12 +4432,14 @@ FSavePackageResultStruct UPackage::Save(UPackage* InOuter, UObject* Base, EObjec
 
 					// Compress the temporarily file to destination.
 					if (bSaveAsync)
-					{						
+					{	
 						FString NewPathToSave = NewPath;
 #if WITH_EDITOR
 
 						if (SaveFlags & SAVE_DiffCallstack)
 						{
+							checkf(AdditionalOutputFiles.IsEmpty(), TEXT("Saving additional output files with the 'SAVE_DiffCallstack' flag is not supported!"));
+
 							const TCHAR* CutoffString = TEXT("UEditorEngine::Save()");
 							FArchiveStackTrace* Writer = (FArchiveStackTrace*)(Linker->Saver);
 							TMap<FName, FArchiveDiffStats> PackageDiffStats;
@@ -4455,6 +4456,8 @@ FSavePackageResultStruct UPackage::Save(UPackage* InOuter, UObject* Base, EObjec
 						}
 						else if (SaveFlags & SAVE_DiffOnly)
 						{
+							checkf(AdditionalOutputFiles.IsEmpty(), TEXT("Saving additional output files with the 'SAVE_DiffOnly' flag is not supported!"));
+
 							FArchiveStackTrace* Writer = (FArchiveStackTrace*)(Linker->Saver);
 							FArchiveDiffMap OutDiffMap;
 							bDiffOnlyIdentical = Writer->GenerateDiffMap(*NewPath, IsEventDrivenLoaderEnabledInCookedBuilds() ? Linker->Summary.TotalHeaderSize : 0, DiffSettings.MaxDiffsToLog, OutDiffMap);
@@ -4473,11 +4476,13 @@ FSavePackageResultStruct UPackage::Save(UPackage* InOuter, UObject* Base, EObjec
 							FLargeMemoryWriter* Writer = static_cast<FLargeMemoryWriter*>(Linker->Saver);
 							const int64 DataSize = Writer->TotalSize();
 
-							if (!(SaveFlags & SAVE_DiffCallstack))  // Avoid double counting the package size if SAVE_DiffCallstack flag is set and DiffSettings.bSaveForDiff == true.
-								TotalPackageSizeUncompressed += DataSize;
-
 							if (IsEventDrivenLoaderEnabledInCookedBuilds() && Linker->IsCooking())
 							{
+								checkf(AdditionalOutputFiles.IsEmpty(), TEXT("Saving additional output files during cooking is not supported!"));
+
+								if (!(SaveFlags & SAVE_DiffCallstack))  // Avoid double counting the package size if SAVE_DiffCallstack flag is set and DiffSettings.bSaveForDiff == true.
+									TotalPackageSizeUncompressed += DataSize;
+
 								if (SavePackageContext != nullptr && SavePackageContext->PackageStoreWriter != nullptr)
 								{
 									FIoBuffer IoBuffer(FIoBuffer::AssumeOwnership, Writer->ReleaseOwnership(), DataSize);
@@ -4509,17 +4514,30 @@ FSavePackageResultStruct UPackage::Save(UPackage* InOuter, UObject* Base, EObjec
 									{
 										WriteOptions |= EAsyncWriteOptions::ComputeHash;
 									}
+
 									SavePackageUtilities::AsyncWriteFileWithSplitExports(AsyncWriteAndHashSequence, FLargeMemoryPtr(Writer->ReleaseOwnership()), DataSize, Linker->Summary.TotalHeaderSize, *NewPathToSave, WriteOptions, Linker->FileRegions);
 								}
 							}
 							else
 							{
+								// Add the uasset file to the list of output files
+								AdditionalOutputFiles.Emplace(NewPathToSave, FLargeMemoryPtr(Writer->ReleaseOwnership()), Linker->FileRegions, DataSize);
+							
 								EAsyncWriteOptions WriteOptions(EAsyncWriteOptions::WriteFileToDisk);
 								if (bComputeHash)
 								{
 									WriteOptions |= EAsyncWriteOptions::ComputeHash;
 								}
-								SavePackageUtilities::AsyncWriteFile(AsyncWriteAndHashSequence, FLargeMemoryPtr(Writer->ReleaseOwnership()), DataSize, *NewPathToSave, WriteOptions, Linker->FileRegions);
+
+								for (FSavePackageOutputFile& Entry : AdditionalOutputFiles)
+								{
+									if (!(SaveFlags & SAVE_DiffCallstack))
+									{
+										TotalPackageSizeUncompressed += Entry.DataSize;
+									}
+
+									SavePackageUtilities::AsyncWriteFile(AsyncWriteAndHashSequence, WriteOptions, Entry);
+								}	
 							}
 						}
 						Linker->CloseAndDestroySaver();
@@ -4531,8 +4549,9 @@ FSavePackageResultStruct UPackage::Save(UPackage* InOuter, UObject* Base, EObjec
 					// Move the temporary file.
 					else
 					{
-						check(TempFilename.IsSet());
+						checkf(TempFilename.IsSet(), TEXT("The package should've been saved to a tmp file first!"));
 
+						// When saving in text format we will have two temp files, so we need to manually delete the non-textbased one
 						if (bTextFormat)
 						{
 							check(TextFormatTempFilename.IsSet());
@@ -4541,29 +4560,25 @@ FSavePackageResultStruct UPackage::Save(UPackage* InOuter, UObject* Base, EObjec
 							TextFormatTempFilename.Reset();
 						}
 
-						UE_LOG(LogSavePackage, Log,  TEXT("Moving '%s' to '%s'"), TempFilename.IsSet() ? *TempFilename.GetValue() : TEXT("UNKNOWN"), *NewPath );
-						TotalPackageSizeUncompressed += PackageSize;
-
-						Success = IFileManager::Get().Move( *NewPath, *TempFilename.GetValue());
-						TempFilename.Reset();
-
-						if (FinalTimeStamp != FDateTime::MinValue())
+						// Add the .uasset file to the list of output files
+						AdditionalOutputFiles.Emplace(NewPath, TempFilename.GetValue(), PackageSize);
+						
+						ESavePackageResult FinalizeResult = SavePackageUtilities::FinalizeTempOutputFiles(TargetPackagePath, AdditionalOutputFiles, bComputeHash, FinalTimeStamp, AsyncWriteAndHashSequence);
+						if (FinalizeResult != ESavePackageResult::Success)
 						{
-							IFileManager::Get().SetTimeStamp(*NewPath, FinalTimeStamp);
-						}
+							bSuccess = false;
+						}	
 
-						if (bComputeHash)
+						// Now add up the total size of the data that was saved
+						for (FSavePackageOutputFile& Entry : AdditionalOutputFiles)
 						{
-							SavePackageUtilities::IncrementOutstandingAsyncWrites();
-							AsyncWriteAndHashSequence.AddWork([NewPath](FMD5& State)
-							{
-								SavePackageUtilities::AddFileToHash(NewPath, State);
-								SavePackageUtilities::DecrementOutstandingAsyncWrites();
-							});
+							TotalPackageSizeUncompressed += Entry.DataSize;
 						}
 					}
 
-					if( Success == false )
+					AdditionalOutputFiles.Empty();
+
+					if( bSuccess == false )
 					{
 						if (SaveFlags & SAVE_NoError)
 						{
@@ -4658,11 +4673,16 @@ FSavePackageResultStruct UPackage::Save(UPackage* InOuter, UObject* Base, EObjec
 #endif
 
 		}
-		if( Success == true )
+		if( bSuccess == true )
 		{
 			// Package has been save, so unmark NewlyCreated flag.
 			InOuter->ClearPackageFlags(PKG_NewlyCreated);
 
+			if (Linker != nullptr)
+			{
+				Linker->OnPostSave(TargetPackagePath);
+			}
+
 			// send a message that the package was saved
 			PRAGMA_DISABLE_DEPRECATION_WARNINGS;
 			UPackage::PackageSavedEvent.Broadcast(Filename, InOuter);
@@ -4675,7 +4695,7 @@ FSavePackageResultStruct UPackage::Save(UPackage* InOuter, UObject* Base, EObjec
 
 		UE_CLOG(!bDiffing, LogSavePackage, Verbose, TEXT("Finished SavePackage %s"), Filename);
 
-		if (Success)
+		if (bSuccess)
 		{
 			// if the save was successful, update the internal package filename path if we aren't currently cooking
 #if WITH_EDITOR
diff --git a/Engine/Source/Runtime/CoreUObject/Private/UObject/SavePackage/SaveContext.h b/Engine/Source/Runtime/CoreUObject/Private/UObject/SavePackage/SaveContext.h
--- a/Engine/Source/Runtime/CoreUObject/Private/UObject/SavePackage/SaveContext.h
+++ b/Engine/Source/Runtime/CoreUObject/Private/UObject/SavePackage/SaveContext.h
@@ -567,7 +567,7 @@ public:
 	int32 OffsetAfterExportMap = 0;
 	TAsyncWorkSequence<FMD5> AsyncWriteAndHashSequence;
 	TArray<FLargeMemoryWriter, TInlineAllocator<4>> AdditionalFilesFromExports;
-
+	FSavePackageOutputFileArray AdditionalPackageFiles;
 private:
 	friend class FPackageHarvester;
 
diff --git a/Engine/Source/Runtime/CoreUObject/Private/UObject/SavePackage/SavePackageUtilities.cpp b/Engine/Source/Runtime/CoreUObject/Private/UObject/SavePackage/SavePackageUtilities.cpp
--- a/Engine/Source/Runtime/CoreUObject/Private/UObject/SavePackage/SavePackageUtilities.cpp
+++ b/Engine/Source/Runtime/CoreUObject/Private/UObject/SavePackage/SavePackageUtilities.cpp
@@ -30,6 +30,7 @@
 #include "UObject/UnrealType.h"
 #include "UObject/UObjectHash.h"
 #include "UObject/UObjectThreadContext.h"
+#include "Virtualization/VirtualizedBulkData.h"
 
 DEFINE_LOG_CATEGORY(LogSavePackage);
 UE_TRACE_CHANNEL_DEFINE(SaveTimeChannel);
@@ -194,6 +195,84 @@ const FName NAME_PrestreamPackage("PrestreamPackage");
 // Switch to SavePackage to Save2, 0:disabled, 1: enabled for cook, 2: enabled for uncooked, 3: enabled everywhere
 static TAutoConsoleVariable<int32> CVarEnablePackageNewSave(TEXT("SavePackage.EnableNewSave"), 3, TEXT("Enable new package save mechanism over the old one."));
 
+/**
+* A utility that records the state of a package's files before we start moving and overwriting them. 
+* This provides an easy way for us to restore the original state of the package incase of failures 
+* while saving.
+*/
+class FPackageBackupUtility
+{
+public:
+	FPackageBackupUtility(const FPackagePath& InPackagePath)
+		: PackagePath(InPackagePath)
+	{
+
+	}
+
+	/** 
+	* Record a file that has been moved. These will need to be moved back to
+	* restore the package.
+	*/
+	void RecordMovedFile( const FString& OriginalPath, const FString& NewLocation)
+	{
+		MovedOriginalFiles.Emplace(OriginalPath, NewLocation);
+	}
+
+	/** 
+	* Record a newly created file that did not exist before. These will need
+	* deleting to restore the package.
+	*/
+	void RecordNewFile(const FString& NewLocation)
+	{
+		NewFiles.Add(NewLocation);
+	}
+
+	/** Restores the package to it's original state */
+	void RestorePackage()
+	{
+		IFileManager& FileSystem = IFileManager::Get();
+
+		UE_LOG(LogSavePackage, Verbose, TEXT("Restoring package '%s'"), *PackagePath.GetDebugName());
+
+		// First we should delete any new file that has been saved for the package
+		for (const FString& Entry : NewFiles)
+		{
+			if (!FileSystem.Delete(*Entry))
+			{
+				UE_LOG(LogSavePackage, Error, TEXT("Failed to delete newly added file '%s' when trying to restore the package state and the package could be unstable, please revert in source control!"), *Entry);
+			}
+		}
+
+		// Now we can move back the original files
+		for (const TPair<FString, FString>& Entry : MovedOriginalFiles)
+		{
+			if (!FileSystem.Move(*Entry.Key, *Entry.Value))
+			{
+				UE_LOG(LogSavePackage, Error, TEXT("Failed to restore package '%s', the file '%s' is in an incorrect state and the package could be unstable, please revert in source control!"), *PackagePath.GetDebugName(), *Entry.Key);
+			}
+		}
+	}
+
+	/** Deletes the backed up files once they are no longer required. */
+	void DiscardBackupFiles()
+	{
+		IFileManager& FileSystem = IFileManager::Get();
+
+		// Note that we do not warn if we fail to delete a backup file as that is probably
+		// the least of the users problems at the moment.
+		for (const TPair<FString, FString>& Entry : MovedOriginalFiles)
+		{
+			FileSystem.Delete(*Entry.Value);
+		}
+	}
+
+private:
+	const FPackagePath& PackagePath;
+
+	TArray<FString> NewFiles;
+	TArray<TPair<FString, FString>> MovedOriginalFiles;	
+};
+
 bool IsNewSaveEnabled(bool bForCooking)
 {
 	int32 EnableNewSave = CVarEnablePackageNewSave.GetValueOnAnyThread();
@@ -575,6 +654,115 @@ void AddFileToHash(FString const& Filename, FMD5& Hash)
 	delete Ar;
 }
 
+ESavePackageResult FinalizeTempOutputFiles(const FPackagePath& PackagePath, const FSavePackageOutputFileArray& OutputFiles, const bool bComputeHash, const FDateTime& FinalTimeStamp, TAsyncWorkSequence<FMD5>& AsyncWriteAndHashSequence)
+{
+	UE_LOG(LogSavePackage, Log,  TEXT("Moving output files for package: %s"), *PackagePath.GetDebugName());
+
+	IFileManager& FileSystem = IFileManager::Get();
+	FPackageBackupUtility OriginalPackageState(PackagePath);
+
+	UE_LOG(LogSavePackage, Verbose, TEXT("Moving existing files to the temp directory"));
+
+	TArray<bool, TInlineAllocator<4>> CanFileBeMoved;
+	CanFileBeMoved.SetNum(OutputFiles.Num());
+
+	// First check if any of the target files that already exist are read only still, if so we can fail the 
+	// whole thing before we try to move any files
+	for (int32 Index = 0; Index < OutputFiles.Num(); ++Index)
+	{	
+		const FSavePackageOutputFile& File = OutputFiles[Index];
+
+		if (File.FileMemoryBuffer.IsValid())
+		{
+			ensureMsgf(false, TEXT("FinalizeTempOutputFiles does not handle async saving files! (%s)"), *PackagePath.GetDebugName());
+			return ESavePackageResult::Error;
+		}
+
+		if (!File.TempFilePath.IsEmpty())
+		{
+			FFileStatData FileStats = FileSystem.GetStatData(*File.TargetPath);
+
+			if (FileStats.bIsValid && FileStats.bIsReadOnly)
+			{
+				UE_LOG(LogSavePackage, Error, TEXT("Cannot remove '%s' as it is read only!"), *File.TargetPath);
+				return ESavePackageResult::Error;
+			}
+
+			CanFileBeMoved[Index] = FileStats.bIsValid;
+		}
+		else
+		{
+			CanFileBeMoved[Index] = false;
+		}
+	}
+
+	// Now we need to move all of the files that we are going to overwrite (if any) so that we 
+	// can restore them if anything goes wrong.
+	for (int32 Index = 0; Index < OutputFiles.Num(); ++Index)
+	{
+		if (CanFileBeMoved[Index]) 
+		{
+			const FSavePackageOutputFile& File = OutputFiles[Index];
+
+			const FString BaseFilename = FPaths::GetBaseFilename(File.TargetPath);
+			const FString TempFilePath = FPaths::CreateTempFilename(*FPaths::ProjectSavedDir(), *BaseFilename.Left(32));
+
+			if (FileSystem.Move(*TempFilePath, *File.TargetPath))
+			{
+				OriginalPackageState.RecordMovedFile(File.TargetPath, TempFilePath);
+			}
+			else
+			{
+				UE_LOG(LogSavePackage, Warning, TEXT("Failed to move '%s' to temp directory"), *File.TargetPath);
+				OriginalPackageState.RestorePackage();
+
+				return ESavePackageResult::Error;
+			}
+		}
+	}
+
+	// Now attempt to move the new files from the temp location to the final location
+	for (const FSavePackageOutputFile& File : OutputFiles)
+	{
+		if (!File.TempFilePath.IsEmpty()) // Only try to move output files that were saved to temp files
+		{
+			UE_LOG(LogSavePackage, Log, TEXT("Moving '%s' to '%s'"), *File.TempFilePath, *File.TargetPath);
+
+			if (FileSystem.Move(*File.TargetPath, *File.TempFilePath))
+			{
+				OriginalPackageState.RecordNewFile(File.TargetPath);
+			}
+			else
+			{
+				UE_LOG(LogSavePackage, Warning, TEXT("Failed to move '%s' from temp directory"), *File.TargetPath);
+				OriginalPackageState.RestorePackage();
+
+				return ESavePackageResult::Error;
+			}
+
+			if (FinalTimeStamp != FDateTime::MinValue())
+			{
+				FileSystem.SetTimeStamp(*File.TargetPath, FinalTimeStamp);
+			}
+
+			if (bComputeHash)
+			{
+				SavePackageUtilities::IncrementOutstandingAsyncWrites();
+				AsyncWriteAndHashSequence.AddWork([NewPath = File.TargetPath](FMD5& State)
+				{
+					SavePackageUtilities::AddFileToHash(NewPath, State);
+					SavePackageUtilities::DecrementOutstandingAsyncWrites();
+				});
+			}
+		}
+	}
+
+	// Finally we can clean up the temp files as we do not need to restore them (failure to delete them will not be considered an error)
+	OriginalPackageState.DiscardBackupFiles();
+
+	return ESavePackageResult::Success;
+}
+
 void WriteToFile(const FString& Filename, const uint8* InDataPtr, int64 InDataSize)
 {
 	IFileManager& FileManager = IFileManager::Get();
@@ -628,6 +816,12 @@ void AsyncWriteFile(TAsyncWorkSequence<FMD5>& AsyncWriteAndHashSequence, FLargeM
 	});
 }
 
+void AsyncWriteFile(TAsyncWorkSequence<FMD5>& AsyncWriteAndHashSequence, EAsyncWriteOptions Options, FSavePackageOutputFile& File)
+{
+	checkf(File.TempFilePath.IsEmpty(), TEXT("AsyncWriteFile does not handle temp files!"));
+	AsyncWriteFile(AsyncWriteAndHashSequence, FLargeMemoryPtr(File.FileMemoryBuffer.Release()), File.DataSize, *File.TargetPath, Options, File.FileRegions);
+}
+
 void AsyncWriteFileWithSplitExports(TAsyncWorkSequence<FMD5>& AsyncWriteAndHashSequence, FLargeMemoryPtr Data, const int64 DataSize, const int64 HeaderSize, const TCHAR* Filename, EAsyncWriteOptions Options, TArrayView<const FFileRegion> InFileRegions)
 {
 	OutstandingAsyncWrites.Increment();
@@ -1915,6 +2109,104 @@ void SaveThumbnails(UPackage* InOuter, FLinkerSave* Linker, FStructuredArchive::
 #endif
 }
 
+ESavePackageResult AppendAdditionalData(FLinkerSave& Linker)
+{
+	FArchive& Ar = Linker;
+
+	for (FLinkerSave::AdditionalDataCallback& Callback : Linker.AdditionalDataToAppend)
+	{
+		Callback(Ar);
+	}
+
+	Linker.AdditionalDataToAppend.Empty();
+
+	// Note that we currently have no failure condition here, but we return a ESavePackageResult
+	// in case one needs to be added in future code.
+	return ESavePackageResult::Success;
+}
+
+ESavePackageResult CreatePayloadSidecarFile(FLinkerSave& Linker, const FPackagePath& PackagePath, const bool bSaveAsync, const bool bWriteToDisk, FSavePackageOutputFileArray& AdditionalPackageFiles)
+{
+	if (Linker.SidecarDataToAppend.IsEmpty())
+	{
+		return ESavePackageResult::Success;
+	}
+
+	// Note since we only allow sidecar file generation when saving a package and not when cooking 
+	// we know that we don't need to generate the hash or check if we should write the file, since those 
+	// operations are cooking only. However we still accept the parameters and check against them for 
+	// safety in case someone tries to add support in the future.
+	// We could add support but it is difficult to test and would be better left for a proper clean up pass
+	// once we enable SavePackage2 only. 
+	checkf(!Linker.IsCooking(), TEXT("Cannot write a sidecar file during cooking! (%s)"), *PackagePath.GetDebugName());
+	checkf(bWriteToDisk, TEXT("Cannot disable writing of the sidecar data to disk, that is supposed to be a cook only operation! (%s)"), *PackagePath.GetDebugName());
+
+	FLargeMemoryWriter Ar(0, true);
+
+	uint32 VersionNumber = UE::Virtualization::FTocEntry::PayloadSidecarFileVersion;
+	Ar << VersionNumber;
+
+	int64 TocPosition = Ar.Tell();
+
+	TArray<UE::Virtualization::FTocEntry> TableOfContents;
+	TableOfContents.SetNum(Linker.SidecarDataToAppend.Num());
+
+	// First we write an empty table of contents to the file to reserve the space
+	Ar << TableOfContents;
+
+	int32 Index = 0;
+	for (FLinkerSave::FSidecarStorageInfo& Info : Linker.SidecarDataToAppend)
+	{
+		// Fill out the entry to the table of contents
+		TableOfContents[Index].Identifier = Info.Identifier;
+		TableOfContents[Index].OffsetInFile = Ar.Tell();
+		TableOfContents[Index].UncompressedSize = Info.Payload.GetRawSize();
+
+		Index++;
+
+		// Now write the payload to the archive
+		for (const FSharedBuffer& Buffer : Info.Payload.GetCompressed().GetSegments())
+		{
+			// Const cast because FArchive requires a non-const pointer!
+			Ar.Serialize(const_cast<void*>(Buffer.GetData()), static_cast<int64>(Buffer.GetSize()));
+		}
+
+		// Reset each payload reference after it has been written to the archive, this could
+		// potentially release memory and keep our high water mark down.
+		Info.Payload.Reset();
+	}
+
+	// Now write out the table of contents again but with valid data
+	int64 EndPos = Ar.Tell();
+	Ar.Seek(TocPosition);
+	Ar << TableOfContents;
+	Ar.Seek(EndPos); 
+
+	const int64 DataSize = Ar.TotalSize();
+	checkf(DataSize > 0, TEXT("The archive should not be empty at this point!"));
+
+	FString TargetFilePath = PackagePath.GetLocalFullPath(EPackageSegment::PayloadSidecar);
+
+	if (bSaveAsync)
+	{
+		AdditionalPackageFiles.Emplace(MoveTemp(TargetFilePath), FLargeMemoryPtr(Ar.ReleaseOwnership()), TArray<FFileRegion>(), DataSize);
+	}
+	else
+	{
+		const FString BaseFilename = FPaths::GetBaseFilename(TargetFilePath);
+		FString TempFilePath = FPaths::CreateTempFilename(*FPaths::ProjectSavedDir(), *BaseFilename.Left(32));
+			
+		SavePackageUtilities::WriteToFile(TempFilePath, Ar.GetData(), DataSize); // TODO: Check the error handling here!
+		UE_LOG(LogSavePackage, Verbose, TEXT("Saved '%s' as temp file '%s'"), *TargetFilePath, *TempFilePath);
+
+		AdditionalPackageFiles.Emplace(MoveTemp(TargetFilePath), MoveTemp(TempFilePath), DataSize);		
+	}
+
+	Linker.SidecarDataToAppend.Empty();
+
+	return ESavePackageResult::Success;
+}
+
 ESavePackageResult SaveBulkData(FLinkerSave* Linker, const UPackage* InOuter, const TCHAR* Filename, const ITargetPlatform* TargetPlatform,
 				  FSavePackageContext* SavePackageContext, uint32 SaveFlags, const bool bTextFormat, const bool bDiffing,
 				  const bool bComputeHash, TAsyncWorkSequence<FMD5>& AsyncWriteAndHashSequence, int64& TotalPackageSizeUncompressed)
@@ -2426,4 +2718,4 @@ namespace AssetRegistry
 		DependencyDataRecord << SA_VALUE(TEXT("SoftPackageUsedInGame"), SoftPackageUsedInGameBits);
 	}
 }
-}
\ No newline at end of file
+}
diff --git a/Engine/Source/Runtime/CoreUObject/Private/UObject/SavePackage/SavePackageUtilities.h b/Engine/Source/Runtime/CoreUObject/Private/UObject/SavePackage/SavePackageUtilities.h
--- a/Engine/Source/Runtime/CoreUObject/Private/UObject/SavePackage/SavePackageUtilities.h
+++ b/Engine/Source/Runtime/CoreUObject/Private/UObject/SavePackage/SavePackageUtilities.h
@@ -18,10 +18,13 @@
 // This file contains private utilities shared by UPackage::Save and UPackage::Save2 
 
 class FMD5;
-class FSavePackageContext;
+class FPackagePath;
 class FSaveContext;
+class FSavePackageContext;
 template<typename StateType> class TAsyncWorkSequence;
 
+enum class ESavePackageResult;
+
 DECLARE_LOG_CATEGORY_EXTERN(LogSavePackage, Log, All);
 
 // Save Time trace
@@ -79,6 +82,47 @@ struct FCanSkipEditorReferencedPackagesWhenCooking
 	FORCEINLINE operator bool() const { return bCanSkipEditorReferencedPackagesWhenCooking; }
 };
 
+
+/** Represents an output file from the package when saving */
+struct FSavePackageOutputFile
+{
+	/** Constructor used for async saving */
+	FSavePackageOutputFile(const FString& InTargetPath, FLargeMemoryPtr&& MemoryBuffer, const TArray<FFileRegion>& InFileRegions, int64 InDataSize)
+		: TargetPath(InTargetPath)
+		, FileMemoryBuffer(MoveTemp(MemoryBuffer))
+		, FileRegions(InFileRegions)
+		, DataSize(InDataSize)
+	{
+
+	}
+
+	/** Constructor used for saving first to a temp file which can be later moved to the target directory */
+	FSavePackageOutputFile(const FString& InTargetPath, const FString& InTempFilePath, int64 InDataSize)
+		: TargetPath(InTargetPath)
+		, TempFilePath(InTempFilePath)
+		, DataSize(InDataSize)
+	{
+
+	}
+
+	/** The final target location of the file once all saving operations are completed */
+	FString TargetPath;
+
+	/** The temp location (if any) that the file is stored at, pending a move to the TargetPath */
+	FString TempFilePath;
+
+	/** The entire file stored as a memory buffer for the async saving path */
+	FLargeMemoryPtr FileMemoryBuffer;
+	/** An array of file regions in FileMemoryBuffer generated during cooking */
+	TArray<FFileRegion> FileRegions;
+
+	/** The size of the file in bytes */
+	int64 DataSize;
+};
+
+// Currently we only expect to store up to 2 files in this, so set the inline capacity to double of this
+using FSavePackageOutputFileArray = TArray<FSavePackageOutputFile, TInlineAllocator<4>>;
+
  /**
   * Helper structure to encapsulate sorting a linker's import table alphabetically, taking into account conforming to other linkers.
   * @note Save2 should not have to use this sorting long term
@@ -345,6 +389,13 @@ namespace SavePackageUtilities
 	ESavePackageResult SaveBulkData(FLinkerSave* Linker, const UPackage* InOuter, const TCHAR* Filename, const ITargetPlatform* TargetPlatform,
 		FSavePackageContext* SavePackageContext, uint32 SaveFlags, const bool bTextFormat, const bool bDiffing,
 		const bool bComputeHash, TAsyncWorkSequence<FMD5>& AsyncWriteAndHashSequence, int64& TotalPackageSizeUncompressed);
+	
+	/** Used to append additional data to the end of the package file by invoking callbacks stored in the linker */
+	ESavePackageResult AppendAdditionalData(FLinkerSave& Linker);
+	
+	/** Used to create the sidecar file (.upayload) from payloads that have been added to the linker */
+	ESavePackageResult CreatePayloadSidecarFile(FLinkerSave& Linker, const FPackagePath& PackagePath, const bool bSaveAsync, const bool bWriteToDisk, FSavePackageOutputFileArray& AdditionalPackageFiles);
+	
 	void SaveWorldLevelInfo(UPackage* InOuter, FLinkerSave* Linker, FStructuredArchive::FRecord Record);
 	EObjectMark GetExcludedObjectMarksForTargetPlatform(const class ITargetPlatform* TargetPlatform);
 	bool HasUnsaveableOuter(UObject* InObj, UPackage* InSavingPackage);
@@ -353,8 +404,17 @@ namespace SavePackageUtilities
 	void FindMostLikelyCulprit(const TArray<UObject*>& BadObjects, UObject*& MostLikelyCulprit, FString& OutReferencer, FSaveContext* InOptionalSaveContext = nullptr);
 	void AddFileToHash(FString const& Filename, FMD5& Hash);
 
+	/** 
+	  * Search 'OutputFiles' for output files that were saved to the temp directory and move those files
+	  * to their final location. Output files that were not saved to the temp directory will be ignored.
+	  * 
+	  * If errors are encountered then the original state of the package will be restored and should continue to work.
+	  */
+	ESavePackageResult FinalizeTempOutputFiles(const FPackagePath& PackagePath, const FSavePackageOutputFileArray& OutputFiles, const bool bComputeHash, const FDateTime& FinalTimeStamp, TAsyncWorkSequence<FMD5>& AsyncWriteAndHashSequence);
+
 	void WriteToFile(const FString& Filename, const uint8* InDataPtr, int64 InDataSize);
 	void AsyncWriteFile(TAsyncWorkSequence<FMD5>& AsyncWriteAndHashSequence, FLargeMemoryPtr Data, const int64 DataSize, const TCHAR* Filename, EAsyncWriteOptions Options, TArrayView<const FFileRegion> InFileRegions);
+	void AsyncWriteFile(TAsyncWorkSequence<FMD5>& AsyncWriteAndHashSequence, EAsyncWriteOptions Options, FSavePackageOutputFile& File);
 	void AsyncWriteFileWithSplitExports(TAsyncWorkSequence<FMD5>& AsyncWriteAndHashSequence, FLargeMemoryPtr Data, const int64 DataSize, const int64 HeaderSize, const TCHAR* Filename, EAsyncWriteOptions Options, TArrayView<const FFileRegion> InFileRegions);
 
 	void GetCDOSubobjects(UObject* CDO, TArray<UObject*>& Subobjects);
@@ -384,4 +444,3 @@ struct FSavePackageStats
 	static void MergeStats(const TMap<FName, FArchiveDiffStats>& ToMerge);
 };
 #endif
-
diff --git a/Engine/Source/Runtime/CoreUObject/Private/UObject/SavePackage2.cpp b/Engine/Source/Runtime/CoreUObject/Private/UObject/SavePackage2.cpp
--- a/Engine/Source/Runtime/CoreUObject/Private/UObject/SavePackage2.cpp
+++ b/Engine/Source/Runtime/CoreUObject/Private/UObject/SavePackage2.cpp
@@ -1746,6 +1746,8 @@ ESavePackageResult FinalizeFile(FStructuredArchive::FRecord& StructuredArchiveRo
 
 		if (SaveContext.IsDiffCallstack())
 		{
+			checkf(SaveContext.AdditionalPackageFiles.IsEmpty(), TEXT("Saving additional output files with the 'SAVE_DiffCallstack' flag is not currently supported! (%s)"), *PathToSave);
+
 			const TCHAR* CutoffString = TEXT("UEditorEngine::Save()");
 			FArchiveStackTrace* Writer = static_cast<FArchiveStackTrace*>(Linker->Saver);
 			TMap<FName, FArchiveDiffStats> PackageDiffStats;
@@ -1762,6 +1764,8 @@ ESavePackageResult FinalizeFile(FStructuredArchive::FRecord& StructuredArchiveRo
 		}
 		else if (SaveContext.IsDiffOnly())
 		{
+			checkf(SaveContext.AdditionalPackageFiles.IsEmpty(), TEXT("Saving additional output files with the 'SAVE_DiffOnly' flag is not currently supported! (%s)"), *PathToSave);
+
 			FArchiveStackTrace* Writer = (FArchiveStackTrace*)(Linker->Saver);
 			FArchiveDiffMap OutDiffMap;
 			SaveContext.bDiffOnlyIdentical = Writer->GenerateDiffMap(*PathToSave, IsEventDrivenLoaderEnabledInCookedBuilds() ? Linker->Summary.TotalHeaderSize : 0, SaveContext.GetMaxDiffsToLog(), OutDiffMap);
@@ -1783,6 +1787,8 @@ ESavePackageResult FinalizeFile(FStructuredArchive::FRecord& StructuredArchiveRo
 				SaveContext.TotalPackageSizeUncompressed += DataSize;
 			}
 
+			checkf(SaveContext.IsCooking() == false || SaveContext.AdditionalPackageFiles.IsEmpty(), TEXT("Saving additional output files during cooking is not currently supported! (%s)"), *PathToSave);
+
 			FSavePackageContext* SavePackageContext = SaveContext.GetSavePackageContext();
 			if (SavePackageContext != nullptr && SavePackageContext->PackageStoreWriter != nullptr && SaveContext.IsCooking())
 			{
@@ -1815,13 +1821,20 @@ ESavePackageResult FinalizeFile(FStructuredArchive::FRecord& StructuredArchiveRo
 				{
 					WriteOptions |= EAsyncWriteOptions::ComputeHash;
 				}
+
 				if (SaveContext.IsCooking())
 				{
 					SavePackageUtilities::AsyncWriteFileWithSplitExports(SaveContext.AsyncWriteAndHashSequence, FLargeMemoryPtr(Writer->ReleaseOwnership()), DataSize, Linker->Summary.TotalHeaderSize, *PathToSave, WriteOptions, Linker->FileRegions);
 				}
 				else
 				{
-					SavePackageUtilities::AsyncWriteFile(SaveContext.AsyncWriteAndHashSequence, FLargeMemoryPtr(Writer->ReleaseOwnership()), DataSize, *PathToSave, WriteOptions, Linker->FileRegions);
+					// Add the uasset file to the list of output files
+					SaveContext.AdditionalPackageFiles.Emplace(PathToSave, FLargeMemoryPtr(Writer->ReleaseOwnership()), Linker->FileRegions, DataSize);
+
+					for (FSavePackageOutputFile& Entry : SaveContext.AdditionalPackageFiles)
+					{
+						SavePackageUtilities::AsyncWriteFile(SaveContext.AsyncWriteAndHashSequence, WriteOptions, Entry);
+					}
 				}
 			}
 			SaveContext.CloseLinkerArchives();
@@ -1830,9 +1843,7 @@ ESavePackageResult FinalizeFile(FStructuredArchive::FRecord& StructuredArchiveRo
 	else
 	{
 		// Destroy archives used for saving, closing file handle.
-		bool bSuccess = SaveContext.CloseLinkerArchives();
-
-		if (!bSuccess)
+		if (SaveContext.CloseLinkerArchives() == false)
 		{
 			UE_LOG(LogSavePackage, Error, TEXT("Error writing temp file '%s' for '%s'"),
 				SaveContext.TempFilename.IsSet() ? *SaveContext.TempFilename.GetValue() : TEXT("UNKNOWN"), SaveContext.GetFilename());
@@ -1840,7 +1851,9 @@ ESavePackageResult FinalizeFile(FStructuredArchive::FRecord& StructuredArchiveRo
 		}
 
 		// Move file to its real destination
-		check(SaveContext.TempFilename.IsSet());
+		checkf(SaveContext.TempFilename.IsSet(), TEXT("The package should've been saved to a tmp file first! (%s)"), *SaveContext.GetFilename());
+
+		// When saving in text format we will have two temp files, so we need to manually delete the non-textbased one
 		if (SaveContext.IsTextFormat())
 		{
 			check(SaveContext.TextFormatTempFilename.IsSet());
@@ -1849,11 +1862,15 @@ ESavePackageResult FinalizeFile(FStructuredArchive::FRecord& StructuredArchiveRo
 			SaveContext.TextFormatTempFilename.Reset();
 		}
 
-		UE_LOG(LogSavePackage, Log, TEXT("Moving '%s' to '%s'"), SaveContext.TempFilename.IsSet() ? *SaveContext.TempFilename.GetValue() : TEXT("UNKNOWN"), SaveContext.GetFilename());
-		bSuccess = IFileManager::Get().Move(SaveContext.GetFilename(), *SaveContext.TempFilename.GetValue());
+		// Add the .uasset file to the list of output files (TODO: Fix the 0 size, it isn't used after this point but needs to be cleaned up)
+		SaveContext.AdditionalPackageFiles.Emplace(SaveContext.GetFilename(), SaveContext.TempFilename.GetValue(), 0);
+
+		ESavePackageResult FinalizeResult = SavePackageUtilities::FinalizeTempOutputFiles(SaveContext.GetTargetPackagePath(), SaveContext.AdditionalPackageFiles, SaveContext.IsComputeHash(), 
+			SaveContext.GetFinalTimestamp(), SaveContext.AsyncWriteAndHashSequence);
+		
 		SaveContext.TempFilename.Reset();
 
-		if (!bSuccess)
+		if (FinalizeResult != ESavePackageResult::Success)
 		{
 			if (SaveContext.IsGenerateSaveError())
 			{
@@ -1864,22 +1881,7 @@ ESavePackageResult FinalizeFile(FStructuredArchive::FRecord& StructuredArchiveRo
 				UE_LOG(LogSavePackage, Error, TEXT("%s"), *FString::Printf(TEXT("Error saving '%s'"), SaveContext.GetFilename()));
 				SaveContext.GetError()->Logf(ELogVerbosity::Warning, TEXT("%s"), *FText::Format(NSLOCTEXT("Core", "SaveWarning", "Error saving '{0}'"), FText::FromString(FString(SaveContext.GetFilename()))).ToString());
 			}
-			return ESavePackageResult::Error;
-		}
-
-		if (SaveContext.GetFinalTimestamp() != FDateTime::MinValue())
-		{
-			IFileManager::Get().SetTimeStamp(SaveContext.GetFilename(), SaveContext.GetFinalTimestamp());
-		}
-
-		if (SaveContext.IsComputeHash())
-		{
-			SavePackageUtilities::IncrementOutstandingAsyncWrites();
-			SaveContext.AsyncWriteAndHashSequence.AddWork([NewPath = FString(SaveContext.GetFilename())](FMD5& State)
-			{
-				SavePackageUtilities::AddFileToHash(NewPath, State);
-				SavePackageUtilities::DecrementOutstandingAsyncWrites();
-			});
+			return FinalizeResult;
 		}
 	}
 
@@ -1929,6 +1931,11 @@ void PostSavePackage(FSaveContext& SaveContext)
 	// Package has been saved, so unmark the NewlyCreated flag.
 	Package->ClearPackageFlags(PKG_NewlyCreated);
 
+	if (SaveContext.Linker != nullptr)
+	{
+		SaveContext.Linker->OnPostSave(SaveContext.GetTargetPackagePath());
+	}
+
 	// Send a message that the package was saved
 	PRAGMA_DISABLE_DEPRECATION_WARNINGS;
 	UPackage::PackageSavedEvent.Broadcast(SaveContext.GetFilename(), Package);
@@ -2053,6 +2060,16 @@ ESavePackageResult InnerSave(FSaveContext& SaveContext)
 	SaveContext.Result = SavePackageUtilities::SaveBulkData(SaveContext.Linker.Get(), SaveContext.GetPackage(), SaveContext.GetFilename(), SaveContext.GetTargetPlatform(),
 		SaveContext.GetSavePackageContext(), SaveContext.GetSaveArgs().SaveFlags, SaveContext.IsTextFormat(), SaveContext.IsDiffing(),
 		SaveContext.IsComputeHash(), SaveContext.AsyncWriteAndHashSequence, SaveContext.TotalPackageSizeUncompressed);
+
+	// Add any pending data blobs to the end of the file by invoking the callbacks
+	SaveContext.Result = SavePackageUtilities::AppendAdditionalData(*SaveContext.Linker.Get());
+	if (SaveContext.Result != ESavePackageResult::Success)
+	{
+		return SaveContext.Result;
+	}
+
+	// Create the payload side car file (if needed)
+	SaveContext.Result = SavePackageUtilities::CreatePayloadSidecarFile(*SaveContext.Linker.Get(), SaveContext.GetTargetPackagePath(), SaveContext.IsSaveAsync(), !SaveContext.IsDiffing(), SaveContext.AdditionalPackageFiles);
 	if (SaveContext.Result != ESavePackageResult::Success)
 	{
 		return SaveContext.Result;
@@ -2077,6 +2094,11 @@ ESavePackageResult InnerSave(FSaveContext& SaveContext)
 	int32 PackageSize = SaveContext.Linker->Tell();
 	SaveContext.TotalPackageSizeUncompressed += PackageSize;
 
+	for (const FSavePackageOutputFile& File : SaveContext.AdditionalPackageFiles)
+	{
+		SaveContext.TotalPackageSizeUncompressed += File.DataSize;
+	}
+
 	// Update Package Header
 	SlowTask.EnterProgressFrame();
 	SaveContext.Result = UpdatePackageHeader(StructuredArchiveRoot, SaveContext);
diff --git a/Engine/Source/Runtime/CoreUObject/Private/Virtualization/VirtualizationManager.cpp b/Engine/Source/Runtime/CoreUObject/Private/Virtualization/VirtualizationManager.cpp
--- a/Engine/Source/Runtime/CoreUObject/Private/Virtualization/VirtualizationManager.cpp
+++ b/Engine/Source/Runtime/CoreUObject/Private/Virtualization/VirtualizationManager.cpp
@@ -2,11 +2,95 @@
 
 #include "Virtualization/VirtualizationManager.h"
 
-#include "HAL/FileManager.h"
+#include "Misc/CommandLine.h"
 #include "Misc/ConfigCacheIni.h"
+#include "Misc/Parse.h"
 #include "Misc/Paths.h"
+#include "Virtualization/IVirtualizationBackend.h"
 
-DEFINE_LOG_CATEGORY(LogVirtualization);
+namespace UE::Virtualization
+{
+
+/** Utility struct, similar to FScopeLock but allows the lock to be enabled/disabled more easily */
+struct FConditionalScopeLock
+{
+	UE_NONCOPYABLE(FConditionalScopeLock);
+
+	FConditionalScopeLock(FCriticalSection* InSyncObject, bool bShouldLock)
+	{
+		checkf(InSyncObject != nullptr, TEXT("InSyncObject must point to a valid FCriticalSection"));
+
+		if (bShouldLock)
+		{		
+			SyncObject = InSyncObject;
+			SyncObject->Lock();
+		}
+		else
+		{
+			SyncObject = nullptr;
+		}
+	}
+
+	/** Destructor that performs a release on the synchronization object. */
+	~FConditionalScopeLock()
+	{
+		if (SyncObject != nullptr)
+		{
+			SyncObject->Unlock();
+		}
+	}
+
+private:
+	FCriticalSection* SyncObject;
+};
+
+/* Utility function for building up a lookup table of all available IBackendFactory interfaces*/
+TMap<FName, IVirtualizationBackendFactory*> FindBackendFactories()
+{
+	TMap<FName, IVirtualizationBackendFactory*> BackendFactories;
+
+	TArray<IVirtualizationBackendFactory*> FactoriesArray = IModularFeatures::Get().GetModularFeatureImplementations<IVirtualizationBackendFactory>(FName("VirtualizationBackendFactory"));
+	for (IVirtualizationBackendFactory* FactoryInterface : FactoriesArray)
+	{
+		checkf(FactoryInterface != nullptr, TEXT("A nullptr was added to the modular features for 'VirtualizationBackendFactory'"));
+
+		const FName FactoryName = FactoryInterface->GetName();
+
+		if (!BackendFactories.Contains(FactoryName))
+		{
+			BackendFactories.Add(FactoryName, FactoryInterface);
+		}
+		else
+		{
+			UE_LOG(LogVirtualization, Error, TEXT("Duplicate IBackendFactory found! Name '%s'"), *FactoryName.ToString());
+		}
+	}
+
+	return BackendFactories;
+}
+
+/* Utility function for finding entries in a given string*/
+TArray<FString> ParseEntries(const FString& Data)
+{
+	TArray<FString> Entries;
+
+	const TCHAR* DataPtr = *Data;
+
+	const TCHAR* EntryLabel = TEXT("Entry=");
+	const int32 EntryLabelLength = FCString::Strlen(EntryLabel);
+
+	FString ConfigEntryName;
+	while (FParse::Value(DataPtr, EntryLabel, ConfigEntryName))
+	{
+		Entries.Add(ConfigEntryName);
+
+		// Skip head so we can look for any additional entries (note that we might not skip past the existing
+		// entry has we have no idea how much whitespace was ignored by FParse, but it will be enough)
+		DataPtr += EntryLabelLength + ConfigEntryName.Len();
+	}
+
+	return Entries;
+}
 
 FVirtualizationManager& FVirtualizationManager::Get()
 {
@@ -18,103 +102,353 @@ FVirtualizationManager& FVirtualizationManager::Get()
 FVirtualizationManager::FVirtualizationManager()
 	: bEnablePayloadPushing(true)
 	, MinPayloadLength(0)
+	, BackendGraphName(TEXT("ContentVirtualizationBackendGraph_None"))
+	, bForceSingleThreaded(false)
+	, bFailPayloadPullOperations(false)
+	, bValidateAfterPushOperation(false)
 {
-	LoadSettingsFromConfigFiles();
+	UE_LOG(LogVirtualization, Log, TEXT("Virtualization manager created"));
+
+	FConfigFile PlatformEngineIni;
+	if (FConfigCacheIni::LoadLocalIniFile(PlatformEngineIni, TEXT("Engine"), true))
+	{
+		ApplySettingsFromConfigFiles(PlatformEngineIni);
+		ApplyDebugSettingsFromConfigFiles(PlatformEngineIni);
+	}
+	else
+	{
+		UE_LOG(LogVirtualization, Error, TEXT("Failed to load config file settings for content virtualization"));
+	}
+
+	ApplySettingsFromCmdline();
+
+	MountBackends();
 }
 
-bool FVirtualizationManager::PushData(const FSharedBuffer& Payload, const FGuid& Guid)
+FVirtualizationManager::~FVirtualizationManager()
 {
-	checkf(Guid.IsValid(), TEXT("VirtualizationManager::PushData called with an invalid guid (%s)."), *Guid.ToString());
+	UE_LOG(LogVirtualization, Log, TEXT("Destroying backends"));
 
-	// Early out if the pushing of payloads is disabled
-	if (bEnablePayloadPushing == false)
+	for (IVirtualizationBackend* Backend : AllBackendsArray)
+	{
+		delete Backend;
+	}
+
+	AllBackendsArray.Empty();
+	PullEnabledBackendsArray.Empty();
+	PushEnabledBackendsArray.Empty();
+
+	UE_LOG(LogVirtualization, Log, TEXT("Virtualization manager destroyed"));
+}
+
+bool FVirtualizationManager::PushData(const FPayloadId& Id, const FCompressedBuffer& Payload)
+{
+	TRACE_CPUPROFILER_EVENT_SCOPE(FVirtualizationManager::PushData);
+
+	if (!Id.IsValid())
+	{
+		// TODO: Should an invalid FPayloadId be an expected input, if so demote this from Warning->Verbose
+		UE_LOG(LogVirtualization, Warning, TEXT("Attempting to push a virtualized payload with an invalid FPayloadId"));
+		return false;
+	}
+
+	FConditionalScopeLock _(&ForceSingleThreadedCS, bForceSingleThreaded);
+
+	// Early out if there are no backends or if the pushing of payloads has been disabled
+	if (PushEnabledBackendsArray.IsEmpty() || bEnablePayloadPushing == false)
 	{
 		return false;
 	}
 
 	// Early out if we have no payload
-	if (Payload.GetSize() == 0)
+	if (Payload.GetCompressedSize() == 0)
 	{
+		// TODO: Should an invalid payload be an expected input, if so demote this from Warning->Verbose
+		UE_LOG(LogVirtualization, Warning, TEXT("Attempting to push an invalid virtualized payload (id: %s)"), *Id.ToString());
 		return false;
 	}
 
 	// Early out if the payload length is below our minimum required length
-	if ((int64)Payload.GetSize() < MinPayloadLength)
+	if ((int64)Payload.GetCompressedSize() < MinPayloadLength)
 	{
+		UE_LOG(	LogVirtualization, 
+				Verbose, 
+				TEXT("Attempting to push a virtualized payload (id: %s) that is smaller (%" UINT64_FMT ") than the MinPayloadLength (%" INT64_FMT ")"), 
+				*Id.ToString(), Payload.GetCompressedSize(), MinPayloadLength);
+
 		return false;
 	}
 
-	// TODO: At this point we should be able to consider which back end actually gets used, the following code is 
-	//		 more of a reference implementation!
+	// TODO: Note that all push operations are currently synchronous, probably 
+	// should change to async at some point, although this makes handling failed
+	// pushed much more difficult.
+	bool bWasPayloadPushed = false;
+	for (IVirtualizationBackend* Backend : PushEnabledBackendsArray)
 	{
-		const FString FilePath = CreateFilePath(Guid);
+		const bool bDidPushPayload = Backend->PushData(Id, Payload);
+		
+		UE_CLOG(bDidPushPayload, LogVirtualization, Verbose, TEXT("[%s] Pushed the payload '%s'"), *Backend->GetDebugString(), *Id.ToString());
+		UE_CLOG(!bDidPushPayload, LogVirtualization, Error, TEXT("[%s] Failed to push the payload '%s'"), *Backend->GetDebugString(), *Id.ToString());
+
+		bWasPayloadPushed |= bDidPushPayload;
+		
+		// Debugging operation where we immediately try to pull the payload after each push (when possible) and assert 
+		// that the pulled payload is the same as the original
+		if (bValidateAfterPushOperation && bDidPushPayload && Backend->SupportsPullOperations())
+		{
+			FCompressedBuffer PulledPayload = Backend->PullData(Id);
+			checkf(Payload.GetRawHash() == PulledPayload.GetRawHash(), TEXT("[%s] Failed to pull payload '%s' after it was pushed to backend"), 
+				*Backend->GetDebugString(), *Id.ToString());
+		}
+	}
+
+	UE_CLOG(!bWasPayloadPushed, LogVirtualization, Fatal, TEXT("Payload '%s' failed to be pushed to any backend'"), *Id.ToString());
+
+	return bWasPayloadPushed;
+}
+
+FCompressedBuffer FVirtualizationManager::PullData(const FPayloadId& Id)
+{
+	TRACE_CPUPROFILER_EVENT_SCOPE(FVirtualizationManager::PullData);
 
-		TUniquePtr<FArchive> FileAr(IFileManager::Get().CreateFileWriter(*FilePath));
-		check(FileAr);
+	if (!Id.IsValid())
+	{
+		// TODO: See below, should errors here be fatal?
+		UE_LOG(LogVirtualization, Error, TEXT("Attempting to pull a virtualized payload with an invalid FPayloadId"));
+		return FCompressedBuffer();
+	}
 
-		uint64 PayloadLength = Payload.GetSize();
-		// Cast away the const because FArchive requires a non-const pointer even if we are only reading from it (saving)
-		void* PayloadPtr = const_cast<void*>(Payload.GetData());
+	if (PullEnabledBackendsArray.IsEmpty())
+	{
+		// TODO: See below, should errors here be fatal?
+		UE_LOG(LogVirtualization, Error, TEXT("Payload '%s' failed to be pulled as there are no backends mounted!'"), *Id.ToString());
+		return FCompressedBuffer();
+	}
 
-		*FileAr << PayloadLength;
-		FileAr->Serialize(PayloadPtr, PayloadLength);
+	if (bFailPayloadPullOperations)
+	{
+		UE_LOG(LogVirtualization, Error, TEXT("Payload '%s' failed to be pulled as the debug option 'FailPayloadPullOperations' is enabled!"), *Id.ToString());
+		return FCompressedBuffer();
 	}
 
-	return true;
+	FConditionalScopeLock _(&ForceSingleThreadedCS, bForceSingleThreaded);
+
+	// TODO: Once a payload is found, other backends should probably be notified 
+	// (a local cache might want to replicate the data for example)
+	for (IVirtualizationBackend* Backend : PullEnabledBackendsArray)
+	{
+		FCompressedBuffer Payload = Backend->PullData(Id);
+		if (Payload)
+		{
+			return Payload;
+		}
+	}
+
+	// TODO: Maybe this should be a fatal error? If we keep it as an error we need to make sure any calling
+	// code handles it properly.
+	// Could be worth extending ::PullData to return error codes instead so we can make a better distinction 
+	// between the payload not being found in any of the backends and one or more of the backends failing.
+	UE_LOG(LogVirtualization, Error, TEXT("Payload '%s' failed to be pulled from any backend'"), *Id.ToString());
+
+	return FCompressedBuffer();
 }
 
-FSharedBuffer FVirtualizationManager::PullData(const FGuid& Guid)
+void FVirtualizationManager::ApplySettingsFromConfigFiles(const FConfigFile& PlatformEngineIni)
 {
-	checkf(Guid.IsValid(), TEXT("VirtualizationManager::PullData called with an invalid guid (%s)."), *Guid.ToString());
-
-	// TODO: At this point we should be able to consider which back end actually gets used, the following code is 
-	//		 more of a reference implementation!
+	UE_LOG(LogVirtualization, Log, TEXT("Loading virtualization manager settings from config files..."));
 	
-	const FString FilePath = CreateFilePath(Guid);
+	bool bEnablePayloadPushingFromIni = false;
+	if (PlatformEngineIni.GetBool(TEXT("Core.ContentVirtualization"), TEXT("EnablePushToBackend"), bEnablePayloadPushingFromIni))
+	{
+		bEnablePayloadPushing = bEnablePayloadPushingFromIni;
+		UE_LOG(LogVirtualization, Log, TEXT("\tEnablePushToBackend : %s"), bEnablePayloadPushing ? TEXT("true") : TEXT("false") );
+	}
+	else
+	{
+		UE_LOG(LogVirtualization, Error, TEXT("Failed to load [Core.ContentVirtualization].EnablePushToBackend from config file!"));
+	}
 
-	TUniquePtr<FArchive> FileAr(IFileManager::Get().CreateFileReader(*FilePath));
-	checkf(FileAr != nullptr, TEXT("Unable to open %s, virtualized data is inaccessible!"), *FilePath);
+	int64 MinPayloadLengthFromIni = 0;
+	if (PlatformEngineIni.GetInt64(TEXT("Core.ContentVirtualization"), TEXT("MinPayloadLength"), MinPayloadLengthFromIni))
+	{
+		MinPayloadLength = MinPayloadLengthFromIni;
+		UE_LOG(LogVirtualization, Log, TEXT("\tMinPayloadLength : %" INT64_FMT), MinPayloadLength );
+	}
+	else
+	{
+		UE_LOG(LogVirtualization, Error, TEXT("Failed to load [Core.ContentVirtualization].MinPayloadLength from config file!"));;
+	}
+
+	FString BackendGraphNameFromIni;
+	if (PlatformEngineIni.GetString(TEXT("Core.ContentVirtualization"), TEXT("BackendGraph"), BackendGraphNameFromIni))
+	{
+		BackendGraphName = BackendGraphNameFromIni;
+		UE_LOG(LogVirtualization, Log, TEXT("\tBackendGraphName : %s"), *BackendGraphName );
+	}
+	else
+	{
+		UE_LOG(LogVirtualization, Error, TEXT("Failed to load [Core.ContentVirtualization].BackendGraph from config file!"));;
+	}
+}
 
-	int64 DataLength = 0;
-	*FileAr << DataLength;
+void FVirtualizationManager::ApplySettingsFromCmdline()
+{
+	FString CmdlineGraphName;
+	if (FParse::Value(FCommandLine::Get(), TEXT("-BackendGraph="), CmdlineGraphName))
+	{
+		UE_LOG(LogVirtualization, Log, TEXT("Backend graph overriden from the cmdline: '%s'"), *CmdlineGraphName);
+		BackendGraphName = CmdlineGraphName;
+	}
 
-	FUniqueBuffer Buffer = FUniqueBuffer::Alloc(DataLength);
-	FileAr->Serialize(Buffer.GetData(), DataLength);
-	return Buffer.MoveToShared();
+	if (FParse::Param(FCommandLine::Get(), TEXT("VirtualizationForceSingleThreaded")))
+	{
+		bForceSingleThreaded = true;
+		UE_LOG(LogVirtualization, Log, TEXT("ForceSingleThreaded overriden from the cmdline: true"));	
+	}
 }
 
-void FVirtualizationManager::LoadSettingsFromConfigFiles()
+void FVirtualizationManager::ApplyDebugSettingsFromConfigFiles(const FConfigFile& PlatformEngineIni)
 {
-	FConfigFile PlatformEngineIni;
-	if (FConfigCacheIni::LoadLocalIniFile(PlatformEngineIni, TEXT("Engine"), true))
+	UE_LOG(LogVirtualization, Log, TEXT("Loading virtualization manager debugging settings from config files..."));
+
+	// Note that the debug settings are optional and could be left out of the config files entirely
+	bool bForceSingleThreadedFromIni = false;
+	if (PlatformEngineIni.GetBool(TEXT("Core.ContentVirtualizationDebugOptions"), TEXT("ForceSingleThreaded"), bForceSingleThreadedFromIni))
 	{
-		bool bEnablePayloadPushingFromIni = false;
-		if (PlatformEngineIni.GetBool(TEXT("Core.ContentVirtualization"), TEXT("EnablePushToBackend"), bEnablePayloadPushingFromIni))
-		{
-			bEnablePayloadPushing = bEnablePayloadPushingFromIni;
-		}
-		else
-		{
-			UE_LOG(LogVirtualization, Error, TEXT("Failed to load [Core.ContentVirtualization].EnablePushToBackend from ini file!"));
-		}
+		bForceSingleThreaded = bForceSingleThreadedFromIni;
+		UE_LOG(LogVirtualization, Log, TEXT("\tForceSingleThreaded : %s"), bForceSingleThreaded ? TEXT("true") : TEXT("false"));
+	}
+
+	bool bFailPayloadPullOperationsFromIni = false;
+	if (PlatformEngineIni.GetBool(TEXT("Core.ContentVirtualizationDebugOptions"), TEXT("FailPayloadPullOperations"), bFailPayloadPullOperationsFromIni))
+	{
+		bFailPayloadPullOperations = bFailPayloadPullOperationsFromIni;
+		UE_LOG(LogVirtualization, Log, TEXT("\tFailPayloadPullOperations : %s"), bFailPayloadPullOperations ? TEXT("true") : TEXT("false"));
+	}
+
+	bool bValidateAfterPushOperationFromIni = false;
+	if (PlatformEngineIni.GetBool(TEXT("Core.ContentVirtualizationDebugOptions"), TEXT("ValidateAfterPushOperation"), bValidateAfterPushOperationFromIni))
+	{
+		bValidateAfterPushOperation = bValidateAfterPushOperationFromIni;
+		UE_LOG(LogVirtualization, Log, TEXT("ValidateAfterPushOperation : %s"), bValidateAfterPushOperation ? TEXT("true") : TEXT("false"));
+	}
 
-		int64 MinPayloadLengthFromIni = 0;
-		if (PlatformEngineIni.GetInt64(TEXT("Core.ContentVirtualization"), TEXT("MinPayloadLength"), MinPayloadLengthFromIni))
+	// Some debug options will cause intentional breaks or slow downs for testing purposes, if these are enabled then we should give warning/errors 
+	// so it is clear in the log that future failures are being caused by the given dev option.
+	UE_CLOG(bForceSingleThreaded, LogVirtualization, Warning, TEXT("ForceSingleThreaded is enabled, virtualization will run in single threaded mode and may be slower!"));
+	UE_CLOG(bFailPayloadPullOperations, LogVirtualization, Error, TEXT("FailPayloadPullOperations is enabled, all virtualization pull operations will fail!"));
+	UE_CLOG(bValidateAfterPushOperation, LogVirtualization, Error, TEXT("ValidateAfterPushOperation is enabled, each push will be followed by a pull to validate it!"));
+}
+
+void FVirtualizationManager::MountBackends()
+{
+	UE_LOG(LogVirtualization, Log, TEXT("Mounting virtualization backends..."));
+
+	TMap<FName, IVirtualizationBackendFactory*> FactoryLookupTable = FindBackendFactories();
+	UE_LOG(LogVirtualization, Verbose, TEXT("Found %d backend factories"), FactoryLookupTable.Num());
+
+	const TCHAR* GraphName = *BackendGraphName;
+	const TCHAR* HierarchyKey = TEXT("Hierarchy");
+
+	UE_LOG(LogVirtualization, Log, TEXT("Using backend graph: '%s'"), GraphName);
+
+	FString HierarchyData;
+	if (!GConfig->GetString(GraphName, HierarchyKey, HierarchyData, GEngineIni))
+	{
+		UE_LOG(LogVirtualization, Fatal, TEXT("Unable to find the '%s' entry for the content virtualization backend graph '%s' [ini=%s]."), HierarchyKey, GraphName, *GEngineIni);
+	}
+
+	if (HierarchyData.IsEmpty())
+	{
+		UE_LOG(LogVirtualization, Fatal, TEXT("The '%s' entry for backend graph '%s' is empty [ini=%s]."), HierarchyKey, GraphName, *GEngineIni);
+	}
+
+	TArray<FString> Entries = ParseEntries(HierarchyData);
+
+	UE_LOG(LogVirtualization, Log, TEXT("The backend graph hierarchy has %d entries"), Entries.Num());
+
+	for (const FString& Entry : Entries)
+	{
+		CreateBackend(GraphName, Entry, FactoryLookupTable);
+	}
+}
+
+bool FVirtualizationManager::CreateBackend(const TCHAR* GraphName, const FString& ConfigEntryName, TMap<FName, IVirtualizationBackendFactory*>& FactoryLookupTable)
+{
+	// All failures in this method are considered fatal, however it still returns true/false in case we decide
+	// to be more forgiving in the future.
+	UE_LOG(LogVirtualization, Log, TEXT("Attempting to create back end entry '%s'"), *ConfigEntryName);
+
+	FString BackendData;
+	if (!GConfig->GetString(GraphName, *ConfigEntryName, BackendData, GEngineIni))
+	{
+		UE_LOG(LogVirtualization, Fatal, TEXT("Unable to find the entry '%s' in the content virtualization backend graph '%s' [ini=%s]."), *ConfigEntryName, GraphName, *GEngineIni);
+		return false;
+	}
+
+	FString BackendType;
+	if (FParse::Value(*BackendData, TEXT("Type="), BackendType) && !BackendType.IsEmpty())
+	{
+		// Put the rest of the ini file entry into a string to pass to the backend.
+		FString Cmdine = BackendData.RightChop(BackendData.Find(BackendType) + BackendType.Len());
+		Cmdine.RemoveFromEnd(TEXT(")"));
+
+		IVirtualizationBackendFactory** FactoryPtr = FactoryLookupTable.Find(FName(BackendType));
+		if (FactoryPtr != nullptr && *FactoryPtr != nullptr)
 		{
-			MinPayloadLength = MinPayloadLengthFromIni;
+			IVirtualizationBackendFactory* Factory = *FactoryPtr;
+			IVirtualizationBackend* Backend = Factory->CreateInstance(ConfigEntryName);
+
+			if (Backend == nullptr)
+			{
+				UE_LOG(LogVirtualization, Fatal, TEXT("IVirtualizationBackendFactory '%s' failed to create an instance!"), *Factory->GetName().ToString());
+				return false;
+
+			}
+			
+			if (Backend->Initialize(Cmdine))
+			{
+				AddBackend(Backend);
+			}
+			else
+			{
+				UE_LOG(LogVirtualization, Fatal, TEXT("Backend '%s' reported errors when initializing"), *ConfigEntryName);
+				delete Backend;
+				return false;
+			}
 		}
 		else
 		{
-			UE_LOG(LogVirtualization, Error, TEXT("Failed to load [Core.ContentVirtualization].MinPayloadLength from ini file!"));;
+			UE_LOG(LogVirtualization, Fatal, TEXT("No backend factory found that can create the type '%s'"), *BackendType);
+			return false;
 		}
 	}
 	else
 	{
-		UE_LOG(LogVirtualization, Error, TEXT("Failed to load 'Core.ContentVirtualization' ini file settings!"));
+		UE_LOG(LogVirtualization, Fatal, TEXT("No 'Type=' entry found for '%s' in the config file"), *ConfigEntryName);
+		return false;
 	}
+
+	return true;
 }
 
-FString FVirtualizationManager::CreateFilePath(const FGuid& Guid)
+void FVirtualizationManager::AddBackend(IVirtualizationBackend* Backend)
 {
-	return FString::Printf(TEXT("%sVirtualizedPayloads/%s.payload"), *FPaths::ProjectSavedDir(), *Guid.ToString());
+	checkf(!AllBackendsArray.Contains(Backend), TEXT("Adding the same virtualization backend (%s) multiple times!"), *Backend->GetDebugString());
+
+	AllBackendsArray.Add(Backend);
+
+	if (Backend->SupportsPullOperations())
+	{
+		PullEnabledBackendsArray.Add(Backend);
+	}
+
+	if (Backend->SupportsPushOperations())
+	{
+		PushEnabledBackendsArray.Add(Backend);
+	}
+
+	UE_LOG(LogVirtualization, Log, TEXT("Mounted backend: %s"), *Backend->GetDebugString());
 }
+
+} // namespace UE::Virtualization
