Cherrypicking Mirage code base (disabled) from Dev-Cooker

diff --git a/Engine/Source/Runtime/CoreUObject/Private/Virtualization/VirtualizedBulkData.cpp b/Engine/Source/Runtime/CoreUObject/Private/Virtualization/VirtualizedBulkData.cpp
--- a/Engine/Source/Runtime/CoreUObject/Private/Virtualization/VirtualizedBulkData.cpp
+++ b/Engine/Source/Runtime/CoreUObject/Private/Virtualization/VirtualizedBulkData.cpp
@@ -3,6 +3,8 @@
 #include "Virtualization/VirtualizedBulkData.h"
 
 #include "HAL/FileManager.h"
+#include "HAL/IConsoleManager.h"
+#include "Misc/ConfigCacheIni.h"
 #include "Misc/PackageSegment.h"
 #include "Misc/SecureHash.h"
 #include "Misc/ScopeLock.h"
@@ -10,28 +12,149 @@
 #include "UObject/LinkerLoad.h"
 #include "UObject/LinkerSave.h"
 #include "UObject/Object.h"
+#include "Virtualization/IVirtualizationSourceControlUtilities.h"
 
 //#if WITH_EDITORONLY_DATA
 
-FVirtualizedUntypedBulkData::~FVirtualizedUntypedBulkData()
+namespace UE::Virtualization
 {
-	Reset();
+
+/** This console variable should only exist for testing */
+static TAutoConsoleVariable<bool> CVarShouldLoadFromSidecar(
+	TEXT("Serialization.LoadFromSidecar"),
+	false,
+	TEXT("When true FVirtualizedUntypedBulkData will load from the sidecar file"));
+
+static TAutoConsoleVariable<bool> CVarShouldValidatePayload(
+	TEXT("Serialization.ValidatePayloads"),
+	false,
+	TEXT("When true FVirtualizedUntypedBulkData validate any payload loaded from the sidecar file"));
+
+static TAutoConsoleVariable<bool> CVarShouldAllowSidecarSyncing(
+	TEXT("Serialization.AllowSidecarSyncing"),
+	false,
+	TEXT("When true FVirtualizedUntypedBulkData will attempt to sync it's .upayload file via sourcecontrol if the first attempt to load from it fails"));
+
+/** Wrapper around the config file option [Core.System.Experimental]EnablePackageSidecarSaving */
+bool ShouldSaveToPackageSidecar()
+{
+	static const struct FSaveToPackageSidecar
+	{
+		bool bEnabled = false;
+
+		FSaveToPackageSidecar()
+		{
+			GConfig->GetBool(TEXT("Core.System.Experimental"), TEXT("EnablePackageSidecarSaving"), bEnabled, GEngineIni);
+		}
+	} ConfigSetting;
+
+	return ConfigSetting.bEnabled;
+}
+
+/** Utility for logging extended error messages when we fail to open a package for reading */
+void LogPackageOpenFailureMessage(const FPackagePath& PackagePath, EPackageSegment PackageSegment)
+{
+	// TODO: Check the various error paths here again!
+	const uint32 SystemError = FPlatformMisc::GetLastError();
+	// If we have a system error we can give a more informative error message but don't output it if the error is zero as 
+	// this can lead to very confusing error messages.
+	if (SystemError != 0)
+	{
+		TCHAR SystemErrorMsg[2048] = { 0 };
+		FPlatformMisc::GetSystemErrorMessage(SystemErrorMsg, sizeof(SystemErrorMsg), SystemError);
+		UE_LOG(LogVirtualization, Error, TEXT("Could not open the file '%s' for reading due to system error: '%s' (%d))"), *PackagePath.GetDebugNameWithExtension(PackageSegment), SystemErrorMsg, SystemError);
+	}
+	else
+	{
+		UE_LOG(LogVirtualization, Error, TEXT("Could not open (%s) to read FVirtualizedUntypedBulkData with an unknown error"), *PackagePath.GetDebugNameWithExtension(PackageSegment));
+	}
+}
+
+/** Utility for accessing IVirtualizationSourceControlUtilities from the modular feature system. */
+IVirtualizationSourceControlUtilities* GetSourceControlInterface()
+{
+	return (IVirtualizationSourceControlUtilities*)IModularFeatures::Get().GetModularFeatureImplementation(FName("VirtualizationSourceControlUtilities"), 0);
+}
+
+FVirtualizedUntypedBulkData::FVirtualizedUntypedBulkData(const FVirtualizedUntypedBulkData& Other)
+{
+	*this = Other;
+}
+
+FVirtualizedUntypedBulkData& FVirtualizedUntypedBulkData::operator=(const FVirtualizedUntypedBulkData& Other)
+{
+	if (!BulkDataId.IsValid() && Other.BulkDataId.IsValid())
+	{
+		BulkDataId = FGuid::NewGuid();
+	}
+
+	PayloadContentId = Other.PayloadContentId;
+	Payload = Other.Payload;
+	PayloadSize = Other.PayloadSize;
+	CompressionFormatToUse = Other.CompressionFormatToUse;
+	OffsetInFile = Other.OffsetInFile;
+	PackagePath = Other.PackagePath;
+	PackageSegment = Other.PackageSegment;
+	bWasKeyGuidDerived = Other.bWasKeyGuidDerived;
+	Flags = Other.Flags;
+	
+	return *this;
 }
 
 void FVirtualizedUntypedBulkData::CreateFromBulkData(FUntypedBulkData& InBulkData, const FGuid& InGuid)
 {
 	TRACE_CPUPROFILER_EVENT_SCOPE(FVirtualizedUntypedBulkData::CreateFromBulkData);
+	
+	checkf(!BulkDataId.IsValid(), 
+		TEXT("Calling ::CreateFromBulkData on a bulkdata object that already has a valid identifier! Package: '%s'"),
+		*InBulkData.GetPackagePath().GetDebugName());
 
 	Reset();
 
+	auto AssignGuid = [this, &InBulkData](const FGuid& InGuid)
+	{
+		// Only use the guid if we will have a valid payload
+		if (InBulkData.GetBulkDataSize() > 0)
+		{
+			BulkDataId = InGuid;
+		}
+
+		PayloadContentId = FPayloadId(InGuid);
+	};
+
+	if (InGuid.IsValid())
+	{
+		AssignGuid(InGuid);
+	}
+	else
+	{
+		UE_LOG(LogVirtualization, Warning,
+			TEXT("CreateFromBulkData recieved an invalid FGuid. A temporary one will be generated until the package is next re-saved! Package: '%s'"),
+			*InBulkData.GetPackagePath().GetDebugName());
+
+		AssignGuid(FGuid::NewGuid());
+	}	
+	
+	PayloadSize = InBulkData.GetBulkDataSize();
+	
+	bWasKeyGuidDerived = true;
 	PackagePath = InBulkData.GetPackagePath();
 	PackageSegment = InBulkData.GetPackageSegment();
-
+	
 	OffsetInFile = InBulkData.GetBulkDataOffsetInFile();
-	PayloadLength = InBulkData.GetBulkDataSize();
-	bIsDataStoredAsCompressed = InBulkData.IsStoredCompressedOnDisk();
 
-	Key = InGuid;
+	// Mark that we are actually referencing a payload stored in an old bulkdata
+	// format.
+	EnumAddFlags(Flags,EFlags::ReferencesLegacyFile);
+
+	if (InBulkData.IsStoredCompressedOnDisk())
+	{
+		EnumAddFlags(Flags, EFlags::LegacyFileIsCompressed);
+	}
+	else 
+	{
+		EnumAddFlags(Flags, EFlags::DisablePayloadCompression);
+	}
 }
 
 void FVirtualizedUntypedBulkData::Serialize(FArchive& Ar, UObject* Owner)
@@ -47,58 +170,106 @@ void FVirtualizedUntypedBulkData::Serialize(FArchive& Ar, UObject* Owner)
 
 		if (bNeedsTransaction)
 		{
-			Ar << Key;
-			Ar << PayloadLength;
-			Ar << bIsVirtualized;
-			Ar << bShouldCompressData;
-			Ar << bIsDataStoredAsCompressed;
+			Ar << Flags;
+			Ar << BulkDataId;
+			Ar << PayloadContentId;
+			Ar << PayloadSize;		
+			Ar << bWasKeyGuidDerived;
 			Ar << PackagePath;
 			Ar << PackageSegment;
 			Ar << OffsetInFile;
 
-			bool bKeepBuffer = !Ar.IsLoading();
-			if (!bIsVirtualized && PackagePath.IsEmpty())
+			// TODO: We could consider compressing the payload so it takes up less space in the 
+			// undo stack or even consider storing as a tmp file on disk rather than keeping it
+			// in memory or some other caching system.
+			// Serializing full 8k texture payloads to memory on each metadata change will empty
+			// the undo stack very quickly.
+			
+			// Note that we will only serialize the payload if it is in memory. Otherwise we can
+			// continue to load the payload as needed from disk or pull from the virtualization system
+			bool bPayloadInArchive = Ar.IsSaving() ? !Payload.IsNull() : false;
+			Ar << bPayloadInArchive;
+
+			if (Ar.IsSaving())
 			{
-				// Data is held in memory so we need to serialize it to the archive
-				bKeepBuffer = SerializeData(Ar, Payload);
+				if (bPayloadInArchive)
+				{
+					FCompressedBuffer CompressedPayload = FCompressedBuffer::Compress(NAME_None, Payload);
+					SerializeData(Ar, CompressedPayload, Flags);
+				}
 			}
-			// TODO:	Do we want to force data to be pulled/loaded and serialized to the archive?
-			//			Assuming the virtualized back end does not discard the payload or the file with 
-			//			the payload remains intact then it should not be needed?
-			
-			if (!bKeepBuffer)
+			else
 			{
-				Payload.Reset();
+				FCompressedBuffer CompressedPayload;
+				if (bPayloadInArchive)
+				{
+					SerializeData(Ar, CompressedPayload, Flags);
+				}
+				
+				Payload = CompressedPayload.Decompress();	
 			}
 		}
 	}
 	else if (Ar.IsPersistent() && !Ar.IsObjectReferenceCollector() && !Ar.ShouldSkipBulkData())
 	{
-		Ar << Key;
-		Ar << PayloadLength;
+		// Check if we need to update the payload id before it is serialized
+		if (Ar.IsSaving())
+		{
+			// TODO: Maybe only do this if we are saving with a FLinkerSave (ie package save to disk)
+			UpdateKeyIfNeeded();
+			checkf(bWasKeyGuidDerived == false, TEXT("bWasKeyGuidDerived should be cleared before saving")); // Sanity check
+		}
+
+		// Store the position in the archive of the flags incase that we need to update it later
+		const int64 SavedFlagsPos = Ar.Tell();
+		Ar << Flags;
+
+		// TODO: Can probably remove these checks before UE5 release
+		check(!Ar.IsSaving() || GetPayloadSize() == 0 || BulkDataId.IsValid()); // Sanity check to stop us saving out bad data
+		check(!Ar.IsSaving() || GetPayloadSize() == 0 || PayloadContentId.IsValid()); // Sanity check to stop us saving out bad data
+		
+		Ar << BulkDataId;
+		Ar << PayloadContentId;
+		Ar << PayloadSize;
+
+		// TODO: Can probably remove these checks before UE5 release
+		check(!Ar.IsLoading() || GetPayloadSize() == 0 || BulkDataId.IsValid()); // Sanity check to stop us loading in bad data
+		check(!Ar.IsLoading() || GetPayloadSize() == 0 || PayloadContentId.IsValid()); // Sanity check to stop us loading in bad data
 
 		if (Ar.IsSaving())
 		{
+			checkf(Ar.IsCooking() == false, TEXT("FVirtualizedUntypedBulkData::Serialize should not be called during a cook"));
+
 			FLinkerSave* LinkerSave = Cast<FLinkerSave>(Ar.GetLinker());
 
 			const bool bCanAttemptVirtualization = LinkerSave != nullptr;
 
 			if (bCanAttemptVirtualization)
 			{
-				PushData();
+				PushData(); // Note this can change various members if we are going from non-virtualized to virtualized
 			}
 
-			Ar << bIsVirtualized;
+			EFlags UpdatedFlags = BuildFlagsForSerialization(Ar);
 
-			if (!bIsVirtualized)
+			// Go back in the archive and update the flags in the archive, we will only apply the updated flags to the current
+			// object later if we detect that the package saved successfully.
+			// TODO: Not a huge fan of this, might be better to find a way to build the flags during serialization and potential callbacks 
+			// later then go back and update the flags in the Ar. Applying the updated flags only if we are saving a package to disk
+			// and the save succeeds continues to make sense.
+			const int64 RestorePos = Ar.Tell();
+			Ar.Seek(SavedFlagsPos);
+			Ar << UpdatedFlags;
+			Ar.Seek(RestorePos);
+
+			if (!IsDataVirtualized())
 			{
 				// Need to load the payload so that we can write it out
-				FSharedBuffer PayloadToSerialize = GetDataInternal();
-
-				// Need to update the flag before it is written out
-				bIsDataStoredAsCompressed = !bIsVirtualized && bShouldCompressData;
+				FCompressedBuffer PayloadToSerialize = GetDataInternal();
+				RecompressForSerialization(PayloadToSerialize, UpdatedFlags);
 
-				Ar << bIsDataStoredAsCompressed;
+				// If we are expecting a valid payload but fail to find one something critical has broken so assert now
+				// to prevent potentially bad data being saved to disk.
+				checkf(PayloadToSerialize || GetPayloadSize() == 0, TEXT("Failed to acquire the payload for saving!"));
 
 				// Write out a dummy value that we will write over once the payload is serialized
 				int64 PlaceholderValue = INDEX_NONE;
@@ -108,11 +279,13 @@ void FVirtualizedUntypedBulkData::Serialize(FArchive& Ar, UObject* Owner)
 
 				// The lambda is mutable so that PayloadToSerialize is not const (due to FArchive api not
 				// accepting const values)
-				auto SerializePayload = [this, OffsetPos, PayloadToSerialize](FArchive& Ar) mutable
+				auto SerializePayload = [this, OffsetPos, PayloadToSerialize, UpdatedFlags](FArchive& Ar) mutable
 				{
+					checkf(Ar.IsCooking() == false, TEXT("FVirtualizedUntypedBulkData::Serialize should not be called during a cook"));
+
 					int64 PayloadOffset = Ar.Tell();
 
-					SerializeData(Ar, PayloadToSerialize);
+					SerializeData(Ar, PayloadToSerialize, UpdatedFlags);
 
 					// Record the current archive offset (probably EOF but we cannot be sure)
 					const int64 ReturnPos = Ar.Tell();
@@ -120,9 +293,49 @@ void FVirtualizedUntypedBulkData::Serialize(FArchive& Ar, UObject* Owner)
 					// Update the offset/size entries that we set up during ::Serialize
 					Ar.Seek(OffsetPos);
 					Ar << PayloadOffset;
-
 					// Restore the archive's offset
 					Ar.Seek(ReturnPos);
+
+					// If we are saving the package to disk (we have access to FLinkerSave and it's filepath is valid) 
+					// then we should register a callback to be received once the package has actually been saved to 
+					// disk so that we can update the object's members to be redirected to the saved file.
+					FLinkerSave* LinkerSave = Cast<FLinkerSave>(Ar.GetLinker());
+					if (LinkerSave != nullptr && !LinkerSave->GetFilename().IsEmpty())
+					{	
+						// At some point saving to the sidecar file will be mutually exclusive with saving to the asset file, at that point
+						// we can split these code paths entirely for clarity. (might need to update ::BuildFlagsForSerialization at that point too!)
+						if (ShouldSaveToPackageSidecar())
+						{
+							FLinkerSave::FSidecarStorageInfo& SidecarData = LinkerSave->SidecarDataToAppend.AddZeroed_GetRef();
+							SidecarData.Identifier = PayloadContentId;
+							SidecarData.Payload = PayloadToSerialize;
+						}
+
+						auto OnSavePackage = [this, PayloadOffset, UpdatedFlags](const FPackagePath& InPackagePath)
+						{
+							this->PackagePath = InPackagePath;
+
+							if (!this->PackagePath.IsEmpty())
+							{
+								this->OffsetInFile = PayloadOffset;
+								this->Flags = UpdatedFlags;
+							}
+							else
+							{
+								// If for what ever reason we no longer have a valid package path then reset these values
+								// so that we get better error messages if they are used in the future.
+								this->OffsetInFile = INDEX_NONE;
+								EnumRemoveFlags(this->Flags, EFlags::ReferencesLegacyFile | EFlags::LegacyFileIsCompressed);
+							}
+
+							if (CanUnloadData())
+							{
+								this->Payload.Reset();
+							}
+						};
+
+						LinkerSave->PostSaveCallbacks.Add(MoveTemp(OnSavePackage));
+					}
 				};
 
 				// If we have a valid linker then we will defer serialization of the payload so that it will
@@ -137,15 +350,19 @@ void FVirtualizedUntypedBulkData::Serialize(FArchive& Ar, UObject* Owner)
 					SerializePayload(Ar);
 				}
 			}
+
+			if (CanUnloadData())
+			{
+				Payload.Reset();
+			}
 		}
 		else if (Ar.IsLoading())
 		{
-			Ar << bIsVirtualized;
-			
-			if (bIsVirtualized)
+			checkf(!IsReferencingOldBulkData(), TEXT("Virtualized bulkdata was saved with the ReferencesLegacyFile flag!"));
+
+			if (IsDataVirtualized())
 			{
 				// We aren't going to use these members so reset them
-				bIsDataStoredAsCompressed = false;
 				OffsetInFile = INDEX_NONE;
 
 				PackagePath.Empty();
@@ -153,8 +370,6 @@ void FVirtualizedUntypedBulkData::Serialize(FArchive& Ar, UObject* Owner)
 			}
 			else
 			{
-				Ar << bIsDataStoredAsCompressed;
-
 				// If we can lazy load then find the PackagePath, otherwise we will want
 				// to serialize immediately.
 				if (Ar.IsAllowingLazyLoading())
@@ -174,110 +389,228 @@ void FVirtualizedUntypedBulkData::Serialize(FArchive& Ar, UObject* Owner)
 				{
 					// If we have no packagepath then we need to load the data immediately
 					// as we will not be able to load it on demand.
-					SerializeData(Ar, Payload);
+					FCompressedBuffer CompressedPayload;
+					SerializeData(Ar, CompressedPayload, Flags);
+					Payload = CompressedPayload.Decompress();
 				}
 			}
 		}
 	}
 }
 
-void FVirtualizedUntypedBulkData::CalculateKey()
+FCompressedBuffer FVirtualizedUntypedBulkData::LoadFromDisk() const
 {
-	if (Payload)
+	TRACE_CPUPROFILER_EVENT_SCOPE(FVirtualizedUntypedBulkData::LoadFromDisk);
+
+	if (PackagePath.IsEmpty())
+	{
+		UE_LOG(LogVirtualization, Error, TEXT("Cannot load a payload with an empty filename!"));
+		return FCompressedBuffer();
+	}
+
+	if (HasPayloadSidecarFile() && CVarShouldLoadFromSidecar.GetValueOnAnyThread())
 	{
-		uint32 Hash[5] = {};
-		FSHA1::HashBuffer(Payload.GetData(), GetBulkDataSize(), (uint8*)Hash);
+		// Note that this code path is purely for debugging and not expected to be enabled by default
+		if (CVarShouldValidatePayload.GetValueOnAnyThread())
+		{
+			UE_LOG(LogVirtualization, Verbose, TEXT("Validating payload loaded from sidecar file: '%s'"), *PackagePath.GetLocalFullPath(EPackageSegment::PayloadSidecar));
+
+			// Load both payloads then generate a FPayloadId from them, since this identifier is a hash of the buffers content
+			// we only need to verify them against PayloadContentId to be sure that the data is correct.
+			FCompressedBuffer SidecarBuffer = LoadFromSidecarFile();
+			FCompressedBuffer AssetBuffer = LoadFromPackageFile();
 
-		Key = FGuid(Hash[0] ^ Hash[4], Hash[1], Hash[2], Hash[3]);
+			FPayloadId SidecarId(SidecarBuffer.Decompress());
+			FPayloadId AssetId(SidecarBuffer.Decompress());
+
+			UE_CLOG(SidecarId != PayloadContentId, LogVirtualization, Error, TEXT("Sidecar content did not hash correctly! Found '%s' Expected '%s'"), *SidecarId.ToString(), *PayloadContentId.ToString());
+			UE_CLOG(AssetId != PayloadContentId, LogVirtualization, Error, TEXT("Asset content did not hash correctly! Found '%s' Expected '%s'"), *SidecarId.ToString(), *PayloadContentId.ToString())
+
+			return SidecarBuffer;
+		}
+		else
+		{
+			return LoadFromSidecarFile();
+		}
 	}
 	else
 	{
-		Key.Invalidate();
+		return LoadFromPackageFile();	 
 	}
 }
 
-FSharedBuffer FVirtualizedUntypedBulkData::LoadFromDisk() const
+FCompressedBuffer FVirtualizedUntypedBulkData::LoadFromPackageFile() const
 {
-	TRACE_CPUPROFILER_EVENT_SCOPE(FVirtualizedUntypedBulkData::LoadFromDisk);
+	TRACE_CPUPROFILER_EVENT_SCOPE(FVirtualizedUntypedBulkData::LoadFromPackageFile);
 
-	FSharedBuffer PayloadFromDisk;
+	UE_LOG(LogVirtualization, Verbose, TEXT("Attempting to load payload from the package file: '%s'"), *PackagePath.GetLocalFullPath(PackageSegment));
+
+	FCompressedBuffer PayloadFromDisk;
 
 	// Open a reader to the file
 	FOpenPackageResult Result = IPackageResourceManager::Get().OpenReadPackage(PackagePath, PackageSegment);
 	if (Result.Archive.IsValid() && Result.Format == EPackageFormat::Binary)
 	{
+		checkf(OffsetInFile != INDEX_NONE, TEXT("Attempting to load '%s' from disk with an invalid OffsetInFile!"), *PackagePath.GetDebugNameWithExtension(PackageSegment));
 		// Move the correct location of the data in the file
 		Result.Archive->Seek(OffsetInFile);
 
 		// Now we can actually serialize it
-		SerializeData(*Result.Archive, PayloadFromDisk);
+		SerializeData(*Result.Archive, PayloadFromDisk, Flags);
 	}
 	else
 	{
-		UE_LOG(LogVirtualization, Error, TEXT("Could not open (%s) to read FVirtualizedUntypedBulkData: %s."), *PackagePath.GetDebugNameWithExtension(PackageSegment),
-			(!Result.Archive.IsValid() ? TEXT("could not find package") : TEXT("package is a TextAsset which is not supported")));
+		LogPackageOpenFailureMessage(PackagePath, PackageSegment);
 	}
 
 	return PayloadFromDisk;
 }
 
-bool FVirtualizedUntypedBulkData::SerializeData(FArchive& Ar, FSharedBuffer& InPayload) const
+FCompressedBuffer FVirtualizedUntypedBulkData::LoadFromSidecarFileInternal(ErrorVerbosity Verbosity) const
 {
-	TRACE_CPUPROFILER_EVENT_SCOPE(FVirtualizedUntypedBulkData::SerializeData);
-
-	const int64 PayloadSize = Ar.IsSaving() ? InPayload.GetSize() : GetBulkDataSize();
+	TRACE_CPUPROFILER_EVENT_SCOPE(FVirtualizedUntypedBulkData::LoadFromSidecarFileInternal);
 
-	if (PayloadSize > 0)
+	FOpenPackageResult Result = IPackageResourceManager::Get().OpenReadPackage(PackagePath, EPackageSegment::PayloadSidecar);
+	if (Result.Archive.IsValid() && Result.Format == EPackageFormat::Binary)
 	{
-		FUniqueBuffer LoadPayload;
+		uint32 Version = INDEX_NONE;
+		*Result.Archive << Version;
 
-		if (Ar.IsLoading())
+		if (Version != FTocEntry::PayloadSidecarFileVersion)
 		{
-			LoadPayload = FUniqueBuffer::Alloc(PayloadSize);
+			UE_CLOG(Verbosity > ErrorVerbosity::None, LogVirtualization, Error, TEXT("Unknown version (%u) found in '%s'"), Version, *PackagePath.GetLocalFullPath(EPackageSegment::PayloadSidecar));
+			return FCompressedBuffer();
 		}
 
-		// Our serialization methods require a non-const pointer even when we are only reading from it
-		// as the methods themselves are bi-directional. This means if we are saving we need to cast 
-		// away the const qualifier.
-		void* DataPtr = Ar.IsLoading() ? LoadPayload.GetData() : const_cast<void*>(InPayload.GetData());
+		// First we load the table of contents so we can find the payload in the file
+		TArray<FTocEntry> TableOfContents;
+		*Result.Archive << TableOfContents;
 
-		checkf(DataPtr != nullptr, TEXT("Attempting to serialize to/from an invalid shared buffer!"));
+		const FTocEntry* Entry = TableOfContents.FindByPredicate([&PayloadContentId = PayloadContentId](const FTocEntry& Entry)
+			{
+				return Entry.Identifier == PayloadContentId;
+			});
 
-		if (bIsDataStoredAsCompressed)
+		if (Entry != nullptr)
 		{
-			Ar.SerializeCompressed(DataPtr, PayloadSize, NAME_Zlib, COMPRESS_NoFlags, false);
+			if (Entry->OffsetInFile != INDEX_NONE)
+			{
+				// Move the correct location of the data in the file
+				Result.Archive->Seek(Entry->OffsetInFile);
+
+				// Now we can actually serialize it
+				FCompressedBuffer PayloadFromDisk;
+				SerializeData(*Result.Archive, PayloadFromDisk, EFlags::None);
+
+				return PayloadFromDisk;
+			}
+			else if(Verbosity > ErrorVerbosity::None)
+			{
+				UE_LOG(LogVirtualization, Error, TEXT("Payload '%s' in '%s' has an invalid OffsetInFile!"), *PayloadContentId.ToString(), *PackagePath.GetLocalFullPath(EPackageSegment::PayloadSidecar));
+			}
 		}
-		else
+		else if(Verbosity > ErrorVerbosity::None)
 		{
-			Ar.Serialize(DataPtr, PayloadSize);
+			UE_LOG(LogVirtualization, Error, TEXT("Unable to find payload '%s' in '%s'"), *PayloadContentId.ToString(), *PackagePath.GetLocalFullPath(EPackageSegment::PayloadSidecar));
 		}
+	}
+	else if(Verbosity > ErrorVerbosity::None)
+	{
+		LogPackageOpenFailureMessage(PackagePath, EPackageSegment::PayloadSidecar);
+	}
 
-		if (Ar.IsLoading())
+	return FCompressedBuffer();
+}
+
+FCompressedBuffer FVirtualizedUntypedBulkData::LoadFromSidecarFile() const
+{
+	TRACE_CPUPROFILER_EVENT_SCOPE(FVirtualizedUntypedBulkData::LoadFromSidecarFile);
+
+	UE_LOG(LogVirtualization, Verbose, TEXT("Attempting to load payload from the sidecar file: '%s'"), 
+		*PackagePath.GetLocalFullPath(EPackageSegment::PayloadSidecar));
+
+	if (CVarShouldAllowSidecarSyncing.GetValueOnAnyThread())
+	{
+		FCompressedBuffer PayloadFromDisk = LoadFromSidecarFileInternal(ErrorVerbosity::None);
+		if (PayloadFromDisk.IsNull())
 		{
-			InPayload = LoadPayload.MoveToShared();
+			UE_LOG(LogVirtualization, Verbose, TEXT("Initial load from sidecar failed, attempting to sync the file: '%s'"), 
+				*PackagePath.GetLocalFullPath(EPackageSegment::PayloadSidecar));
+
+			if (IVirtualizationSourceControlUtilities* SourceControlInterface = GetSourceControlInterface())
+			{
+				// SyncPayloadSidecarFile should log failure cases, so there is no need for us to add log messages here
+				if (SourceControlInterface->SyncPayloadSidecarFile(PackagePath))
+				{
+					PayloadFromDisk = LoadFromSidecarFileInternal(ErrorVerbosity::All);
+				}
+			}
+			else
+			{
+				UE_LOG(LogVirtualization, Error, TEXT("Failed to find IVirtualizationSourceControlUtilities, unable to try and sync: '%s'"), 
+					*PackagePath.GetLocalFullPath(EPackageSegment::PayloadSidecar));
+			}
 		}
 
-		return true;
+		return PayloadFromDisk;
 	}
 	else
 	{
-		// If we are loading and we failed to serialize anything then we should reset the SharedBuffer.
-		if (Ar.IsLoading())
+		return LoadFromSidecarFileInternal(ErrorVerbosity::All);
+	}
+}
+
+bool FVirtualizedUntypedBulkData::SerializeData(FArchive& Ar, FCompressedBuffer& InPayload, const EFlags PayloadFlags) const
+{
+	TRACE_CPUPROFILER_EVENT_SCOPE(FVirtualizedUntypedBulkData::SerializeData);
+
+	if (Ar.IsSaving()) // Saving to virtualized bulkdata format
+	{
+		for (const FSharedBuffer& Buffer : InPayload.GetCompressed().GetSegments())
+		{
+			// Const cast because FArchive requires a non-const pointer!
+			Ar.Serialize(const_cast<void*>(Buffer.GetData()), static_cast<int64>(Buffer.GetSize()));
+		}
+		
+		return true;
+	}
+	else if(Ar.IsLoading() && !EnumHasAnyFlags(PayloadFlags, EFlags::ReferencesLegacyFile)) // Loading from virtualized bulkdata format
+	{
+		InPayload = FCompressedBuffer::FromCompressed(Ar);
+		return InPayload.IsNull();	
+	}
+	else if (Ar.IsLoading()) // Loading from old bulkdata format
+	{
+		const int64 Size = GetBulkDataSize();
+		FUniqueBuffer LoadPayload = FUniqueBuffer::Alloc(Size);
+
+		if (EnumHasAnyFlags(PayloadFlags, EFlags::LegacyFileIsCompressed))
 		{
-			InPayload.Reset();
+			Ar.SerializeCompressed(LoadPayload.GetData(), Size, NAME_Zlib, COMPRESS_NoFlags, false);
 		}
+		else
+		{
+			Ar.Serialize(LoadPayload.GetData(), Size);
+		}
+
+		InPayload = FCompressedBuffer::Compress(NAME_None, LoadPayload.MoveToShared());
 
+		return true;
+	}
+	else
+	{
 		return false;
 	}
 }
 
 void FVirtualizedUntypedBulkData::PushData()
 {
-	checkf(bIsVirtualized == false || Payload.IsNull(), TEXT("Cannot have a valid payload in memory if the payload is virtualized!")); // Sanity check
+	checkf(IsDataVirtualized() == false || Payload.IsNull(), TEXT("Cannot have a valid payload in memory if the payload is virtualized!")); // Sanity check
 
-	// We only need to push if the payload is not currently virtualized (either we have an updated
-	// payload in memory or the payload is currently non-virtualized and stored on disk)
-	if (!bIsVirtualized)
+	// We only need to push if the payload if it actually has data and it is not 
+	// currently virtualized (either we have an updated payload in memory or the 
+	// payload is currently non-virtualized and stored on disk)
+	if (!IsDataVirtualized() && GetPayloadSize() > 0)
 	{ 
 		TRACE_CPUPROFILER_EVENT_SCOPE(FVirtualizedUntypedBulkData::PushData);
 
@@ -285,31 +618,42 @@ void FVirtualizedUntypedBulkData::PushData()
 		// a non-virtualized payload to a virtualized one. If the bulkdata is merely being
 		// edited then we should have the payload in memory already and are just accessing a
 		// reference to it.
-		FSharedBuffer PayloadToPush = GetDataInternal();
 
-		if (FVirtualizationManager::Get().PushData(PayloadToPush, Key))
+		FCompressedBuffer PayloadToPush = GetDataInternal();
+		// TODO: If the push fails we will end up potentially re-compressing this payload for
+		// serialization, we need a better way to save the results of 'RecompressForSerialization'
+		RecompressForSerialization(PayloadToPush, Flags);
+
+		if (FVirtualizationManager::Get().PushData(PayloadContentId, PayloadToPush))
 		{
-			bIsVirtualized = true;
+			EnumAddFlags(Flags, EFlags::IsVirtualized);
+			EnumRemoveFlags(Flags, EFlags::ReferencesLegacyFile | EFlags::LegacyFileIsCompressed);
+			
+			// Clear members associated with non-virtualized data and release the in-memory
+			// buffer.
+			PackagePath.Empty();
+			PackageSegment = EPackageSegment::Header;
+			OffsetInFile = INDEX_NONE;
 		}
 	}	
 }
 
-FSharedBuffer FVirtualizedUntypedBulkData::PullData() const
+FCompressedBuffer FVirtualizedUntypedBulkData::PullData() const
 {
 	TRACE_CPUPROFILER_EVENT_SCOPE(FVirtualizedUntypedBulkData::PullData);
 
-	FSharedBuffer PulledPayload = FVirtualizationManager::Get().PullData(Key);
+	FCompressedBuffer PulledPayload = FVirtualizationManager::Get().PullData(PayloadContentId);
 
 	if (PulledPayload)
 	{
-		checkf(	PayloadLength == PulledPayload.GetSize(),
+		checkf(	PayloadSize == PulledPayload.GetRawSize(),
 				TEXT("Mismatch between serialized length (%" INT64_FMT ") and virtualized data length (%" UINT64_FMT ")"),
-				PayloadLength,
-				PulledPayload.GetSize());
+				PayloadSize,
+				PulledPayload.GetRawSize());
 	}
 	else
 	{
-		UE_LOG(LogVirtualization, Warning, TEXT("Failed to pull virtual data with guid (%s)"), *Key.ToString());
+		UE_LOG(LogVirtualization, Warning, TEXT("Failed to pull virtual data with guid (%s)"), *PayloadContentId.ToString());
 	}
 
 	return PulledPayload;
@@ -337,74 +681,99 @@ FPackagePath FVirtualizedUntypedBulkData::GetPackagePathFromOwner(UObject* Owner
 bool FVirtualizedUntypedBulkData::CanUnloadData() const
 {
 	// We cannot unload the data if are unable to reload it from a file
-	return bIsVirtualized || PackagePath.IsEmpty() == false;
+	return IsDataVirtualized() || PackagePath.IsEmpty() == false;
 }
 
 void FVirtualizedUntypedBulkData::Reset()
 {
+	// Note that we do not reset the BulkDataId
+	PayloadContentId.Reset();
+	Payload.Reset();
+	PayloadSize = 0;
+	CompressionFormatToUse = NAME_Default;
+	OffsetInFile = INDEX_NONE;
 	PackagePath.Empty();
 	PackageSegment = EPackageSegment::Header;
-	OffsetInFile = INDEX_NONE;
-	Key.Invalidate();
-	Payload.Reset();
-	PayloadLength = 0;
-	bIsVirtualized = false;
-	bIsDataStoredAsCompressed = false;
+	bWasKeyGuidDerived = false;
+	Flags = EFlags::None;	
 }
 
 void FVirtualizedUntypedBulkData::UnloadData()
 {
-	Payload.Reset();
+	if (CanUnloadData())
+	{
+		Payload.Reset();
+	}
+}
+
+FGuid FVirtualizedUntypedBulkData::GetIdentifier() const
+{
+	if (GetPayloadSize() > 0)
+	{
+		checkf(BulkDataId.IsValid(), TEXT("If bulkdata has a valid payload then it should have a valid BulkDataId"));
+		return BulkDataId;
+	}
+	else
+	{
+		return FGuid();
+	}
 }
 
-FSharedBuffer FVirtualizedUntypedBulkData::GetDataInternal() const
+FCompressedBuffer FVirtualizedUntypedBulkData::GetDataInternal() const
 {
 	TRACE_CPUPROFILER_EVENT_SCOPE(FVirtualizedUntypedBulkData::GetDataInternal);
 
 	// Early out there isn't any data to actually load
-	if (GetBulkDataSize() == 0)
+	if (GetPayloadSize() == 0)
 	{
-		return FSharedBuffer();
+		return FCompressedBuffer();
 	}
 
 	// Check if we already have the data in memory
 	if (Payload)
 	{
-		return Payload;
+		return FCompressedBuffer::Compress(NAME_None, Payload);
 	}
 
-	if (bIsVirtualized)
+	if (IsDataVirtualized())
 	{
-		FSharedBuffer Buffer = PullData();
-		check(Payload.IsNull()); //Make sure that we did not assign the buffer internally
+		FCompressedBuffer Buffer = PullData();
+		checkf(Payload.IsNull(), TEXT("Pulling data somehow assigned it to the bulk data object!")); //Make sure that we did not assign the buffer internally
 		return Buffer;
 	}
 	else
 	{
-		FSharedBuffer Buffer = LoadFromDisk();
+		FCompressedBuffer Buffer = LoadFromDisk();
 		check(Payload.IsNull()); //Make sure that we did not assign the buffer internally
 		return Buffer;
 	}
 }
 
-TFuture<FSharedBuffer> FVirtualizedUntypedBulkData::GetData() const
+TFuture<FSharedBuffer> FVirtualizedUntypedBulkData::GetPayload() const
 {
 	TPromise<FSharedBuffer> Promise;
 
-	FSharedBuffer SharedBuffer = GetDataInternal();
+	FSharedBuffer UncompressedPayload = GetDataInternal().Decompress();
 
 	// TODO: Not actually async yet!
-	Promise.SetValue(MoveTemp(SharedBuffer));
+	Promise.SetValue(MoveTemp(UncompressedPayload));
 
 	return Promise.GetFuture();
 }
 
-void FVirtualizedUntypedBulkData::GetData(OnDataReadyCallback&& Callback) const
+TFuture<FCompressedBuffer>FVirtualizedUntypedBulkData::GetCompressedPayload() const
 {
-	Callback(GetDataInternal());
+	TPromise<FCompressedBuffer> Promise;
+
+	FCompressedBuffer CompressedPayload = GetDataInternal();
+
+	// TODO: Not actually async yet!
+	Promise.SetValue(MoveTemp(CompressedPayload));
+
+	return Promise.GetFuture();
 }
 
-void FVirtualizedUntypedBulkData::UpdatePayload(const FSharedBuffer& InPayload)
+void FVirtualizedUntypedBulkData::UpdatePayload(FSharedBuffer InPayload, FName InCompressionFormat)
 {
 	TRACE_CPUPROFILER_EVENT_SCOPE(FVirtualizedUntypedBulkData::UpdatePayload);
 
@@ -412,16 +781,122 @@ void FVirtualizedUntypedBulkData::UpdatePayload(const FSharedBuffer& InPayload)
 
 	// Make sure that we own the memory in the shared buffer
 	Payload = InPayload.MakeOwned();
-	PayloadLength = (int64)Payload.GetSize();
+	PayloadSize = (int64)Payload.GetSize();
+	
+	EnumRemoveFlags(Flags,	EFlags::IsVirtualized | 
+							EFlags::DisablePayloadCompression | 
+							EFlags::ReferencesLegacyFile | 
+							EFlags::LegacyFileIsCompressed);
+
+	//TODO: Should we validate now or let FCompressedBuffer do that later?
+	CompressionFormatToUse = InCompressionFormat;
 
-	bIsVirtualized = false;
-	bIsDataStoredAsCompressed = false;
+	if (InCompressionFormat == NAME_None)
+	{
+		EnumAddFlags(Flags, EFlags::DisablePayloadCompression);
+	}
 
 	PackagePath.Empty();
 	PackageSegment = EPackageSegment::Header;
 	OffsetInFile = INDEX_NONE;
 
-	CalculateKey();
+	if (!BulkDataId.IsValid() && PayloadSize > 0)
+	{
+		BulkDataId = FGuid::NewGuid();
+	}
+
+	PayloadContentId = FPayloadId(Payload);
 }
 
+FCustomVersionContainer FVirtualizedUntypedBulkData::GetCustomVersions(FArchive& InlineArchive)
+{
+	return InlineArchive.GetCustomVersions();
+}
+
+void FVirtualizedUntypedBulkData::UpdateKeyIfNeeded()
+{
+	// If this was created from old BulkData then the key is generated from an older FGuid, we
+	// should recalculate it based off the payload to keep the key consistent in the future.
+	if (bWasKeyGuidDerived)
+	{
+		checkf(IsDataVirtualized() == false, TEXT("Cannot have a virtualized payload if bWasCreatedFromOldBulkData is still true")); // Sanity check
+
+		// Load the payload from disk (or memory) so that we can hash it
+		FSharedBuffer InPayload = GetDataInternal().Decompress();
+		PayloadContentId = FPayloadId(InPayload);
+
+		// Store as the in memory payload, since this method is only called during saving 
+		// we know it will get cleared anyway.
+		Payload = InPayload;
+
+		bWasKeyGuidDerived = false;
+	}
+}
+
+void FVirtualizedUntypedBulkData::RecompressForSerialization(FCompressedBuffer& InOutPayload, EFlags PayloadFlags) const
+{
+	const FName CurrentMethod = InOutPayload.GetFormatName();
+
+	FName TargetMethod;
+	
+	if (EnumHasAnyFlags(PayloadFlags, EFlags::DisablePayloadCompression))
+	{
+		TargetMethod = NAME_None;
+	}
+	else if (CompressionFormatToUse != NAME_Default)
+	{
+		check(!CompressionFormatToUse.IsNone()); // Should be caught by the DisablePayloadCompression flag
+		TargetMethod = CompressionFormatToUse;
+	}
+	else
+	{
+		// TODO: Do we want to add more logic, min size etc?
+		TargetMethod = NAME_LZ4;
+	}
+	
+	if (TargetMethod != CurrentMethod) // No change in compression so we can just return it
+	{
+		FCompositeBuffer DecompressedBuffer = InOutPayload.DecompressToComposite();
+
+		// If the buffer actually decompressed we can have both the compressed and the uncompressed version of the
+		// payload in memory. Compressing it will create a third version so before doing that we should reset
+		// the original compressed buffer in case that we can release it to reduce higher water mark pressure.
+		InOutPayload.Reset();
+		InOutPayload = FCompressedBuffer::Compress(TargetMethod, DecompressedBuffer);
+	}
+}
+
+FVirtualizedUntypedBulkData::EFlags FVirtualizedUntypedBulkData::BuildFlagsForSerialization(FArchive& Ar) const
+{
+	if (Ar.IsSaving())
+	{
+		EFlags UpdatedFlags = Flags;
+
+		// Now update any changes to the flags that we might need to make when serializing.
+		// Note that these changes are not applied to the current object UNLESS we are saving
+		// the package, in which case the newly modified flags will be applied once we confirm
+		// that the package has saved.
+		
+		const FLinkerSave* LinkerSave = Cast<FLinkerSave>(Ar.GetLinker());
+		if (LinkerSave != nullptr && !LinkerSave->GetFilename().IsEmpty() && ShouldSaveToPackageSidecar())
+		{
+			EnumAddFlags(UpdatedFlags, EFlags::HasPayloadSidecarFile);
+		}
+		else
+		{
+			EnumRemoveFlags(UpdatedFlags, EFlags::HasPayloadSidecarFile);
+		}
+
+		EnumRemoveFlags(UpdatedFlags,EFlags::ReferencesLegacyFile | EFlags::LegacyFileIsCompressed);
+
+		return UpdatedFlags;
+	}
+	else
+	{
+		return Flags;
+	}
+}
+
+} // namespace UE::Virtualization
+
 //#endif //WITH_EDITORONLY_DATA
diff --git a/Engine/Source/Runtime/CoreUObject/Public/Misc/PackagePath.h b/Engine/Source/Runtime/CoreUObject/Public/Misc/PackagePath.h
--- a/Engine/Source/Runtime/CoreUObject/Public/Misc/PackagePath.h
+++ b/Engine/Source/Runtime/CoreUObject/Public/Misc/PackagePath.h
@@ -59,8 +59,9 @@ enum class EPackageExtension : uint8
 	BulkDataDefault,
 	BulkDataOptional,
 	BulkDataMemoryMapped,
+	PayloadSidecar,
 };
-constexpr int EPackageExtensionCount = static_cast<int>(EPackageExtension::BulkDataMemoryMapped) + 1;
+constexpr int EPackageExtensionCount = static_cast<int>(EPackageExtension::PayloadSidecar) + 1;
 COREUOBJECT_API const TCHAR* LexToString(EPackageExtension PackageExtension);
 
 
diff --git a/Engine/Source/Runtime/CoreUObject/Public/Misc/PackageSegment.h b/Engine/Source/Runtime/CoreUObject/Public/Misc/PackageSegment.h
--- a/Engine/Source/Runtime/CoreUObject/Public/Misc/PackageSegment.h
+++ b/Engine/Source/Runtime/CoreUObject/Public/Misc/PackageSegment.h
@@ -21,6 +21,8 @@ enum class EPackageSegment : uint8
 	BulkDataOptional,
 	/** BulkData that is stored in a separate file aligned for memory mapping */
 	BulkDataMemoryMapped,
+	/** BulkData payloads that are stored in a separate file along side the Header in the workspace domain*/
+	PayloadSidecar,
 };
 const TCHAR* LexToString(EPackageSegment PackageSegment);
 
diff --git a/Engine/Source/Runtime/CoreUObject/Public/UObject/LinkerSave.h b/Engine/Source/Runtime/CoreUObject/Public/UObject/LinkerSave.h
--- a/Engine/Source/Runtime/CoreUObject/Public/UObject/LinkerSave.h
+++ b/Engine/Source/Runtime/CoreUObject/Public/UObject/LinkerSave.h
@@ -3,12 +3,14 @@
 #pragma once
 
 #include "CoreMinimal.h"
+#include "Compression/CompressedBuffer.h"
 #include "Serialization/ArchiveUObject.h"
 #include "Serialization/FileRegions.h"
-#include "UObject/ObjectResource.h"
+#include "Templates/RefCounting.h"
 #include "UObject/Linker.h"
+#include "UObject/ObjectResource.h"
 #include "UObject/UObjectThreadContext.h"
-#include "Templates/RefCounting.h"
+#include "Virtualization/PayloadId.h"
 
 struct FUntypedBulkData;
 
@@ -91,6 +93,23 @@ public:
 	 */
 	bool bUpdatingLoadedPath = false;
 
+	/** Used by FVirtualizedUntypedBulkData to add payloads to be added to the payload sidecar file */
+	struct FSidecarStorageInfo
+	{
+		UE::Virtualization::FPayloadId Identifier;
+		FCompressedBuffer Payload;
+	};
+	TArray<FSidecarStorageInfo> SidecarDataToAppend;
+
+	/** 
+	 * Array of callbacks that will be invoked when the package has successfully saved to disk.
+	 * The callbacks will not be invoked if the package fails to save for some reason.
+	 * Unlike subscribing to the UPackage::PackageSavedEvent this callback allows custom data
+	 * via lambda capture. 
+	 * @param PackagePath The path of the package
+	 */
+	TArray<TUniqueFunction<void(const FPackagePath& PackagePath)>> PostSaveCallbacks;
+
 	/** A mapping of package name to generated script SHA keys */
 	COREUOBJECT_API static TMap<FString, TArray<uint8> > PackagesToScriptSHAMap;
 
@@ -141,6 +160,9 @@ public:
 	using FLinker::Serialize;
 	void Serialize( void* V, int64 Length );
 
+	/** Invoke all of the callbacks in PostSaveCallbacks and then empty it. */
+	void OnPostSave(const FPackagePath& PackagePath);
+
 	// FLinker interface
 	virtual FString GetDebugName() const override;
 
diff --git a/Engine/Source/Runtime/CoreUObject/Public/Virtualization/IVirtualizationBackend.h b/Engine/Source/Runtime/CoreUObject/Public/Virtualization/IVirtualizationBackend.h
new file mode 100644
--- /dev/null
+++ b/Engine/Source/Runtime/CoreUObject/Public/Virtualization/IVirtualizationBackend.h
@@ -0,0 +1,145 @@
+// Copyright Epic Games, Inc. All Rights Reserved.
+
+#pragma once
+
+#include "Compression/CompressedBuffer.h"
+#include "Containers/StringView.h"
+#include "Features/IModularFeature.h"
+#include "Features/IModularFeatures.h"
+
+namespace UE::Virtualization
+{
+
+class FPayloadId;
+
+/**
+ * The interface to derive from to create a new backend implementation.
+ * 
+ * Note that virtualization backends are instantiated FVirtualizationManager via  
+ * IVirtualizationBackendFactory so each new backend derived from IVirtualizationBackend  
+ * will also need a factory derived from IVirtualizationBackendFactory. You can either do
+ * this manually or use the helper macro 'UE_REGISTER_VIRTUALIZATION_BACKEND_FACTORY' to 
+ * generate the code for you.
+ *
+ */
+class IVirtualizationBackend
+{
+protected:
+	/** Enum detailing which operations a backend can support */
+	enum class EOperations : uint8
+	{
+		/** Supports only push operations */
+		Push,
+		/** Supports only pull operations */
+		Pull,
+		/** Supports both push and pull operations */
+		Both
+	};
+
+	IVirtualizationBackend(EOperations InSupportedOperations) 
+		: SupportedOperations(InSupportedOperations)
+	{
+
+	}
+
+public:
+	virtual ~IVirtualizationBackend() = default;
+
+	/**
+	 * This will be called during the setup of the backend hierarchy. The entry config file
+	 * entry that caused the backend to be created will be passed to the method so that any
+	 * additional settings may be parsed from it.
+	 * Take care to clearly log any error that occurs so that the end user has a clear way 
+	 * to fix them.
+	 * 
+	 * @param ConfigEntry	The entry for the backend from the config ini file that may 
+	 *						contain additional settings.
+	 * @return				Returning false indicates that initialization failed in a way 
+	 *						that the backend will not be able to function correctly.
+	 */
+	virtual bool Initialize(const FString& ConfigEntry) = 0;
+
+	/**
+	 * The backend will attempt to store the given payload by what ever method the backend uses.
+	 * NOTE: It is assumed that the virtualization manager will run all appropriate validation
+	 * on the payload and it's id and that the inputs to PushData can be trusted.
+	 * 
+	 * @param Id		The Id of the payload
+	 * @param Payload	A potentially compressed buffer representing the payload
+	 * @return			The result of the push operation
+	 */
+	virtual bool PushData(const FPayloadId& Id, const FCompressedBuffer& Payload) = 0;
+
+	/** 
+	 * The backend will attempt to retrieve the given payload by what ever method the backend uses.
+	 * NOTE: It is assumed that the virtualization manager will validate the returned payload to 
+	 * make sure that it matches the requested id so there is no need for each backend to do this/
+	 * 
+	 * @param Id	The Id of a payload to try and pull from the backend.
+	 * 
+	 * @return		A valid FCompressedBuffer containing the payload if the pull
+	 *				operation succeeded and a null FCompressedBuffer
+	 *				if it did not.
+	 */
+	virtual FCompressedBuffer PullData(const FPayloadId& Id) = 0;
+
+	/** Return true if the backend supports push operations. Returning true allows ::PushData to be called.  */
+	bool SupportsPushOperations() const { return SupportedOperations == EOperations::Push || SupportedOperations == EOperations::Both;  }
+
+	/** Return true if the backend supports pull operations. Returning true allows ::PullData to be called.  */
+	bool SupportsPullOperations() const { return SupportedOperations == EOperations::Pull || SupportedOperations == EOperations::Both;  }
+
+	/** Returns a string that can be used to identify the backend for debugging and logging purposes */
+	virtual FString GetDebugString() const = 0;
+
+private:
+
+	/** The operations that this backend supports */
+	EOperations SupportedOperations;
+};
+
+/** 
+ * Derive from this interface to implement a factory to return a backend type.
+ * An instance of the factory should be created and then registered to 
+ * IModularFeatures with the feature name "VirtualizationBackendFactory" to
+ * give 'FVirtualizationManager' access to it. 
+ * The macro 'UE_REGISTER_VIRTUALIZATION_BACKEND_FACTORY' can be used to create 
+ * a factory easily if you do not want to specialize the behaviour.
+*/
+class IVirtualizationBackendFactory : public IModularFeature
+{
+public:
+	/** 
+	 * Creates a new backend instance.
+	 * 
+	 * @param ConfigName	The name given to the back end in the config ini file
+	 * @return A new backend instance
+	 */
+	virtual IVirtualizationBackend* CreateInstance(FStringView ConfigName) = 0;
+
+	/** Returns the name used to identify the type in config ini files */
+	virtual FName GetName() = 0;
+};
+
+/**
+ * This macro is used to generate a backend factories boilerplate code if you do not
+ * need anything more than the default behaviour.
+ * As well as creating the class, a single instance will be created which will register the factory with
+ * 'IModularFeatures' so that it is ready for use.
+ * 
+ * @param BackendClass The name of the class derived from 'IVirtualizationBackend' that the factory should create
+ * @param The name used in config ini files to reference this backend type.
+ */
+#define UE_REGISTER_VIRTUALIZATION_BACKEND_FACTORY(BackendClass, ConfigName) \
+	class BackendClass##Factory : public IVirtualizationBackendFactory \
+	{ \
+	public: \
+		BackendClass##Factory() { IModularFeatures::Get().RegisterModularFeature(FName("VirtualizationBackendFactory"), this); }\
+		virtual ~BackendClass##Factory() { IModularFeatures::Get().UnregisterModularFeature(FName("VirtualizationBackendFactory"), this); } \
+	private: \
+		virtual IVirtualizationBackend* CreateInstance(FStringView ConfigName) override { return new BackendClass(ConfigName); } \
+		virtual FName GetName() override { return FName(#ConfigName); } \
+	}; \
+	static BackendClass##Factory BackendClass##Factory##Instance;
+
+} // namespace UE::Virtualization
diff --git a/Engine/Source/Runtime/CoreUObject/Public/Virtualization/IVirtualizationSourceControlUtilities.h b/Engine/Source/Runtime/CoreUObject/Public/Virtualization/IVirtualizationSourceControlUtilities.h
new file mode 100644
--- /dev/null
+++ b/Engine/Source/Runtime/CoreUObject/Public/Virtualization/IVirtualizationSourceControlUtilities.h
@@ -0,0 +1,37 @@
+// Copyright Epic Games, Inc. All Rights Reserved.
+
+#pragma once
+
+#include "Features/IModularFeature.h"
+
+class FPackagePath;
+
+namespace UE
+{
+namespace Virtualization
+{
+
+class IVirtualizationSourceControlUtilities : public IModularFeature
+{
+public:
+	/**
+	 * Given a package path this method will attempt to sync th e.upayload file that is compatible with
+	 * the .uasset file of the package.
+	 * 
+	 * We can make the following assumptions about the relationship between .uasset and .upayload files: 
+	 * 1) The .uasset may be submitted to perforce without the .upayload (if the payload is unmodified)
+	 * 2) If the payload is modified then the .uasset and .upayload file must be submitted at the same time.
+	 * 3) The caller has already checked the existing .upayload file (if any) to see if it contains the payload
+	 * that they are looking for.
+	 * 
+	 * If the above is true then we can sync the .upayload file to the same perforce changelist as the 
+	 * * .uasset and be sure that we have the correct version.
+	 * 
+	 * Note that this has only been tested with perforce and so other source control solutions are currently
+	 * unsupported.
+	 */
+	virtual bool SyncPayloadSidecarFile(const FPackagePath& PackagePath) = 0;
+};
+
+} // namespace Virtualization
+} // namespace UE
\ No newline at end of file
diff --git a/Engine/Source/Runtime/CoreUObject/Public/Virtualization/VirtualizationManager.h b/Engine/Source/Runtime/CoreUObject/Public/Virtualization/VirtualizationManager.h
--- a/Engine/Source/Runtime/CoreUObject/Public/Virtualization/VirtualizationManager.h
+++ b/Engine/Source/Runtime/CoreUObject/Public/Virtualization/VirtualizationManager.h
@@ -2,15 +2,14 @@
 
 #pragma once
 
+#include "Compression/CompressedBuffer.h"
+#include "HAL/CriticalSection.h"
 #include "Logging/LogMacros.h"
-#include "Misc/Guid.h"
-#include "Memory/SharedBuffer.h"
+#include "Virtualization/PayloadId.h"
 
 // For UE_USE_VIRTUALBULKDATA
 #include "UObject/UE5CookerObjectVersion.h"
 
-COREUOBJECT_API DECLARE_LOG_CATEGORY_EXTERN(LogVirtualization, Log, All);
-
 // TODO: Do we want to keep this header public for UE5 release? If the only interaction should be
 //		 via FVirtualizedUntypedBulkData or other such classes we might not want to expose this 
 //		 at all.
@@ -23,6 +22,45 @@ COREUOBJECT_API DECLARE_LOG_CATEGORY_EXTERN(LogVirtualization, Log, All);
 	#define UE_VBD_CONST
 #endif //UE_USE_VIRTUALBULKDATA
 
+/**
+ * Configuring the backend hierarchy
+ * 
+ * The [Core.ContentVirtualization] section can contain a string 'BackendGraph' which will set with the name of  
+ * the backend graph, if not set then the default 'ContentVirtualizationBackendGraph_None' will be used instead. 
+ * This value can also be overridden from the command line by using 'BackendGraph=FooBar' where FooBar is the 
+ * name of the graph.
+ * 
+ * The first entry in the graph to be parsed will be the 'Hierarchy' which describes which backends should be
+ * mounted and in which order. For example 'Hierarchy=(Entry=Foo, Entry=Bar)' which should mount two backends
+ * 'Foo' and 'Bar' in that order. 
+ * 
+ * Each referenced backend in the hierarchy will then require it's own entry in the graph where the key will be
+ * it's name in the hierarchy and the value a string describing how to set it up. 
+ * The value must contain 'Type=X' where X is the name used to find the correct IVirtualizationBackendFactory 
+ * to create the backend with. 
+ * Once the backend is created then reset of the string will be passed to it, so that additional customization
+ * can be extracted. Depending on the backend implementation these values may or may not be required.
+ * 
+ * Example graph:
+ * [ContentVirtualizationBackendGraph_Example] 
+ * Hierarchy=(Entry=MemoryCache, Entry=NetworkShare)
+ * MemoryCache=(Type=InMemory)
+ * NetworkShare=(Type=FileSystem, Path="\\path\to\somewhere")
+ * 
+ * The graph is named 'ContentVirtualizationBackendGraph_Example'.
+ * The hierarchy contains two entries 'InMemory' and 'NetworkShare' to be mounted in that order
+ * MemoryCache creates a backend of type 'InMemory' and has no additional customization
+ * NetworkShare creates a backend of type 'FileSystem' and provides an additional path, the filesystem backend would 
+ * fatal error without this value.
+ */
+
+namespace UE::Virtualization
+{
+
+class IVirtualizationBackend;
+class IVirtualizationBackendFactory;
+
+
 /** This is used as a wrapper around the various potential back end implementations. 
 	The calling code shouldn't need to care about which back ends are actually in use. */
 class COREUOBJECT_API FVirtualizationManager
@@ -32,22 +70,74 @@ public:
 	static FVirtualizationManager& Get();
 
 	FVirtualizationManager();
-	~FVirtualizationManager() = default;
+	~FVirtualizationManager();
 
-	/** Push a payload to the backends*/
-	bool PushData(const FSharedBuffer& Payload, const FGuid& Guid);
+	/** 
+	 * Push a payload to the virtualization backends.
+	 * 
+	 * @param Id
+	 * @param Payload
+	 * @returns
+	 */
+	bool PushData(const FPayloadId& Id, const FCompressedBuffer& Payload);
 
-	/** Pull a payload from the backends*/
-	FSharedBuffer PullData(const FGuid& Guid);
+	/** 
+	 * Pull a payload from the virtualization backends.
+	 *
+	 * @param Id
+	 * @returns
+	 */
+	FCompressedBuffer PullData(const FPayloadId& Id);
 
 private:
 
-	void LoadSettingsFromConfigFiles();
-	FString CreateFilePath(const FGuid& Guid);
+	void ApplySettingsFromConfigFiles(const FConfigFile& PlatformEngineIni);
+	void ApplySettingsFromCmdline();
+
+	void ApplyDebugSettingsFromConfigFiles(const FConfigFile& PlatformEngineIni);
+
+	void MountBackends();
+	bool CreateBackend(const TCHAR* GraphName, const FString& ConfigEntryName, TMap<FName, IVirtualizationBackendFactory*>& FactoryLookupTable);
+
+	void AddBackend(class IVirtualizationBackend* Backend);
 
 	/** Are payloads allowed to be virtualized. Defaults to true. */
 	bool bEnablePayloadPushing;
 
 	/** The minimum length for a payload to be considered for virtualization. Defaults to 0 bytes. */
 	int64 MinPayloadLength;
+
+	/** The name of the backend graph to load from the config ini file that will describe the backend hierarchy */
+	FString BackendGraphName;
+
+	/** Debugging option: When enabled all public operations will be performed as single threaded. This is intended to aid debugging and not for production use.*/
+	bool bForceSingleThreaded;
+	
+	/** 
+	 * Debugging option: When enabled all pull operations will fail so we can see which systems cannot survive ::PullData failing to find the virtualization 
+	 * data at all. If ::PullData failures become fatal errors at some point then this option will cease to be useful.
+	 * This is intended to aid debugging and not for production use.
+	 */
+	bool bFailPayloadPullOperations;
+
+	/**
+	 * Debugging option: When enabled we will immediately 'pull' each payload after it has been 'pushed' and compare it to the original payload source to make 
+	 * sure that it can be pulled correctly.
+	 * This is intended to aid debugging and not for production use.
+	 */
+	bool bValidateAfterPushOperation;
+
+	/** The critical section used to force single threaded access if bForceSingleThreaded is true */
+	FCriticalSection ForceSingleThreadedCS;
+
+	/** Collection of all the instantiated backends */
+	TArray<IVirtualizationBackend*> AllBackendsArray;
+
+	/** Collection of all the instantiated backends that support pull operations*/
+	TArray<IVirtualizationBackend*> PullEnabledBackendsArray;
+
+	/** Collection of all the instantiated backends that support push operations*/
+	TArray<IVirtualizationBackend*> PushEnabledBackendsArray;
 };
+
+} // namespace UE::Virtualization
diff --git a/Engine/Source/Runtime/CoreUObject/Public/Virtualization/VirtualizationUtilities.h b/Engine/Source/Runtime/CoreUObject/Public/Virtualization/VirtualizationUtilities.h
--- a/Engine/Source/Runtime/CoreUObject/Public/Virtualization/VirtualizationUtilities.h
+++ b/Engine/Source/Runtime/CoreUObject/Public/Virtualization/VirtualizationUtilities.h
@@ -8,6 +8,11 @@
 
 #if WITH_EDITORONLY_DATA
 
+namespace UE
+{
+namespace Virtualization
+{
+
 class FVirtualizedBulkDataWriter : public FArchive
 {
 public:
@@ -17,7 +22,7 @@ public:
 		SetIsSaving(true);
 		SetIsPersistent(bIsPersistent);
 
-		FSharedBuffer Payload = InBulkData.GetData().Get();
+		FSharedBuffer Payload = InBulkData.GetPayload().Get();
 
 		if (Payload)
 		{
@@ -51,6 +56,12 @@ public:
 		BulkData.UpdatePayload(FSharedBuffer::TakeOwnership(Buffer, DataLength, FMemory::Free));
 	}
 
+	/** Returns if the FVirtualizedBulkDataWriter has a valid bulkdata payload or not */
+	bool IsValid() const 
+	{
+		return Buffer != nullptr;
+	}
+
 	virtual void Serialize(void* Data, int64 Num)
 	{
 		// Determine if we need to reallocate the buffer to fit the next item
@@ -126,7 +137,7 @@ namespace UE4VirtualizedBulkData_Private
 	{
 	protected:
 		DataAccessWrapper(FVirtualizedUntypedBulkData& InBulkData)
-			: Payload(InBulkData.GetData().Get())
+			: Payload(InBulkData.GetPayload().Get())
 
 		{
 
@@ -134,6 +145,11 @@ namespace UE4VirtualizedBulkData_Private
 
 		virtual ~DataAccessWrapper() = default;
 
+		bool IsValid() const
+		{
+			return !Payload.IsNull();
+		}
+
 		void* GetData() const
 		{
 			// It's okay to remove the const qualifier here as it will only be passed
@@ -152,8 +168,7 @@ namespace UE4VirtualizedBulkData_Private
 	};
 }
 
-class FVirtualizedBulkDataReader :	public UE4VirtualizedBulkData_Private::DataAccessWrapper,
-	public FBufferReaderBase
+class FVirtualizedBulkDataReader : protected UE4VirtualizedBulkData_Private::DataAccessWrapper, public FBufferReaderBase
 {
 public:
 	FVirtualizedBulkDataReader(FVirtualizedUntypedBulkData& InBulkData, bool bIsPersistent = false)
@@ -164,6 +179,12 @@ public:
 
 	virtual ~FVirtualizedBulkDataReader() = default;
 
+	/** Returns if the FVirtualizedBulkDataReader has a valid bulkdata payload or not */
+	bool IsValid() const 
+	{
+		return UE4VirtualizedBulkData_Private::DataAccessWrapper::IsValid();
+	}
+
 	using FArchive::operator<<; // For visibility of the overloads we don't override
 
 	virtual FArchive& operator<<(class FName& Name) override
@@ -181,4 +202,7 @@ public:
 	}	
 };
 
+} // namespace Virtualization
+} // namespace UE
+
 #endif //WITH_EDITORONLY_DATA
diff --git a/Engine/Source/Runtime/CoreUObject/Public/Virtualization/VirtualizedBulkData.h b/Engine/Source/Runtime/CoreUObject/Public/Virtualization/VirtualizedBulkData.h
--- a/Engine/Source/Runtime/CoreUObject/Public/Virtualization/VirtualizedBulkData.h
+++ b/Engine/Source/Runtime/CoreUObject/Public/Virtualization/VirtualizedBulkData.h
@@ -7,6 +7,7 @@
 #include "Memory/SharedBuffer.h"
 #include "Misc/Guid.h"
 #include "Misc/PackagePath.h"
+#include "Virtualization/IVirtualizedData.h"
 #include "Virtualization/VirtualizationManager.h"
 
 class FArchive;
@@ -16,6 +17,9 @@ class UObject;
 //		never used at runtime (requires a little too much reworking of some assets for now though)
 //#if WITH_EDITORONLY_DATA
 
+namespace UE::Virtualization
+{
+	
 /**
  * The goal of this class is to provide an editor time version of BulkData that will work with the content
  * virtualization system.
@@ -43,21 +47,25 @@ class UObject;
  */
 
 /** The base class with no type */
-class COREUOBJECT_API FVirtualizedUntypedBulkData
+class COREUOBJECT_API FVirtualizedUntypedBulkData : public IVirtualizedData
 {
 public:
-	using OnDataReadyCallback = TUniqueFunction<void(FSharedBuffer)>;
-
 	FVirtualizedUntypedBulkData() = default;
-	~FVirtualizedUntypedBulkData();
+	FVirtualizedUntypedBulkData(FVirtualizedUntypedBulkData&&) = default;
+	FVirtualizedUntypedBulkData& operator=(FVirtualizedUntypedBulkData&& ) = default;
+
+	FVirtualizedUntypedBulkData(const FVirtualizedUntypedBulkData& Other);
+	FVirtualizedUntypedBulkData& operator=(const FVirtualizedUntypedBulkData& Other);
+
+	virtual ~FVirtualizedUntypedBulkData() = default;
 
 	/** 
-	* Convenience method to make it easier to convert from BulkData to FVirtualizedBulkData and sets the Guid 
-	*
-	* @param BulkData	The bulkdata object to create from.
-	* @param Guid		A guid associated with the bulkdata object which will be used to identify the payload.
-	*					This MUST remain the same between sessions so that the payloads key remains consistent!
-	*/
+	 * Convenience method to make it easier to convert from BulkData to FVirtualizedBulkData and sets the Guid 
+	 *
+	 * @param BulkData	The bulkdata object to create from.
+	 * @param Guid		A guid associated with the bulkdata object which will be used to identify the payload.
+	 *					This MUST remain the same between sessions so that the payloads key remains consistent!
+	 */
 	void CreateFromBulkData(FUntypedBulkData& BulkData, const FGuid& Guid);
 
 	void Serialize(FArchive& Ar, UObject* Owner);
@@ -70,11 +78,24 @@ public:
 	/** Unloads the data (if possible) but leaves it in a state where the data can be reloaded */
 	void UnloadData();
 
-	/** Returns the a key that identifies the payload */
-	const FGuid& GetKey() const { return Key; }
+	/**
+	 * Returns a unique identifier for the object itself.
+	 * This should only return a valid FGuid as long as the object owns a valid payload.
+	 * If an object with a valid payload, has that payload removed then it should start
+	 * returning an invalid FGuid instead.
+	 * Should that object be given a new payload it should then return the original
+	 * identifier, there is no need to generate a new one.
+	 */
+	virtual FGuid GetIdentifier() const override;
+
+	/** Returns an unique identifier for the content of the payload */
+	virtual const FPayloadId& GetPayloadId() const override { return PayloadContentId; }
 
 	/** Returns the size of the payload in bytes */
-	int64 GetBulkDataSize() const { return PayloadLength; }
+	virtual int64 GetPayloadSize() const override { return PayloadSize; }
+	
+	/** Temp method, to make it easier to transition older code to virtualized bulkdata, remove when we remove UE_USE_VIRTUALBULKDATA */
+	FORCEINLINE int64 GetBulkDataSize() const { return GetPayloadSize(); }
 
 	// TODO: This (IsDataLoaded) is only needed for TextureDerivedData.cpp (which assumes that the data 
 	// needs to be loaded in order to be able to run on a background thread. Since FVirtualizedUntypedBulkData 
@@ -82,72 +103,121 @@ public:
 	/** Returns if the data is being held in memory (true) or will be loaded from disk (false) */
 	bool IsDataLoaded() const { return !Payload.IsNull(); }
 
-	/** 
-	 * Enable/Disable if the payload should be compressed if it ends up being serialized to disk on the local
-	 * machine (if virtualization is disabled or if the payload does not meet the virtualization requirements
-	 * of the project for example)
-	 */
-	void EnableDataCompression(bool bFlag) { bShouldCompressData = bFlag;  }
+	/** Returns an immutable FCompressedBuffer reference to the payload data. */
+	virtual TFuture<FSharedBuffer> GetPayload() const override;
 
 	/**
-	 * Returns an immutable FSharedBuffer reference to the payload data. 
+	 * Returns an immutable FCompressedBuffer reference to the payload data.
 	 *
-	 * To convert to a mutable reference call'FSharedBufferRef Mutable = MakeSharedBuffer(FSharedBuffer::Clone, Immutable->GetView());'
-	 * which will return a newly allocated clone of the payload which the caller will have full ownership over.
+	 * Note that depending on the internal storage formats, the payload might not actually be compressed, but that
+	 * will be handled by the FCompressedBuffer interface. Call FCompressedBuffer::Decompress() to get access to
+	 * the payload in FSharedBuffer format.
 	 */
-	TFuture<FSharedBuffer> GetData() const;
+	virtual TFuture<FCompressedBuffer> GetCompressedPayload() const override;
 
-	/** Returns an immutable FSharedBuffer reference to the payload data via the OnDataReadyCallback */
-	void GetData(OnDataReadyCallback&& Callback) const;
-		
 	/**
 	 * Allows the existing payload to be replaced with a new one.
 	 *
-	 * To pass in a raw pointer use 'MakeSharedBuffer(...)' to create a valid FSharedBufferConstRef
-	 * Use 'FSharedBuffer::Wrap' if you want to retain ownership on the data being passed in, and use 
-	 * 'FSharedBuffer::AssumeOwnership' if you are okay with the bulkdata object taking over ownership of it.
-	 * The bulkdata object must own it's internal buffer, so if you pass in a non-owned FSharedBuffer (ie 
-	 * by using 'FSharedBuffer::Wrap') then a clone of the data will be created internally and assigned to
-	 * the bulkdata object.
+	 * To pass in a raw pointer use 'FSharedBuffer::...(Data, Size)' to create a valid FSharedBuffer.
+	 * Use 'FSharedBuffer::MakeView' if you want to retain ownership on the data being passed in, and use
+	 * 'FSharedBuffer::TakeOwnership' if you are okay with the bulkdata object taking over ownership of it.
+	 * The bulkdata object must own its internal buffer, so if you pass in a non-owned FSharedBuffer (ie
+	 * by using 'FSharedBuffer::MakeView') then a clone of the data will be created internally and assigned
+	 * to the bulkdata object.
+	 *
+	 * @param InPayload The payload to update the bulkdata with
+	 * @param CompressionFormat (optional) The compression format to use, NAME_None indicates that the
+	 * payload is already in a compressed format and will not gain from being compressed again. These
+	 * payloads will never be compressed. NAME_Default will apply which ever compression format that the
+	 * underlying code deems appropriate. Other specific compression formats may be allowed, see the
+	 * documentation of FCompressedBuffer for details.
 	 */
-	void UpdatePayload(const FSharedBuffer& InPayload);
+	virtual void UpdatePayload(FSharedBuffer InPayload, FName CompressionFormat = NAME_Default) override;
 
-private:
-
-	FSharedBuffer GetDataInternal() const;
-
-	void CalculateKey();
+	/**
+	* Get the CustomVersions used in the file containing the payload. Currently this is assumed
+	* to always be the versions in the InlineArchive
+	* 
+	* @param InlineArchive The archive that was used to load this object
+	* 
+	* @return The CustomVersions that apply to the interpretation of the payload.
+	*/
+	FCustomVersionContainer GetCustomVersions(FArchive& InlineArchive);
 
-	FSharedBuffer LoadFromDisk() const;
-	bool SerializeData(FArchive& Ar, FSharedBuffer& Payload) const;
+private:
+	/** Flags used to store additional meta information about the bulk data */
+	enum class EFlags : uint32
+	{
+		/** No flags are set */
+		None						= 0,
+		/** Is the data actually virtualized or not? */
+		IsVirtualized				= 1 << 0,
+		/** Does the package have access to a .upayload file? */
+		HasPayloadSidecarFile		= 1 << 1,
+		/** The bulkdata object is currently referencing a payload saved under old bulkdata formats */
+		ReferencesLegacyFile		= 1 << 2,
+		/** The legacy file being referenced is stored with Zlib compression format */
+		LegacyFileIsCompressed		= 1 << 3,
+		/** The payload should not have compression applied to it. It is assumed that the payload is already 
+			in some sort of compressed format, see the compression documentation above for more details. */
+		DisablePayloadCompression	= 1 << 4,
+	};
+
+	/** Used to control what level of error reporting we return from some methods */
+	enum ErrorVerbosity
+	{
+		/** No errors should be logged */
+		None = 0,
+		/** Everything should be logged */
+		All
+	};
+
+	FRIEND_ENUM_CLASS_FLAGS(EFlags);
+
+	FCompressedBuffer GetDataInternal() const;
+
+	FCompressedBuffer LoadFromDisk() const;
+	FCompressedBuffer LoadFromPackageFile() const;
+	FCompressedBuffer LoadFromSidecarFile() const;
+	FCompressedBuffer LoadFromSidecarFileInternal(ErrorVerbosity Verbosity) const;
+
+	bool SerializeData(FArchive& Ar, FCompressedBuffer& Payload, const EFlags PayloadFlags) const;
 
 	void PushData();
-	FSharedBuffer PullData() const;
+	FCompressedBuffer PullData() const;
 
 	FPackagePath GetPackagePathFromOwner(UObject* Owner, EPackageSegment& OutPackageSegment) const;
 
 	bool CanUnloadData() const;
 
-	/** Unique identifier for the payload*/
-	FGuid Key;
+	void UpdateKeyIfNeeded();
+
+	void RecompressForSerialization(FCompressedBuffer& InOutPayload, EFlags PayloadFlags) const;
+	EFlags BuildFlagsForSerialization(FArchive& Ar) const;
+
+	bool IsDataVirtualized() const { return EnumHasAnyFlags(Flags, EFlags::IsVirtualized); }
+	bool HasPayloadSidecarFile() const { return EnumHasAnyFlags(Flags, EFlags::HasPayloadSidecarFile); }
+	bool IsReferencingOldBulkData() const { return EnumHasAnyFlags(Flags, EFlags::ReferencesLegacyFile); }
+
+	/** Unique identifier for the bulkdata object itself */
+	FGuid BulkDataId;
+
+	/** Unique identifier for the contents of the payload*/
+	FPayloadId PayloadContentId;
 
 	/** Pointer to the payload if it is held in memory (it has been updated but not yet saved to disk for example) */
 	FSharedBuffer Payload;
 
 	/** Length of the payload in bytes */
-	int64 PayloadLength = 0;
+	int64 PayloadSize = 0;
 
-	/** Is the data actually virtualized or not? */
-	bool bIsVirtualized = false;
-	
 	//---- The remaining members are used when the payload is not virtualized.
 	
-	/** Should we compress the data next time (don't want to serialize this)
-	TODO: Not currently affected by ::Reset which is an odd design and needs fixing */
-	bool bShouldCompressData = false;
+	/** The compression algorithm to use when saving the member 'Payload' */
+	FName CompressionFormatToUse = NAME_Default;
 
-	/** If the data has been saved to disk rather than being Virtualized, was it stored in a compressed format? */
-	bool bIsDataStoredAsCompressed = false ;
+	/** Offset of the payload in the file that contains it (INDEX_NONE if the payload does not come from a file)*/
+	int64 OffsetInFile = INDEX_NONE;
 
 	/** PackagePath containing the payload (this will be empty if the payload does not come from PackageResourceManager)*/
 	FPackagePath PackagePath;
@@ -155,10 +225,15 @@ private:
 	/** PackageSegment to load with the packagepath (unused if the payload does not come from PackageResourceManager) */
 	EPackageSegment PackageSegment;
 
-	/** Offset of the payload in the file that contains it (INDEX_NONE if the payload does not come from a file)*/
-	int64 OffsetInFile = INDEX_NONE;
+	/** Records if the payload key was created from a FGuid via CreateFromBulkData */
+	bool bWasKeyGuidDerived = false;
+
+	/** A 32bit bitfield of flags */
+	EFlags Flags = EFlags::None;
 };
 
+ENUM_CLASS_FLAGS(FVirtualizedUntypedBulkData::EFlags);
+
 //TODO: Probably remove this and change FVirtualizedUntypedBulkData to always be TVirtualizedBulkData<uint8>
 /** Type safe versions */
 template<typename DataType>
@@ -174,4 +249,32 @@ using FWordVirtualizedBulkData	= TVirtualizedBulkData<uint16>;
 using FIntVirtualizedBulkData	= TVirtualizedBulkData<int32>;
 using FFloatVirtualizedBulkData	= TVirtualizedBulkData<float>;
 
+/** 
+  * Represents an entry to the table of contents found at the start of a payload sidecar file.
+  * This might be moved to it's own header and the table of contents made into a proper class 
+  * if we decide that we want to make access of the payload sidecar file a generic feature.
+  */
+struct FTocEntry
+{
+	static constexpr uint32 PayloadSidecarFileVersion = 1;
+
+	/** Identifier for the payload */
+	FPayloadId Identifier;
+	/** The offset into the file where we can find the payload */
+	int64 OffsetInFile = INDEX_NONE;
+	/** The size of the payload WHEN uncompressed. */
+	int64 UncompressedSize = INDEX_NONE;
+
+	friend FArchive& operator<<(FArchive& Ar, FTocEntry& Entry)
+	{
+		Ar << Entry.Identifier;
+		Ar << Entry.OffsetInFile;
+		Ar << Entry.UncompressedSize;
+
+		return Ar;
+	}
+};
+
+} // namespace UE::Virtualization
+
 //#endif //WITH_EDITORONLY_DATA
diff --git a/Engine/Source/Runtime/Engine/Classes/Engine/Texture.h b/Engine/Source/Runtime/Engine/Classes/Engine/Texture.h
--- a/Engine/Source/Runtime/Engine/Classes/Engine/Texture.h
+++ b/Engine/Source/Runtime/Engine/Classes/Engine/Texture.h
@@ -19,6 +19,7 @@
 #include "Engine/StreamableRenderAsset.h"
 #include "PerPlatformProperties.h"
 #include "Misc/FieldAccessor.h"
+#include "Virtualization/VirtualizedBulkData.h"
 
 #if WITH_EDITOR
 #include "Templates/DontCopy.h"
@@ -175,6 +176,8 @@ struct FTextureSource
 	FORCEINLINE static bool IsHDR(ETextureSourceFormat Format) { return (Format == TSF_BGRE8 || Format == TSF_RGBA16F); }
 
 #if WITH_EDITOR
+	// FwdDecl for member structs
+	struct FMipData;
 
 	ENGINE_API void InitBlocked(const ETextureSourceFormat* InLayerFormats,
 		const FTextureSourceBlock* InBlocks,
@@ -273,9 +276,12 @@ struct FTextureSource
 
 	/** Unlock a mip. */
 	ENGINE_API void UnlockMip(int32 BlockIndex, int32 LayerIndex, int32 MipIndex);
-
+	
 	/** Retrieve a copy of the data for a particular mip. */
 	ENGINE_API bool GetMipData(TArray64<uint8>& OutMipData, int32 BlockIndex, int32 LayerIndex, int32 MipIndex, class IImageWrapperModule* ImageWrapperModule = nullptr);
+	
+	/** Returns a FMipData structure that wraps around the entire mip chain for read only operations. This is more efficient than calling the above method once per mip. */
+	ENGINE_API FMipData GetMipData(class IImageWrapperModule* ImageWrapperModule);
 
 	/** Computes the size of a single mip. */
 	ENGINE_API int64 CalcMipSize(int32 BlockIndex, int32 LayerIndex, int32 MipIndex) const;
@@ -314,10 +320,44 @@ struct FTextureSource
 	FORCEINLINE int32 GetNumBlocks() const { return Blocks.Num() + 1; }
 	FORCEINLINE ETextureSourceFormat GetFormat(int32 LayerIndex = 0) const { return (LayerIndex == 0) ? Format : LayerFormat[LayerIndex]; }
 	FORCEINLINE bool IsPNGCompressed() const { return bPNGCompressed; }
-	FORCEINLINE int64 GetSizeOnDisk() const { return BulkData.GetBulkDataSize(); }
-	FORCEINLINE bool IsBulkDataLoaded() const { return BulkData.IsBulkDataLoaded(); }
-	FORCEINLINE bool LoadBulkDataWithFileReader() { return BulkData.LoadBulkDataWithFileReader(); }
-	FORCEINLINE void RemoveBulkData() { BulkData.RemoveBulkData(); }
+	FORCEINLINE int64 GetSizeOnDisk() const 
+	{ 
+#if UE_USE_VIRTUALBULKDATA
+		return BulkData.GetPayloadSize(); 
+#else
+		return BulkData.GetBulkDataSize(); 
+#endif //UE_USE_VIRTUALBULKDATA
+	}
+	
+	FORCEINLINE bool IsBulkDataLoaded() const
+	{
+#if UE_USE_VIRTUALBULKDATA
+		return BulkData.IsDataLoaded();
+#else
+		return BulkData.IsBulkDataLoaded();
+#endif //UE_USE_VIRTUALBULKDATA
+	}
+
+	FORCEINLINE bool LoadBulkDataWithFileReader()
+	{
+		// TODO: When we remove UE_USE_VIRTUALBULKDATA and go all in, places that call this
+		// should be changed.
+#if UE_USE_VIRTUALBULKDATA
+		return true; // Currently a NOP
+#else
+		return BulkData.LoadBulkDataWithFileReader(); // Kick off an async load if needed
+#endif
+	}
+
+	FORCEINLINE void RemoveBulkData()
+	{
+#if UE_USE_VIRTUALBULKDATA
+		BulkData.UnloadData();
+#else
+		BulkData.RemoveBulkData();
+#endif //UE_USE_VIRTUALBULKDATA
+	}
+
 	
 	/** Sets the GUID to use, and whether that GUID is actually a hash of some data. */
 	ENGINE_API void SetId(const FGuid& InId, bool bInGuidIsHash);
@@ -335,7 +375,88 @@ struct FTextureSource
 	FORCEINLINE uint8* LockMip(int32 MipIndex) { return LockMip(0, 0, MipIndex); }
 	FORCEINLINE void UnlockMip(int32 MipIndex) { UnlockMip(0, 0, MipIndex); }
 
-#endif
+	struct FMipAllocation
+	{
+		/** Create an empty object */
+		FMipAllocation() = default;
+		/** Take a read only FSharedBuffer, will allocate a new buffer and copy from this if Read/Write access is requested */
+		FMipAllocation(FSharedBuffer SrcData);
+		/** Allocate a new buffer and copy from the source data */
+		FMipAllocation(const void* SrcData, int64 DataLength);
+		/** Lock the bulkdata and use it's memory buffer */
+		FMipAllocation(FByteBulkData& BulkData);
+
+		// Do not actually do anything for copy constructor or assignments, this is required for as long as
+		// we need to support the old bulkdata code path (although previously storing these allocations as
+		// raw pointers would allow it to be assigned, this would most likely cause a mismatch in lock counts,
+		// either in FTextureSource or the underlying bulkdata and was never actually safe)
+		FMipAllocation(const FMipAllocation&) {}
+		FMipAllocation& operator =(const FMipAllocation&) { return *this; }
+
+		// We do allow rvalue assignment
+		FMipAllocation(FMipAllocation&&);
+		FMipAllocation& operator =(FMipAllocation&&);
+
+		~FMipAllocation() = default;
+
+		/** Release all currently owned data and return the object to the default state */
+		void Reset();
+
+		/** Returns true if the object contains no data */
+		bool IsNull() const { return ReadOnlyReference.IsNull(); }
+
+		/** Returns the overall size of the data in bytes */
+		int64 GetSize() const { return ReadOnlyReference.GetSize(); }
+
+		/** Returns a FSharedBuffer that contains the current texture data but cannot be directly modified */
+		const FSharedBuffer& GetDataReadOnly() const { return ReadOnlyReference; }
+
+		/** Returns a pointer that contains the current texture data and can be written to */
+		uint8* GetDataReadWrite();
+
+		/** Returns the internal FSharedBuffer and relinquish ownership, used to transfer the data to virtualized bulkdata */
+		FSharedBuffer Release();
+
+	private:
+		void CreateReadWriteBuffer(const void* SrcData, int64 DataLength);
+
+		FSharedBuffer		ReadOnlyReference;
+		TUniquePtr<uint8[]>	ReadWriteBuffer;
+		uint8* BulkDataPtr{ nullptr };
+	};
+
+	/** Structure that encapsulates the decompressed texture data and can be accessed per mip */
+	struct ENGINE_API FMipData
+	{
+		~FMipData();
+
+		/** Get a copy of a given texture mip, to be stored in OutMipData */
+		bool GetMipData(TArray64<uint8>& OutMipData, int32 BlockIndex, int32 LayerIndex, int32 MipIndex) const;
+
+		/** Allow the copy constructor by rvalue*/
+		FMipData(FMipData&& Other)
+			: TextureSource(Other.TextureSource)
+		{
+			MipData = MoveTemp(Other.MipData);
+		}
+
+		/** Disallow everything else so we don't get duplicates */
+		FMipData() = delete;
+		FMipData(const FMipData&) = delete;
+		FMipData& operator=(const FMipData&) = delete;
+		FMipData& operator=(FMipData&& Other) = delete;
+
+	private:
+		// We only want to allow FTextureSource to create FMipData objects
+		friend struct FTextureSource;
+
+		FMipData(const FTextureSource& InSource, FMipAllocation&& InMipData, FRWLock* InReadLock = nullptr);
+
+		const FTextureSource& TextureSource;
+		FMipAllocation MipData;
+		FRWLock* ReadLock;
+	};
+#endif // WITH_EDITOR
 
 private:
 	/** Allow UTexture access to internals. */
@@ -350,9 +471,8 @@ private:
 	TDontCopy<FRWLock> BulkDataLock;
 #endif
 	/** The bulk source data. */
-	FByteBulkData BulkData;
-	/** Pointer to locked mip data, if any. */
-	uint8* LockedMipData;
+	TChooseClass<UE_USE_VIRTUALBULKDATA, UE::Virtualization::FByteVirtualizedBulkData, FByteBulkData>::Result BulkData;
+	
 	/** Number of mips that are locked. */
 	uint32 NumLockedMips;
 
@@ -366,11 +486,19 @@ private:
 	ELockState LockState;
 
 #if WITH_EDITOR
+	/** Pointer to locked mip data, if any. */
+	FMipAllocation LockedMipData;
 
 	// Internal implementation for locking the mip data, called by LockMipReadOnly or LockMip */
 	uint8* LockMipInternal(int32 BlockIndex, int32 LayerIndex, int32 MipIndex, ELockState RequestedLockState);
+	
+	/** Returns the source data fully decompressed */
+	FMipAllocation Decompress(class IImageWrapperModule* ImageWrapperModule) UE_VBD_CONST;
+
+	FMipAllocation TryDecompressPngData(IImageWrapperModule* ImageWrapperModule) UE_VBD_CONST;
 	/** Attempt to decompress the source data from Jpeg format. All failures will be logged and result in the method returning false */
-	bool TryDecompressJpegData(IImageWrapperModule* ImageWrapperModule, const void* RawSourceData, int32 MipIndex, TArray64<uint8>& OutRawData);
+	FMipAllocation TryDecompressJpegData(IImageWrapperModule* ImageWrapperModule) UE_VBD_CONST;
+	
 	/** Return true if the source art is not png compressed but could be. */
 	bool CanPNGCompress() const;
 	/** Removes source data. */
diff --git a/Engine/Source/Runtime/Engine/Private/Streaming/Texture2DStreamIn_DDC.cpp b/Engine/Source/Runtime/Engine/Private/Streaming/Texture2DStreamIn_DDC.cpp
--- a/Engine/Source/Runtime/Engine/Private/Streaming/Texture2DStreamIn_DDC.cpp
+++ b/Engine/Source/Runtime/Engine/Private/Streaming/Texture2DStreamIn_DDC.cpp
@@ -195,7 +195,7 @@ void FTexture2DStreamIn_DDC::DoLoadNewMipsFromDDC(const FContext& Context)
 					}
 					else
 					{
-						UE_LOG(LogTexture, Error, TEXT("DDC mip size (%d) not as expected."), MipIndex);
+						UE_LOG(LogTexture, Error, TEXT("DDC mip size (%d) not as expected (%d)."), MipIndex, ExpectedMipSize);
 						MarkAsCancelled();
 					}
 				}
