Create function UStruct::GetSchemaHash. This function provides a hash that can be used to check whether tagged property serialization stored in a saved package is a match for the current binary's UClasses. This hash can be used to test whether versioned tagged property serialization will produce unchanged results, and whether unversioned property serialization will succeed.
#rnx

diff --git a/Engine/Source/Editor/UnrealEd/Private/EditorServer.cpp b/Engine/Source/Editor/UnrealEd/Private/EditorServer.cpp
--- a/Engine/Source/Editor/UnrealEd/Private/EditorServer.cpp
+++ b/Engine/Source/Editor/UnrealEd/Private/EditorServer.cpp
@@ -5605,6 +5605,8 @@ void ListMapPackageDependencies(const TCHAR* InStr)
 	}
 }
 
+COREUOBJECT_API void DumpClassSchemas(const TCHAR* Str, FOutputDevice& Ar);
+
 bool UEditorEngine::Exec( UWorld* InWorld, const TCHAR* Stream, FOutputDevice& Ar )
 {
 	TCHAR ErrorTemp[256]=TEXT("Setup: ");
@@ -5884,6 +5886,10 @@ bool UEditorEngine::Exec( UWorld* InWorld, const TCHAR* Stream, FOutputDevice& A
 	{
 		HandleDumpPublicCommand( Str, Ar );
 	}
+	else if (FParse::Command(&Str, TEXT("DumpClassSchemas")))
+	{
+		DumpClassSchemas(Str, Ar);
+	}
 	else if( FParse::Command(&Str,TEXT("JUMPTO")) )
 	{
 		return HandleJumpToCommand( Str, Ar );
diff --git a/Engine/Source/Runtime/Core/Private/UObject/UnrealNames.cpp b/Engine/Source/Runtime/Core/Private/UObject/UnrealNames.cpp
--- a/Engine/Source/Runtime/Core/Private/UObject/UnrealNames.cpp
+++ b/Engine/Source/Runtime/Core/Private/UObject/UnrealNames.cpp
@@ -28,6 +28,7 @@
 #include "Serialization/MemoryImage.h"
 #include "Serialization/MemoryReader.h"
 #include "Serialization/MemoryWriter.h"
+#include "Hash/Blake3.h"
 #include "Hash/CityHash.h"
 #include "Templates/AlignmentTemplates.h"
 
@@ -2322,6 +2323,22 @@ bool FName::TryAppendAnsiString(FAnsiStringBuilderBase& Out) const
 	return true;
 }
 
+void AppendHash(FBlake3& Builder, FName In)
+{
+	FNameBuffer DecodeBuffer;
+	FNameStringView NameEntryView = In.GetDisplayNameEntry()->MakeView(DecodeBuffer);
+	if (NameEntryView.IsAnsi())
+	{
+		Builder.Update(NameEntryView.Ansi, NameEntryView.Len * sizeof(NameEntryView.Ansi[0]));
+	}
+	else
+	{
+		Builder.Update(NameEntryView.Wide, NameEntryView.Len * sizeof(NameEntryView.Wide[0]));
+	}
+	int32 Number = INTEL_ORDER64(In.GetNumber());
+	Builder.Update(&Number, sizeof(Number));
+}
+
 void FName::DisplayHash(FOutputDevice& Ar)
 {
 	GetNamePool().LogStats(Ar);
diff --git a/Engine/Source/Runtime/Core/Public/UObject/NameTypes.h b/Engine/Source/Runtime/Core/Public/UObject/NameTypes.h
--- a/Engine/Source/Runtime/Core/Public/UObject/NameTypes.h
+++ b/Engine/Source/Runtime/Core/Public/UObject/NameTypes.h
@@ -1295,3 +1295,7 @@ public:
 };
 
 template <> struct TIsContiguousContainer<FNameBuilder> { static constexpr bool Value = true; };
+
+/** Update the Hash with the FName's text and number */
+class FBlake3;
+CORE_API void AppendHash(FBlake3& Builder, FName In);
\ No newline at end of file
diff --git a/Engine/Source/Runtime/CoreUObject/Private/Serialization/UnversionedPropertySerialization.cpp b/Engine/Source/Runtime/CoreUObject/Private/Serialization/UnversionedPropertySerialization.cpp
--- a/Engine/Source/Runtime/CoreUObject/Private/Serialization/UnversionedPropertySerialization.cpp
+++ b/Engine/Source/Runtime/CoreUObject/Private/Serialization/UnversionedPropertySerialization.cpp
@@ -2,10 +2,16 @@
 
 #include "Serialization/UnversionedPropertySerialization.h"
 #include "Serialization/UnversionedPropertySerializationTest.h"
+#include "Hash/Blake3.h"
 #include "Interfaces/ITargetPlatform.h"
 #include "Misc/ScopeRWLock.h"
 #include "UObject/UnrealType.h"
 
+#if WITH_EDITORONLY_DATA
+#include "Misc/FileHelper.h"
+#include "UObject/UObjectIterator.h"
+#endif
+
 // Caches a property array per UStruct to avoid link-walking and touching FProperty data.
 //
 // As a reference point, this uses ~6MB of memory in an internal project and makes all
@@ -278,11 +284,17 @@ private:
 // Serialization is based on indices into this property array
 struct FUnversionedStructSchema
 {
+#if WITH_EDITORONLY_DATA
+	FBlake3Hash SchemaHash;
+#endif
 	uint32 Num;
 	FUnversionedPropertySerializer Serializers[0];
 
 	FORCEINLINE static FUnversionedStructSchema* Create(const UStruct* Struct, bool bSkipEditorOnly)
 	{
+#if WITH_EDITORONLY_DATA
+		FBlake3 HashBuilder;
+#endif
 		TArray<FUnversionedPropertySerializer, TInlineAllocator<256>> Serializers;
 		for (FProperty* Property = Struct->PropertyLink; Property; Property = Property->PropertyLinkNext)
 		{
@@ -290,6 +302,9 @@ struct FUnversionedStructSchema
 			if (!bSkipEditorOnly || !Property->IsEditorOnlyProperty())
 #endif
 			{
+#if WITH_EDITORONLY_DATA
+				Property->AppendSchemaHash(HashBuilder, bSkipEditorOnly);
+#endif
 				for (int32 ArrayIdx = 0, ArrayDim = Property->ArrayDim; ArrayIdx < ArrayDim; ++ArrayIdx)
 				{
 					Serializers.Add(FUnversionedPropertySerializer(Property, ArrayIdx));
@@ -299,11 +314,30 @@ struct FUnversionedStructSchema
 
 		uint32 Bytes = sizeof(FUnversionedStructSchema) + Serializers.Num() * sizeof(FUnversionedPropertySerializer);
 		FUnversionedStructSchema* Schema = reinterpret_cast<FUnversionedStructSchema*>(FMemory::Malloc(Bytes, alignof(FUnversionedPropertySerializer)));
+		
+#if WITH_EDITORONLY_DATA
+		Schema->SchemaHash = HashBuilder.Finalize();
+#endif
 		Schema->Num = Serializers.Num();
 		FMemory::Memcpy(Schema->Serializers, Serializers.GetData(), Serializers.Num() * sizeof(FUnversionedPropertySerializer));
 
 		return Schema;
 	}
+
+#if WITH_EDITORONLY_DATA
+	static FBlake3Hash CalculateSchemaHash(UStruct* Struct, bool bSkipEditorOnly)
+	{
+		FBlake3 HashBuilder;
+		for (FProperty* Property = Struct->PropertyLink; Property; Property = Property->PropertyLinkNext)
+		{
+			if (!bSkipEditorOnly || !Property->IsEditorOnlyProperty())
+			{
+				Property->AppendSchemaHash(HashBuilder, bSkipEditorOnly);
+			}
+		}
+		return HashBuilder.Finalize();
+	}
+#endif
 };
 
 static const FUnversionedStructSchema*& GetUnversionedSchema(const UStruct* Struct, bool bSkipEditorOnly)
@@ -878,3 +912,44 @@ void SerializeUnversionedProperties(const UStruct* Struct, FStructuredArchive::F
 		}
 	}
 }
+
+#if WITH_EDITORONLY_DATA
+const FBlake3Hash& GetSchemaHash(const UStruct* Struct, bool bSkipEditorOnly)
+{
+	return GetOrCreateUnversionedSchema(Struct, bSkipEditorOnly).SchemaHash;
+}
+
+COREUOBJECT_API void DumpClassSchemas(const TCHAR* Str, FOutputDevice& Ar)
+{
+	TArray<FString> Lines;
+	TArray<UStruct*> Structs;
+	for (TObjectIterator<UStruct> It; It; ++It)
+	{
+		Structs.Add(*It);
+	}
+	Algo::Sort(Structs, [](UStruct* A, UStruct* B)
+		{
+			FNameBuilder AName;
+			FNameBuilder BName;
+			A->GetPathName(nullptr, AName);
+			B->GetPathName(nullptr, BName);
+			return FStringView(AName).Compare(FStringView(BName), ESearchCase::IgnoreCase) < 0;
+		});
+	bool bSkipEditorOnly = false;
+	FParse::Bool(Str, TEXT("-SkipEditorOnly="), bSkipEditorOnly);
+	for (UStruct* Struct : Structs)
+	{
+		const FBlake3Hash& ExistingHash = Struct->GetSchemaHash(bSkipEditorOnly);
+		FBlake3Hash NewHash = FUnversionedStructSchema::CalculateSchemaHash(Struct, bSkipEditorOnly);
+		ensureMsgf(ExistingHash == NewHash, TEXT("Hash mismatch for %s. Stored hash=%s, Current hash=%s"),
+			*Struct->GetFullName(nullptr), *LexToString(ExistingHash), *LexToString(NewHash));
+		Lines.Add(FString::Printf(TEXT("%-80s, %s"), *Struct->GetPathName(nullptr), *LexToString(ExistingHash)));
+	}
+	FString DumpFilename;
+	FParse::Value(Str, TEXT("-FILE="), DumpFilename);
+	if (!DumpFilename.IsEmpty())
+	{
+		FFileHelper::SaveStringArrayToFile(Lines, *DumpFilename);
+	}
+}
+#endif
diff --git a/Engine/Source/Runtime/CoreUObject/Private/Serialization/UnversionedPropertySerialization.h b/Engine/Source/Runtime/CoreUObject/Private/Serialization/UnversionedPropertySerialization.h
--- a/Engine/Source/Runtime/CoreUObject/Private/Serialization/UnversionedPropertySerialization.h
+++ b/Engine/Source/Runtime/CoreUObject/Private/Serialization/UnversionedPropertySerialization.h
@@ -4,10 +4,20 @@
 
 #include "UObject/UnrealType.h"
 
+struct FBlake3Hash;
+
 // Check if unversioned property serialization is configured to be used on target platform
 bool CanUseUnversionedPropertySerialization(const ITargetPlatform* Target);
 
 // Serialize sparse unversioned properties for a particular struct
 void SerializeUnversionedProperties(const UStruct* Struct, FStructuredArchive::FSlot Slot, uint8* Data, UStruct* DefaultsStruct, uint8* Defaults);
 
-void DestroyUnversionedSchema(const UStruct* Struct);
\ No newline at end of file
+void DestroyUnversionedSchema(const UStruct* Struct);
+
+#if WITH_EDITORONLY_DATA
+/**
+ * Read the SchemaHash for the given struct, with our without editor-only properties.
+ * This hash covers the Struct's data for both versioned and unversioned serialization.
+ */
+const FBlake3Hash& GetSchemaHash(const UStruct* Struct, bool bSkipEditorOnly);
+#endif
\ No newline at end of file
diff --git a/Engine/Source/Runtime/CoreUObject/Private/UObject/Class.cpp b/Engine/Source/Runtime/CoreUObject/Private/UObject/Class.cpp
--- a/Engine/Source/Runtime/CoreUObject/Private/UObject/Class.cpp
+++ b/Engine/Source/Runtime/CoreUObject/Private/UObject/Class.cpp
@@ -1278,6 +1278,13 @@ void UStruct::SerializeTaggedProperties(FStructuredArchive::FSlot Slot, uint8* D
 	}
 }
 
+#if WITH_EDITORONLY_DATA
+const FBlake3Hash& UStruct::GetSchemaHash(bool bSkipEditorOnly) const
+{
+	return ::GetSchemaHash(this, bSkipEditorOnly);
+}
+#endif
+
 void UStruct::SerializeVersionedTaggedProperties(FStructuredArchive::FSlot Slot, uint8* Data, UStruct* DefaultsStruct, uint8* Defaults, const UObject* BreakRecursionIfFullyLoad) const
 {
 	FArchive& UnderlyingArchive = Slot.GetUnderlyingArchive();
@@ -6396,4 +6403,4 @@ IMPLEMENT_CORE_INTRINSIC_CLASS(UDynamicClass, UClass,
 	#endif
 #endif
 
-#include "UObject/DefineUPropertyMacros.h"
\ No newline at end of file
+#include "UObject/DefineUPropertyMacros.h"
diff --git a/Engine/Source/Runtime/CoreUObject/Private/UObject/EnumProperty.cpp b/Engine/Source/Runtime/CoreUObject/Private/UObject/EnumProperty.cpp
--- a/Engine/Source/Runtime/CoreUObject/Private/UObject/EnumProperty.cpp
+++ b/Engine/Source/Runtime/CoreUObject/Private/UObject/EnumProperty.cpp
@@ -10,6 +10,7 @@
 #include "Algo/Find.h"
 #include "UObject/LinkerLoad.h"
 #include "Misc/NetworkVersion.h"
+#include "Hash/Blake3.h"
 
 // WARNING: This should always be the last include in any file that needs it (except .generated.h)
 #include "UObject/UndefineUPropertyMacros.h"
@@ -461,6 +462,24 @@ EConvertFromTypeResult FEnumProperty::ConvertFromType(const FPropertyTag& Tag, F
 	return EConvertFromTypeResult::Converted;
 }
 
+#if WITH_EDITORONLY_DATA
+void FEnumProperty::AppendSchemaHash(FBlake3& Builder, bool bSkipEditorOnly) const
+{
+	Super::AppendSchemaHash(Builder, bSkipEditorOnly);
+	if (Enum)
+	{
+		FNameBuilder NameBuilder;
+		Enum->GetPathName(nullptr, NameBuilder);
+		Builder.Update(NameBuilder.GetData(), NameBuilder.Len() * sizeof(NameBuilder.GetData()[0]));
+		int32 Num = Enum->NumEnums();
+		for (int32 Index = 0; Index < Num; ++Index)
+		{
+			AppendHash(Builder, Enum->GetNameByIndex(Index));
+		}
+	}
+}
+#endif
+
 uint32 FEnumProperty::GetValueTypeHashInternal(const void* Src) const
 {
 	check(UnderlyingProp);
diff --git a/Engine/Source/Runtime/CoreUObject/Private/UObject/Property.cpp b/Engine/Source/Runtime/CoreUObject/Private/UObject/Property.cpp
--- a/Engine/Source/Runtime/CoreUObject/Private/UObject/Property.cpp
+++ b/Engine/Source/Runtime/CoreUObject/Private/UObject/Property.cpp
@@ -5,6 +5,7 @@
 =============================================================================*/
 
 #include "CoreMinimal.h"
+#include "Hash/Blake3.h"
 #include "Misc/AsciiSet.h"
 #include "Misc/Guid.h"
 #include "Misc/StringBuilder.h"
@@ -956,6 +957,16 @@ bool FProperty::SupportsNetSharedSerialization() const
 	return true;
 }
 
+#if WITH_EDITORONLY_DATA
+void FProperty::AppendSchemaHash(FBlake3& Builder, bool bSkipEditorOnly) const
+{
+	AppendHash(Builder, NamePrivate);
+	Builder.Update(&ArrayDim, sizeof(ArrayDim));
+	AppendHash(Builder, GetID());
+}
+#endif
+
+
 //
 // Return whether the property should be exported.
 //
@@ -1655,4 +1666,4 @@ FProperty* UStruct::FindPropertyByName(FName InName) const
 }
 
 
-#include "UObject/DefineUPropertyMacros.h"
\ No newline at end of file
+#include "UObject/DefineUPropertyMacros.h"
diff --git a/Engine/Source/Runtime/CoreUObject/Private/UObject/PropertyArray.cpp b/Engine/Source/Runtime/CoreUObject/Private/UObject/PropertyArray.cpp
--- a/Engine/Source/Runtime/CoreUObject/Private/UObject/PropertyArray.cpp
+++ b/Engine/Source/Runtime/CoreUObject/Private/UObject/PropertyArray.cpp
@@ -605,6 +605,17 @@ const TCHAR* FArrayProperty::ImportTextInnerItem( const TCHAR* Buffer, const FPr
 	return Buffer;
 }
 
+#if WITH_EDITORONLY_DATA
+void FArrayProperty::AppendSchemaHash(FBlake3& Builder, bool bSkipEditorOnly) const
+{
+	Super::AppendSchemaHash(Builder, bSkipEditorOnly);
+	if (Inner)
+	{
+		Inner->AppendSchemaHash(Builder, bSkipEditorOnly);
+	}
+}
+#endif
+
 void FArrayProperty::AddCppProperty(FProperty* Property)
 {
 	check(!Inner);
diff --git a/Engine/Source/Runtime/CoreUObject/Private/UObject/PropertyBool.cpp b/Engine/Source/Runtime/CoreUObject/Private/UObject/PropertyBool.cpp
--- a/Engine/Source/Runtime/CoreUObject/Private/UObject/PropertyBool.cpp
+++ b/Engine/Source/Runtime/CoreUObject/Private/UObject/PropertyBool.cpp
@@ -1,6 +1,8 @@
 // Copyright Epic Games, Inc. All Rights Reserved.
 
 #include "CoreMinimal.h"
+
+#include "Hash/Blake3.h"
 #include "UObject/ObjectMacros.h"
 #include "UObject/UObjectGlobals.h"
 #include "UObject/UnrealType.h"
@@ -291,6 +293,17 @@ EConvertFromTypeResult FBoolProperty::ConvertFromType(const FPropertyTag& Tag, F
 	return EConvertFromTypeResult::Converted;
 }
 
+#if WITH_EDITORONLY_DATA
+void FBoolProperty::AppendSchemaHash(FBlake3& Builder, bool bSkipEditorOnly) const
+{
+	Super::AppendSchemaHash(Builder, bSkipEditorOnly);
+	Builder.Update(&ByteOffset, sizeof(ByteOffset));
+	Builder.Update(&ByteMask, sizeof(ByteMask));
+	Builder.Update(&FieldMask, sizeof(FieldMask));
+}
+#endif
+
+
 void FBoolProperty::ExportTextItem( FString& ValueStr, const void* PropertyValue, const void* DefaultValue, UObject* Parent, int32 PortFlags, UObject* ExportRootScope ) const
 {
 	check(FieldSize != 0);
diff --git a/Engine/Source/Runtime/CoreUObject/Private/UObject/PropertyByte.cpp b/Engine/Source/Runtime/CoreUObject/Private/UObject/PropertyByte.cpp
--- a/Engine/Source/Runtime/CoreUObject/Private/UObject/PropertyByte.cpp
+++ b/Engine/Source/Runtime/CoreUObject/Private/UObject/PropertyByte.cpp
@@ -11,6 +11,7 @@
 #include "Algo/Find.h"
 #include "UObject/LinkerLoad.h"
 #include "Misc/NetworkVersion.h"
+#include "Hash/Blake3.h"
 
 // WARNING: This should always be the last include in any file that needs it (except .generated.h)
 #include "UObject/UndefineUPropertyMacros.h"
@@ -324,6 +325,24 @@ EConvertFromTypeResult FByteProperty::ConvertFromType(const FPropertyTag& Tag, F
 	return EConvertFromTypeResult::Converted;
 }
 
+#if WITH_EDITORONLY_DATA
+void FByteProperty::AppendSchemaHash(FBlake3& Builder, bool bSkipEditorOnly) const
+{
+	Super::AppendSchemaHash(Builder, bSkipEditorOnly);
+	if (Enum)
+	{
+		FNameBuilder NameBuilder;
+		Enum->GetPathName(nullptr, NameBuilder);
+		Builder.Update(NameBuilder.GetData(), NameBuilder.Len() * sizeof(NameBuilder.GetData()[0]));
+		int32 Num = Enum->NumEnums();
+		for (int32 Index = 0; Index < Num; ++Index)
+		{
+			AppendHash(Builder, Enum->GetNameByIndex(Index));
+		}
+	}
+}
+#endif
+
 void FByteProperty::ExportTextItem( FString& ValueStr, const void* PropertyValue, const void* DefaultValue, UObject* Parent, int32 PortFlags, UObject* ExportRootScope ) const
 {
 	if (0 != (PortFlags & PPF_ExportCpp))
diff --git a/Engine/Source/Runtime/CoreUObject/Private/UObject/PropertyClass.cpp b/Engine/Source/Runtime/CoreUObject/Private/UObject/PropertyClass.cpp
--- a/Engine/Source/Runtime/CoreUObject/Private/UObject/PropertyClass.cpp
+++ b/Engine/Source/Runtime/CoreUObject/Private/UObject/PropertyClass.cpp
@@ -8,6 +8,7 @@
 #include "UObject/LinkerPlaceholderClass.h"
 #include "Misc/ConfigCacheIni.h"
 #include "UObject/PropertyHelper.h"
+#include "Hash/Blake3.h"
 
 // WARNING: This should always be the last include in any file that needs it (except .generated.h)
 #include "UObject/UndefineUPropertyMacros.h"
@@ -178,4 +179,17 @@ bool FClassProperty::Identical( const void* A, const void* B, uint32 PortFlags )
 	return (ObjectA == ObjectB);
 }
 
+#if WITH_EDITORONLY_DATA
+void FClassProperty::AppendSchemaHash(FBlake3& Builder, bool bSkipEditorOnly) const
+{
+	Super::AppendSchemaHash(Builder, bSkipEditorOnly);
+	if (MetaClass)
+	{
+		// Hash the class's name instead of recursively hashing the class; the class's schema does not impact how we serialize our pointer to it
+		FNameBuilder ObjectPath;
+		MetaClass->GetPathName(nullptr, ObjectPath);
+		Builder.Update(ObjectPath.GetData(), ObjectPath.Len() * sizeof(ObjectPath.GetData()[0]));
+	}
+}
+#endif
 #include "UObject/DefineUPropertyMacros.h"
\ No newline at end of file
diff --git a/Engine/Source/Runtime/CoreUObject/Private/UObject/PropertyDelegate.cpp b/Engine/Source/Runtime/CoreUObject/Private/UObject/PropertyDelegate.cpp
--- a/Engine/Source/Runtime/CoreUObject/Private/UObject/PropertyDelegate.cpp
+++ b/Engine/Source/Runtime/CoreUObject/Private/UObject/PropertyDelegate.cpp
@@ -8,6 +8,7 @@
 #include "UObject/PropertyHelper.h"
 #include "UObject/LinkerPlaceholderFunction.h"
 #include "Serialization/ArchiveUObjectFromStructuredArchive.h"
+#include "Hash/Blake3.h"
 
 // WARNING: This should always be the last include in any file that needs it (except .generated.h)
 #include "UObject/UndefineUPropertyMacros.h"
@@ -190,6 +191,21 @@ bool FDelegateProperty::SameType(const FProperty* Other) const
 	return Super::SameType(Other) && (SignatureFunction == ((FDelegateProperty*)Other)->SignatureFunction);
 }
 
+#if WITH_EDITORONLY_DATA
+void FDelegateProperty::AppendSchemaHash(FBlake3& Builder, bool bSkipEditorOnly) const
+{
+	Super::AppendSchemaHash(Builder, bSkipEditorOnly);
+	if (SignatureFunction)
+	{
+		// Hash the function's name instead of recursively hashing the function; the function's schema does not impact how we serialize our pointer to it
+		FNameBuilder ObjectPath;
+		SignatureFunction->GetPathName(nullptr, ObjectPath);
+		Builder.Update(ObjectPath.GetData(), ObjectPath.Len() * sizeof(ObjectPath.GetData()[0]));
+	}
+}
+#endif
+
+
 void FDelegateProperty::BeginDestroy()
 {
 #if USE_CIRCULAR_DEPENDENCY_LOAD_DEFERRING
diff --git a/Engine/Source/Runtime/CoreUObject/Private/UObject/PropertyInterface.cpp b/Engine/Source/Runtime/CoreUObject/Private/UObject/PropertyInterface.cpp
--- a/Engine/Source/Runtime/CoreUObject/Private/UObject/PropertyInterface.cpp
+++ b/Engine/Source/Runtime/CoreUObject/Private/UObject/PropertyInterface.cpp
@@ -9,6 +9,7 @@
 #include "UObject/UnrealType.h"
 #include "UObject/UnrealTypePrivate.h"
 #include "UObject/LinkerPlaceholderClass.h"
+#include "Hash/Blake3.h"
 
 // WARNING: This should always be the last include in any file that needs it (except .generated.h)
 #include "UObject/UndefineUPropertyMacros.h"
@@ -300,6 +301,20 @@ bool FInterfaceProperty::SameType(const FProperty* Other) const
 	return Super::SameType(Other) && (InterfaceClass == ((FInterfaceProperty*)Other)->InterfaceClass);
 }
 
+#if WITH_EDITORONLY_DATA
+void FInterfaceProperty::AppendSchemaHash(FBlake3& Builder, bool bSkipEditorOnly) const
+{
+	Super::AppendSchemaHash(Builder, bSkipEditorOnly);
+	if (InterfaceClass)
+	{
+		// Hash the class's name instead of recursively hashing the class; the class's schema does not impact how we serialize our pointer to it
+		FNameBuilder ObjectPath;
+		InterfaceClass->GetPathName(nullptr, ObjectPath);
+		Builder.Update(ObjectPath.GetData(), ObjectPath.Len() * sizeof(ObjectPath.GetData()[0]));
+	}
+}
+#endif
+
 void FInterfaceProperty::AddReferencedObjects(FReferenceCollector& Collector)
 {
 	Collector.AddReferencedObject(InterfaceClass);
diff --git a/Engine/Source/Runtime/CoreUObject/Private/UObject/PropertyMap.cpp b/Engine/Source/Runtime/CoreUObject/Private/UObject/PropertyMap.cpp
--- a/Engine/Source/Runtime/CoreUObject/Private/UObject/PropertyMap.cpp
+++ b/Engine/Source/Runtime/CoreUObject/Private/UObject/PropertyMap.cpp
@@ -1176,6 +1176,21 @@ EConvertFromTypeResult FMapProperty::ConvertFromType(const FPropertyTag& Tag, FS
 	return EConvertFromTypeResult::UseSerializeItem;
 }
 
+#if WITH_EDITORONLY_DATA
+void FMapProperty::AppendSchemaHash(FBlake3& Builder, bool bSkipEditorOnly) const
+{
+	Super::AppendSchemaHash(Builder, bSkipEditorOnly);
+	if (KeyProp)
+	{
+		KeyProp->AppendSchemaHash(Builder, bSkipEditorOnly);
+	}
+	if (ValueProp)
+	{
+		ValueProp->AppendSchemaHash(Builder, bSkipEditorOnly);
+	}
+}
+#endif
+
 void FScriptMapHelper::Rehash()
 {
 	WithScriptMap([this](auto* Map)
diff --git a/Engine/Source/Runtime/CoreUObject/Private/UObject/PropertySet.cpp b/Engine/Source/Runtime/CoreUObject/Private/UObject/PropertySet.cpp
--- a/Engine/Source/Runtime/CoreUObject/Private/UObject/PropertySet.cpp
+++ b/Engine/Source/Runtime/CoreUObject/Private/UObject/PropertySet.cpp
@@ -984,6 +984,17 @@ EConvertFromTypeResult FSetProperty::ConvertFromType(const FPropertyTag& Tag, FS
 	return EConvertFromTypeResult::UseSerializeItem;
 }
 
+#if WITH_EDITORONLY_DATA
+void FSetProperty::AppendSchemaHash(FBlake3& Builder, bool bSkipEditorOnly) const
+{
+	Super::AppendSchemaHash(Builder, bSkipEditorOnly);
+	if (ElementProp)
+	{
+		ElementProp->AppendSchemaHash(Builder, bSkipEditorOnly);
+	}
+}
+#endif
+
 void FScriptSetHelper::Rehash()
 {
 	// Moved out-of-line to maybe fix a weird link error
diff --git a/Engine/Source/Runtime/CoreUObject/Private/UObject/PropertyStruct.cpp b/Engine/Source/Runtime/CoreUObject/Private/UObject/PropertyStruct.cpp
--- a/Engine/Source/Runtime/CoreUObject/Private/UObject/PropertyStruct.cpp
+++ b/Engine/Source/Runtime/CoreUObject/Private/UObject/PropertyStruct.cpp
@@ -9,6 +9,7 @@
 #include "UObject/PropertyHelper.h"
 #include "UObject/LinkerPlaceholderBase.h"
 #include "Serialization/ArchiveUObjectFromStructuredArchive.h"
+#include "Hash/Blake3.h"
 
 // WARNING: This should always be the last include in any file that needs it (except .generated.h)
 #include "UObject/UndefineUPropertyMacros.h"
@@ -405,4 +406,17 @@ EConvertFromTypeResult FStructProperty::ConvertFromType(const FPropertyTag& Tag,
 	return EConvertFromTypeResult::UseSerializeItem;
 }
 
+#if WITH_EDITORONLY_DATA
+void FStructProperty::AppendSchemaHash(FBlake3& Builder, bool bSkipEditorOnly) const
+{
+	Super::AppendSchemaHash(Builder, bSkipEditorOnly);
+	if (Struct)
+	{
+		const FIoHash& StructSchemaHash = Struct->GetSchemaHash(bSkipEditorOnly);
+		Builder.Update(&StructSchemaHash, sizeof(StructSchemaHash));
+	}
+}
+#endif
+
+
 #include "UObject/DefineUPropertyMacros.h"
\ No newline at end of file
diff --git a/Engine/Source/Runtime/CoreUObject/Public/UObject/Class.h b/Engine/Source/Runtime/CoreUObject/Public/UObject/Class.h
--- a/Engine/Source/Runtime/CoreUObject/Public/UObject/Class.h
+++ b/Engine/Source/Runtime/CoreUObject/Public/UObject/Class.h
@@ -30,6 +30,7 @@
 #include "UObject/UObjectGlobals.h"
 #include "UObject/FieldPath.h"
 
+struct FBlake3Hash;
 struct FCustomPropertyListNode;
 struct FFrame;
 struct FNetDeltaSerializeInfo;
@@ -348,6 +349,9 @@ public:
 #if WITH_EDITORONLY_DATA
 	/** Cached schema for optimized unversioned property serialization, with editor data, owned by this. */
 	mutable const struct FUnversionedStructSchema* UnversionedEditorSchema = nullptr;
+
+	/** Get the Schema Hash for this struct - the hash of its property names and types. */
+	const FBlake3Hash& GetSchemaHash(bool bSkipEditorOnly) const;
 #endif
 
 public:
diff --git a/Engine/Source/Runtime/CoreUObject/Public/UObject/EnumProperty.h b/Engine/Source/Runtime/CoreUObject/Public/UObject/EnumProperty.h
--- a/Engine/Source/Runtime/CoreUObject/Public/UObject/EnumProperty.h
+++ b/Engine/Source/Runtime/CoreUObject/Public/UObject/EnumProperty.h
@@ -50,6 +50,9 @@ public:
 	virtual int32 GetMinAlignment() const override;
 	virtual bool SameType(const FProperty* Other) const override;
 	virtual EConvertFromTypeResult ConvertFromType(const FPropertyTag& Tag, FStructuredArchive::FSlot Slot, uint8* Data, UStruct* DefaultsStruct) override;
+#if WITH_EDITORONLY_DATA
+	virtual void AppendSchemaHash(FBlake3& Builder, bool bSkipEditorOnly) const override;
+#endif
 	// End of FProperty interface
 
 	/**
diff --git a/Engine/Source/Runtime/CoreUObject/Public/UObject/UnrealType.h b/Engine/Source/Runtime/CoreUObject/Public/UObject/UnrealType.h
--- a/Engine/Source/Runtime/CoreUObject/Public/UObject/UnrealType.h
+++ b/Engine/Source/Runtime/CoreUObject/Public/UObject/UnrealType.h
@@ -37,6 +37,7 @@
 // WARNING: This should always be the last include in any file that needs it (except .generated.h)
 #include "UObject/UndefineUPropertyMacros.h"
 
+class FBlake3;
 class UPropertyWrapper;
 
 COREUOBJECT_API DECLARE_LOG_CATEGORY_EXTERN(LogType, Log, All);
@@ -409,6 +410,21 @@ public:
 		PortFlags |= EPropertyPortFlags::PPF_UseDeprecatedProperties; // Imports should always process deprecated properties
 		return ImportText_Internal( Buffer, Data, PortFlags, OwnerObject, ErrorText );
 	}
+
+#if WITH_EDITORONLY_DATA
+	/**
+	 * Updates the given HashBuilder with name and type information of this Property.
+	 * Contract: the hashed data is different from any property that serializes differently in Tagged Property Serialization.
+	 * If necessary to follow the contract, subclasses should override and add further information after calling
+	 * Super::AppendSchemaHash. e.g. FStructProperty needs to append the schema hash of its UStruct.
+	 * 
+	 * @param HashBuilder The builder to Update with property information
+	 * @param bSkipEditorOnly Used by subclasses with sub- properties that may be editor-only.
+	 *                        If true, sub- properties that are editor-only should not be appended to the hash.
+	 *                        This property's base data is appended without regard for bSkipEditorOnly.
+	 */
+	virtual void AppendSchemaHash(FBlake3& Builder, bool bSkipEditorOnly) const;
+#endif
 protected:
 	virtual const TCHAR* ImportText_Internal( const TCHAR* Buffer, void* Data, int32 PortFlags, UObject* OwnerObject, FOutputDevice* ErrorText ) const PURE_VIRTUAL(FProperty::ImportText,return NULL;);
 public:
@@ -1652,6 +1668,9 @@ class COREUOBJECT_API FByteProperty : public TProperty_Numeric<uint8>
 	virtual void ExportTextItem( FString& ValueStr, const void* PropertyValue, const void* DefaultValue, UObject* Parent, int32 PortFlags, UObject* ExportRootScope ) const override;
 	virtual const TCHAR* ImportText_Internal( const TCHAR* Buffer, void* Data, int32 PortFlags, UObject* Parent, FOutputDevice* ErrorText ) const override;
 	virtual EConvertFromTypeResult ConvertFromType(const FPropertyTag& Tag, FStructuredArchive::FSlot Slot, uint8* Data, UStruct* DefaultsStruct) override;
+#if WITH_EDITORONLY_DATA
+	virtual void AppendSchemaHash(FBlake3& Builder, bool bSkipEditorOnly) const override;
+#endif
 	// End of FProperty interface
 
 	// FNumericProperty interface.
@@ -2022,6 +2041,9 @@ public:
 	virtual void InitializeValueInternal( void* Dest ) const override;
 	virtual int32 GetMinAlignment() const override;
 	virtual EConvertFromTypeResult ConvertFromType(const FPropertyTag& Tag, FStructuredArchive::FSlot Slot, uint8* Data, UStruct* DefaultsStruct) override;
+#if WITH_EDITORONLY_DATA
+	virtual void AppendSchemaHash(FBlake3& Builder, bool bSkipEditorOnly) const override;
+#endif
 	// End of FProperty interface
 
 	// Emulate the CPP type API, see TPropertyTypeFundamentals
@@ -2601,6 +2623,9 @@ public:
 	virtual const TCHAR* ImportText_Internal( const TCHAR* Buffer, void* Data, int32 PortFlags, UObject* OwnerObject, FOutputDevice* ErrorText ) const override;
 	virtual bool SameType(const FProperty* Other) const override;
 	virtual bool Identical( const void* A, const void* B, uint32 PortFlags ) const override;
+#if WITH_EDITORONLY_DATA
+	virtual void AppendSchemaHash(FBlake3& Builder, bool bSkipEditorOnly) const override;
+#endif
 	// End of FProperty interface
 
 	virtual FString GetCPPTypeCustom(FString* ExtendedTypeText, uint32 CPPExportFlags, const FString& InnerNativeTypeName)  const override;
@@ -2779,6 +2804,9 @@ public:
 	virtual const TCHAR* ImportText_Internal( const TCHAR* Buffer, void* Data, int32 PortFlags, UObject* OwnerObject, FOutputDevice* ErrorText ) const override;
 	virtual bool ContainsObjectReference(TArray<const FStructProperty*>& EncounteredStructProps, EPropertyObjectReferenceType InReferenceType = EPropertyObjectReferenceType::Strong) const override;
 	virtual bool SameType(const FProperty* Other) const override;
+#if WITH_EDITORONLY_DATA
+	virtual void AppendSchemaHash(FBlake3& Builder, bool bSkipEditorOnly) const override;
+#endif
 	// End of FProperty interface
 
 	// UObject interface
@@ -3025,6 +3053,9 @@ public:
 
 	/** Called by ImportTextItem, but can also be used by a non-ArrayProperty whose ArrayDim is > 1. ArrayHelper should be supplied by ArrayProperties and nullptr for fixed-size arrays. */
 	static const TCHAR* ImportTextInnerItem(const TCHAR* Buffer, const FProperty* Inner, void* Data, int32 PortFlags, UObject* OwnerObject, FScriptArrayHelper* ArrayHelper = nullptr, FOutputDevice* ErrorText = (FOutputDevice*)GWarn);
+#if WITH_EDITORONLY_DATA
+	virtual void AppendSchemaHash(FBlake3& Builder, bool bSkipEditorOnly) const override;
+#endif
 
 private:
 	FORCEINLINE void SetElementSize()
@@ -3179,6 +3210,9 @@ public:
 	{
 		return ValueProp;
 	}
+#if WITH_EDITORONLY_DATA
+	virtual void AppendSchemaHash(FBlake3& Builder, bool bSkipEditorOnly) const override;
+#endif
 };
 
 // need to break this out a different type so that the DECLARE_CASTED_CLASS_INTRINSIC macro can digest the comma
@@ -3281,6 +3315,9 @@ public:
 		FScriptSet* Set = (FScriptSet*)InSet;
 		return (uint8*)Set->GetData(Index, SetLayout);
 	}
+#if WITH_EDITORONLY_DATA
+	virtual void AppendSchemaHash(FBlake3& Builder, bool bSkipEditorOnly) const override;
+#endif
 };
 
 /**
@@ -5054,6 +5091,9 @@ public:
 	virtual void EmitReferenceInfo(UClass& OwnerClass, int32 BaseOffset, TArray<const FStructProperty*>& EncounteredStructProps, FGCStackSizeHelper& StackSizeHelper) override;
 	virtual bool SameType(const FProperty* Other) const override;
 	virtual EConvertFromTypeResult ConvertFromType(const FPropertyTag& Tag, FStructuredArchive::FSlot Slot, uint8* Data, UStruct* DefaultsStruct) override;
+#if WITH_EDITORONLY_DATA
+	virtual void AppendSchemaHash(FBlake3& Builder, bool bSkipEditorOnly) const override;
+#endif
 	// End of FProperty interface
 
 	bool UseBinaryOrNativeSerialization(const FArchive& Ar) const;
@@ -5132,6 +5172,9 @@ public:
 	virtual void EmitReferenceInfo(UClass& OwnerClass, int32 BaseOffset, TArray<const FStructProperty*>& EncounteredStructProps, FGCStackSizeHelper& StackSizeHelper) override;
 	virtual void InstanceSubobjects( void* Data, void const* DefaultData, UObject* Owner, struct FObjectInstancingGraph* InstanceGraph ) override;
 	virtual bool SameType(const FProperty* Other) const override;
+#if WITH_EDITORONLY_DATA
+	virtual void AppendSchemaHash(FBlake3& Builder, bool bSkipEditorOnly) const override;
+#endif
 	// End of FProperty interface
 };
 
