GetSchemaHash: Fix buildbreak with programs that have WITH_EDITOR_DATA but not CACHE_UNVERSIONED_PROPERTY_SCHEMA.
#rnx

diff --git a/Engine/Source/Runtime/CoreUObject/Private/Serialization/UnversionedPropertySerialization.cpp b/Engine/Source/Runtime/CoreUObject/Private/Serialization/UnversionedPropertySerialization.cpp
--- a/Engine/Source/Runtime/CoreUObject/Private/Serialization/UnversionedPropertySerialization.cpp
+++ b/Engine/Source/Runtime/CoreUObject/Private/Serialization/UnversionedPropertySerialization.cpp
@@ -916,11 +916,17 @@ void SerializeUnversionedProperties(const UStruct* Struct, FStructuredArchive::F
 #if WITH_EDITORONLY_DATA
 const FBlake3Hash& GetSchemaHash(const UStruct* Struct, bool bSkipEditorOnly)
 {
+#if CACHE_UNVERSIONED_PROPERTY_SCHEMA
 	return GetOrCreateUnversionedSchema(Struct, bSkipEditorOnly).SchemaHash;
+#else
+	static FBlake3Hash Placeholder;
+	return Placeholder;
+#endif
 }
 
 COREUOBJECT_API void DumpClassSchemas(const TCHAR* Str, FOutputDevice& Ar)
 {
+#if CACHE_UNVERSIONED_PROPERTY_SCHEMA
 	TArray<FString> Lines;
 	TArray<UStruct*> Structs;
 	for (TObjectIterator<UStruct> It; It; ++It)
@@ -951,5 +957,6 @@ COREUOBJECT_API void DumpClassSchemas(const TCHAR* Str, FOutputDevice& Ar)
 	{
 		FFileHelper::SaveStringArrayToFile(Lines, *DumpFilename);
 	}
+#endif
 }
 #endif
