Fix AssetRegistry error "UpdateAssetData called with a change in ObjectPath" when cooking a package and moving an actor asset out of the external package into the main package.
#rnx

diff --git a/Engine/Source/Runtime/AssetRegistry/Private/AssetRegistry.cpp b/Engine/Source/Runtime/AssetRegistry/Private/AssetRegistry.cpp
--- a/Engine/Source/Runtime/AssetRegistry/Private/AssetRegistry.cpp
+++ b/Engine/Source/Runtime/AssetRegistry/Private/AssetRegistry.cpp
@@ -4097,6 +4097,11 @@ void UAssetRegistryImpl::ProcessLoadedAssetsToUpdateCache(UE::AssetRegistry::Imp
 		while (Index < CurrentBatchSize)
 		{
 			UObject* LoadedObject = BatchObjects[Index++];
+			if (!LoadedObject->IsAsset())
+			{
+				// If the object has changed and is no longer an asset, ignore it. This can happen when an Actor is modified during cooking to no longer have an external package
+				continue;
+			}
 			BatchAssetDatas.Add(FAssetData(LoadedObject, true /* bAllowBlueprintClass */));
 
 			// Check to see if we have run out of time in this tick
diff --git a/Engine/Source/Runtime/AssetRegistry/Private/AssetRegistryState.cpp b/Engine/Source/Runtime/AssetRegistry/Private/AssetRegistryState.cpp
--- a/Engine/Source/Runtime/AssetRegistry/Private/AssetRegistryState.cpp
+++ b/Engine/Source/Runtime/AssetRegistry/Private/AssetRegistryState.cpp
@@ -1623,7 +1623,7 @@ void FAssetRegistryState::UpdateAssetData(FAssetData* AssetData, const FAssetDat
 	}
 
 	// Update ObjectPath
-	if (AssetData->PackageName != NewAssetData.PackageName || AssetData->AssetName != NewAssetData.AssetName)
+	if (AssetData->ObjectPath != NewAssetData.ObjectPath)
 	{
 		int32 NumRemoved = CachedAssetsByObjectPath.Remove(AssetData->ObjectPath);
 		check(NumRemoved <= 1);
diff --git a/Engine/Source/Runtime/CoreUObject/Private/AssetRegistry/AssetData.cpp b/Engine/Source/Runtime/CoreUObject/Private/AssetRegistry/AssetData.cpp
--- a/Engine/Source/Runtime/CoreUObject/Private/AssetRegistry/AssetData.cpp
+++ b/Engine/Source/Runtime/CoreUObject/Private/AssetRegistry/AssetData.cpp
@@ -96,18 +96,18 @@ FAssetData::FAssetData(const UObject* InAsset, bool bAllowBlueprintClass)
 			InAsset = InClass->ClassGeneratedBy;
 		}
 
-		const UPackage* Outermost = InAsset->GetOutermost();
+		const UPackage* Package = InAsset->GetPackage();
 
-		PackageName = Outermost->GetFName();
-		PackagePath = FName(*FPackageName::GetLongPackagePath(Outermost->GetName()));
+		PackageName = Package->GetFName();
+		PackagePath = FName(*FPackageName::GetLongPackagePath(Package->GetName()));
 		AssetName = InAsset->GetFName();
 		AssetClass = InAsset->GetClass()->GetFName();
 		ObjectPath = FName(*InAsset->GetPathName());
 
 		InAsset->GetAssetRegistryTags(*this);
 
-		ChunkIDs = Outermost->GetChunkIDs();
-		PackageFlags = Outermost->GetPackageFlags();
+		ChunkIDs = Package->GetChunkIDs();
+		PackageFlags = Package->GetPackageFlags();
 	}
 }
 
