#rnx

diff --git a/Engine/Source/Runtime/CoreUObject/Public/UObject/GarbageCollection.h b/Engine/Source/Runtime/CoreUObject/Public/UObject/GarbageCollection.h
--- a/Engine/Source/Runtime/CoreUObject/Public/UObject/GarbageCollection.h
+++ b/Engine/Source/Runtime/CoreUObject/Public/UObject/GarbageCollection.h
@@ -31,8 +31,9 @@ DECLARE_STATS_GROUP(TEXT("Garbage Collection"), STATGROUP_GC, STATCAT_Advanced);
 
 /**
  * Do extra checks on GC'd function references to catch uninitialized pointers?
+ * These checks are possibly producing false positives now that our memory use is going over 128Gb = 2^39.
  */
-#define DO_POINTER_CHECKS_ON_GC WITH_EDITORONLY_DATA
+#define DO_POINTER_CHECKS_ON_GC 0 
 
 /*-----------------------------------------------------------------------------
 	Realtime garbage collection helper classes.
