Add support for objects that changed class in LinkerLoad.
- With external actors, it is now possible for an actor to change class and be saved without it's referencer objects to be resaved.
- That was impossible to happen in standard levels, as you would have resaved the level and all possible referencers were resaved at the same time.
- To fix that, change how exports are hashed in the linker, to be able to resolve imports solely by name (removed the class and class package filters).
- Object properties will properly filter incompatible objects based on their expected class.

diff --git a/Engine/Source/Runtime/CoreUObject/Private/UObject/LinkerLoad.cpp b/Engine/Source/Runtime/CoreUObject/Private/UObject/LinkerLoad.cpp
--- a/Engine/Source/Runtime/CoreUObject/Private/UObject/LinkerLoad.cpp
+++ b/Engine/Source/Runtime/CoreUObject/Private/UObject/LinkerLoad.cpp
@@ -405,13 +405,6 @@ static FTexture2DResourceMem* CreateResourceMem(int32 SizeX, int32 SizeY, int32
 	return ResourceMem;
 }
 
-static inline int32 HashNames(FName Object, FName Class, FName Package)
-{
-	return GetTypeHash(Object.GetComparisonIndex())
-		+ 7 * GetTypeHash(Class.GetComparisonIndex())
-		+ 31 * GetTypeHash(FPackageName::GetShortFName(Package).GetComparisonIndex());
-}
-
 static FORCEINLINE bool IsCoreUObjectPackage(const FName& PackageName)
 {
 	return PackageName == NAME_CoreUObject || PackageName == GLongCoreUObjectPackageName || PackageName == NAME_Core || PackageName == GLongCorePackageName;
@@ -906,6 +899,8 @@ FLinkerLoad::FLinkerLoad(UPackage* InParent, const FPackagePath& InPackagePath,
 , DeferredCDOIndex(INDEX_NONE)
 #endif // USE_CIRCULAR_DEPENDENCY_LOAD_DEFERRING
 {
+	static_assert((ExportHashCount & (ExportHashCount - 1)) == 0, "ExportHashCount must be power of two");
+
 	if (InPackagePath.GetHeaderExtension() == EPackageExtension::Unspecified)
 	{
 		UE_LOG(LogPackageName, Error, TEXT("PackagePath %s is missing header extension when assigned to LinkerLoad"), *InPackagePath.GetDebugName());
@@ -2447,7 +2442,7 @@ FLinkerLoad::ELinkerStatus FLinkerLoad::CreateExportHash()
 	{
 		FObjectExport& Export = ExportMap[ExportHashIndex];
 
-		const int32 iHash = HashNames( Export.ObjectName, GetExportClassName(ExportHashIndex), GetExportClassPackage(ExportHashIndex) ) & (ExportHashCount-1);
+		const int32 iHash = GetHashBucket( Export.ObjectName );
 		Export.HashNext = ExportHash[iHash];
 		ExportHash[iHash] = ExportHashIndex;
 
@@ -3446,38 +3441,7 @@ bool FLinkerLoad::VerifyImportInner(const int32 ImportIndex, FString& WarningSuf
 			Pkg = Import.SourceLinker->LinkerRoot;
 
 			// Find this import within its existing linker.
-			int32 iHash = HashNames( Import.ObjectName, Import.ClassName, Import.ClassPackage) & (ExportHashCount-1);
-
-			//@Package name transition, if we can match without shortening the names, then we must not take a shortened match
-			bool bMatchesWithoutShortening = false;
-			FName TestName = Import.ClassPackage;
-			
-			for( int32 j=Import.SourceLinker->ExportHash[iHash]; j!=INDEX_NONE; j=Import.SourceLinker->ExportMap[j].HashNext )
-			{
-				if (!Import.SourceLinker->ExportMap.IsValidIndex(j))
-				{
-					UE_LOG(LogLinker, Error, TEXT("Invalid index [%d/%d] while attempting to import '%s' with LinkerRoot '%s'"), j, Import.SourceLinker->ExportMap.Num(), *ImportObjectName, *GetNameSafe(Import.SourceLinker->LinkerRoot));
-					break;
-				}
-				else
-				{
-					FObjectExport& SourceExport = Import.SourceLinker->ExportMap[ j ];
-					if
-						(
-						SourceExport.ObjectName == Import.ObjectName
-						&&	Import.SourceLinker->GetExportClassName(j) == Import.ClassName
-						&&  Import.SourceLinker->GetExportClassPackage(j) == Import.ClassPackage 
-						)
-					{
-						bMatchesWithoutShortening = true;
-						break;
-					}
-				}
-			}
-			if (!bMatchesWithoutShortening)
-			{
-				TestName = FPackageName::GetShortFName(TestName);
-			}
+			int32 iHash = GetHashBucket( Import.ObjectName );
 
 			for( int32 j=Import.SourceLinker->ExportHash[iHash]; j!=INDEX_NONE; j=Import.SourceLinker->ExportMap[j].HashNext )
 			{
@@ -3491,9 +3455,9 @@ bool FLinkerLoad::VerifyImportInner(const int32 ImportIndex, FString& WarningSuf
 					FObjectExport& SourceExport = Import.SourceLinker->ExportMap[ j ];
 					if
 					(	
-						SourceExport.ObjectName==Import.ObjectName               
-						&&	Import.SourceLinker->GetExportClassName(j)==Import.ClassName
-						&&  (bMatchesWithoutShortening ? Import.SourceLinker->GetExportClassPackage(j) : FPackageName::GetShortFName(Import.SourceLinker->GetExportClassPackage(j))) == TestName 
+						SourceExport.ObjectName == Import.ObjectName
+						// If we are not explicitly looking for a redirector, skip for now as it will be properly handled in VerifyImport
+						&& ((Import.ClassName == NAME_ObjectRedirector) == (Import.SourceLinker->GetExportClassName(j) == NAME_ObjectRedirector))
 					)
 					{
 						// at this point, SourceExport is an FObjectExport in another linker that looks like it
@@ -3531,12 +3495,18 @@ bool FLinkerLoad::VerifyImportInner(const int32 ImportIndex, FString& WarningSuf
 									// if the import and its outer do not share a source linker, validate the import entry of the outer in the source linker matches otherwise skip resolveing the outer
 									check(SourceExport.OuterIndex.IsImport())
 									FObjectImport& SourceExportOuter = Import.SourceLinker->Imp(SourceExport.OuterIndex);
-									if (SourceExportOuter.ObjectName != OuterImport.ObjectName
-										|| SourceExportOuter.ClassName != OuterImport.ClassName
-										|| SourceExportOuter.ClassPackage != OuterImport.ClassPackage)
+									if (SourceExportOuter.ObjectName != OuterImport.ObjectName)
 									{
 										continue;
 									}
+									else
+									{
+										if ((SourceExportOuter.ClassName != OuterImport.ClassName) || (SourceExportOuter.ClassPackage != OuterImport.ClassPackage))
+										{
+											UE_LOG(LogLinker, Warning, TEXT("Resolved outer import with a different class: import class '%s.%s', package class '%s.%s'. Resave %s to fix."), 
+												*OuterImport.ClassPackage.ToString(), *OuterImport.ClassName.ToString(), *SourceExportOuter.ClassPackage.ToString(), *SourceExportOuter.ClassName.ToString(), *GetNameSafe(LinkerRoot));
+										}
+									}
 								}
 							}
 						}
@@ -3607,6 +3577,12 @@ bool FLinkerLoad::VerifyImportInner(const int32 ImportIndex, FString& WarningSuf
 						}
 
 						// Found the FObjectExport for this import
+						if ((Import.ClassName != Import.SourceLinker->GetExportClassName(j)) || (Import.ClassPackage != Import.SourceLinker->GetExportClassPackage(j)))
+						{
+							UE_LOG(LogLinker, Warning, TEXT("Resolved import with a different class: import class '%s.%s', package class '%s.%s'. Resave %s to fix."), 
+								*Import.ClassPackage.ToString(), *Import.ClassName.ToString(), *Import.SourceLinker->GetExportClassPackage(j).ToString(), *Import.SourceLinker->GetExportClassName(j).ToString(), *GetNameSafe(LinkerRoot));
+						}
+
 						Import.SourceIndex = j;
 						break;
 					}
@@ -3999,9 +3975,10 @@ UObject* FLinkerLoad::ResolveResource(FPackageIndex Index)
 }
 
 // Find the index of a specified object without regard to specific package.
+///@todo - this function is only used through FLinkerLoad::Create to perform package conform compatibility, which is deprecated.
 int32 FLinkerLoad::FindExportIndex( FName ClassName, FName ClassPackage, FName ObjectName, FPackageIndex ExportOuterIndex )
 {
-	int32 iHash = HashNames( ObjectName, ClassName, ClassPackage ) & (ExportHashCount-1);
+	int32 iHash = GetHashBucket( ObjectName );
 
 	for( int32 i=ExportHash[iHash]; i!=INDEX_NONE; i=ExportMap[i].HashNext )
 	{
@@ -4012,13 +3989,18 @@ int32 FLinkerLoad::FindExportIndex( FName ClassName, FName ClassPackage, FName O
 		else
 		{
 			if
-			(  (ExportMap[i].ObjectName  ==ObjectName                              )
-				&& (GetExportClassPackage(i) ==ClassPackage                            )
-				&& (GetExportClassName   (i) ==ClassName                               ) 
-				&& (ExportMap[i].OuterIndex  ==ExportOuterIndex 
+			(  (ExportMap[i].ObjectName == ObjectName)
+				&& ((GetExportClassName(i) == NAME_ObjectRedirector) == (ClassName == NAME_ObjectRedirector)) // If we are not explicitly looking for a redirector, skip for now as it will be properly handled in VerifyImport
+				&& (ExportMap[i].OuterIndex == ExportOuterIndex 
 				|| ExportOuterIndex.IsImport()) // this is very not legit to be passing INDEX_NONE into this function to mean "ignore"
 			)
 			{
+				if ((ClassPackage != GetExportClassPackage(i)) || (ClassName != GetExportClassName(i)))
+				{
+					UE_LOG(LogLinker, Warning, TEXT("Resolved export with a different class: export class '%s.%s', package class '%s.%s'. Resave %s to fix."), 
+						 *GetExportClassPackage(i).ToString(), *GetExportClassName(i).ToString(), *ClassPackage.ToString(), *ClassName.ToString(), *GetNameSafe(LinkerRoot));
+				}
+
 				return i;
 			}
 		}
@@ -4057,6 +4039,7 @@ int32 FLinkerLoad::FindExportIndex( FName ClassName, FName ClassPackage, FName O
  * @param Checked		Whether or not a failure will throw an error
  * @return The created object, or (UObject*)-1 if this is just verifying
  */
+///@todo - this function is only used to perform package conform compatibility (through ValidateConformCompatibility), which is deprecated.
 UObject* FLinkerLoad::Create( UClass* ObjectClass, FName ObjectName, UObject* Outer, uint32 InLoadFlags, bool Checked )
 {
 	// We no longer handle a NULL outer, which used to mean look in any outer, but we need fully qualified names now
diff --git a/Engine/Source/Runtime/CoreUObject/Public/UObject/LinkerLoad.h b/Engine/Source/Runtime/CoreUObject/Public/UObject/LinkerLoad.h
--- a/Engine/Source/Runtime/CoreUObject/Public/UObject/LinkerLoad.h
+++ b/Engine/Source/Runtime/CoreUObject/Public/UObject/LinkerLoad.h
@@ -226,6 +226,11 @@ public:
 	static constexpr int32 ExportHashCount = 256;
 	TUniquePtr<int32[]> ExportHash;
 
+	FORCEINLINE static int32 GetHashBucket(FName Object)
+	{
+		return GetTypeHash(Object.GetComparisonIndex()) & (ExportHashCount - 1);
+	}
+
 	/**
 	* List of imports and exports that must be serialized before other exports...all packed together, see FirstExportDependency
 	*/
