FUnversionedStructSchema: fix Address Sanitzer's warning about calling delete on a pointer that was not allocated with new.
#rnx

diff --git a/Engine/Source/Runtime/CoreUObject/Private/Serialization/UnversionedPropertySerialization.cpp b/Engine/Source/Runtime/CoreUObject/Private/Serialization/UnversionedPropertySerialization.cpp
--- a/Engine/Source/Runtime/CoreUObject/Private/Serialization/UnversionedPropertySerialization.cpp
+++ b/Engine/Source/Runtime/CoreUObject/Private/Serialization/UnversionedPropertySerialization.cpp
@@ -324,6 +324,15 @@ struct FUnversionedStructSchema
 		return Schema;
 	}
 
+	FORCEINLINE static void Delete(FUnversionedStructSchema* Schema)
+	{
+		if (Schema)
+		{
+			Schema->~FUnversionedStructSchema();
+			FMemory::Free(Schema);
+		}
+	}
+
 #if WITH_EDITORONLY_DATA
 	static FBlake3Hash CalculateSchemaHash(UStruct* Struct, bool bSkipEditorOnly)
 	{
@@ -361,7 +370,7 @@ const FUnversionedStructSchema& GetOrCreateUnversionedSchema(const UStruct* Stru
 	void** CachedSchemaPtr = reinterpret_cast<void**>(const_cast<FUnversionedStructSchema**>(&GetUnversionedSchema(Struct, bSkipEditorOnly)));
 	if (const FUnversionedStructSchema* ExistingSchema = reinterpret_cast<const FUnversionedStructSchema*>(FPlatformAtomics::InterlockedCompareExchangePointer(CachedSchemaPtr, CreatedSchema, nullptr)))
 	{
-		delete CreatedSchema;
+		FUnversionedStructSchema::Delete(CreatedSchema);
 		return *ExistingSchema;
 	}
 
@@ -812,10 +821,10 @@ bool CanUseUnversionedPropertySerialization(const ITargetPlatform* Target)
 void DestroyUnversionedSchema(const UStruct* Struct)
 {
 #if CACHE_UNVERSIONED_PROPERTY_SCHEMA
-	delete Struct->UnversionedGameSchema;
+	FUnversionedStructSchema::Delete(const_cast<FUnversionedStructSchema*>(Struct->UnversionedGameSchema));
 	Struct->UnversionedGameSchema = nullptr;
 #if WITH_EDITORONLY_DATA
-	delete Struct->UnversionedEditorSchema;
+	FUnversionedStructSchema::Delete(const_cast<FUnversionedStructSchema*>(Struct->UnversionedEditorSchema));
 	Struct->UnversionedEditorSchema = nullptr;
 #endif
 #endif
