Fixed FMemberReference::ResolveMember failing to resolve fields from a parent interface

This was due to TFieldIterator skipping interfaces by default, which is what the various FindField/Property functions use internally.

This change also cleans-up TFieldIterator so that you can control its behavior with a single flags value (rather than 3), via a new constructor that takes a EFieldIterationFlags value (now the default constructor), with the older-style EFieldIteratorFlags values being proxied from the older multi-parameter constructor.

diff --git a/Engine/Source/Runtime/CoreUObject/Public/UObject/UnrealType.h b/Engine/Source/Runtime/CoreUObject/Public/UObject/UnrealType.h
--- a/Engine/Source/Runtime/CoreUObject/Public/UObject/UnrealType.h
+++ b/Engine/Source/Runtime/CoreUObject/Public/UObject/UnrealType.h
@@ -5702,24 +5702,38 @@ TFieldIterator.
 -----------------------------------------------------------------------------*/
 
 /** TFieldIterator construction flags */
+enum class EFieldIterationFlags : uint8
+{
+	None = 0,
+	IncludeSuper = 1<<0,		// Include super class
+	IncludeDeprecated = 1<<1,	// Include deprecated properties
+	IncludeInterfaces = 1<<2,	// Include interfaces
+
+	IncludeAll = IncludeSuper | IncludeDeprecated | IncludeInterfaces,
+
+	Default = IncludeSuper | IncludeDeprecated,
+};
+ENUM_CLASS_FLAGS(EFieldIterationFlags);
+
+/** Old-style TFieldIterator construction flags */
 namespace EFieldIteratorFlags
 {
 	enum SuperClassFlags
 	{
-		ExcludeSuper = 0,	// Exclude super class
-		IncludeSuper		// Include super class
+		ExcludeSuper = (uint8)EFieldIterationFlags::None,
+		IncludeSuper = (uint8)EFieldIterationFlags::IncludeSuper,
 	};
 
 	enum DeprecatedPropertyFlags
 	{
-		ExcludeDeprecated = 0,	// Exclude deprecated properties
-		IncludeDeprecated		// Include deprecated properties
+		ExcludeDeprecated = (uint8)EFieldIterationFlags::None,
+		IncludeDeprecated = (uint8)EFieldIterationFlags::IncludeDeprecated,
 	};
 
 	enum InterfaceClassFlags
 	{
-		ExcludeInterfaces = 0,	// Exclude interfaces
-		IncludeInterfaces		// Include interfaces
+		ExcludeInterfaces = (uint8)EFieldIterationFlags::None,
+		IncludeInterfaces = (uint8)EFieldIterationFlags::IncludeInterfaces,
 	};
 }
 
@@ -5764,20 +5778,26 @@ private:
 	const bool bIncludeInterface;
 
 public:
-	TFieldIterator(const UStruct*                               InStruct,
-	               EFieldIteratorFlags::SuperClassFlags         InSuperClassFlags      = EFieldIteratorFlags::IncludeSuper,
-	               EFieldIteratorFlags::DeprecatedPropertyFlags InDeprecatedFieldFlags = EFieldIteratorFlags::IncludeDeprecated,
-	               EFieldIteratorFlags::InterfaceClassFlags     InInterfaceFieldFlags  = EFieldIteratorFlags::ExcludeInterfaces)
+	TFieldIterator(const UStruct* InStruct, EFieldIterationFlags InIterationFlags = EFieldIterationFlags::Default)
 		: Struct            ( InStruct )
 		, Field             ( InStruct ? GetChildFieldsFromStruct<typename T::BaseFieldClass>(InStruct) : NULL )
 		, InterfaceIndex    ( -1 )
-		, bIncludeSuper     ( InSuperClassFlags      == EFieldIteratorFlags::IncludeSuper )
-		, bIncludeDeprecated( InDeprecatedFieldFlags == EFieldIteratorFlags::IncludeDeprecated )
-		, bIncludeInterface ( InInterfaceFieldFlags  == EFieldIteratorFlags::IncludeInterfaces && InStruct && InStruct->IsA(UClass::StaticClass()) )
+		, bIncludeSuper     ( EnumHasAnyFlags(InIterationFlags, EFieldIterationFlags::IncludeSuper) )
+		, bIncludeDeprecated( EnumHasAnyFlags(InIterationFlags, EFieldIterationFlags::IncludeDeprecated) )
+		, bIncludeInterface ( EnumHasAnyFlags(InIterationFlags, EFieldIterationFlags::IncludeInterfaces) && InStruct && InStruct->IsA(UClass::StaticClass()) )
 	{
 		IterateToNext();
 	}
 
+	/** Legacy version taking the flags as 3 separate values */
+	TFieldIterator(const UStruct*                               InStruct,
+	               EFieldIteratorFlags::SuperClassFlags         InSuperClassFlags,
+	               EFieldIteratorFlags::DeprecatedPropertyFlags InDeprecatedFieldFlags = EFieldIteratorFlags::IncludeDeprecated,
+	               EFieldIteratorFlags::InterfaceClassFlags     InInterfaceFieldFlags  = EFieldIteratorFlags::ExcludeInterfaces)
+		: TFieldIterator(InStruct, (EFieldIterationFlags)(InSuperClassFlags | InDeprecatedFieldFlags | InInterfaceFieldFlags))
+	{
+	}
+
 	/** conversion to "bool" returning true if the iterator is valid. */
 	FORCEINLINE explicit operator bool() const
 	{ 
@@ -5880,11 +5900,17 @@ protected:
 template <typename T>
 struct TFieldRange
 {
+	TFieldRange(const UStruct* InStruct, EFieldIterationFlags InIterationFlags = EFieldIterationFlags::Default)
+		: Begin(InStruct, InIterationFlags)
+	{
+	}
+
+	/** Legacy version taking the flags as 3 separate values */
 	TFieldRange(const UStruct*                               InStruct,
-	            EFieldIteratorFlags::SuperClassFlags         InSuperClassFlags      = EFieldIteratorFlags::IncludeSuper,
+	            EFieldIteratorFlags::SuperClassFlags         InSuperClassFlags,
 	            EFieldIteratorFlags::DeprecatedPropertyFlags InDeprecatedFieldFlags = EFieldIteratorFlags::IncludeDeprecated,
 	            EFieldIteratorFlags::InterfaceClassFlags     InInterfaceFieldFlags  = EFieldIteratorFlags::ExcludeInterfaces)
-		: Begin(InStruct, InSuperClassFlags, InDeprecatedFieldFlags, InInterfaceFieldFlags)
+		: TFieldRange(InStruct, (EFieldIterationFlags)(InSuperClassFlags | InDeprecatedFieldFlags | InInterfaceFieldFlags))
 	{
 	}
 
@@ -5934,7 +5960,7 @@ T* FindField( const UStruct* Owner, const TCHAR* FieldName )
 }
 
 template <class T> 
-typename TEnableIf<TIsDerivedFrom<T, UField>::IsDerived, T*>::Type FindUField(const UStruct* Owner, FName FieldName)
+typename TEnableIf<TIsDerivedFrom<T, UField>::IsDerived, T*>::Type FindUField(const UStruct* Owner, FName FieldName, EFieldIterationFlags IterationFlags = EFieldIterationFlags::Default)
 {
 	static_assert(sizeof(T) > 0, "T must not be an incomplete type");
 
@@ -5945,7 +5971,7 @@ typename TEnableIf<TIsDerivedFrom<T, UField>::IsDerived, T*>::Type FindUField(co
 	}
 
 	// Search by comparing FNames (INTs), not strings
-	for (TFieldIterator<T>It(Owner); It; ++It)
+	for (TFieldIterator<T>It(Owner, IterationFlags); It; ++It)
 	{
 		if (It->GetFName() == FieldName)
 		{
@@ -5958,17 +5984,17 @@ typename TEnableIf<TIsDerivedFrom<T, UField>::IsDerived, T*>::Type FindUField(co
 }
 
 template <class T> 
-typename TEnableIf<TIsDerivedFrom<T, UField>::IsDerived, T*>::Type FindUField(const UStruct* Owner, const TCHAR* FieldName)
+typename TEnableIf<TIsDerivedFrom<T, UField>::IsDerived, T*>::Type FindUField(const UStruct* Owner, const TCHAR* FieldName, EFieldIterationFlags IterationFlags = EFieldIterationFlags::Default)
 {
 	static_assert(sizeof(T) > 0, "T must not be an incomplete type");
 
 	// lookup the string name in the Name hash
 	FName Name(FieldName, FNAME_Find);
-	return FindUField<T>(Owner, Name);
+	return FindUField<T>(Owner, Name, IterationFlags);
 }
 
 template <class T>
-typename TEnableIf<TIsDerivedFrom<T, FField>::IsDerived, T*>::Type FindFProperty(const UStruct* Owner, FName FieldName)
+typename TEnableIf<TIsDerivedFrom<T, FField>::IsDerived, T*>::Type FindFProperty(const UStruct* Owner, FName FieldName, EFieldIterationFlags IterationFlags = EFieldIterationFlags::Default)
 {
 	static_assert(sizeof(T) > 0, "T must not be an incomplete type");
 
@@ -5979,7 +6005,7 @@ typename TEnableIf<TIsDerivedFrom<T, FField>::IsDerived, T*>::Type FindFProperty
 	}
 
 	// Search by comparing FNames (INTs), not strings
-	for (TFieldIterator<T>It(Owner); It; ++It)
+	for (TFieldIterator<T>It(Owner, IterationFlags); It; ++It)
 	{
 		if (It->GetFName() == FieldName)
 		{
@@ -5992,33 +6018,33 @@ typename TEnableIf<TIsDerivedFrom<T, FField>::IsDerived, T*>::Type FindFProperty
 }
 
 template <class T>
-typename TEnableIf<TIsDerivedFrom<T, FField>::IsDerived, T*>::Type FindFProperty(const UStruct* Owner, const TCHAR* FieldName)
+typename TEnableIf<TIsDerivedFrom<T, FField>::IsDerived, T*>::Type FindFProperty(const UStruct* Owner, const TCHAR* FieldName, EFieldIterationFlags IterationFlags = EFieldIterationFlags::Default)
 {
 	static_assert(sizeof(T) > 0, "T must not be an incomplete type");
 
 	// lookup the string name in the Name hash
 	FName Name(FieldName, FNAME_Find);
-	return FindFProperty<T>(Owner, Name);
+	return FindFProperty<T>(Owner, Name, IterationFlags);
 }
 
 /** Finds FProperties or UFunctions and UEnums */
-inline FFieldVariant FindUFieldOrFProperty(const UStruct* Owner, FName FieldName)
+inline FFieldVariant FindUFieldOrFProperty(const UStruct* Owner, FName FieldName, EFieldIterationFlags IterationFlags = EFieldIterationFlags::Default)
 {
 	// Look for properties first as they're most often the runtime thing higher level code wants to find
-	FFieldVariant Result = FindFProperty<FProperty>(Owner, FieldName);
+	FFieldVariant Result = FindFProperty<FProperty>(Owner, FieldName, IterationFlags);
 	if (!Result)
 	{
-		Result = FindUField<UField>(Owner, FieldName);
+		Result = FindUField<UField>(Owner, FieldName, IterationFlags);
 	}
 	return Result;
 }
 
 /** Finds FProperties or UFunctions and UEnums */
-inline FFieldVariant FindUFieldOrFProperty(const UStruct* Owner, const TCHAR* FieldName)
+inline FFieldVariant FindUFieldOrFProperty(const UStruct* Owner, const TCHAR* FieldName, EFieldIterationFlags IterationFlags = EFieldIterationFlags::Default)
 {
 	// lookup the string name in the Name hash
 	FName Name(FieldName, FNAME_Find);
-	return FindUFieldOrFProperty(Owner, Name);
+	return FindUFieldOrFProperty(Owner, Name, IterationFlags);
 }
 
 /**
diff --git a/Engine/Source/Runtime/Engine/Classes/Engine/MemberReference.h b/Engine/Source/Runtime/Engine/Classes/Engine/MemberReference.h
--- a/Engine/Source/Runtime/Engine/Classes/Engine/MemberReference.h
+++ b/Engine/Source/Runtime/Engine/Classes/Engine/MemberReference.h
@@ -372,7 +372,7 @@ public:
 			UStruct* MemberScopeStruct = FindUField<UStruct>(SelfScope, *MemberScope);
 
 			// Find in target scope
-			ReturnField = FindUFieldOrFProperty(MemberScopeStruct, MemberName).Get<TFieldType>();
+			ReturnField = FindUFieldOrFProperty(MemberScopeStruct, MemberName, EFieldIterationFlags::IncludeAll).Get<TFieldType>();
 
 #if WITH_EDITOR
 			if(ReturnField == nullptr)
@@ -381,7 +381,7 @@ public:
 				const FName RenamedMemberName = RefreshLocalVariableName(SelfScope);
 				if (RenamedMemberName != NAME_None)
 				{
-					ReturnField = FindUFieldOrFProperty(MemberScopeStruct, MemberName).Get<TFieldType>();
+					ReturnField = FindUFieldOrFProperty(MemberScopeStruct, MemberName, EFieldIterationFlags::IncludeAll).Get<TFieldType>();
 				}
 			}
 #endif
@@ -431,11 +431,11 @@ public:
 					UScriptStruct* SparseClassDataStruct = TargetScope->GetSparseClassDataStruct();
 					if (SparseClassDataStruct)
 					{
-						ReturnField = FindUFieldOrFProperty(SparseClassDataStruct, MemberName).Get<TFieldType>();
+						ReturnField = FindUFieldOrFProperty(SparseClassDataStruct, MemberName, EFieldIterationFlags::IncludeAll).Get<TFieldType>();
 					}
 					if (ReturnField == nullptr)
 					{
-						ReturnField = FindUFieldOrFProperty(TargetScope, MemberName).Get<TFieldType>();
+						ReturnField = FindUFieldOrFProperty(TargetScope, MemberName, EFieldIterationFlags::IncludeAll).Get<TFieldType>();
 					}
 				}
 
@@ -454,7 +454,7 @@ public:
 					if (RenamedMemberName != NAME_None)
 					{
 						MemberName = RenamedMemberName;
-						ReturnField = FindUFieldOrFProperty(TargetScope, MemberName).Get<TFieldType>();
+						ReturnField = FindUFieldOrFProperty(TargetScope, MemberName, EFieldIterationFlags::IncludeAll).Get<TFieldType>();
 					}
 				}
 #endif
