Enabling the use of Virtualized BulkData api for both static meshes and textures. NOTE: This will not start virtualizing content data
This will cause the old bulkdata objects to be upgraded when loaded in the editor and saved out in the new format when the package is next saved.

#rnx

* UE5CookerObjectVersion
- Note that this work was done in CL 16266112 (it was easier to update this on the DevCooker branch)
- Moved UE_USE_VIRTUALBULKDATA to VirtualizationManager.
- Moved the virtualization versions to UE5MainStreamObjectVersion.
- Note that no package would have been saved with UE_USE_VIRTUALBULKDATA so it is safe to just move the versions.

* UE5MainStreamObjectVersion
- Added the virtialization versions here
- Fixed up some white space issues

* VirtualizationManager
- UE_USE_VIRTUALBULKDATA is now controlled from here.
- When disabled (set to 0) a new code path UE_VBD_TO_OLD_BULKDATA_PATH will be enabled that will attempt to allow packages saved with VBD to be reverted safely.
-- This requires that a new version be added to FUE5MainStreamObjectVersion, there is a static assert to warn people about this.

* VirtualizedBulkData
- Added a new method ::ConvertToOldBulkData which will load the VBD payload from disk and then pass it to the old bulkdata object provided. This is not as efficient as the OldBulkData->VBD path but is only provided as an emergancy way for us to back out if needed, we do not expect this code to actually be used.
- Fixed a bug where converting an empty bulkdata object to VBD would give a warning about invalid guids. If the provided bulkdata object is empty then it is likely that we will not have generated a unique guid for it anyway. A VBD object with no payload would expect to have an unset FPayloadId AND an invalid bulkdata identifier anyway so we can just do that and only warn in the case that a valid payload has no valid guid provided. This has a bonus of making the code easier to read as well.

* StaticMeshSourceData
- Updated MESHDATAKEY_STATICMESH_DERIVEDDATA_VER when UE_USE_VIRTUALBULKDATA is enabled again.
-- This probably isn't required but in theory some of our tests could've written mesh description data to the DDC with the old key, so we should update it to play it safe.

* Texture
- Converted to use FUE5MainStreamObjectVersion for it's version number
- Added a code path that will allow us to convert back to using old style bulkdata if UE_USE_VIRTUALBULKDATA is disabled and UE_VBD_TO_OLD_BULKDATA_PATH is enabled.

*MeshDescription
- Converted to use FUE5MainStreamObjectVersion for it's version number
-- Note that we don't have to tell the archive that we are using this because this was already happening.
- Added a code path that will allow us to convert back to using old style bulkdata if UE_USE_VIRTUALBULKDATA is disabled and UE_VBD_TO_OLD_BULKDATA_PATH is enabled.

diff --git a/Engine/Source/Runtime/Core/Public/UObject/UE5MainStreamObjectVersion.h b/Engine/Source/Runtime/Core/Public/UObject/UE5MainStreamObjectVersion.h
--- a/Engine/Source/Runtime/Core/Public/UObject/UE5MainStreamObjectVersion.h
+++ b/Engine/Source/Runtime/Core/Public/UObject/UE5MainStreamObjectVersion.h
@@ -104,7 +104,7 @@ struct CORE_API FUE5MainStreamObjectVersion
 		BlendSpaceRuntimeTriangulation,
 
 		// Fix to the Cubic smoothing, plus introduction of new smoothing types
-        BlendSpaceSmoothingImprovements,
+		BlendSpaceSmoothingImprovements,
 
 		// Removing Tessellation parameters from Materials
 		RemovingTessellationParameters,
@@ -133,7 +133,13 @@ struct CORE_API FUE5MainStreamObjectVersion
 		// Changed how actor descriptors compute their bounds
 		WorldPartitionActorDescGetStreamingBounds,
 
-        // -----<new versions can be added above this line>-------------------------------------------------
+		// Switch FMeshDescriptionBulkData to use virtualized bulkdata
+		MeshDescriptionVirtualization,
+		
+		// Switch FTextureSource to use virtualized bulkdata
+		TextureSourceVirtualization,
+
+		// -----<new versions can be added above this line>-------------------------------------------------
 		VersionPlusOne,
 		LatestVersion = VersionPlusOne - 1
 	};
diff --git a/Engine/Source/Runtime/CoreUObject/Private/Virtualization/VirtualizedBulkData.cpp b/Engine/Source/Runtime/CoreUObject/Private/Virtualization/VirtualizedBulkData.cpp
--- a/Engine/Source/Runtime/CoreUObject/Private/Virtualization/VirtualizedBulkData.cpp
+++ b/Engine/Source/Runtime/CoreUObject/Private/Virtualization/VirtualizedBulkData.cpp
@@ -111,33 +111,28 @@ void FVirtualizedUntypedBulkData::CreateFromBulkData(FUntypedBulkData& InBulkDat
 
 	Reset();
 
-	auto AssignGuid = [this, &InBulkData](const FGuid& InGuid)
+	// We only need to set up the bulkdata/content identifiers if we have a valid payload
+	if (InBulkData.GetBulkDataSize() > 0)
 	{
-		// Only use the guid if we will have a valid payload
-		if (InBulkData.GetBulkDataSize() > 0)
+		if (InGuid.IsValid())
 		{
 			BulkDataId = InGuid;
 		}
+		else
+		{
+			UE_LOG(LogVirtualization, Warning,
+				TEXT("CreateFromBulkData recieved an invalid FGuid. A temporary one will be generated until the package is next re-saved! Package: '%s'"),
+				*InBulkData.GetPackagePath().GetDebugName());
 
-		PayloadContentId = FPayloadId(InGuid);
-	};
+			BulkDataId = FGuid::NewGuid();
+		}
 
-	if (InGuid.IsValid())
-	{
-		AssignGuid(InGuid);
+		PayloadContentId = FPayloadId(InGuid);
+		bWasKeyGuidDerived = true;
 	}
-	else
-	{
-		UE_LOG(LogVirtualization, Warning,
-			TEXT("CreateFromBulkData recieved an invalid FGuid. A temporary one will be generated until the package is next re-saved! Package: '%s'"),
-			*InBulkData.GetPackagePath().GetDebugName());
-
-		AssignGuid(FGuid::NewGuid());
-	}	
 	
 	PayloadSize = InBulkData.GetBulkDataSize();
 	
-	bWasKeyGuidDerived = true;
 	PackagePath = InBulkData.GetPackagePath();
 	PackageSegment = InBulkData.GetPackageSegment();
 	
@@ -157,6 +152,36 @@ void FVirtualizedUntypedBulkData::CreateFromBulkData(FUntypedBulkData& InBulkDat
 	}
 }
 
+#if UE_VBD_TO_OLD_BULKDATA_PATH
+void FVirtualizedUntypedBulkData::ConvertToOldBulkData(FUntypedBulkData& BulkData)
+{
+	TRACE_CPUPROFILER_EVENT_SCOPE(FVirtualizedUntypedBulkData::ConvertToOldBulkData);
+
+	FSharedBuffer UncompressedPayload = GetDataInternal().Decompress();
+
+	checkf(!IsDataVirtualized(), TEXT("Trying to convert data that has been virtualized, this system should've been removed before this could occur!"));
+	checkf(UncompressedPayload.GetSize() % BulkData.GetElementSize() == 0, TEXT("Payload does not match the element size for bulkdata!"));
+
+	const int64 ElementCount = UncompressedPayload.GetSize() / BulkData.GetElementSize();
+
+	BulkData.Lock(LOCK_READ_WRITE);
+
+	// It would be faster if we could 'release' the FSharedBuffer data to bulkdata but neither API supports
+	// this sort of thing. Give this is throw away code that we don't expect to have run it is not worth 
+	// adding support if this is the only usecase.
+	void* BulkDataPtr = BulkData.Realloc(ElementCount);
+	FMemory::Memcpy(BulkDataPtr, UncompressedPayload.GetData(), UncompressedPayload.GetSize());
+
+	BulkData.Unlock();
+
+	// If the virtual bulkdata has not disabled compression we should apply it to the old bulkdata
+	if (!EnumHasAnyFlags(Flags, EFlags::DisablePayloadCompression))
+	{
+		BulkData.StoreCompressedOnDisk(NAME_Zlib);
+	}
+}
+#endif // UE_VBD_TO_OLD_BULKDATA_PATH
+
 void FVirtualizedUntypedBulkData::Serialize(FArchive& Ar, UObject* Owner)
 {
 	TRACE_CPUPROFILER_EVENT_SCOPE(FVirtualizedUntypedBulkData::Serialize);
diff --git a/Engine/Source/Runtime/CoreUObject/Public/Virtualization/VirtualizationManager.h b/Engine/Source/Runtime/CoreUObject/Public/Virtualization/VirtualizationManager.h
--- a/Engine/Source/Runtime/CoreUObject/Public/Virtualization/VirtualizationManager.h
+++ b/Engine/Source/Runtime/CoreUObject/Public/Virtualization/VirtualizationManager.h
@@ -7,7 +7,22 @@
 #include "Logging/LogMacros.h"
 #include "Virtualization/PayloadId.h"
 
-#define UE_USE_VIRTUALBULKDATA 0
+#define UE_USE_VIRTUALBULKDATA 1
+
+// Enables a code path that will allow VBD to be converted back to old style bulkdata if we need to revert for some reason
+#if !UE_USE_VIRTUALBULKDATA
+	#define UE_VBD_TO_OLD_BULKDATA_PATH 1
+
+	// When reverting to old bulkdata we will need to add a new version to FUE5MainStreamObjectVersion, but since it is unlikely
+	// that we will ever do this, I'd prefer not to submit it to FUE5MainStreamObjectVersion in a disabled form.
+	// This static assert will remind anyone turning on this path that the version needs to be updated and hopefully we can just
+	// delete this without ever enabling it once we take the plunge and strip out UE_USE_VIRTUALBULKDATA.
+	static_assert(false, "Add the following two lines to FUE5MainStreamObjectVersion!");
+	//	// Disabled content virtualization due to crippling problems, sadly StaticMesh and Textures saved with virtualization will need reverting.
+	//	DisabledVirtualization,
+#else
+	#define UE_VBD_TO_OLD_BULKDATA_PATH 0
+#endif //UE_VBD_TO_OLD_BULKDATA_PATH
 
 // TODO: Do we want to keep this header public for UE5 release? If the only interaction should be
 //		 via FVirtualizedUntypedBulkData or other such classes we might not want to expose this 
diff --git a/Engine/Source/Runtime/CoreUObject/Public/Virtualization/VirtualizedBulkData.h b/Engine/Source/Runtime/CoreUObject/Public/Virtualization/VirtualizedBulkData.h
--- a/Engine/Source/Runtime/CoreUObject/Public/Virtualization/VirtualizedBulkData.h
+++ b/Engine/Source/Runtime/CoreUObject/Public/Virtualization/VirtualizedBulkData.h
@@ -68,6 +68,19 @@ public:
 	 */
 	void CreateFromBulkData(FUntypedBulkData& BulkData, const FGuid& Guid);
 
+#if UE_VBD_TO_OLD_BULKDATA_PATH
+	/**
+	 * Convenience method to make it easier to convert from FVirtualizedBulkData to BulkData. This is not very
+	 * efficient as it will load the payload from disk and the old bulkdata object will hold it in memory. It 
+	 * is provided in case virtualized bulkdata needs to be disabled once it goes live and we need a way to 
+	 * re-save packages back to the older format.
+	 * 
+	 * @param BulkData	The bulkdata object to be provided with the payload that the virtualized bulkdata object
+	 *					represents.
+	 */
+	void ConvertToOldBulkData(FUntypedBulkData& BulkData);
+#endif //UE_VBD_TO_OLD_BULKDATA_PATH
+
 	void Serialize(FArchive& Ar, UObject* Owner);
 
 	/** Reset to a truly empty state */
diff --git a/Engine/Source/Runtime/Engine/Private/StaticMeshSourceData.cpp b/Engine/Source/Runtime/Engine/Private/StaticMeshSourceData.cpp
--- a/Engine/Source/Runtime/Engine/Private/StaticMeshSourceData.cpp
+++ b/Engine/Source/Runtime/Engine/Private/StaticMeshSourceData.cpp
@@ -361,7 +361,7 @@ void FStaticMeshSourceModel::FillMaterialName(TMap<int32, FName>& OutMaterialMap
 // If static mesh derived data needs to be rebuilt (new format, serialization differences, etc.) replace the version GUID below with a new one.
 // In case of merge conflicts with DDC versions, you MUST generate a new GUID and set this new GUID as the version.
 #if UE_USE_VIRTUALBULKDATA
-#define MESHDATAKEY_STATICMESH_DERIVEDDATA_VER TEXT("80264C99E2AD4D83A1DDBBB8D93D8AEF")
+#define MESHDATAKEY_STATICMESH_DERIVEDDATA_VER TEXT("97E1F285E9824AD83693BB8BF976461")
 #else
 #define MESHDATAKEY_STATICMESH_DERIVEDDATA_VER TEXT("B3CE4F8E59E542D4A721EE56068AD214")
 #endif
diff --git a/Engine/Source/Runtime/Engine/Private/Texture.cpp b/Engine/Source/Runtime/Engine/Private/Texture.cpp
--- a/Engine/Source/Runtime/Engine/Private/Texture.cpp
+++ b/Engine/Source/Runtime/Engine/Private/Texture.cpp
@@ -28,7 +28,7 @@
 #include "TextureDerivedDataTask.h"
 #include "Interfaces/ITextureFormat.h"
 #include "Interfaces/ITextureFormatModule.h"
-#include "UObject/UE5CookerObjectVersion.h"
+#include "UObject/UE5MainStreamObjectVersion.h"
 
 #if WITH_EDITOR
 #include "TextureCompiler.h"
@@ -454,9 +454,7 @@ void UTexture::PostEditChangeProperty(FPropertyChangedEvent& PropertyChangedEven
 
 void UTexture::Serialize(FArchive& Ar)
 {
-#if UE_USE_VIRTUALBULKDATA
-	Ar.UsingCustomVersion(FUE5CookerObjectVersion::GUID);
-#endif //UE_USE_VIRTUALBULKDATA
+	Ar.UsingCustomVersion(FUE5MainStreamObjectVersion::GUID);
 
 	Super::Serialize(Ar);
 
@@ -470,18 +468,31 @@ void UTexture::Serialize(FArchive& Ar)
 		FWriteScopeLock BulkDataExclusiveScope(Source.BulkDataLock.Get());
 #endif
 
-#if UE_USE_VIRTUALBULKDATA
-		if (Ar.IsLoading() && Ar.CustomVer(FUE5CookerObjectVersion::GUID) < FUE5CookerObjectVersion::TextureSourceVirtualization)
+		if (Ar.IsLoading() && Ar.CustomVer(FUE5MainStreamObjectVersion::GUID) < FUE5MainStreamObjectVersion::TextureSourceVirtualization)
 		{
+#if UE_USE_VIRTUALBULKDATA
 			FByteBulkData TempBulkData;
 			TempBulkData.Serialize(Ar, this);
 
 			Source.BulkData.CreateFromBulkData(TempBulkData, Source.GetId());
+#else
+			Source.BulkData.Serialize(Ar, this);
+#endif //UE_USE_VIRTUALBULKDATA
 		}
 		else
-#endif //UE_USE_VIRTUALBULKDATA
 		{
-			Source.BulkData.Serialize(Ar, this);
+#if !UE_USE_VIRTUALBULKDATA && UE_VBD_TO_OLD_BULKDATA_PATH
+			if (Ar.IsLoading() && Ar.CustomVer(FUE5MainStreamObjectVersion::GUID) < FUE5MainStreamObjectVersion::DisabledVirtualization)
+			{
+				UE::Virtualization::FByteVirtualizedBulkData TempVirtualBulkData;
+				TempVirtualBulkData.Serialize(Ar, this);
+				TempVirtualBulkData.ConvertToOldBulkData(Source.BulkData);
+			}
+			else
+#endif // !UE_USE_VIRTUALBULKDATA && UE_VBD_TO_OLD_BULKDATA_PATH
+			{
+				Source.BulkData.Serialize(Ar, this);
+			}
 		}
 	}
 
diff --git a/Engine/Source/Runtime/MeshDescription/Private/MeshDescription.cpp b/Engine/Source/Runtime/MeshDescription/Private/MeshDescription.cpp
--- a/Engine/Source/Runtime/MeshDescription/Private/MeshDescription.cpp
+++ b/Engine/Source/Runtime/MeshDescription/Private/MeshDescription.cpp
@@ -1960,9 +1960,6 @@ void FMeshDescriptionBulkData::Serialize( FArchive& Ar, UObject* Owner )
 	Ar.UsingCustomVersion( FReleaseObjectVersion::GUID );
 	Ar.UsingCustomVersion( FEnterpriseObjectVersion::GUID );
 	Ar.UsingCustomVersion( FUE5MainStreamObjectVersion::GUID );
-#if UE_USE_VIRTUALBULKDATA
-	Ar.UsingCustomVersion(FUE5CookerObjectVersion::GUID);
-#endif //UE_USE_VIRTUALBULKDATA
 
 	// Make sure to serialize only actual data
 	if (Ar.ShouldSkipBulkData() || Ar.IsObjectReferenceCollector())
@@ -2007,19 +2004,32 @@ void FMeshDescriptionBulkData::Serialize( FArchive& Ar, UObject* Owner )
 		}
 	}
 
-#if UE_USE_VIRTUALBULKDATA
 	bool bSerializedOldDataTypes = false;
 	FByteBulkData TempBulkData;
-	if (Ar.IsLoading() && Ar.CustomVer(FUE5CookerObjectVersion::GUID) < FUE5CookerObjectVersion::MeshDescriptionVirtualization)
+	if (Ar.IsLoading() && Ar.CustomVer(FUE5MainStreamObjectVersion::GUID) < FUE5MainStreamObjectVersion::MeshDescriptionVirtualization)
 	{
+#if UE_USE_VIRTUALBULKDATA
 		// Serialize the old BulkData format and mark that we require a conversion after the guid has been serialized
 		TempBulkData.Serialize(Ar, Owner);
 		bSerializedOldDataTypes = true;
+#else
+		BulkData.Serialize(Ar, Owner);
+#endif //UE_USE_VIRTUALBULKDATA
 	}
 	else
-#endif //UE_USE_VIRTUALBULKDATA	
 	{
-		BulkData.Serialize(Ar, Owner);
+#if !UE_USE_VIRTUALBULKDATA && UE_VBD_TO_OLD_BULKDATA_PATH
+		if (Ar.IsLoading() && Ar.CustomVer(FUE5MainStreamObjectVersion::GUID) < FUE5MainStreamObjectVersion::DisabledVirtualization)
+		{
+			UE::Virtualization::FByteVirtualizedBulkData TempVirtualBulkData;
+			TempVirtualBulkData.Serialize(Ar, Owner);
+			TempVirtualBulkData.ConvertToOldBulkData(BulkData);
+		}
+		else
+#endif // !UE_USE_VIRTUALBULKDATA && UE_VBD_TO_OLD_BULKDATA_PATH
+		{
+			BulkData.Serialize(Ar, Owner);
+		}
 	}
 
 	if( Ar.IsLoading() && Ar.CustomVer( FEditorObjectVersion::GUID ) < FEditorObjectVersion::MeshDescriptionBulkDataGuid )
