Don't try to resolve manually selected debug target if none.
Saves up to 8ms per frame running demo project in PIE.

diff --git a/Engine/Source/Editor/UnrealEd/Private/Kismet2/KismetDebugUtilities.cpp b/Engine/Source/Editor/UnrealEd/Private/Kismet2/KismetDebugUtilities.cpp
--- a/Engine/Source/Editor/UnrealEd/Private/Kismet2/KismetDebugUtilities.cpp
+++ b/Engine/Source/Editor/UnrealEd/Private/Kismet2/KismetDebugUtilities.cpp
@@ -234,7 +234,7 @@ void FKismetDebugUtilities::OnScriptException(const UObject* ActiveObject, const
 		UWorld* WorldBeingDebugged = BlueprintObj->GetWorldBeingDebugged();
 		const FString& PathToDebug = BlueprintObj->GetObjectPathToDebug();
 		
-		if (ObjectBeingDebugged == nullptr)
+		if (ObjectBeingDebugged == nullptr && !PathToDebug.IsEmpty())
 		{
 			// Check if we need to update the object being debugged
 			UObject* ObjectToDebug = FindObjectSafe<UObject>(nullptr, *PathToDebug);
diff --git a/Engine/Source/Runtime/CoreUObject/Private/UObject/ScriptCore.cpp b/Engine/Source/Runtime/CoreUObject/Private/UObject/ScriptCore.cpp
--- a/Engine/Source/Runtime/CoreUObject/Private/UObject/ScriptCore.cpp
+++ b/Engine/Source/Runtime/CoreUObject/Private/UObject/ScriptCore.cpp
@@ -2282,6 +2282,7 @@ DEFINE_FUNCTION(UObject::execBreakpoint)
 {
 	if (FBlueprintCoreDelegates::IsDebuggingEnabled())
 	{
+		TRACE_CPUPROFILER_EVENT_SCOPE(UObject::execBreakpoint);
 		FBlueprintExceptionInfo BreakpointExceptionInfo(EBlueprintExceptionType::Breakpoint);
 		FBlueprintCoreDelegates::ThrowScriptException(P_THIS, Stack, BreakpointExceptionInfo);
 	}
@@ -2292,6 +2293,7 @@ DEFINE_FUNCTION(UObject::execTracepoint)
 {
 	if (FBlueprintCoreDelegates::IsDebuggingEnabled())
 	{
+		TRACE_CPUPROFILER_EVENT_SCOPE(UObject::execTracepoint);
 		FBlueprintExceptionInfo TracepointExceptionInfo(EBlueprintExceptionType::Tracepoint);
 		FBlueprintCoreDelegates::ThrowScriptException(P_THIS, Stack, TracepointExceptionInfo);
 	}
@@ -2302,6 +2304,7 @@ DEFINE_FUNCTION(UObject::execWireTracepoint)
 {
 	if (FBlueprintCoreDelegates::IsDebuggingEnabled())
 	{
+		TRACE_CPUPROFILER_EVENT_SCOPE(UObject::execWireTracepoint);
 		FBlueprintExceptionInfo TracepointExceptionInfo(EBlueprintExceptionType::WireTracepoint);
 		FBlueprintCoreDelegates::ThrowScriptException(P_THIS, Stack, TracepointExceptionInfo);
 	}
