Added profile tag in ProcessAsyncLoading.

diff --git a/Engine/Source/Runtime/CoreUObject/Private/Serialization/AsyncPackageLoader.cpp b/Engine/Source/Runtime/CoreUObject/Private/Serialization/AsyncPackageLoader.cpp
--- a/Engine/Source/Runtime/CoreUObject/Private/Serialization/AsyncPackageLoader.cpp
+++ b/Engine/Source/Runtime/CoreUObject/Private/Serialization/AsyncPackageLoader.cpp
@@ -662,6 +662,7 @@ int32 GetNumAsyncPackages()
 EAsyncPackageState::Type ProcessAsyncLoading(bool bUseTimeLimit, bool bUseFullTimeLimit, float TimeLimit)
 {
 	LLM_SCOPE(ELLMTag::AsyncLoading);
+	TRACE_CPUPROFILER_EVENT_SCOPE(ProcessAsyncLoading);
 	return GetAsyncPackageLoader().ProcessLoading(bUseTimeLimit, bUseFullTimeLimit, TimeLimit);
 }
 
