Optimize FPackageName::TryGetMountPointForPath by caching mount point absolute paths at creation time

diff --git a/Engine/Source/Runtime/CoreUObject/Private/Misc/PackageName.cpp b/Engine/Source/Runtime/CoreUObject/Private/Misc/PackageName.cpp
--- a/Engine/Source/Runtime/CoreUObject/Private/Misc/PackageName.cpp
+++ b/Engine/Source/Runtime/CoreUObject/Private/Misc/PackageName.cpp
@@ -159,15 +159,19 @@ struct FPathPair
 	// The physical relative path (e.g., "../../../Engine/Content/")
 	const FString ContentPath;
 
+	// The physical absolute path
+	const FString AbsolutePath;
+
 	bool operator ==(const FPathPair& Other) const
 	{
 		return RootPath == Other.RootPath && ContentPath == Other.ContentPath;
 	}
 
 	// Construct a path pair
-	FPathPair(const FString& InRootPath, const FString& InContentPath)
+	FPathPair(const FString& InRootPath, const FString& InContentPath, const FString& InAbsolutePath)
 		: RootPath(InRootPath)
 		, ContentPath(InContentPath)
+		, AbsolutePath(InAbsolutePath)
 	{
 	}
 };
@@ -262,7 +266,10 @@ struct FLongPackagePathsSingleton
 			RelativeContentPath += TEXT( "/" );
 		}
 
-		FPathPair Pair(RootPath, RelativeContentPath);
+		TStringBuilder<256> AbsolutePath;
+		FPathViews::ToAbsolutePath(RelativeContentPath, AbsolutePath);
+		FPathPair Pair(RootPath, RelativeContentPath, AbsolutePath.ToString());
+		
 		{
 			FWriteScopeLock ScopeLock(ContentMountPointCriticalSection);
 			ContentRootToPath.Insert(Pair, 0);
@@ -291,7 +298,8 @@ struct FLongPackagePathsSingleton
 			FWriteScopeLock ScopeLock(ContentMountPointCriticalSection);
 			if ( MountPointRootPaths.Remove(RootPath) > 0 )
 			{
-				FPathPair Pair(RootPath, RelativeContentPath);
+				// Absolute path doesn't need to be computed here as we are removing mount point
+				FPathPair Pair(RootPath, RelativeContentPath, FString());
 				ContentRootToPath.Remove(Pair);
 				ContentPathToRoot.Remove(Pair);
 				MountPointRootPaths.Remove(RootPath);
@@ -345,35 +353,43 @@ private:
 		FWriteScopeLock ScopeLock(ContentMountPointCriticalSection);
 
 		ContentPathToRoot.Empty(13);
-		ContentPathToRoot.Emplace(EngineRootPath, EngineContentPath);
+		TStringBuilder<256> AbsolutePath;
+		auto AddPathPair = [&AbsolutePath](TArray<FPathPair>& Paths, const FString& RootPath, const FString& ContentPath)
+		{
+			AbsolutePath.Reset();
+			FPathViews::ToAbsolutePath(ContentPath, AbsolutePath);
+			Paths.Emplace(RootPath, ContentPath, AbsolutePath.ToString());
+		};
+		
+		AddPathPair(ContentPathToRoot, EngineRootPath, EngineContentPath);
 		if (FPaths::IsSamePath(GameContentPath, ContentPathShort))
 		{
-			ContentPathToRoot.Emplace(GameRootPath, ContentPathShort);
+			AddPathPair(ContentPathToRoot, GameRootPath, ContentPathShort);
 		}
 		else
 		{
-			ContentPathToRoot.Emplace(EngineRootPath, ContentPathShort);
+			AddPathPair(ContentPathToRoot, EngineRootPath, ContentPathShort);
 		}
-		ContentPathToRoot.Emplace(EngineRootPath, EngineShadersPath);
-		ContentPathToRoot.Emplace(EngineRootPath, EngineShadersPathShort);
-		ContentPathToRoot.Emplace(GameRootPath,   GameContentPath);
-		ContentPathToRoot.Emplace(ScriptRootPath, GameScriptPath);
-		ContentPathToRoot.Emplace(TempRootPath,   GameSavedPath);
-		ContentPathToRoot.Emplace(GameRootPath,   GameContentPathRebased);
-		ContentPathToRoot.Emplace(ScriptRootPath, GameScriptPathRebased);
-		ContentPathToRoot.Emplace(TempRootPath,   GameSavedPathRebased);
-		ContentPathToRoot.Emplace(ConfigRootPath, GameConfigPath);
+		AddPathPair(ContentPathToRoot, EngineRootPath, EngineShadersPath);
+		AddPathPair(ContentPathToRoot, EngineRootPath, EngineShadersPathShort);
+		AddPathPair(ContentPathToRoot, GameRootPath,   GameContentPath);
+		AddPathPair(ContentPathToRoot, ScriptRootPath, GameScriptPath);
+		AddPathPair(ContentPathToRoot, TempRootPath,   GameSavedPath);
+		AddPathPair(ContentPathToRoot, GameRootPath,   GameContentPathRebased);
+		AddPathPair(ContentPathToRoot, ScriptRootPath, GameScriptPathRebased);
+		AddPathPair(ContentPathToRoot, TempRootPath,   GameSavedPathRebased);
+		AddPathPair(ContentPathToRoot, ConfigRootPath, GameConfigPath);
 
 		ContentRootToPath.Empty(11);
-		ContentRootToPath.Emplace(EngineRootPath, EngineContentPath);
-		ContentRootToPath.Emplace(EngineRootPath, EngineShadersPath);
-		ContentRootToPath.Emplace(GameRootPath,   GameContentPath);
-		ContentRootToPath.Emplace(ScriptRootPath, GameScriptPath);
-		ContentRootToPath.Emplace(TempRootPath,   GameSavedPath);
-		ContentRootToPath.Emplace(GameRootPath,   GameContentPathRebased);
-		ContentRootToPath.Emplace(ScriptRootPath, GameScriptPathRebased);
-		ContentRootToPath.Emplace(TempRootPath,   GameSavedPathRebased);
-		ContentRootToPath.Emplace(ConfigRootPath, GameConfigPathRebased);
+		AddPathPair(ContentRootToPath, EngineRootPath, EngineContentPath);
+		AddPathPair(ContentRootToPath, EngineRootPath, EngineShadersPath);
+		AddPathPair(ContentRootToPath, GameRootPath,   GameContentPath);
+		AddPathPair(ContentRootToPath, ScriptRootPath, GameScriptPath);
+		AddPathPair(ContentRootToPath, TempRootPath,   GameSavedPath);
+		AddPathPair(ContentRootToPath, GameRootPath,   GameContentPathRebased);
+		AddPathPair(ContentRootToPath, ScriptRootPath, GameScriptPathRebased);
+		AddPathPair(ContentRootToPath, TempRootPath,   GameSavedPathRebased);
+		AddPathPair(ContentRootToPath, ConfigRootPath, GameConfigPathRebased);
 
 		// Allow the plugin manager to mount new content paths by exposing access through a delegate.  PluginManager is 
 		// a Core class, but content path functionality is added at the CoreUObject level.
@@ -1114,9 +1130,6 @@ bool FPackageName::TryGetMountPointForPath(FStringView InFilePathOrPackageName,
 	FReadScopeLock ScopeLock(ContentMountPointCriticalSection);
 	for (const auto& Pair : Paths.ContentRootToPath)
 	{
-		TStringBuilder<256> RootFileAbsPath;
-		FPathViews::ToAbsolutePath(Pair.ContentPath, RootFileAbsPath);
-
 		FStringView RelPath;
 		if (FPathViews::TryMakeChildPathRelativeTo(InFilePathOrPackageName, Pair.RootPath, RelPath))
 		{
@@ -1140,7 +1153,7 @@ bool FPackageName::TryGetMountPointForPath(FStringView InFilePathOrPackageName,
 			}
 			return true;
 		}
-		else if (FPathViews::TryMakeChildPathRelativeTo(PossibleAbsFilePath, RootFileAbsPath, RelPath))
+		else if (FPathViews::TryMakeChildPathRelativeTo(PossibleAbsFilePath, Pair.AbsolutePath, RelPath))
 		{
 			OutMountPointPackageName << Pair.RootPath;
 			OutMountPointFilePath << Pair.ContentPath;
