Added profile tag for FPackageName::DoesPackageExist

diff --git a/Engine/Source/Runtime/CoreUObject/Private/Misc/PackageName.cpp b/Engine/Source/Runtime/CoreUObject/Private/Misc/PackageName.cpp
--- a/Engine/Source/Runtime/CoreUObject/Private/Misc/PackageName.cpp
+++ b/Engine/Source/Runtime/CoreUObject/Private/Misc/PackageName.cpp
@@ -1314,6 +1314,8 @@ bool FPackageName::DoesPackageExist(const FPackagePath& PackagePath, FPackagePat
 
 bool FPackageName::DoesPackageExist(const FPackagePath& PackagePath, const FGuid* Guid, bool bMatchCaseOnDisk, FPackagePath* OutPackagePath)
 {
+	TRACE_CPUPROFILER_EVENT_SCOPE(FPackageName::DoesPackageExist);
+
 	// DoesPackageExist returns false for local filenames that are in unmounted directories, even if those files exist on the local disk
 	if (!PackagePath.IsMountedPath())
 	{
