PR #7908: FPrimaryAssetID can now be serialized into a FArchive (Contributed by [REDACTED_PERSON])

diff --git a/Engine/Source/Runtime/CoreUObject/Public/UObject/PrimaryAssetId.h b/Engine/Source/Runtime/CoreUObject/Public/UObject/PrimaryAssetId.h
--- a/Engine/Source/Runtime/CoreUObject/Public/UObject/PrimaryAssetId.h
+++ b/Engine/Source/Runtime/CoreUObject/Public/UObject/PrimaryAssetId.h
@@ -161,6 +161,13 @@ struct FPrimaryAssetId
 		return *this;
 	}
 
+	FORCEINLINE friend FArchive& operator<<(FArchive& Ar, FPrimaryAssetId& Other)
+	{
+        Ar << Other.PrimaryAssetType;
+		Ar << Other.PrimaryAssetName;
+		return Ar;
+	}
+
 	/** UStruct Overrides */
 	bool ExportTextItem(FString& ValueStr, FPrimaryAssetId const& DefaultValue, UObject* Parent, int32 PortFlags, UObject* ExportRootScope) const;
 	bool ImportTextItem(const TCHAR*& Buffer, int32 PortFlags, UObject* Parent, FOutputDevice* ErrorText);
