Updating a few deferred dependency verification checks to take FFieldVariants into account. These checks were only compiled if the USE_DEFERRED_DEPENDENCY_CHECK_VERIFICATION_TESTS build flag was enabled, and off by default.

diff --git a/Engine/Source/Runtime/CoreUObject/Private/UObject/LinkerPlaceholderBase.cpp b/Engine/Source/Runtime/CoreUObject/Private/UObject/LinkerPlaceholderBase.cpp
--- a/Engine/Source/Runtime/CoreUObject/Private/UObject/LinkerPlaceholderBase.cpp
+++ b/Engine/Source/Runtime/CoreUObject/Private/UObject/LinkerPlaceholderBase.cpp
@@ -99,7 +99,7 @@ int32 FLinkerPlaceholderObjectImpl::ResolvePlaceholderValues(const TArray<FField
 		else if (const FArrayProperty* ArrayProperty = Property.Get<FArrayProperty>())
 		{
 #if USE_DEFERRED_DEPENDENCY_CHECK_VERIFICATION_TESTS
-			const FProperty* NextProperty = PropertyChain[PropertyIndex - 1];
+			const FProperty* NextProperty = PropertyChain[PropertyIndex - 1].Get<FProperty>();
 			check(NextProperty == ArrayProperty->Inner);
 #endif // USE_DEFERRED_DEPENDENCY_CHECK_VERIFICATION_TESTS
 
@@ -124,7 +124,7 @@ int32 FLinkerPlaceholderObjectImpl::ResolvePlaceholderValues(const TArray<FField
 		else if (const FSetProperty* SetProperty = Property.Get<FSetProperty>())
 		{
 #if USE_DEFERRED_DEPENDENCY_CHECK_VERIFICATION_TESTS
-			const FProperty* NextProperty = PropertyChain[PropertyIndex - 1];
+			const FProperty* NextProperty = PropertyChain[PropertyIndex - 1].Get<FProperty>();
 			check(NextProperty == SetProperty->ElementProp);
 #endif // USE_DEFERRED_DEPENDENCY_CHECK_VERIFICATION_TESTS
 
@@ -472,7 +472,7 @@ bool FLinkerPlaceholderBase::AddReferencingPropertyValue(FFieldVariant Referenci
 	if (ReferencingContainer != nullptr)
 	{
 #if USE_DEFERRED_DEPENDENCY_CHECK_VERIFICATION_TESTS
-		check(ReferencingProperty->GetObjectPropertyValue(DataPtr) == GetPlaceholderAsUObject());
+		check(ReferencingProperty.Get<FObjectProperty>()->GetObjectPropertyValue(DataPtr) == GetPlaceholderAsUObject());
 		check(PropertyChain.IsValid());
 #endif // USE_DEFERRED_DEPENDENCY_CHECK_VERIFICATION_TESTS
 
diff --git a/Engine/Source/Runtime/CoreUObject/Private/UObject/LinkerPlaceholderBase.inl b/Engine/Source/Runtime/CoreUObject/Private/UObject/LinkerPlaceholderBase.inl
--- a/Engine/Source/Runtime/CoreUObject/Private/UObject/LinkerPlaceholderBase.inl
+++ b/Engine/Source/Runtime/CoreUObject/Private/UObject/LinkerPlaceholderBase.inl
@@ -87,19 +87,13 @@ void TLinkerImportPlaceholder<PlaceholderType>::AddReferencingProperty(FFieldVar
 	UObject* ThisAsObject = GetPlaceholderAsUObject();
 	check(ThisAsObject != nullptr);
 
-	FObjectImport* PlaceholderImport = nullptr;
-	if (FLinkerLoad* PropertyLinker = ReferencingProperty->GetLinker())
+	if (!ReferencingProperty.IsUObject())
 	{
-		for (FObjectImport& Import : PropertyLinker->ImportMap)
+		if (FLinkerLoad* PropertyLinker = ReferencingProperty.Get<FProperty>()->GetLinker())
 		{
-			if (Import.XObject == ThisAsObject)
-			{
-				PlaceholderImport = &Import;
-				break;
-			}
+			check(ThisAsObject->GetOutermost() == PropertyLinker->LinkerRoot);
+			check(PropertyLinker->LoadFlags & LOAD_DeferDependencyLoads);
 		}
-		check(ThisAsObject->GetOutermost() == PropertyLinker->LinkerRoot);
-		check(PropertyLinker->LoadFlags & LOAD_DeferDependencyLoads);
 	}
 	// if this check hits, then we're adding dependencies after we've 
 	// already resolved the placeholder (it won't be resolved again)
