Fixed errors re format string / parameter list for various UE_LOGs:
  - number of supplied parameters less or more than the number of specifiers used in format string
  - invalid format string specifiers used

diff --git a/Engine/Source/Developer/ScreenShotComparisonTools/Private/ScreenShotManager.cpp b/Engine/Source/Developer/ScreenShotComparisonTools/Private/ScreenShotManager.cpp
--- a/Engine/Source/Developer/ScreenShotComparisonTools/Private/ScreenShotManager.cpp
+++ b/Engine/Source/Developer/ScreenShotComparisonTools/Private/ScreenShotManager.cpp
@@ -160,7 +160,7 @@ TArray<FString> FScreenShotManager::FindApprovedImages(const FAutomationScreensh
 	{
 		FString CurrentPlatformRHI = GetPathComponentForPlatformAndRHI(IncomingMetaData);
 
-		UE_LOG(LogScreenShotManager, Log, TEXT("No ideal-image found at {0}. Checking fallback images"), *ApprovedPath);
+		UE_LOG(LogScreenShotManager, Log, TEXT("No ideal-image found at %s. Checking fallback images"), *ApprovedPath);
 
 		while (ApprovedImages.Num() == 0)
 		{
diff --git a/Engine/Source/Runtime/AssetRegistry/Private/AssetDataGatherer.cpp b/Engine/Source/Runtime/AssetRegistry/Private/AssetDataGatherer.cpp
--- a/Engine/Source/Runtime/AssetRegistry/Private/AssetDataGatherer.cpp
+++ b/Engine/Source/Runtime/AssetRegistry/Private/AssetDataGatherer.cpp
@@ -1455,7 +1455,7 @@ void FAssetDataDiscovery::TickInternal()
 			}
 			if (FPathViews::GetPathLeaf(RelPath).Len() != RelPath.Len())
 			{
-				UE_LOG(LogAssetRegistry, Warning, TEXT("IterateDirectoryStat returned unexpected result %* which is not a direct child of the requested path %s."), InPackageFilename, DirLocalAbsPath.ToString());
+				UE_LOG(LogAssetRegistry, Warning, TEXT("IterateDirectoryStat returned unexpected result %s which is not a direct child of the requested path %s."), InPackageFilename, DirLocalAbsPath.ToString());
 				return true;
 			}
 			ON_SCOPE_EXIT{ DirLongPackageName.RemoveSuffix(DirLongPackageName.Len() - DirLongPackageNameRootLen); };
diff --git a/Engine/Source/Runtime/Core/Private/Internationalization/TextLocalizationManager.cpp b/Engine/Source/Runtime/Core/Private/Internationalization/TextLocalizationManager.cpp
--- a/Engine/Source/Runtime/Core/Private/Internationalization/TextLocalizationManager.cpp
+++ b/Engine/Source/Runtime/Core/Private/Internationalization/TextLocalizationManager.cpp
@@ -245,7 +245,7 @@ void ApplyDefaultCultureSettings(const ELocalizationLoadFlags LocLoadFlags)
 				if (bRequireExactMatch && InRequestedCulture != ValidCultureName)
 				{
 					TargetCultureName = ValidCultureName;
-					UE_LOG(LogTextLocalizationManager, Log, TEXT("No specific localization for '%s' exists, so the '%s' localization will be used."), *InRequestedCulture, *ValidCultureName, InLogDesc);
+					UE_LOG(LogTextLocalizationManager, Log, TEXT("No specific localization for '%s' exists, so '%s' will be used for the %s."), *InRequestedCulture, *TargetCultureName, InLogDesc);
 				}
 			}
 			else if (!ValidFallbackCultureName.IsEmpty())
diff --git a/Engine/Source/Runtime/CoreUObject/Private/Serialization/AsyncLoading.cpp b/Engine/Source/Runtime/CoreUObject/Private/Serialization/AsyncLoading.cpp
--- a/Engine/Source/Runtime/CoreUObject/Private/Serialization/AsyncLoading.cpp
+++ b/Engine/Source/Runtime/CoreUObject/Private/Serialization/AsyncLoading.cpp
@@ -2623,7 +2623,7 @@ void FAsyncPackage::DumpDependencies(const TCHAR* Label, UObject* Obj)
 	UE_LOG(LogStreaming, Error, TEXT("****DumpDependencies [%s]:"), Label);
 	if (!Obj)
 	{
-		UE_LOG(LogStreaming, Error, TEXT("    Obj is nullptr"), Label);
+		UE_LOG(LogStreaming, Error, TEXT("    Obj is nullptr"));
 		return;
 	}
 	UE_LOG(LogStreaming, Error, TEXT("    Obj is %s"), *Obj->GetFullName());
diff --git a/Engine/Source/Runtime/CoreUObject/Private/UObject/GarbageCollection.cpp b/Engine/Source/Runtime/CoreUObject/Private/UObject/GarbageCollection.cpp
--- a/Engine/Source/Runtime/CoreUObject/Private/UObject/GarbageCollection.cpp
+++ b/Engine/Source/Runtime/CoreUObject/Private/UObject/GarbageCollection.cpp
@@ -2149,8 +2149,7 @@ bool UnhashUnreachableObjects(bool bUseTimeLimit, float TimeLimit)
 		UE_LOG(LogGarbage, Log, TEXT("%f ms for %sunhashing unreachable objects (%d objects unhashed)"),
 		(FPlatformTime::Seconds() - StartTime) * 1000,
 		bUseTimeLimit ? TEXT("incrementally ") : TEXT(""),
-			Items,
-		GUnrechableObjectIndex, GUnreachableObjects.Num());
+		GUnreachableObjects.Num());
 	}
 	else if (!bTimeLimitReached)
 	{
diff --git a/Engine/Source/Runtime/CoreUObject/Private/UObject/LinkerLoad.cpp b/Engine/Source/Runtime/CoreUObject/Private/UObject/LinkerLoad.cpp
--- a/Engine/Source/Runtime/CoreUObject/Private/UObject/LinkerLoad.cpp
+++ b/Engine/Source/Runtime/CoreUObject/Private/UObject/LinkerLoad.cpp
@@ -1757,7 +1757,7 @@ FLinkerLoad::ELinkerStatus FLinkerLoad::FixupImportMap()
 							// Change object name
 							Import.ObjectName = NewObjectName.ObjectName;
 
-							UE_LOG(LogLinker, Verbose, TEXT("FLinkerLoad::FixupImportMap() - Renamed Object %s -> %s"), *LinkerRoot->GetName(), *OldObjectName.ToString(), *NewObjectName.ToString());
+							UE_LOG(LogLinker, Verbose, TEXT("FLinkerLoad::FixupImportMap() - Pkg<%s> - Renamed Object %s -> %s"), *LinkerRoot->GetName(), *OldObjectName.ToString(), *NewObjectName.ToString());
 						}
 					}
 
@@ -1793,7 +1793,7 @@ FLinkerLoad::ELinkerStatus FLinkerLoad::FixupImportMap()
 							Import.ObjectName = FName(*NewDefaultObjectName);
 						}
 
-						UE_LOG(LogLinker, Verbose, TEXT("FLinkerLoad::FixupImportMap() - Renamed Class %s -> %s"), *LinkerRoot->GetName(), *OldClassName.ToString(), *NewClassName.ToString());
+						UE_LOG(LogLinker, Verbose, TEXT("FLinkerLoad::FixupImportMap() - Pkg<%s> - Renamed Class %s -> %s"), *LinkerRoot->GetName(), *OldClassName.ToString(), *NewClassName.ToString());
 					}	
 				}
 
diff --git a/Engine/Source/Runtime/RHI/Public/RHICommandList.h b/Engine/Source/Runtime/RHI/Public/RHICommandList.h
--- a/Engine/Source/Runtime/RHI/Public/RHICommandList.h
+++ b/Engine/Source/Runtime/RHI/Public/RHICommandList.h
@@ -3582,7 +3582,7 @@ public:
 
 		if (InInfo.bTooManyUAVs)
 		{
-			UE_LOG(LogRHI, Warning, TEXT("RenderPass %s has too many UAVs"));
+			UE_LOG(LogRHI, Warning, TEXT("RenderPass %s has too many UAVs"), Name);
 		}
 		InInfo.Validate();
 
diff --git a/Engine/Source/Runtime/RenderCore/Private/ShaderCodeLibrary.cpp b/Engine/Source/Runtime/RenderCore/Private/ShaderCodeLibrary.cpp
--- a/Engine/Source/Runtime/RenderCore/Private/ShaderCodeLibrary.cpp
+++ b/Engine/Source/Runtime/RenderCore/Private/ShaderCodeLibrary.cpp
@@ -1904,7 +1904,7 @@ public:
 					// In case we run under such circumstances, we need to discover the components ourselves
 					if (FPlatformFileManager::Get().FindPlatformFile(TEXT("PakFile")) == nullptr)
 					{
-						UE_LOG(LogShaderLibrary, Display, TEXT("Running without a pakfile and did not find a monolithic library '%s' - attempting disk search for its chunks"), *Name, Library->GetNumComponents());
+						UE_LOG(LogShaderLibrary, Display, TEXT("Running without a pakfile and did not find a monolithic library '%s' - attempting disk search for its chunks"), *Name);
 
 						TArray<FString> UshaderbytecodeFiles;
 						FString SearchMask = Directory / FString::Printf(TEXT("ShaderArchive-*%s*.ushaderbytecode"), *Name);
diff --git a/Engine/Source/Runtime/RenderCore/Private/ShaderResource.cpp b/Engine/Source/Runtime/RenderCore/Private/ShaderResource.cpp
--- a/Engine/Source/Runtime/RenderCore/Private/ShaderResource.cpp
+++ b/Engine/Source/Runtime/RenderCore/Private/ShaderResource.cpp
@@ -237,8 +237,7 @@ void FShaderMapResourceCode::AddShaderCode(EShaderFrequency InFrequency, const F
 				Entry.UncompressedSize = ShaderCode.Num();
 				UE_LOG(LogShaders, Warning, TEXT("Shader %s is expected to be compressed with %s, but it arrived uncompressed. Assuming compressing made it longer and storing uncompressed."),
 					*InHash.ToString(),
-					*ShaderCompressionFormat.ToString(),
-					*InCode.GetCompressionFormat().ToString()
+					*ShaderCompressionFormat.ToString()
 				);
 			}
 		}
