FPackageName::DoesPackageExist: downgrade ensurefailure+Error to a warning in the case that input LongPackageName is not mounted.
This ensure has been hit twice, and in both cases the calling code would have correctly handled the false return value, and we needed to alter it to avoid the ensure.
#rnx

diff --git a/Engine/Source/Runtime/CoreUObject/Private/Misc/PackageName.cpp b/Engine/Source/Runtime/CoreUObject/Private/Misc/PackageName.cpp
--- a/Engine/Source/Runtime/CoreUObject/Private/Misc/PackageName.cpp
+++ b/Engine/Source/Runtime/CoreUObject/Private/Misc/PackageName.cpp
@@ -1279,15 +1279,8 @@ bool FPackageName::DoesPackageExist(const FString& LongPackageName, const FGuid*
 		EErrorCode FailureReason;
 		if (!FPackageName::TryConvertToMountedPathComponents(LongPackageName, PackageNameRoot, FilePathRoot, RelPath, UnusedObjectName, Extension, CustomExtension, nullptr /* OutFlexNameType */, &FailureReason))
 		{
-			FString Message = FString::Printf(TEXT("Illegal call to DoesPackageExist: %s"), *FormatErrorAsString(LongPackageName, FailureReason));
-			UE_LOG(LogPackageName, Error, TEXT("%s"), *Message);
-
-			// Only ensure if the given package name is not an empty string to 
-			// support legacy behavior
-			if(!LongPackageName.IsEmpty())
-			{
-				ensureMsgf(false, TEXT("%s"), *Message);
-			}
+			FString Message = FString::Printf(TEXT("DoesPackageExist was called with an invalid LongPackageName; this will always return false. Reason: %s"), *FormatErrorAsString(LongPackageName, FailureReason));
+			UE_LOG(LogPackageName, Warning, TEXT("%s"), *Message);
 			return false;
 		}
 		PackagePath = FPackagePath::FromMountedComponents(PackageNameRoot, FilePathRoot, RelPath, Extension, CustomExtension);
