VirtualBulkData and EditorDomain: Do not update the PackagePath for VirtualBulkData if SavePackage is saving to a copy rather than updating the Package's PackagePath.
#rnx

diff --git a/Engine/Source/Runtime/CoreUObject/Private/UObject/LinkerSave.cpp b/Engine/Source/Runtime/CoreUObject/Private/UObject/LinkerSave.cpp
--- a/Engine/Source/Runtime/CoreUObject/Private/UObject/LinkerSave.cpp
+++ b/Engine/Source/Runtime/CoreUObject/Private/UObject/LinkerSave.cpp
@@ -10,6 +10,7 @@
 #include "UObject/LazyObjectPtr.h"
 #include "UObject/SoftObjectPtr.h"
 #include "Internationalization/TextPackageNamespaceUtil.h"
+#include "UObject/ObjectSaveContext.h"
 #include "UObject/UObjectThreadContext.h"
 #include "UObject/UnrealType.h"
 #include "HAL/PlatformStackWalk.h"
@@ -274,11 +275,11 @@ void FLinkerSave::Serialize( void* V, int64 Length )
 	Saver->Serialize( V, Length );
 }
 
-void FLinkerSave::OnPostSave(const FPackagePath& PackagePath)
+void FLinkerSave::OnPostSave(const FPackagePath& PackagePath, FObjectPostSaveContext ObjectSaveContext)
 {
-	for (TUniqueFunction<void(const FPackagePath&)>& Callback : PostSaveCallbacks)
+	for (TUniqueFunction<void(const FPackagePath&, FObjectPostSaveContext)>& Callback : PostSaveCallbacks)
 	{
-		Callback(PackagePath);
+		Callback(PackagePath, ObjectSaveContext);
 	}
 
 	PostSaveCallbacks.Empty();
diff --git a/Engine/Source/Runtime/CoreUObject/Private/UObject/SavePackage.cpp b/Engine/Source/Runtime/CoreUObject/Private/UObject/SavePackage.cpp
--- a/Engine/Source/Runtime/CoreUObject/Private/UObject/SavePackage.cpp
+++ b/Engine/Source/Runtime/CoreUObject/Private/UObject/SavePackage.cpp
@@ -4680,7 +4680,7 @@ FSavePackageResultStruct UPackage::Save(UPackage* InOuter, UObject* Base, EObjec
 
 			if (Linker != nullptr)
 			{
-				Linker->OnPostSave(TargetPackagePath);
+				Linker->OnPostSave(TargetPackagePath, FObjectPostSaveContext(ObjectSaveContext));
 			}
 
 			// send a message that the package was saved
diff --git a/Engine/Source/Runtime/CoreUObject/Private/UObject/SavePackage2.cpp b/Engine/Source/Runtime/CoreUObject/Private/UObject/SavePackage2.cpp
--- a/Engine/Source/Runtime/CoreUObject/Private/UObject/SavePackage2.cpp
+++ b/Engine/Source/Runtime/CoreUObject/Private/UObject/SavePackage2.cpp
@@ -1933,7 +1933,7 @@ void PostSavePackage(FSaveContext& SaveContext)
 
 	if (SaveContext.Linker != nullptr)
 	{
-		SaveContext.Linker->OnPostSave(SaveContext.GetTargetPackagePath());
+		SaveContext.Linker->OnPostSave(SaveContext.GetTargetPackagePath(), FObjectPostSaveContext(SaveContext.GetObjectSaveContext()));
 	}
 
 	// Send a message that the package was saved
diff --git a/Engine/Source/Runtime/CoreUObject/Private/Virtualization/VirtualizedBulkData.cpp b/Engine/Source/Runtime/CoreUObject/Private/Virtualization/VirtualizedBulkData.cpp
--- a/Engine/Source/Runtime/CoreUObject/Private/Virtualization/VirtualizedBulkData.cpp
+++ b/Engine/Source/Runtime/CoreUObject/Private/Virtualization/VirtualizedBulkData.cpp
@@ -12,6 +12,7 @@
 #include "UObject/LinkerLoad.h"
 #include "UObject/LinkerSave.h"
 #include "UObject/Object.h"
+#include "UObject/ObjectSaveContext.h"
 #include "Virtualization/IVirtualizationSourceControlUtilities.h"
 
 //#if WITH_EDITORONLY_DATA
@@ -321,7 +322,7 @@ void FVirtualizedUntypedBulkData::Serialize(FArchive& Ar, UObject* Owner)
 					// Restore the archive's offset
 					Ar.Seek(ReturnPos);
 
-					// If we are saving the package to disk (we have access to FLinkerSave and it's filepath is valid) 
+					// If we are saving the package to disk (we have access to FLinkerSave and its filepath is valid) 
 					// then we should register a callback to be received once the package has actually been saved to 
 					// disk so that we can update the object's members to be redirected to the saved file.
 					FLinkerSave* LinkerSave = Cast<FLinkerSave>(Ar.GetLinker());
@@ -336,8 +337,13 @@ void FVirtualizedUntypedBulkData::Serialize(FArchive& Ar, UObject* Owner)
 							SidecarData.Payload = PayloadToSerialize;
 						}
 
-						auto OnSavePackage = [this, PayloadOffset, UpdatedFlags](const FPackagePath& InPackagePath)
+						auto OnSavePackage = [this, PayloadOffset, UpdatedFlags](const FPackagePath& InPackagePath, FObjectPostSaveContext ObjectSaveContext)
 						{
+							if (!ObjectSaveContext.IsUpdatingLoadedPath())
+							{
+								return;
+							}
+
 							this->PackagePath = InPackagePath;
 
 							if (!this->PackagePath.IsEmpty())
diff --git a/Engine/Source/Runtime/CoreUObject/Public/UObject/LinkerSave.h b/Engine/Source/Runtime/CoreUObject/Public/UObject/LinkerSave.h
--- a/Engine/Source/Runtime/CoreUObject/Public/UObject/LinkerSave.h
+++ b/Engine/Source/Runtime/CoreUObject/Public/UObject/LinkerSave.h
@@ -12,6 +12,7 @@
 #include "UObject/UObjectThreadContext.h"
 #include "Virtualization/PayloadId.h"
 
+class FObjectPostSaveContext;
 struct FUntypedBulkData;
 
 /*----------------------------------------------------------------------------
@@ -108,7 +109,7 @@ public:
 	 * via lambda capture. 
 	 * @param PackagePath The path of the package
 	 */
-	TArray<TUniqueFunction<void(const FPackagePath& PackagePath)>> PostSaveCallbacks;
+	TArray<TUniqueFunction<void(const FPackagePath& PackagePath, FObjectPostSaveContext ObjectSaveContext)>> PostSaveCallbacks;
 
 	/** A mapping of package name to generated script SHA keys */
 	COREUOBJECT_API static TMap<FString, TArray<uint8> > PackagesToScriptSHAMap;
@@ -161,7 +162,7 @@ public:
 	void Serialize( void* V, int64 Length );
 
 	/** Invoke all of the callbacks in PostSaveCallbacks and then empty it. */
-	void OnPostSave(const FPackagePath& PackagePath);
+	void OnPostSave(const FPackagePath& PackagePath, FObjectPostSaveContext ObjectSaveContext);
 
 	// FLinker interface
 	virtual FString GetDebugName() const override;
