Allow for time-sliced unhashing of uobjects when in the editor. All unhashing will still occur before streaming is continued, but async loading will provide additional time to the unhashing, allowing it to finish faster, but over a few frames.

diff --git a/Engine/Source/Runtime/CoreUObject/Private/Serialization/AsyncLoading.cpp b/Engine/Source/Runtime/CoreUObject/Private/Serialization/AsyncLoading.cpp
--- a/Engine/Source/Runtime/CoreUObject/Private/Serialization/AsyncLoading.cpp
+++ b/Engine/Source/Runtime/CoreUObject/Private/Serialization/AsyncLoading.cpp
@@ -4531,18 +4531,26 @@ EAsyncPackageState::Type FAsyncLoadingThread::TickAsyncLoading(bool bUseTimeLimi
 
 	if (!bLoadingSuspended)
 	{
+		double TickStartTime = FPlatformTime::Seconds();
+		double TimeLimitUsedForProcessLoaded = 0;
+
 		// First make sure there's no objects pending to be unhashed. This is important in uncooked builds since we don't 
 		// detach linkers immediately there and we may end up in getting unreachable objects from Linkers in CreateImports
 		if (FPlatformProperties::RequiresCookedData() == false && IsIncrementalUnhashPending() && IsAsyncLoadingPackages())
 		{
 			// Call ConditionalBeginDestroy on all pending objects. CBD is where linkers get detached from objects.
-			UnhashUnreachableObjects(false);
+			UnhashUnreachableObjects(bUseTimeLimit, TimeLimit);
+
+			TimeLimitUsedForProcessLoaded = FPlatformTime::Seconds() - TickStartTime;
+
+			if(IsIncrementalUnhashPending() || TimeLimitUsedForProcessLoaded > TimeLimit)
+			{
+				return EAsyncPackageState::TimeOut;
+			}
 		}
 
 		const bool bIsMultithreaded = FAsyncLoadingThread::IsMultithreaded();
-		double TickStartTime = FPlatformTime::Seconds();
-		double TimeLimitUsedForProcessLoaded = 0;
-
+		
 		bool bDidSomething = false;
 		{
 			Result = ProcessLoadedPackages(bUseTimeLimit, bUseFullTimeLimit, TimeLimit, bDidSomething, FlushTree);
