Adding CLASSCAST flags for the new LWC property.  It was acting as a numeric when used in an IsA<> test.  This resulted in UHT accepting any numeric as a valid blueprint type.

#rnx

diff --git a/Engine/Source/Runtime/CoreUObject/Private/UObject/PropertyNumeric.cpp b/Engine/Source/Runtime/CoreUObject/Private/UObject/PropertyNumeric.cpp
--- a/Engine/Source/Runtime/CoreUObject/Private/UObject/PropertyNumeric.cpp
+++ b/Engine/Source/Runtime/CoreUObject/Private/UObject/PropertyNumeric.cpp
@@ -17,6 +17,7 @@ IMPLEMENT_FIELD(FUInt32Property)
 IMPLEMENT_FIELD(FUInt64Property)
 IMPLEMENT_FIELD(FFloatProperty)
 IMPLEMENT_FIELD(FDoubleProperty)
+IMPLEMENT_FIELD(FLargeWorldCoordinatesRealProperty)
 
 int64 FNumericProperty::ReadEnumAsInt64(FStructuredArchive::FSlot Slot, UStruct* DefaultsStruct, const FPropertyTag& Tag)
 {
diff --git a/Engine/Source/Runtime/CoreUObject/Public/UObject/ObjectMacros.h b/Engine/Source/Runtime/CoreUObject/Public/UObject/ObjectMacros.h
--- a/Engine/Source/Runtime/CoreUObject/Public/UObject/ObjectMacros.h
+++ b/Engine/Source/Runtime/CoreUObject/Public/UObject/ObjectMacros.h
@@ -350,6 +350,7 @@ enum EClassCastFlags : uint64
 	CASTCLASS_FFieldPathProperty			= 0x0010000000000000,
 	CASTCLASS_FObjectPtrProperty			= 0x0020000000000000,
 	CASTCLASS_FClassPtrProperty				= 0x0040000000000000,
+	CASTCLASS_FLargeWorldCoordinatesRealProperty = 0x0080000000000000,
 };
 
 #define CASTCLASS_AllFlags ((EClassCastFlags)0xFFFFFFFFFFFFFFFF)
diff --git a/Engine/Source/Runtime/CoreUObject/Public/UObject/UnrealType.h b/Engine/Source/Runtime/CoreUObject/Public/UObject/UnrealType.h
--- a/Engine/Source/Runtime/CoreUObject/Public/UObject/UnrealType.h
+++ b/Engine/Source/Runtime/CoreUObject/Public/UObject/UnrealType.h
@@ -1973,7 +1973,8 @@ class COREUOBJECT_API FDoubleProperty : public TProperty_Numeric<double>
 // Note: Stub only. Used to provide FNativeClassHeaderGenerator::PropertyNew with a way to track an as yet unaliased FFloatProperty/FDoubleProperty within the header tool, which will be resolved to the correct type at run time.
 class COREUOBJECT_API FLargeWorldCoordinatesRealProperty : public TProperty_Numeric<double>
 {
-public:
+	DECLARE_FIELD(FLargeWorldCoordinatesRealProperty, TProperty_Numeric<double>, CASTCLASS_FLargeWorldCoordinatesRealProperty)
+
 	FLargeWorldCoordinatesRealProperty(FFieldVariant InOwner, const FName& InName, EObjectFlags InObjectFlags)
 		: TProperty_Numeric(InOwner, InName, InObjectFlags)
 	{
