Added support for PIE remapping of FLazyObjectPtr.

diff --git a/Engine/Source/Runtime/CoreUObject/Private/UObject/LazyObjectPtr.cpp b/Engine/Source/Runtime/CoreUObject/Private/UObject/LazyObjectPtr.cpp
--- a/Engine/Source/Runtime/CoreUObject/Private/UObject/LazyObjectPtr.cpp
+++ b/Engine/Source/Runtime/CoreUObject/Private/UObject/LazyObjectPtr.cpp
@@ -11,13 +11,26 @@
 #include "UObject/PropertyPortFlags.h"
 #include "UObject/UObjectAnnotation.h"
 
-
 /** Annotation associating objects with their guids **/
 static FUObjectAnnotationSparseSearchable<FUniqueObjectGuid,true> GuidAnnotation;
 
 #define MAX_PIE_INSTANCES 10
 static TMap<FGuid, FGuid> PIEGuidMap[MAX_PIE_INSTANCES];
 
+static FGuid RemapGuid(const FUniqueObjectGuid& Guid, int32 PIEInstanceID)
+{
+	check(PIEInstanceID != INDEX_NONE);
+	check(PIEInstanceID < MAX_PIE_INSTANCES);
+	FGuid& FoundGuid = PIEGuidMap[PIEInstanceID].FindOrAdd(Guid.GetGuid());
+
+	if (!FoundGuid.IsValid())
+	{
+		FoundGuid = FGuid::NewGuid();
+	}
+
+	return FoundGuid;
+};
+
 /*-----------------------------------------------------------------------------
 	FUniqueObjectGuid
 -----------------------------------------------------------------------------*/
@@ -29,16 +42,7 @@ FUniqueObjectGuid::FUniqueObjectGuid(const class UObject* InObject)
 
 FUniqueObjectGuid FUniqueObjectGuid::FixupForPIE(int32 PlayInEditorID) const
 {
-	FUniqueObjectGuid Temp(*this);
-
-	check(PlayInEditorID != -1)
-	const FGuid *FoundGuid = PIEGuidMap[PlayInEditorID % MAX_PIE_INSTANCES].Find(Temp.GetGuid());
-
-	if (FoundGuid)
-	{
-		Temp = *FoundGuid;
-	}
-	return Temp;
+	return RemapGuid(GetGuid(), PlayInEditorID);
 }
 
 UObject* FUniqueObjectGuid::ResolveObject() const
@@ -99,20 +103,6 @@ void FLazyObjectPtr::PossiblySerializeObjectGuid(UObject *Object, FStructuredArc
 {
 	FArchive& UnderlyingArchive = Record.GetUnderlyingArchive();
 
-	auto RemapGuid = [](FUniqueObjectGuid& Guid, int32 PIEInstanceID)
-	{
-		check(PIEInstanceID != INDEX_NONE);
-		FGuid& FoundGuid = PIEGuidMap[PIEInstanceID % MAX_PIE_INSTANCES].FindOrAdd(Guid.GetGuid());
-		if (!FoundGuid.IsValid())
-		{
-			Guid = FoundGuid = FGuid::NewGuid();
-		}
-		else
-		{
-			Guid = FoundGuid;
-		}
-	};
-
 	if (UnderlyingArchive.IsSaving() || UnderlyingArchive.IsCountingMemory())
 	{
 		FUniqueObjectGuid Guid = GuidAnnotation.GetAnnotation(Object);
@@ -121,7 +111,7 @@ void FLazyObjectPtr::PossiblySerializeObjectGuid(UObject *Object, FStructuredArc
 		{
 			if (UnderlyingArchive.GetPortFlags() & PPF_DuplicateForPIE)
 			{
-				RemapGuid(Guid, GPlayInEditorID);
+				Guid = RemapGuid(Guid, GPlayInEditorID);
 			}
 
 			GuidSlot.GetValue() << Guid;
@@ -157,7 +147,7 @@ void FLazyObjectPtr::PossiblySerializeObjectGuid(UObject *Object, FStructuredArc
 #if WITH_EDITOR
 						else if (Object->GetOutermostObject()->GetPackage()->HasAnyPackageFlags(PKG_PlayInEditor))
 						{
-							RemapGuid(Guid, Object->GetOutermostObject()->GetPackage()->PIEInstanceID);
+							Guid = RemapGuid(Guid, Object->GetOutermostObject()->GetPackage()->PIEInstanceID);
 							GuidAnnotation.AddAnnotation(Object, Guid);
 						}
 #endif
@@ -194,6 +184,6 @@ void FLazyObjectPtr::PossiblySerializeObjectGuid(UObject *Object, FStructuredArc
 void FLazyObjectPtr::ResetPIEFixups()
 {
 	check(GPlayInEditorID != -1);
-	PIEGuidMap[GPlayInEditorID % MAX_PIE_INSTANCES].Reset();
-}
-
+	check(GPlayInEditorID < MAX_PIE_INSTANCES);
+	PIEGuidMap[GPlayInEditorID].Reset();
+}
\ No newline at end of file
