Removed unncessary FLinkerLoad::StaticInit.

This appears to be a really old holdover from early UE4 development. It was originally introduced in CL 1258792 to satisfy a lingering Engine dependency. However, that dependency was eventually removed in CL 3119255, which this change probably should've been lumped with.

diff --git a/Engine/Source/Runtime/CoreUObject/Private/UObject/LinkerLoad.cpp b/Engine/Source/Runtime/CoreUObject/Private/UObject/LinkerLoad.cpp
--- a/Engine/Source/Runtime/CoreUObject/Private/UObject/LinkerLoad.cpp
+++ b/Engine/Source/Runtime/CoreUObject/Private/UObject/LinkerLoad.cpp
@@ -73,8 +73,6 @@ static FAutoConsoleVariableRef CVarLinkerAllowDynamicClasses(
 	);
 #endif
 
-UClass* FLinkerLoad::UTexture2DStaticClass = NULL;
-
 FName FLinkerLoad::NAME_LoadErrors("LoadErrors");
 
 
@@ -414,11 +412,6 @@ static FORCEINLINE bool IsCoreUObjectPackage(const FName& PackageName)
 	FLinkerLoad.
 ----------------------------------------------------------------------------*/
 
-void FLinkerLoad::StaticInit(UClass* InUTexture2DStaticClass)
-{
-	UTexture2DStaticClass = InUTexture2DStaticClass;
-}
-
 FLinkerLoad* FLinkerLoad::CreateLinker(FUObjectSerializeContext* LoadContext, UPackage* Parent, const TCHAR* Filename, uint32 LoadFlags, FArchive* InLoader, const FLinkerInstancingContext* InstancingContext)
 {
 	return CreateLinker(LoadContext, Parent, FPackagePath::FromLocalPath(Filename), LoadFlags, InLoader, InstancingContext);
diff --git a/Engine/Source/Runtime/CoreUObject/Public/UObject/LinkerLoad.h b/Engine/Source/Runtime/CoreUObject/Public/UObject/LinkerLoad.h
--- a/Engine/Source/Runtime/CoreUObject/Public/UObject/LinkerLoad.h
+++ b/Engine/Source/Runtime/CoreUObject/Public/UObject/LinkerLoad.h
@@ -403,9 +403,6 @@ private:
 
 private:
 
-	/** Allows access to UTexture2D::StaticClass() without linking Core with Engine											*/
-	static UClass* UTexture2DStaticClass;
-
 	static FName NAME_LoadErrors;
 
 	/** Makes sure the deprecated active redirects inis have been read */
@@ -442,11 +439,6 @@ private:
 
 public:
 
-	/**
-	 * Initialize the static variables
-	 */
-	COREUOBJECT_API static void StaticInit(UClass* InUTexture2DStaticClass);
-
 	/**
 	 * Add redirects to FLinkerLoad static map
 	 */
diff --git a/Engine/Source/Runtime/Engine/Private/UnrealEngine.cpp b/Engine/Source/Runtime/Engine/Private/UnrealEngine.cpp
--- a/Engine/Source/Runtime/Engine/Private/UnrealEngine.cpp
+++ b/Engine/Source/Runtime/Engine/Private/UnrealEngine.cpp
@@ -1559,7 +1559,6 @@ void UEngine::Init(IEngineLoop* InEngineLoop)
 
 	// Subsystems.
 	FURL::StaticInit();
-	FLinkerLoad::StaticInit(UTexture2D::StaticClass());
 	EngineSubsystemCollection->Initialize(this);
 
 #if !UE_BUILD_SHIPPING
