Do not sort loaded objects by the linker pointer when loading objects in EndLoad to keep the relative load order,

diff --git a/Engine/Source/Runtime/CoreUObject/Private/UObject/UObjectGlobals.cpp b/Engine/Source/Runtime/CoreUObject/Private/UObject/UObjectGlobals.cpp
--- a/Engine/Source/Runtime/CoreUObject/Private/UObject/UObjectGlobals.cpp
+++ b/Engine/Source/Runtime/CoreUObject/Private/UObject/UObjectGlobals.cpp
@@ -1569,7 +1569,7 @@ struct FCompareUObjectByLinkerAndOffset
 			// Sort by pointer address.
 			else
 			{
-				return LinkerA < LinkerB;
+				return false;
 			}
 		}
 		// Neither objects have a linker, don't do anything.
@@ -1634,7 +1634,7 @@ void EndLoad(FUObjectSerializeContext* LoadContext, TArray<UPackage*>* OutLoaded
 			LoadContext->AppendLoadedObjectsAndEmpty(ObjLoaded);
 
 			// Sort by Filename and Offset.
-			ObjLoaded.Sort(FCompareUObjectByLinkerAndOffset());
+			ObjLoaded.StableSort(FCompareUObjectByLinkerAndOffset());
 
 			// Finish loading everything.
 			{
