Fixed issue with LiveCoding where changing the default value on a class would not be reflected in the editor.

#rnx

diff --git a/Engine/Source/Runtime/CoreUObject/Private/UObject/DeferredRegistry.h b/Engine/Source/Runtime/CoreUObject/Private/UObject/DeferredRegistry.h
--- a/Engine/Source/Runtime/CoreUObject/Private/UObject/DeferredRegistry.h
+++ b/Engine/Source/Runtime/CoreUObject/Private/UObject/DeferredRegistry.h
@@ -94,9 +94,11 @@ public:
 			}
 			if (bAdd)
 			{
-				if (!bHasChanged && *ExistingInfo != &InInfo)
+				if (!bHasChanged)
 				{
-					CopyExistingSingletons(InInfo, **ExistingInfo);
+					// With live coding, the existing might be the same as the new info.  
+					// We still invoke the copy method to allow UClasses to clear the singletons.
+					UpdateSingletons(InInfo, **ExistingInfo);
 				}
 				*ExistingInfo = &InInfo;
 			}
@@ -284,7 +286,7 @@ public:
 	/**
 	* When registering a new info, when the contents haven't changed, we copy singleton pointers 
 	*/
-	void CopyExistingSingletons(TInfo& NewInfo, const TInfo& OldInfo)
+	void UpdateSingletons(TInfo& NewInfo, const TInfo& OldInfo)
 	{
 		NewInfo.InnerSingleton = OldInfo.InnerSingleton;
 		NewInfo.OuterSingleton = OldInfo.OuterSingleton;
@@ -335,10 +337,13 @@ inline typename TDeferredRegistry<FClassRegistrationInfo>::TType* TDeferredRegis
 	return Registrant.InnerRegisterFn();
 }
 
-// Specialization for class to not copy the existing singletons when reloading a class that hasn't changed
+// Specialization for class to not copy the existing singletons when reloading a class that hasn't changed.
+// In the case of live coding, we must clear out the pointer for classes to make sure it invokes the registration code.
 template <>
-inline void TDeferredRegistry<FClassRegistrationInfo>::CopyExistingSingletons(FClassRegistrationInfo& NewInfo, const FClassRegistrationInfo& OldInfo)
+inline void TDeferredRegistry<FClassRegistrationInfo>::UpdateSingletons(FClassRegistrationInfo& NewInfo, const FClassRegistrationInfo& OldInfo)
 {
+	NewInfo.InnerSingleton = nullptr;
+	NewInfo.OuterSingleton = nullptr;
 }
 
 using FClassDeferredRegistry = TDeferredRegistry<FClassRegistrationInfo>;
