Avoid new compile warning stemming from the use of "range-loop-analysis" feature when compiling for Apple target platforms.

diff --git a/Engine/Source/Runtime/CoreUObject/Public/AssetRegistry/AssetData.h b/Engine/Source/Runtime/CoreUObject/Public/AssetRegistry/AssetData.h
--- a/Engine/Source/Runtime/CoreUObject/Public/AssetRegistry/AssetData.h
+++ b/Engine/Source/Runtime/CoreUObject/Public/AssetRegistry/AssetData.h
@@ -377,7 +377,7 @@ public:
 		UE_LOG(LogAssetData, Log, TEXT("        AssetClass: %s"), *AssetClass.ToString());
 		UE_LOG(LogAssetData, Log, TEXT("        TagsAndValues: %d"), TagsAndValues.Num());
 
-		for (const auto& TagValue: TagsAndValues)
+		for (const auto TagValue: TagsAndValues)
 		{
 			UE_LOG(LogAssetData, Log, TEXT("            %s : %s"), *TagValue.Key.ToString(), *TagValue.Value.AsString());
 		}
