Fix data conversion warnings in the bulkdata code

- FUntypedBulkData::SetBulkDataAlignment now takes a uint16 as that is the storage type used anyway.
- FBulkDataBuffer now uses TArrayView64 as the wrapper around it's data. This has been requested by another team anyway.
-- SoundWave.cpp will pass on the return type of FBulkDataBuffer to other systems. So I moved the 'is the data too big for TArrayView checks there and did not change the SoundWave api.

#rnx

diff --git a/Engine/Source/Runtime/CoreUObject/Private/Serialization/BulkData.cpp b/Engine/Source/Runtime/CoreUObject/Private/Serialization/BulkData.cpp
--- a/Engine/Source/Runtime/CoreUObject/Private/Serialization/BulkData.cpp
+++ b/Engine/Source/Runtime/CoreUObject/Private/Serialization/BulkData.cpp
@@ -831,7 +831,7 @@ uint32 FUntypedBulkData::GetBulkDataFlags() const
  *
  * @param BulkDataAlignmentToSet	Bulk data alignment to set
  */
-void FUntypedBulkData::SetBulkDataAlignment( uint32 BulkDataAlignmentToSet )
+void FUntypedBulkData::SetBulkDataAlignment(uint16 BulkDataAlignmentToSet)
 {
 	BulkDataAlignment = BulkDataAlignmentToSet;
 }
@@ -2376,7 +2376,7 @@ void FFloatBulkDataOld::SerializeElement( FArchive& Ar, void* Data, int64 Elemen
 	Ar << FloatData;
 }
 
-void FFormatContainer::Serialize(FArchive& Ar, UObject* Owner, const TArray<FName>* FormatsToSave, bool bSingleUse, uint32 InAlignment, bool bInline, bool bMapped)
+void FFormatContainer::Serialize(FArchive& Ar, UObject* Owner, const TArray<FName>* FormatsToSave, bool bSingleUse, uint16 InAlignment, bool bInline, bool bMapped)
 {
 	if (Ar.IsLoading())
 	{
diff --git a/Engine/Source/Runtime/CoreUObject/Private/Tests/Virtualization/VirtualizedBulkDataTests.cpp b/Engine/Source/Runtime/CoreUObject/Private/Tests/Virtualization/VirtualizedBulkDataTests.cpp
--- a/Engine/Source/Runtime/CoreUObject/Private/Tests/Virtualization/VirtualizedBulkDataTests.cpp
+++ b/Engine/Source/Runtime/CoreUObject/Private/Tests/Virtualization/VirtualizedBulkDataTests.cpp
@@ -132,7 +132,7 @@ bool FVirtualizationWrapperTestUpdatePayload::RunTest(const FString& Parameters)
 		TestEqual(TEXT("Updated payload length"), (int64)Payload.GetSize(), BufferSize);
 		TestEqual(TEXT("Payload and EditablePayload should have the same memory addresses"), (uint8*)Payload.GetData(), (uint8*)EditedPayload.GetData());
 
-		const bool bAllElementsCorrect = Algo::AllOf(TArrayView<uint8>((uint8*)Payload.GetData(), Payload.GetSize()), [NewValue](uint8 Val)
+		const bool bAllElementsCorrect = Algo::AllOf(TArrayView64<uint8>((uint8*)Payload.GetData(), (int64)Payload.GetSize()), [NewValue](uint8 Val)
 			{ 
 				return Val == NewValue; 
 			});
@@ -154,7 +154,7 @@ bool FVirtualizationWrapperTestUpdatePayload::RunTest(const FString& Parameters)
 		TestEqual(TEXT("Payload and OriginalDataPtr should have the same memory addresses"), (uint8*)Payload.GetData(), OriginalDataPtr);
 
 		// The original data was all zeros, so we can test for that to make sure that the contents are correct.
-		const bool bAllElementsCorrect = Algo::AllOf(TArrayView<uint8>((uint8*)Payload.GetData(), Payload.GetSize()), [](uint8 Val)
+		const bool bAllElementsCorrect = Algo::AllOf(TArrayView64<uint8>((uint8*)Payload.GetData(), (int64)Payload.GetSize()), [](uint8 Val)
 			{
 				return Val == 0;
 			});
diff --git a/Engine/Source/Runtime/CoreUObject/Public/Serialization/BulkData.h b/Engine/Source/Runtime/CoreUObject/Public/Serialization/BulkData.h
--- a/Engine/Source/Runtime/CoreUObject/Public/Serialization/BulkData.h
+++ b/Engine/Source/Runtime/CoreUObject/Public/Serialization/BulkData.h
@@ -491,7 +491,7 @@ public:
 	 *
 	 * @param BulkDataAlignmentToSet	Bulk data alignment to set
 	 */
-	void SetBulkDataAlignment( uint32 BulkDataAlignmentToSet );
+	void SetBulkDataAlignment(uint16 BulkDataAlignmentToSet);
 
 	/**
 	* Gets the current bulk data alignment.
@@ -1089,6 +1089,6 @@ public:
 		}
 		Formats.Empty();
 	}
-	COREUOBJECT_API void Serialize(FArchive& Ar, UObject* Owner, const TArray<FName>* FormatsToSave = nullptr, bool bSingleUse = true, uint32 InAlignment = DEFAULT_ALIGNMENT, bool bInline = true, bool bMapped = false);
+	COREUOBJECT_API void Serialize(FArchive& Ar, UObject* Owner, const TArray<FName>* FormatsToSave = nullptr, bool bSingleUse = true, uint16 InAlignment = DEFAULT_ALIGNMENT, bool bInline = true, bool bMapped = false);
 	COREUOBJECT_API void SerializeAttemptMappedLoad(FArchive& Ar, UObject* Owner);
 };
diff --git a/Engine/Source/Runtime/CoreUObject/Public/Serialization/BulkDataBuffer.h b/Engine/Source/Runtime/CoreUObject/Public/Serialization/BulkDataBuffer.h
--- a/Engine/Source/Runtime/CoreUObject/Public/Serialization/BulkDataBuffer.h
+++ b/Engine/Source/Runtime/CoreUObject/Public/Serialization/BulkDataBuffer.h
@@ -15,7 +15,7 @@ template<typename DataType>
 class FBulkDataBuffer
 {
 public:
-	using ViewType = TArrayView<DataType>;
+	using ViewType = TArrayView64<DataType>;
 
 	/** Constructor. */
 	FBulkDataBuffer() = default;
@@ -47,10 +47,10 @@ public:
 	* @param InBuffer Pointer to memory to take ownership of, must've have been allocated via FMemory::Malloc/Realloc.
 	* @param InNumberOfElements The number of elements in the buffer.
 	*/
-	FBulkDataBuffer(DataType* InBuffer, uint64 InNumberOfElements)
+	FBulkDataBuffer(DataType* InBuffer, int64 InNumberOfElements)
 		: View(InBuffer, InNumberOfElements)
 	{
-		check(InNumberOfElements <= TNumericLimits<uint32>::Max());
+
 	}
 
 	/** Destructor. */
diff --git a/Engine/Source/Runtime/CoreUObject/Public/Serialization/BulkDataCommon.h b/Engine/Source/Runtime/CoreUObject/Public/Serialization/BulkDataCommon.h
--- a/Engine/Source/Runtime/CoreUObject/Public/Serialization/BulkDataCommon.h
+++ b/Engine/Source/Runtime/CoreUObject/Public/Serialization/BulkDataCommon.h
@@ -100,7 +100,7 @@ inline void SerializeBulkDataSizeInt(FArchive& Ar, int64& Value, EBulkDataFlags
 	else
 	{
 		check(!Ar.IsSaving() || (MIN_int32 <= Value && Value <= MAX_int32));
-		int32 ValueAsInt32 = Value;
+		int32 ValueAsInt32 = static_cast<int32>(Value);
 		Ar << ValueAsInt32;
 		Value = ValueAsInt32;
 	}
diff --git a/Engine/Source/Runtime/Engine/Private/SoundWave.cpp b/Engine/Source/Runtime/Engine/Private/SoundWave.cpp
--- a/Engine/Source/Runtime/Engine/Private/SoundWave.cpp
+++ b/Engine/Source/Runtime/Engine/Private/SoundWave.cpp
@@ -2446,7 +2446,9 @@ TArrayView<const uint8> USoundWave::GetZerothChunk(bool bForImmediatePlayback)
 			IStreamingManager::Get().GetAudioStreamingManager().RequestChunk(InternalProxy, 1, [](EAudioChunkLoadResult InResult) {}, ENamedThreads::AnyThread, bForImmediatePlayback);
 		}
 
-		return ZerothChunkData->GetView();
+		FBulkDataBuffer<uint8>::ViewType View = ZerothChunkData->GetView();
+		check(View.Num() <= MAX_int32);
+		return TArrayView<uint8>(View.GetData(), static_cast<int32>(View.Num()));
 	}
 	else
 	{
@@ -3239,7 +3241,9 @@ bool FSoundWaveProxy::IsZerothChunkDataLoaded() const
 
 const TArrayView<uint8> FSoundWaveProxy::GetZerothChunkDataView() const
 {
-	return ZerothChunkData->GetView();
+	FBulkDataBuffer<uint8>::ViewType View = ZerothChunkData->GetView();
+	check(View.Num() <= MAX_int32);
+	return TArrayView<uint8>(View.GetData(), static_cast<int32>(View.Num()));
 }
 
 void FSoundWaveProxy::EnsureZerothChunkIsLoaded()
