Clean up FVirtualizedUntypedBulkData::SerializeData to use the recently added serialization operator for FCompressedBuffer.

#rnx

diff --git a/Engine/Source/Runtime/CoreUObject/Private/Virtualization/VirtualizedBulkData.cpp b/Engine/Source/Runtime/CoreUObject/Private/Virtualization/VirtualizedBulkData.cpp
--- a/Engine/Source/Runtime/CoreUObject/Private/Virtualization/VirtualizedBulkData.cpp
+++ b/Engine/Source/Runtime/CoreUObject/Private/Virtualization/VirtualizedBulkData.cpp
@@ -597,17 +597,12 @@ bool FVirtualizedUntypedBulkData::SerializeData(FArchive& Ar, FCompressedBuffer&
 
 	if (Ar.IsSaving()) // Saving to virtualized bulkdata format
 	{
-		for (const FSharedBuffer& Buffer : InPayload.GetCompressed().GetSegments())
-		{
-			// Const cast because FArchive requires a non-const pointer!
-			Ar.Serialize(const_cast<void*>(Buffer.GetData()), static_cast<int64>(Buffer.GetSize()));
-		}
-		
+		Ar << InPayload;
 		return true;
 	}
 	else if(Ar.IsLoading() && !EnumHasAnyFlags(PayloadFlags, EFlags::ReferencesLegacyFile)) // Loading from virtualized bulkdata format
 	{
-		InPayload = FCompressedBuffer::FromCompressed(Ar);
+		Ar << InPayload;
 		return InPayload.IsNull();	
 	}
 	else if (Ar.IsLoading()) // Loading from old bulkdata format
