AsyncLoading2:
Fixed pending io requests trace counter
Added total loaded trace counter
#rnx

diff --git a/Engine/Source/Runtime/CoreUObject/Private/Serialization/AsyncLoading2.cpp b/Engine/Source/Runtime/CoreUObject/Private/Serialization/AsyncLoading2.cpp
--- a/Engine/Source/Runtime/CoreUObject/Private/Serialization/AsyncLoading2.cpp
+++ b/Engine/Source/Runtime/CoreUObject/Private/Serialization/AsyncLoading2.cpp
@@ -298,7 +298,8 @@ if ((Condition)) \
 CSV_DECLARE_CATEGORY_MODULE_EXTERN(CORE_API, Basic);
 CSV_DECLARE_CATEGORY_MODULE_EXTERN(CORE_API, FileIO);
 
-TRACE_DECLARE_INT_COUNTER(PendingBundleIoRequests, TEXT("AsyncLoading/PendingBundleIoRequests"));
+TRACE_DECLARE_INT_COUNTER(AsyncLoadingPendingIoRequests, TEXT("AsyncLoading/PendingIoRequests"));
+TRACE_DECLARE_MEMORY_COUNTER(AsyncLoadingTotalLoaded, TEXT("AsyncLoading/TotalLoaded"));
 
 struct FAsyncPackage2;
 class FAsyncLoadingThread2;
@@ -1938,7 +1939,7 @@ private:
 	TMap<FPackageId, FAsyncPackage2*> AsyncPackageLookup;
 
 	TQueue<FAsyncPackage2*, EQueueMode::Mpsc> ExternalReadQueue;
-	FThreadSafeCounter WaitingForIoBundleCounter;
+	TAtomic<int32> WaitingForIoBundleCounter{ 0 };
 	
 	/** List of all pending package requests */
 	TSet<int32> PendingRequests;
@@ -2575,7 +2576,8 @@ bool FAsyncLoadingThread2::CreateAsyncPackagesFromQueue(FAsyncLoadingThreadState
 
 void FAsyncLoadingThread2::AddBundleIoRequest(FAsyncPackage2* Package)
 {
-	WaitingForIoBundleCounter.Increment();
+	int32 LocalWaitingForIoBundleCount = WaitingForIoBundleCounter.IncrementExchange() + 1;
+	TRACE_COUNTER_SET(AsyncLoadingPendingIoRequests, LocalWaitingForIoBundleCount);
 	WaitingIoRequests.HeapPush({ Package });
 }
 
@@ -2621,10 +2623,11 @@ void FAsyncLoadingThread2::StartBundleIoRequests()
 					TEXT("Failed reading chunk for package: %s"), *Result.Status().ToString());
 				Package->bLoadHasFailed = true;
 			}
-			Package->AsyncLoadingThread.WaitingForIoBundleCounter.Decrement();
+			int32 LocalWaitingForIoBundleCounter = Package->AsyncLoadingThread.WaitingForIoBundleCounter.DecrementExchange() - 1;
+			TRACE_COUNTER_SET(AsyncLoadingPendingIoRequests, LocalWaitingForIoBundleCounter);
+			TRACE_COUNTER_ADD(AsyncLoadingTotalLoaded, Package->IoBuffer.DataSize());
 			Package->GetPackageNode(EEventLoadNode2::Package_ProcessSummary).ReleaseBarrier();
 		});
-		TRACE_COUNTER_DECREMENT(PendingBundleIoRequests);
 	}
 	IoBatch.Issue();
 }
@@ -4731,7 +4734,7 @@ uint32 FAsyncLoadingThread2::Run()
 			if (!bDidSomething)
 			{
 				FAsyncPackage2* Package = nullptr;
-				if (WaitingForIoBundleCounter.GetValue() > 0)
+				if (WaitingForIoBundleCounter.Load() > 0)
 				{
 					TRACE_CPUPROFILER_EVENT_SCOPE(AsyncLoadingTime);
 					TRACE_CPUPROFILER_EVENT_SCOPE(WaitingForIo);
