FScriptInterface/TScriptInterface improvements

 - Added type-safe overloads of GetInterface/SetInterface on TScriptInterface.
 - Added GetTypeHash to FScriptInterface.
 - Cleaned-up operator= overloads with non-standard return types.
 - Ensure that FScriptInterface is always constructed with both an object and interface argument specified.

diff --git a/Engine/Source/Runtime/CoreUObject/Public/UObject/ScriptInterface.h b/Engine/Source/Runtime/CoreUObject/Public/UObject/ScriptInterface.h
--- a/Engine/Source/Runtime/CoreUObject/Public/UObject/ScriptInterface.h
+++ b/Engine/Source/Runtime/CoreUObject/Public/UObject/ScriptInterface.h
@@ -9,6 +9,7 @@
 #include "CoreMinimal.h"
 #include "UObject/UObjectGlobals.h"
 #include "Templates/Casts.h"
+#include "Templates/UnrealTemplate.h"
 
 /**
  * FScriptInterface
@@ -21,12 +22,12 @@ private:
 	/**
 	 * A pointer to a UObject that implements a native interface.
 	 */
-	UObject*	ObjectPointer;
+	UObject*	ObjectPointer = nullptr;
 
 	/**
 	 * Pointer to the location of the interface object within the UObject referenced by ObjectPointer.
 	 */
-	void*		InterfacePointer;
+	void*		InterfacePointer = nullptr;
 
 	/**
 	 * Serialize ScriptInterface
@@ -37,15 +38,20 @@ public:
 	/**
 	 * Default constructor
 	 */
-	FScriptInterface( UObject* InObjectPointer=NULL, void* InInterfacePointer=NULL )
+	FScriptInterface() = default;
+
+	/**
+	 * Construction from object and interface
+	 */
+	FScriptInterface( UObject* InObjectPointer, void* InInterfacePointer )
 	: ObjectPointer(InObjectPointer), InterfacePointer(InInterfacePointer)
 	{}
 
-	void operator=(const FScriptInterface& Other)
-	{
-		ObjectPointer = Other.ObjectPointer;
-		InterfacePointer = Other.InterfacePointer;
-	}
+	/**
+	 * Copyable
+	 */
+	FScriptInterface(const FScriptInterface&) = default;
+	FScriptInterface& operator=(const FScriptInterface&) = default;
 
 	/**
 	 * Returns the ObjectPointer contained by this FScriptInterface
@@ -70,7 +76,7 @@ public:
 	{
 		// only allow access to InterfacePointer if we have a valid ObjectPointer.  This is necessary because the garbage collector will set ObjectPointer to NULL
 		// without using the accessor methods
-		return ObjectPointer != NULL ? InterfacePointer : NULL;
+		return ObjectPointer ? InterfacePointer : nullptr;
 	}
 
 	/**
@@ -79,9 +85,9 @@ public:
 	FORCEINLINE void SetObject( UObject* InObjectPointer )
 	{
 		ObjectPointer = InObjectPointer;
-		if ( ObjectPointer == NULL )
+		if ( ObjectPointer == nullptr )
 		{
-			SetInterface(NULL);
+			SetInterface(nullptr);
 		}
 	}
 
@@ -109,6 +115,11 @@ public:
 	{
 		Collector.AddReferencedObject(ObjectPointer);
 	}
+
+	friend inline uint32 GetTypeHash(const FScriptInterface& Instance)
+	{
+		return GetTypeHash(Instance.InterfacePointer);
+	}
 };
 
 
@@ -119,13 +130,14 @@ template<> struct TIsZeroConstructType<class FScriptInterface> { enum { Value =
 /**
  * Templated version of FScriptInterface, which provides accessors and operators for referencing the interface portion of a UObject that implements a native interface.
  */
-template< class InterfaceType > class TScriptInterface : public FScriptInterface
+template <typename InterfaceType>
+class TScriptInterface : public FScriptInterface
 {
 public:
 	/**
 	 * Default constructor
 	 */
-	TScriptInterface() {}
+	TScriptInterface() = default;
 
 	/**
 	 * Construction from nullptr
@@ -133,44 +145,40 @@ public:
 	TScriptInterface(TYPE_OF_NULLPTR) {}
 
 	/**
-	 * Standard constructor.
-	 *
-	 * @param	SourceObject	a pointer to a UObject that implements the InterfaceType native interface class.
+	 * Assignment from an object type that implements the InterfaceType native interface class
 	 */
-	template <class UObjectType> TScriptInterface( UObjectType* SourceObject )
+	template <typename ObjectType, typename = decltype(ImplicitConv<UObject*>((ObjectType*)nullptr))>
+	TScriptInterface(ObjectType* SourceObject)
 	{
-		(*this) = SourceObject;
+		SetObject(SourceObject);
+
+		InterfaceType* SourceInterface = Cast<InterfaceType>(SourceObject);
+		SetInterface(SourceInterface);
 	}
+
 	/**
-	 * Copy constructor
+	 * Copyable
 	 */
-	TScriptInterface( const TScriptInterface& Other )
-	{
-		SetObject(Other.GetObject());
-		SetInterface(Other.GetInterface());
-	}
+	TScriptInterface(const TScriptInterface&) = default;
+	TScriptInterface& operator=(const TScriptInterface&) = default;
 
 	/**
-	 * Assignment operator.
-	 *
-	 * @param	SourceObject	a pointer to a UObject that implements the InterfaceType native interface class.
+	 * Assignment from nullptr
 	 */
-	template<class UObjectType> InterfaceType& operator=( UObjectType* SourceObject )
+	TScriptInterface& operator=(TYPE_OF_NULLPTR)
 	{
-		SetObject(SourceObject);
-		
-		InterfaceType* SourceInterface = Cast<InterfaceType>(SourceObject);
-		SetInterface( SourceInterface );
-
-		return *((InterfaceType*)GetInterface());
+		*this = TScriptInterface();
+		return *this;
 	}
 
 	/**
-	 * Assignment from nullptr
+	 * Assignment from an object type that implements the InterfaceType native interface class
 	 */
-	void operator=(TYPE_OF_NULLPTR)
+	template <typename ObjectType, typename = decltype(ImplicitConv<UObject*>((ObjectType*)nullptr))>
+	TScriptInterface& operator=(ObjectType* SourceObject)
 	{
-		*this = TScriptInterface();
+		*this = TScriptInterface(SourceObject);
+		return *this;
 	}
 
 	/**
@@ -202,7 +210,7 @@ public:
 	 */
 	FORCEINLINE InterfaceType* operator->() const
 	{
-		return (InterfaceType*)GetInterface();
+		return GetInterface();
 	}
 
 	/**
@@ -212,7 +220,23 @@ public:
 	 */
 	FORCEINLINE InterfaceType& operator*() const
 	{
-		return *((InterfaceType*)GetInterface());
+		return *GetInterface();
+	}
+
+	/**
+	 * Returns the pointer to the interface
+	 */
+	FORCEINLINE InterfaceType* GetInterface() const
+	{
+		return (InterfaceType*)FScriptInterface::GetInterface();
+	}
+
+	/**
+	 * Sets the value of the InterfacePointer for this TScriptInterface
+	 */
+	FORCEINLINE void SetInterface(InterfaceType* InInterfacePointer)
+	{
+		FScriptInterface::SetInterface(InInterfacePointer);
 	}
 
 	/**
@@ -222,7 +246,7 @@ public:
 	 */
 	FORCEINLINE explicit operator bool() const
 	{
-		return GetInterface() != NULL;
+		return GetInterface() != nullptr;
 	}
 
 	friend FArchive& operator<<( FArchive& Ar, TScriptInterface& Interface )
