Compiletime improvement: remove LinkerSave.h from Package.h and CoreUObject.h.
#rnx

diff --git a/Engine/Plugins/FX/Niagara/Source/Niagara/Private/NiagaraScript.cpp b/Engine/Plugins/FX/Niagara/Source/Niagara/Private/NiagaraScript.cpp
--- a/Engine/Plugins/FX/Niagara/Source/Niagara/Private/NiagaraScript.cpp
+++ b/Engine/Plugins/FX/Niagara/Source/Niagara/Private/NiagaraScript.cpp
@@ -36,6 +36,7 @@
 
 #include "UObject/FortniteMainBranchObjectVersion.h"
 #include "UObject/RenderingObjectVersion.h"
+#include "UObject/UObjectThreadContext.h"
 #include "Serialization/ObjectAndNameAsStringProxyArchive.h"
 
 #include "VectorVM.h"
diff --git a/Engine/Plugins/FX/Niagara/Source/Niagara/Private/NiagaraSystem.cpp b/Engine/Plugins/FX/Niagara/Source/Niagara/Private/NiagaraSystem.cpp
--- a/Engine/Plugins/FX/Niagara/Source/Niagara/Private/NiagaraSystem.cpp
+++ b/Engine/Plugins/FX/Niagara/Source/Niagara/Private/NiagaraSystem.cpp
@@ -27,6 +27,7 @@
 #include "ProfilingDebugging/CookStats.h"
 #include "UObject/ObjectSaveContext.h"
 #include "UObject/Package.h"
+#include "UObject/UObjectThreadContext.h"
 #include "ShaderCompiler.h"
 
 #define LOCTEXT_NAMESPACE "NiagaraSystem"
diff --git a/Engine/Source/Editor/UnrealEd/Private/Commandlets/SavePackageUtilitiesCommandlet.cpp b/Engine/Source/Editor/UnrealEd/Private/Commandlets/SavePackageUtilitiesCommandlet.cpp
--- a/Engine/Source/Editor/UnrealEd/Private/Commandlets/SavePackageUtilitiesCommandlet.cpp
+++ b/Engine/Source/Editor/UnrealEd/Private/Commandlets/SavePackageUtilitiesCommandlet.cpp
@@ -12,6 +12,7 @@
 #include "Misc/PackageName.h"
 #include "Misc/Paths.h"
 #include "UObject/LinkerDiff.h"
+#include "UObject/LinkerSave.h"
 #include "UObject/Package.h"
 #include "UObject/UObjectGlobals.h"
 #include "UObject/UObjectHash.h"
diff --git a/Engine/Source/Editor/UnrealEd/Private/CookOnTheFlyServer.cpp b/Engine/Source/Editor/UnrealEd/Private/CookOnTheFlyServer.cpp
--- a/Engine/Source/Editor/UnrealEd/Private/CookOnTheFlyServer.cpp
+++ b/Engine/Source/Editor/UnrealEd/Private/CookOnTheFlyServer.cpp
@@ -40,6 +40,7 @@
 #include "UObject/Class.h"
 #include "UObject/UObjectIterator.h"
 #include "UObject/LinkerDiff.h"
+#include "UObject/LinkerSave.h"
 #include "UObject/Package.h"
 #include "UObject/SavePackage.h"
 #include "UObject/MetaData.h"
diff --git a/Engine/Source/Runtime/CoreUObject/Private/UObject/Field.cpp b/Engine/Source/Runtime/CoreUObject/Private/UObject/Field.cpp
--- a/Engine/Source/Runtime/CoreUObject/Private/UObject/Field.cpp
+++ b/Engine/Source/Runtime/CoreUObject/Private/UObject/Field.cpp
@@ -29,7 +29,6 @@ Field.cpp: Defines FField property system fundamentals
 #include "UObject/Stack.h"
 #include "Misc/PackageName.h"
 #include "UObject/ObjectResource.h"
-#include "UObject/LinkerSave.h"
 #include "UObject/Interface.h"
 #include "UObject/LinkerLoad.h"
 #include "UObject/LinkerPlaceholderClass.h"
diff --git a/Engine/Source/Runtime/CoreUObject/Private/UObject/Package.cpp b/Engine/Source/Runtime/CoreUObject/Private/UObject/Package.cpp
--- a/Engine/Source/Runtime/CoreUObject/Private/UObject/Package.cpp
+++ b/Engine/Source/Runtime/CoreUObject/Private/UObject/Package.cpp
@@ -7,6 +7,7 @@
 #include "Misc/ITransaction.h"
 #include "Misc/PackageName.h"
 #include "UObject/LinkerLoad.h"
+#include "UObject/LinkerSave.h"
 #include "UObject/LinkerManager.h"
 #include "UObject/MetaData.h"
 #include "UObject/PackageResourceManager.h"
diff --git a/Engine/Source/Runtime/CoreUObject/Private/UObject/PropertyStruct.cpp b/Engine/Source/Runtime/CoreUObject/Private/UObject/PropertyStruct.cpp
--- a/Engine/Source/Runtime/CoreUObject/Private/UObject/PropertyStruct.cpp
+++ b/Engine/Source/Runtime/CoreUObject/Private/UObject/PropertyStruct.cpp
@@ -10,6 +10,7 @@
 #include "UObject/LinkerPlaceholderBase.h"
 #include "Serialization/ArchiveUObjectFromStructuredArchive.h"
 #include "Hash/Blake3.h"
+#include "IO/IoHash.h"
 
 // WARNING: This should always be the last include in any file that needs it (except .generated.h)
 #include "UObject/UndefineUPropertyMacros.h"
diff --git a/Engine/Source/Runtime/CoreUObject/Private/UObject/SavePackage.cpp b/Engine/Source/Runtime/CoreUObject/Private/UObject/SavePackage.cpp
--- a/Engine/Source/Runtime/CoreUObject/Private/UObject/SavePackage.cpp
+++ b/Engine/Source/Runtime/CoreUObject/Private/UObject/SavePackage.cpp
@@ -2259,7 +2259,7 @@ FSavePackageResultStruct UPackage::Save(UPackage* InOuter, UObject* Base, EObjec
 
 		// Size of serialized out package in bytes. This is before compression.
 		int32 PackageSize = INDEX_NONE;
-		TUniquePtr<FLinkerSave> Linker = nullptr;
+		TPimplPtr<FLinkerSave> Linker = nullptr;
 		uint32 SerializedPackageFlags = 0;
 		{			
 			// TODO: Require a SavePackageContext and move to EditorEngine
@@ -2403,13 +2403,13 @@ FSavePackageResultStruct UPackage::Save(UPackage* InOuter, UObject* Base, EObjec
 						// The entire package will be serialized to memory and then compared against package on disk.
 						// Each difference will be log with its Serialize call stack trace
 						FArchive* Saver = new FArchiveStackTrace(FindAssetInPackage(InOuter), *InOuter->GetLoadedPath().GetPackageName(), true, InOutDiffMap);
-						Linker = TUniquePtr<FLinkerSave>(new FLinkerSave(InOuter, Saver, bForceByteSwapping, bSaveUnversioned));
+						Linker = MakePimpl<FLinkerSave>(InOuter, Saver, bForceByteSwapping, bSaveUnversioned);
 					}
 					else if (TargetPlatform != nullptr && (SaveFlags & SAVE_DiffOnly))
 					{
 						// The entire package will be serialized to memory and then compared against package on disk
 						FArchive* Saver = new FArchiveStackTrace(FindAssetInPackage(InOuter), *InOuter->GetLoadedPath().GetPackageName(), false);
-						Linker = TUniquePtr<FLinkerSave>(new FLinkerSave(InOuter, Saver, bForceByteSwapping, bSaveUnversioned));
+						Linker = MakePimpl<FLinkerSave>(InOuter, Saver, bForceByteSwapping, bSaveUnversioned);
 					}
 					else if ((!!TargetPlatform) && FParse::Value(FCommandLine::Get(), TEXT("DiffCookedPackages="), DiffCookedPackagesPath))
 					{
@@ -2424,7 +2424,7 @@ FSavePackageResultStruct UPackage::Save(UPackage* InOuter, UObject* Base, EObjec
 
 						FArchive* TestArchive = IFileManager::Get().CreateFileReader(*TestArchiveFilename);
 						FArchive* Saver = new FDiffSerializeArchive(*InOuter->GetLoadedPath().GetPackageName(), TestArchive);
-						Linker = TUniquePtr<FLinkerSave>(new FLinkerSave(InOuter, Saver, bForceByteSwapping));
+						Linker = MakePimpl<FLinkerSave>(InOuter, Saver, bForceByteSwapping);
 					}
 					else
 #endif
@@ -2432,13 +2432,13 @@ FSavePackageResultStruct UPackage::Save(UPackage* InOuter, UObject* Base, EObjec
 						if (bSaveAsync)
 						{
 							// Allocate the linker with a memory writer, forcing byte swapping if wanted.
-							Linker = TUniquePtr<FLinkerSave>(new FLinkerSave(InOuter, bForceByteSwapping, bSaveUnversioned));
+							Linker = MakePimpl<FLinkerSave>(InOuter, bForceByteSwapping, bSaveUnversioned);
 						}
 						else
 						{
 							// Allocate the linker, forcing byte swapping if wanted.
 							TempFilename = FPaths::CreateTempFilename(*FPaths::ProjectSavedDir(), *BaseFilename.Left(32));
-							Linker = TUniquePtr<FLinkerSave>(new FLinkerSave(InOuter, *TempFilename.GetValue(), bForceByteSwapping, bSaveUnversioned));
+							Linker = MakePimpl<FLinkerSave>(InOuter, *TempFilename.GetValue(), bForceByteSwapping, bSaveUnversioned);
 						}
 					}
 					Linker->bProceduralSave = ObjectSaveContext.bProceduralSave;
diff --git a/Engine/Source/Runtime/CoreUObject/Private/UObject/SavePackage/SaveContext.h b/Engine/Source/Runtime/CoreUObject/Private/UObject/SavePackage/SaveContext.h
--- a/Engine/Source/Runtime/CoreUObject/Private/UObject/SavePackage/SaveContext.h
+++ b/Engine/Source/Runtime/CoreUObject/Private/UObject/SavePackage/SaveContext.h
@@ -12,6 +12,7 @@
 #include "Serialization/PropertyLocalizationDataGathering.h"
 #include "Serialization/StructuredArchive.h"
 #include "Serialization/UnversionedPropertySerialization.h"
+#include "Templates/PimplPtr.h"
 #include "Templates/UniquePtr.h"
 #include "UObject/AsyncWorkSequence.h"
 #include "UObject/LinkerSave.h"
@@ -550,7 +551,7 @@ public:
 public:
 	ESavePackageResult Result;
 
-	TUniquePtr<FLinkerSave> Linker;
+	TPimplPtr<FLinkerSave> Linker;
 	TUniquePtr<FArchive> TextFormatArchive;
 	TUniquePtr<FArchiveFormatterType> Formatter;
 	TUniquePtr<FStructuredArchive> StructuredArchive;
diff --git a/Engine/Source/Runtime/CoreUObject/Private/UObject/SavePackage2.cpp b/Engine/Source/Runtime/CoreUObject/Private/UObject/SavePackage2.cpp
--- a/Engine/Source/Runtime/CoreUObject/Private/UObject/SavePackage2.cpp
+++ b/Engine/Source/Runtime/CoreUObject/Private/UObject/SavePackage2.cpp
@@ -685,20 +685,20 @@ ESavePackageResult CreateLinker(FSaveContext& SaveContext)
 			// The entire package will be serialized to memory and then compared against package on disk.
 			// Each difference will be log with its Serialize call stack trace if IsDiffCallstack is true
 			FArchive* Saver = new FArchiveStackTrace(SaveContext.GetAsset(), *SaveContext.GetPackage()->GetLoadedPath().GetPackageName(), SaveContext.IsDiffCallstack(), SaveContext.GetDiffMapPtr());
-			SaveContext.Linker = MakeUnique<FLinkerSave>(SaveContext.GetPackage(), Saver, SaveContext.IsForceByteSwapping(), SaveContext.IsSaveUnversioned());
+			SaveContext.Linker = MakePimpl<FLinkerSave>(SaveContext.GetPackage(), Saver, SaveContext.IsForceByteSwapping(), SaveContext.IsSaveUnversioned());
 		}
 		else
 		{
 			if (SaveContext.IsSaveAsync())
 			{
 				// Allocate the linker with a memory writer, forcing byte swapping if wanted.
-				SaveContext.Linker = MakeUnique<FLinkerSave>(SaveContext.GetPackage(), SaveContext.IsForceByteSwapping(), SaveContext.IsSaveUnversioned());
+				SaveContext.Linker = MakePimpl<FLinkerSave>(SaveContext.GetPackage(), SaveContext.IsForceByteSwapping(), SaveContext.IsSaveUnversioned());
 			}
 			else
 			{
 				// Allocate the linker, forcing byte swapping if wanted.
 				SaveContext.TempFilename = FPaths::CreateTempFilename(*FPaths::ProjectSavedDir(), *BaseFilename.Left(32));
-				SaveContext.Linker = MakeUnique<FLinkerSave>(SaveContext.GetPackage(), *SaveContext.TempFilename.GetValue(), SaveContext.IsForceByteSwapping(), SaveContext.IsSaveUnversioned());
+				SaveContext.Linker = MakePimpl<FLinkerSave>(SaveContext.GetPackage(), *SaveContext.TempFilename.GetValue(), SaveContext.IsForceByteSwapping(), SaveContext.IsSaveUnversioned());
 			}
 		}
 		if (SaveContext.IsGenerateSaveError())
diff --git a/Engine/Source/Runtime/CoreUObject/Public/CoreUObject.h b/Engine/Source/Runtime/CoreUObject/Public/CoreUObject.h
--- a/Engine/Source/Runtime/CoreUObject/Public/CoreUObject.h
+++ b/Engine/Source/Runtime/CoreUObject/Public/CoreUObject.h
@@ -79,7 +79,6 @@ MONOLITHIC_HEADER_BOILERPLATE()
 #include "UObject/ObjectResource.h"
 #include "UObject/Linker.h"
 #include "UObject/LinkerLoad.h"
-#include "UObject/LinkerSave.h"
 #include "UObject/GCObject.h"
 #include "Misc/NotifyHook.h"
 #include "Misc/RedirectCollector.h"
diff --git a/Engine/Source/Runtime/CoreUObject/Public/UObject/Package.h b/Engine/Source/Runtime/CoreUObject/Public/UObject/Package.h
--- a/Engine/Source/Runtime/CoreUObject/Public/UObject/Package.h
+++ b/Engine/Source/Runtime/CoreUObject/Public/UObject/Package.h
@@ -7,7 +7,6 @@
 #include "UObject/UObjectGlobals.h"
 #include "UObject/Object.h"
 #include "UObject/PackageId.h"
-#include "UObject/LinkerSave.h"
 #include "Misc/Guid.h"
 #include "Misc/OutputDeviceError.h"
 #include "Misc/ObjectThumbnail.h"
@@ -15,6 +14,7 @@
 #include "Misc/SecureHash.h"
 #include "Misc/WorldCompositionUtility.h"
 #include "Serialization/CustomVersion.h"
+#include "Templates/PimplPtr.h"
 #include "Templates/UniquePtr.h"
 #include "Async/Future.h"
 
@@ -23,6 +23,7 @@ class Error;
 // This is a dummy type which is not implemented anywhere. It's only 
 // used to flag a deprecated Conform argument to package save functions.
 class FLinkerNull;
+class FLinkerSave;
 struct FPackageSaveInfo;
 class FSavePackageContext;
 struct FSavePackageArgs;
@@ -70,13 +71,13 @@ struct FSavePackageResultStruct
 	uint32 SerializedPackageFlags;
 
 	/** Linker for linker comparison after save. */
-	TUniquePtr<FLinkerSave> LinkerSave;
+	TPimplPtr<FLinkerSave> LinkerSave;
 
 	/** Constructors, it will implicitly construct from the result enum */
 	FSavePackageResultStruct() : Result(ESavePackageResult::Error), TotalFileSize(0), SerializedPackageFlags(0) {}
 	FSavePackageResultStruct(ESavePackageResult InResult) : Result(InResult), TotalFileSize(0), SerializedPackageFlags(0) {}
 	FSavePackageResultStruct(ESavePackageResult InResult, int64 InTotalFileSize) : Result(InResult), TotalFileSize(InTotalFileSize), SerializedPackageFlags(0) {}
-	FSavePackageResultStruct(ESavePackageResult InResult, int64 InTotalFileSize, TFuture<FMD5Hash>&& InHash, uint32 InSerializedPackageFlags, TUniquePtr<FLinkerSave> Linker = nullptr) : Result(InResult), TotalFileSize(InTotalFileSize), CookedHash(MoveTemp(InHash)), SerializedPackageFlags(InSerializedPackageFlags), LinkerSave(MoveTemp(Linker)) {}
+	FSavePackageResultStruct(ESavePackageResult InResult, int64 InTotalFileSize, TFuture<FMD5Hash>&& InHash, uint32 InSerializedPackageFlags, TPimplPtr<FLinkerSave> Linker = nullptr) : Result(InResult), TotalFileSize(InTotalFileSize), CookedHash(MoveTemp(InHash)), SerializedPackageFlags(InSerializedPackageFlags), LinkerSave(MoveTemp(Linker)) {}
 
 	bool operator==(const FSavePackageResultStruct& Other) const
 	{
diff --git a/Engine/Source/Runtime/Engine/Private/AsyncCompilationHelpers.cpp b/Engine/Source/Runtime/Engine/Private/AsyncCompilationHelpers.cpp
--- a/Engine/Source/Runtime/Engine/Private/AsyncCompilationHelpers.cpp
+++ b/Engine/Source/Runtime/Engine/Private/AsyncCompilationHelpers.cpp
@@ -19,6 +19,7 @@
 #include "ShaderCompiler.h"
 #include "Misc/ScopeRWLock.h"
 #include "ProfilingDebugging/StallDetector.h"
+#include "UObject/UObjectThreadContext.h"
 
 #endif // #if WITH_EDITOR
 
