Fix problem when saving map redirector with Save2

diff --git a/Engine/Source/Runtime/CoreUObject/Private/UObject/SavePackage2.cpp b/Engine/Source/Runtime/CoreUObject/Private/UObject/SavePackage2.cpp
--- a/Engine/Source/Runtime/CoreUObject/Private/UObject/SavePackage2.cpp
+++ b/Engine/Source/Runtime/CoreUObject/Private/UObject/SavePackage2.cpp
@@ -575,6 +575,26 @@ ESavePackageResult ValidateImports(FSaveContext& SaveContext)
 		return false;
 
 	};
+	auto IsMapReferenceAllowed = [&SaveContext](UObject* InImport) -> bool
+	{
+		if (!InImport->HasAnyFlags(RF_Public))
+		{
+			return false;
+		}
+
+		// If we have a public import from a map (i.e. the world) only redirector are allowed to have a hard reference
+		for (const auto& Pair : SaveContext.GetObjectDependencies())
+		{
+			if (Pair.Key->GetClass() != UObjectRedirector::StaticClass())
+			{
+				if (Pair.Value.Contains(InImport))
+				{
+					return false;
+				}
+			}
+		}
+		return true;
+	};
 
 	FString PackageName = SaveContext.GetPackage()->GetName();
 
@@ -628,8 +648,8 @@ ESavePackageResult ValidateImports(FSaveContext& SaveContext)
 			continue;
 		}
 
-		// See whether the object we are referencing is in another map package.
-		if (ImportPackage->ContainsMap())
+		// See whether the object we are referencing is in another map package and if it is allowed (i.e. from redirector).
+		if (ImportPackage->ContainsMap() && !IsMapReferenceAllowed(Import))
 		{
 			ObjectsInOtherMaps.Add(Import);
 		}
