Removed trace tags from UObject tracepoints since the cost of these functions has been fixed in 16266767.
Saves ~1ms in PIE based on SuperLuminal sampling.

diff --git a/Engine/Source/Runtime/CoreUObject/Private/UObject/ScriptCore.cpp b/Engine/Source/Runtime/CoreUObject/Private/UObject/ScriptCore.cpp
--- a/Engine/Source/Runtime/CoreUObject/Private/UObject/ScriptCore.cpp
+++ b/Engine/Source/Runtime/CoreUObject/Private/UObject/ScriptCore.cpp
@@ -2282,7 +2282,6 @@ DEFINE_FUNCTION(UObject::execBreakpoint)
 {
 	if (FBlueprintCoreDelegates::IsDebuggingEnabled())
 	{
-		TRACE_CPUPROFILER_EVENT_SCOPE(UObject::execBreakpoint);
 		FBlueprintExceptionInfo BreakpointExceptionInfo(EBlueprintExceptionType::Breakpoint);
 		FBlueprintCoreDelegates::ThrowScriptException(P_THIS, Stack, BreakpointExceptionInfo);
 	}
@@ -2293,7 +2292,6 @@ DEFINE_FUNCTION(UObject::execTracepoint)
 {
 	if (FBlueprintCoreDelegates::IsDebuggingEnabled())
 	{
-		TRACE_CPUPROFILER_EVENT_SCOPE(UObject::execTracepoint);
 		FBlueprintExceptionInfo TracepointExceptionInfo(EBlueprintExceptionType::Tracepoint);
 		FBlueprintCoreDelegates::ThrowScriptException(P_THIS, Stack, TracepointExceptionInfo);
 	}
@@ -2304,7 +2302,6 @@ DEFINE_FUNCTION(UObject::execWireTracepoint)
 {
 	if (FBlueprintCoreDelegates::IsDebuggingEnabled())
 	{
-		TRACE_CPUPROFILER_EVENT_SCOPE(UObject::execWireTracepoint);
 		FBlueprintExceptionInfo TracepointExceptionInfo(EBlueprintExceptionType::WireTracepoint);
 		FBlueprintCoreDelegates::ThrowScriptException(P_THIS, Stack, TracepointExceptionInfo);
 	}
