Optimize FPropertyValueIterator to alleviate the GC performance regression from 16048941.
Add separate stack entries instead of inserting elements in one single array to minimize mem copying.
Minimize processing of non matching properties.
Do expensive IsA-checks only once per property.
Use TInlineAllocator to reduce allocations.
Fix broken GetPropertyChain feature.

Iterator Tests:
Extend test with recursion.
Add test without recursion.
Add test for SkipRecursiveProperty feature.
Add test for previously broken GetPropertyChain feature.

Impact on TPropertyValueIterator in EngineTest on Win64:
~2X speedup for TPropertyValueIterator<T> on UEngineTestPropertiesObject for any T != FProperty, less when actually iterating all properties

Impact on Garbage Collection Reachability Analysis for a big internal project on a slower platform running reachability analysis in parallel in 4 task queues:
~130 ms => ~005 ms for UClass::AddReferencedObjects (TPropertyValueIterator<FObjectProperty>) for the worst case UAnimBlueprintGeneratedClass
~150 ms => ~035 ms total wall time
~300 ms => ~130 ms total thread time

Impact on Garbage Collection Reachability Analysis for a big internal project on a faster platform running reachability analysis in parallel in 7 task queues:
~040 ms => ~002 ms for UClass::AddReferencedObjects (TPropertyValueIterator<FObjectProperty>) for the worst case UAnimBlueprintGeneratedClass
~045 ms => ~022 ms total wall time
~200 ms => ~150 ms total thread time

#rnx

diff --git a/Engine/Source/Runtime/CoreUObject/Private/UObject/UnrealType.cpp b/Engine/Source/Runtime/CoreUObject/Private/UObject/UnrealType.cpp
--- a/Engine/Source/Runtime/CoreUObject/Private/UObject/UnrealType.cpp
+++ b/Engine/Source/Runtime/CoreUObject/Private/UObject/UnrealType.cpp
@@ -5,130 +5,195 @@
 
 DEFINE_LOG_CATEGORY(LogType);
 
-bool FPropertyValueIterator::NextValue(EPropertyValueIteratorFlags InRecursionFlags)
+FPropertyValueIterator::FPropertyValueIterator(
+	FFieldClass* InPropertyClass,
+	const UStruct* InStruct,
+	const void* InStructValue,
+	EPropertyValueIteratorFlags InRecursionFlags,
+	EFieldIteratorFlags::DeprecatedPropertyFlags InDeprecatedPropertyFlags)
+	: PropertyClass(InPropertyClass)
+	, RecursionFlags(InRecursionFlags)
+	, DeprecatedPropertyFlags(InDeprecatedPropertyFlags)
+	, bSkipRecursionOnce(false)
+	, bMatchAll(InPropertyClass == FProperty::StaticClass())
 {
-	if (PropertyIteratorStack.Num() == 0)
+	FPropertyValueStackEntry Entry(InStructValue);
+	FillStructProperties(InStruct, Entry);
+	if (Entry.ValueArray.Num() > 0)
+	{
+		PropertyIteratorStack.Emplace(MoveTemp(Entry));
+		IterateToNext();
+	}
+}
+
+uint8 FPropertyValueIterator::GetPropertyValueFlags(const FProperty* Property)
+{
+	uint8 Flags = 0;
+	if (RecursionFlags == EPropertyValueIteratorFlags::FullRecursion)
+	{
+		if (Property->IsA(FArrayProperty::StaticClass()))
+		{
+			Flags = (uint8)EPropertyValueFlags::IsArray;
+		}
+		else if (Property->IsA(FMapProperty::StaticClass()))
+		{
+			Flags = (uint8)EPropertyValueFlags::IsMap;
+		}
+		else if (Property->IsA(FSetProperty::StaticClass()))
+		{
+			Flags = (uint8)EPropertyValueFlags::IsSet;
+		}
+		else if (Property->IsA(FStructProperty::StaticClass()))
+		{
+			Flags = (uint8)EPropertyValueFlags::IsStruct;
+		}
+	}
+	if (bMatchAll || Property->IsA(PropertyClass))
 	{
-		// Stack is done, return
-		return false;
+		Flags |= (uint8)EPropertyValueFlags::IsMatch;
 	}
+	return Flags;
+}
 
+void FPropertyValueIterator::FillStructProperties(const UStruct* Struct, FPropertyValueStackEntry& Entry)
+{
+	FPropertyValueStackEntry::FValueArrayType& ValueArray = Entry.ValueArray;
+	for (TFieldIterator<FProperty> It(Struct, EFieldIteratorFlags::IncludeSuper, DeprecatedPropertyFlags, EFieldIteratorFlags::ExcludeInterfaces); It; ++It)
+	{
+		const FProperty* Property  = *It;
+		uint8 PropertyValueFlags = GetPropertyValueFlags(Property);
+		if (PropertyValueFlags)
+		{
+			int32 Num = Property->ArrayDim;
+			ValueArray.Reserve(ValueArray.Num() + Num);
+			for (int32 StaticIndex = 0; StaticIndex < Num; ++StaticIndex)
+			{
+				const void* PropertyValue = Property->ContainerPtrToValuePtr<void>(Entry.Owner, StaticIndex);
+				ValueArray.Emplace(BasePairType(Property, PropertyValue), PropertyValueFlags);
+			}
+		}
+	}
+}
+
+bool FPropertyValueIterator::NextValue(EPropertyValueIteratorFlags InRecursionFlags)
+{
+	check(PropertyIteratorStack.Num() > 0)
 	FPropertyValueStackEntry& Entry = PropertyIteratorStack.Last();
 
 	// If we have pending values, deal with them
-	if (Entry.ValueIndex < Entry.ValueArray.Num())
+	if (Entry.NextValueIndex < Entry.ValueArray.Num())
 	{
-		// Look for recursion on current value first
-		const FProperty* Property = Entry.ValueArray[Entry.ValueIndex].Key;
-		const void* PropertyValue = Entry.ValueArray[Entry.ValueIndex].Value;
+		const bool bIsPropertyMatchProcessed = Entry.ValueIndex == Entry.NextValueIndex;
+		Entry.ValueIndex = Entry.NextValueIndex;
+		Entry.NextValueIndex = Entry.ValueIndex + 1;
 
-		// For containers, insert at next index ahead of others
-		int32 InsertIndex = Entry.ValueIndex + 1;
+		const FProperty* Property = Entry.ValueArray[Entry.ValueIndex].Key.Key;
+		const void* PropertyValue = Entry.ValueArray[Entry.ValueIndex].Key.Value;
+		const uint8 PropertyValueFlags = Entry.ValueArray[Entry.ValueIndex].Value;
+		check(PropertyValueFlags);
 
-		// Handle container properties
-		if (const FArrayProperty* ArrayProperty = CastField<FArrayProperty>(Property))
+		// Handle matching properties
+		if (!bIsPropertyMatchProcessed && (PropertyValueFlags & EPropertyValueFlags::IsMatch))
 		{
-			if (InRecursionFlags == EPropertyValueIteratorFlags::FullRecursion)
+			if (PropertyValueFlags & EPropertyValueFlags::ContainerMask)
 			{
-				FScriptArrayHelper Helper(ArrayProperty, PropertyValue);
-				for (int32 DynamicIndex = 0; DynamicIndex < Helper.Num(); ++DynamicIndex)
-				{
-					Entry.ValueArray.EmplaceAt(InsertIndex++, ArrayProperty->Inner, Helper.GetRawPtr(DynamicIndex));
-				}
+				// this match is also a container/struct, so recurse into it next time
+				Entry.NextValueIndex = Entry.ValueIndex;
 			}
+			return false; // Break at this matching property
 		}
-		else if (const FMapProperty* MapProperty = CastField<FMapProperty>(Property))
+
+		// Handle container properties
+		check(PropertyValueFlags & EPropertyValueFlags::ContainerMask);
+		if (InRecursionFlags == EPropertyValueIteratorFlags::FullRecursion)
 		{
-			if (InRecursionFlags == EPropertyValueIteratorFlags::FullRecursion)
+			FPropertyValueStackEntry NewEntry(PropertyValue);
+
+			if (PropertyValueFlags & EPropertyValueFlags::IsArray)
 			{
-				FScriptMapHelper Helper(MapProperty, PropertyValue);
-				int32 Num = Helper.Num();
-				for (int32 DynamicIndex = 0; Num; ++DynamicIndex)
+				const FArrayProperty* ArrayProperty = CastFieldChecked<FArrayProperty>(Property);
+				const FProperty* InnerProperty = ArrayProperty->Inner;
+				uint8 InnerFlags = GetPropertyValueFlags(InnerProperty);
+				if (InnerFlags)
 				{
-					if (Helper.IsValidIndex(DynamicIndex))
+					FScriptArrayHelper Helper(ArrayProperty, PropertyValue);
+					const int32 Num = Helper.Num();
+					NewEntry.ValueArray.Reserve(Num);
+					for (int32 DynamicIndex = 0; DynamicIndex < Num; ++DynamicIndex)
 					{
-						Entry.ValueArray.EmplaceAt(InsertIndex++, MapProperty->KeyProp, Helper.GetKeyPtr(DynamicIndex));
-						Entry.ValueArray.EmplaceAt(InsertIndex++, MapProperty->ValueProp, Helper.GetValuePtr(DynamicIndex));
-
-						--Num;
+						NewEntry.ValueArray.Emplace(
+							BasePairType(InnerProperty, Helper.GetRawPtr(DynamicIndex)), InnerFlags);
 					}
 				}
 			}
-		}
-		else if (const FSetProperty* SetProperty = CastField<FSetProperty>(Property))
-		{
-			if (InRecursionFlags == EPropertyValueIteratorFlags::FullRecursion)
+			else if (PropertyValueFlags & EPropertyValueFlags::IsMap)
 			{
-				FScriptSetHelper Helper(SetProperty, PropertyValue);
-				int32 Num = Helper.Num();
-				for (int32 DynamicIndex = 0; Num; ++DynamicIndex)
+				const FMapProperty* MapProperty = CastFieldChecked<FMapProperty>(Property);
+				const FProperty* KeyProperty = MapProperty->KeyProp;
+				const FProperty* ValueProperty = MapProperty->ValueProp;
+				uint8 KeyFlags = GetPropertyValueFlags(KeyProperty);
+				uint8 ValueFlags = GetPropertyValueFlags(ValueProperty);
+				if (KeyFlags | ValueFlags)
 				{
-					if (Helper.IsValidIndex(DynamicIndex))
+					FScriptMapHelper Helper(MapProperty, PropertyValue);
+					const int32 Num = Helper.Num();
+					for (int32 DynamicIndex = 0; DynamicIndex < Num; ++DynamicIndex)
 					{
-						Entry.ValueArray.EmplaceAt(InsertIndex++, SetProperty->ElementProp, Helper.GetElementPtr(DynamicIndex));
-
-						--Num;
+						if (Helper.IsValidIndex(DynamicIndex))
+						{
+							if (KeyFlags)
+							{
+								NewEntry.ValueArray.Emplace(
+									BasePairType(KeyProperty, Helper.GetKeyPtr(DynamicIndex)), KeyFlags);
+							}
+							if (ValueFlags)
+							{
+								NewEntry.ValueArray.Emplace(
+									BasePairType(ValueProperty, Helper.GetValuePtr(DynamicIndex)), ValueFlags);
+							}
+						}
 					}
 				}
 			}
-		}
-		else if (const FStructProperty* StructProperty = CastField<FStructProperty>(Property))
-		{
-			if (InRecursionFlags == EPropertyValueIteratorFlags::FullRecursion)
+			else if (PropertyValueFlags & EPropertyValueFlags::IsSet)
 			{
-				// We don't need recursion on these - this will happen naturally as we process these values
-				for (FPropertyValueIterator Iter(FProperty::StaticClass(), StructProperty->Struct, PropertyValue, EPropertyValueIteratorFlags::NoRecursion, DeprecatedPropertyFlags); Iter; ++Iter)
+				const FSetProperty* SetProperty = CastFieldChecked<FSetProperty>(Property);
+				const FProperty* InnerProperty = SetProperty->ElementProp;
+				uint8 InnerFlags = GetPropertyValueFlags(InnerProperty);
+				if (InnerFlags)
 				{
-					Entry.ValueArray.EmplaceAt(InsertIndex++, Iter->Key, Iter->Value);
+					FScriptSetHelper Helper(SetProperty, PropertyValue);
+					const int32 Num = Helper.Num();
+					for (int32 DynamicIndex = 0; DynamicIndex < Num; ++DynamicIndex)
+					{
+						if (Helper.IsValidIndex(DynamicIndex))
+						{
+							NewEntry.ValueArray.Emplace(
+								BasePairType(InnerProperty, Helper.GetElementPtr(DynamicIndex)), InnerFlags);
+						}
+					}
 				}
 			}
-		}
-
-		// Else this is a normal property and has nothing to expand
-		// We don't expand enum properties because EnumProperty handles value wrapping for us
-
-		// Increment next value to check
-		Entry.ValueIndex++;
-	}
-
-	// Out of pending values, try to add more
-	if (Entry.ValueIndex == Entry.ValueArray.Num())
-	{
-		// If Field iterator is done, pop stack and run again as Entry is invalidated
-		if (!Entry.FieldIterator)
-		{
-			PropertyIteratorStack.Pop();
-
-			if (PropertyIteratorStack.Num() > 0)
+			else if (PropertyValueFlags & EPropertyValueFlags::IsStruct)
 			{
-				// Increment value index as we delayed incrementing it when entering recursion
-				PropertyIteratorStack.Last().ValueIndex++;
-
-				return NextValue(InRecursionFlags);
+				const FStructProperty* StructProperty = CastFieldChecked<FStructProperty>(Property);
+				FillStructProperties(StructProperty->Struct, NewEntry);
+			}
+			if (NewEntry.ValueArray.Num() > 0)
+			{
+				PropertyIteratorStack.Emplace(MoveTemp(NewEntry));
+				return true; // NextValue should be called again to move to the top of the stack
 			}
-
-			return false;
 		}
+	}
 
-		// If nothing left in value array, add base properties for current field and increase field iterator
-		const FProperty* Property = *Entry.FieldIterator;
-		++Entry.FieldIterator;
-
-		// Clear out existing value array
-		Entry.ValueArray.Reset();
-		Entry.ValueIndex = 0;
-
-		// Handle static arrays 
-		for (int32 StaticIndex = 0; StaticIndex != Property->ArrayDim; ++StaticIndex)
-		{
-			const void* PropertyValue = Property->ContainerPtrToValuePtr<void>(Entry.StructValue, StaticIndex);
-
-			Entry.ValueArray.Emplace(Property, PropertyValue);
-		}
-		return true;
+	if (Entry.NextValueIndex == Entry.ValueArray.Num())
+	{
+		PropertyIteratorStack.Pop();
 	}
-	
-	return true;
+
+	// NextValue should be called again to continue iteration
+	return PropertyIteratorStack.Num() > 0;
 }
 
 void FPropertyValueIterator::IterateToNext()
@@ -143,13 +208,6 @@ void FPropertyValueIterator::IterateToNext()
 
 	while (NextValue(LocalRecursionFlags))
 	{
-		// If this property is valid type, stop iteration
-		FPropertyValueStackEntry& Entry = PropertyIteratorStack.Last();
-		if (Entry.GetPropertyValue().Key->IsA(PropertyClass))
-		{
-			return;
-		}
-
 		// Reset recursion override as we've skipped the first property
 		LocalRecursionFlags = RecursionFlags;
 	}
@@ -157,19 +215,14 @@ void FPropertyValueIterator::IterateToNext()
 
 void FPropertyValueIterator::GetPropertyChain(TArray<const FProperty*>& PropertyChain) const
 {
-	// Iterate over UStruct nesting, starting at the inner most property
+	PropertyChain.Reserve(PropertyIteratorStack.Num());
+	// Iterate over struct/container property stack, starting at the inner most property
 	for (int32 StackIndex = PropertyIteratorStack.Num() - 1; StackIndex >= 0; StackIndex--)
 	{
 		const FPropertyValueStackEntry& Entry = PropertyIteratorStack[StackIndex];
 
 		// Index should always be valid
-		const FProperty* Property = Entry.ValueArray[Entry.ValueIndex].Key;
-
-		while (Property)
-		{
-			// This handles container property nesting
-			PropertyChain.Add(Property);
-			Property = Property->GetOwner<FProperty>();
-		}
+		const FProperty* Property = Entry.ValueArray[Entry.ValueIndex].Key.Key;
+		PropertyChain.Add(Property);
 	}
-}
\ No newline at end of file
+}
diff --git a/Engine/Source/Runtime/CoreUObject/Public/UObject/UnrealType.h b/Engine/Source/Runtime/CoreUObject/Public/UObject/UnrealType.h
--- a/Engine/Source/Runtime/CoreUObject/Public/UObject/UnrealType.h
+++ b/Engine/Source/Runtime/CoreUObject/Public/UObject/UnrealType.h
@@ -6103,17 +6103,9 @@ public:
 	 * @param InRecursionFlags	Rather to recurse into container and struct properties
 	 * @param InDeprecatedPropertyFlags	Rather to iterate over deprecated properties
 	 */
-	FPropertyValueIterator(FFieldClass* InPropertyClass, const UStruct* InStruct, const void* InStructValue,
+	COREUOBJECT_API FPropertyValueIterator(FFieldClass* InPropertyClass, const UStruct* InStruct, const void* InStructValue,
 		EPropertyValueIteratorFlags						InRecursionFlags = EPropertyValueIteratorFlags::FullRecursion,
-		EFieldIteratorFlags::DeprecatedPropertyFlags	InDeprecatedPropertyFlags = EFieldIteratorFlags::IncludeDeprecated)
-		: PropertyClass(InPropertyClass)
-		, RecursionFlags(InRecursionFlags)
-		, DeprecatedPropertyFlags(InDeprecatedPropertyFlags)
-		, bSkipRecursionOnce(false)
-	{
-		PropertyIteratorStack.Emplace(InStruct, InStructValue, InDeprecatedPropertyFlags);
-		IterateToNext();
-	}
+		EFieldIteratorFlags::DeprecatedPropertyFlags	InDeprecatedPropertyFlags = EFieldIteratorFlags::IncludeDeprecated);
 
 	/** Invalid iterator, start with empty stack */
 	FPropertyValueIterator()
@@ -6190,54 +6182,81 @@ public:
 	COREUOBJECT_API void GetPropertyChain(TArray<const FProperty*>& PropertyChain) const;
 
 private:
-	struct FPropertyValueStackEntry
+	enum EPropertyValueFlags : uint8
 	{
-		/** Field iterator within a UStruct */
-		TFieldIterator<const FProperty> FieldIterator;
+		IsMatch = 0x01,
+
+		ContainerMask = 0xF0,
+		IsArray = 0x10,
+		IsMap = 0x20,
+		IsSet = 0x40,
+		IsStruct = 0x80,
+	};
 
-		/** Address of owning UStruct */
-		const void* StructValue;
+	struct FPropertyValueStackEntry
+	{
+		/** Address of owning UStruct or FProperty container */
+		const void* Owner = nullptr;
 		
 		/** List of current root property+value pairs for the current top level FProperty */
-		TArray<BasePairType> ValueArray;
+		typedef TPair<BasePairType, uint8> BasePairAndFlags;
+		typedef TArray<BasePairAndFlags, TInlineAllocator<8>> FValueArrayType;
+		FValueArrayType ValueArray;
 
 		/** Current position inside ValueArray */
-		int32 ValueIndex;
+		int32 ValueIndex = -1;
+
+		/** Next position inside ValueArray */
+		int32 NextValueIndex = 0;
+
+		FPropertyValueStackEntry(const void* InValue)
+			: Owner(InValue)
+		{}
 
 		FPropertyValueStackEntry(const UStruct* InStruct, const void* InValue, EFieldIteratorFlags::DeprecatedPropertyFlags InDeprecatedPropertyFlags)
-			: FieldIterator(InStruct, EFieldIteratorFlags::IncludeSuper, InDeprecatedPropertyFlags, EFieldIteratorFlags::ExcludeInterfaces)
-			, StructValue(InValue)
-			, ValueIndex(0)
+			: Owner(InValue)
 		{}
 
 		FORCEINLINE friend bool operator==(const FPropertyValueStackEntry& Lhs, const FPropertyValueStackEntry& Rhs)
 		{
-			return Lhs.ValueIndex == Rhs.ValueIndex && Lhs.FieldIterator == Rhs.FieldIterator && Lhs.StructValue == Rhs.StructValue;
+			return Lhs.Owner == Rhs.Owner && Lhs.ValueIndex == Rhs.ValueIndex;
 		}
 
 		FORCEINLINE const BasePairType& GetPropertyValue() const
 		{
 			// Index has to be valid to get this far
-			return ValueArray[ValueIndex];
+			return ValueArray[ValueIndex].Key;
 		}
 	};
 
-	/** Internal stack, one per UStruct */
-	TArray<FPropertyValueStackEntry> PropertyIteratorStack;
+	/** Internal stack, one per continer/struct */
+	TArray<FPropertyValueStackEntry, TInlineAllocator<8>> PropertyIteratorStack;
 
 	/** Property type that is explicitly checked for */
-	FFieldClass* PropertyClass;
+	FFieldClass* PropertyClass = nullptr;
 
-	/** Whether to recurse into containers and StructProperties */
+	/** Whether to recurse into containers/structs */
 	const EPropertyValueIteratorFlags RecursionFlags;
 
 	/** Inherits to child field iterator */
 	const EFieldIteratorFlags::DeprecatedPropertyFlags DeprecatedPropertyFlags;
 
 	/** If true, next iteration will skip recursing into containers/structs */
-	bool bSkipRecursionOnce;
+	bool bSkipRecursionOnce = false;
+
+	/** If true, all properties will be matched without checking IsA(PropertyClass) */
+	bool bMatchAll = false;
 
-	/** Goes to the next Property/value pair. Returns true if next value is valid */
+	/** Returns EPropertyValueFlags to describe if this Property is a match and/or a container/struct */
+	uint8 GetPropertyValueFlags(const FProperty* Property);
+
+	/** Fills the Entry.ValueArray with all relevant properties found in Struct */
+	void FillStructProperties(const UStruct* Struct, FPropertyValueStackEntry& Entry);
+
+	/**
+	 * Goes to the next Property/value pair.
+	 * Returns false on a match or out of properties, true when iteration should continue
+	 */
 	bool NextValue(EPropertyValueIteratorFlags RecursionFlags);
 
 	/** Iterates to next property being checked for or until reaching the end of the structure */
@@ -6390,4 +6409,4 @@ class COREUOBJECT_API UMulticastInlineDelegatePropertyWrapper : public UMulticas
 	DECLARE_CLASS_INTRINSIC(UMulticastInlineDelegatePropertyWrapper, UMulticastDelegatePropertyWrapper, CLASS_Transient, TEXT("/Script/CoreUObject"));
 };
 
-#include "UObject/DefineUPropertyMacros.h"
\ No newline at end of file
+#include "UObject/DefineUPropertyMacros.h"
