Removed UE::Virtualization::IVirtualizedData from core now that there is no need for it.

diff --git a/Engine/Source/Runtime/Core/Public/Virtualization/IVirtualizedData.h b/Engine/Source/Runtime/Core/Public/Virtualization/IVirtualizedData.h
deleted file mode 100644
diff --git a/Engine/Source/Runtime/CoreUObject/Public/Virtualization/VirtualizedBulkData.h b/Engine/Source/Runtime/CoreUObject/Public/Virtualization/VirtualizedBulkData.h
--- a/Engine/Source/Runtime/CoreUObject/Public/Virtualization/VirtualizedBulkData.h
+++ b/Engine/Source/Runtime/CoreUObject/Public/Virtualization/VirtualizedBulkData.h
@@ -7,7 +7,6 @@
 #include "Memory/SharedBuffer.h"
 #include "Misc/Guid.h"
 #include "Misc/PackagePath.h"
-#include "Virtualization/IVirtualizedData.h"
 #include "Virtualization/VirtualizationManager.h"
 
 class FArchive;
@@ -47,7 +46,7 @@ namespace UE::Virtualization
  */
 
 /** The base class with no type */
-class COREUOBJECT_API FVirtualizedUntypedBulkData : public IVirtualizedData
+class COREUOBJECT_API FVirtualizedUntypedBulkData
 {
 public:
 	FVirtualizedUntypedBulkData() = default;
@@ -57,7 +56,7 @@ public:
 	FVirtualizedUntypedBulkData(const FVirtualizedUntypedBulkData& Other);
 	FVirtualizedUntypedBulkData& operator=(const FVirtualizedUntypedBulkData& Other);
 
-	virtual ~FVirtualizedUntypedBulkData() = default;
+	~FVirtualizedUntypedBulkData() = default;
 
 	/** 
 	 * Convenience method to make it easier to convert from BulkData to FVirtualizedBulkData and sets the Guid 
@@ -99,13 +98,13 @@ public:
 	 * Should that object be given a new payload it should then return the original
 	 * identifier, there is no need to generate a new one.
 	 */
-	virtual FGuid GetIdentifier() const override;
+	FGuid GetIdentifier() const;
 
 	/** Returns an unique identifier for the content of the payload */
-	virtual const FPayloadId& GetPayloadId() const override { return PayloadContentId; }
+	const FPayloadId& GetPayloadId() const { return PayloadContentId; }
 
 	/** Returns the size of the payload in bytes */
-	virtual int64 GetPayloadSize() const override { return PayloadSize; }
+	int64 GetPayloadSize() const { return PayloadSize; }
 	
 	/** Temp method, to make it easier to transition older code to virtualized bulkdata, remove when we remove UE_USE_VIRTUALBULKDATA */
 	FORCEINLINE int64 GetBulkDataSize() const { return GetPayloadSize(); }
@@ -117,7 +116,7 @@ public:
 	bool IsDataLoaded() const { return !Payload.IsNull(); }
 
 	/** Returns an immutable FCompressedBuffer reference to the payload data. */
-	virtual TFuture<FSharedBuffer> GetPayload() const override;
+	TFuture<FSharedBuffer> GetPayload() const;
 
 	/**
 	 * Returns an immutable FCompressedBuffer reference to the payload data.
@@ -126,7 +125,7 @@ public:
 	 * will be handled by the FCompressedBuffer interface. Call FCompressedBuffer::Decompress() to get access to
 	 * the payload in FSharedBuffer format.
 	 */
-	virtual TFuture<FCompressedBuffer> GetCompressedPayload() const override;
+	TFuture<FCompressedBuffer> GetCompressedPayload() const;
 
 	/**
 	 * Allows the existing payload to be replaced with a new one.
@@ -145,7 +144,7 @@ public:
 	 * underlying code deems appropriate. Other specific compression formats may be allowed, see the
 	 * documentation of FCompressedBuffer for details.
 	 */
-	virtual void UpdatePayload(FSharedBuffer InPayload, FName CompressionFormat = NAME_Default) override;
+	void UpdatePayload(FSharedBuffer InPayload, FName CompressionFormat = NAME_Default);
 
 	/**
 	* Get the CustomVersions used in the file containing the payload. Currently this is assumed
