EditorDomain: Fix version warning when receiving EditorDomain packages from other users on the shared DDC.
Warning: Asset has been saved with a newer engine and can't be loaded.
#rnx

diff --git a/Engine/Source/Editor/UnrealEd/Private/EditorDomain/EditorDomainArchive.h b/Engine/Source/Editor/UnrealEd/Private/EditorDomain/EditorDomainArchive.h
--- a/Engine/Source/Editor/UnrealEd/Private/EditorDomain/EditorDomainArchive.h
+++ b/Engine/Source/Editor/UnrealEd/Private/EditorDomain/EditorDomainArchive.h
@@ -48,6 +48,7 @@ public:
 	virtual void Flush() override;
 	virtual void FlushCache() override;
 	virtual bool Precache(int64 PrecacheOffset, int64 PrecacheSize) override;
+	virtual bool NeedsEngineVersionChecks() const override { return false; }
 
 private:
 	enum class ESource : uint8
diff --git a/Engine/Source/Runtime/Core/Public/Serialization/Archive.h b/Engine/Source/Runtime/Core/Public/Serialization/Archive.h
--- a/Engine/Source/Runtime/Core/Public/Serialization/Archive.h
+++ b/Engine/Source/Runtime/Core/Public/Serialization/Archive.h
@@ -1871,13 +1871,19 @@ public:
 	 *
 	 * @return true if dependency has been added, false if Archive does not support them
 	 */
-	virtual bool AttachExternalReadDependency(FExternalReadCallback& ReadCallback) { return false; };
+	virtual bool AttachExternalReadDependency(FExternalReadCallback& ReadCallback) { return false; }
 
 	/**
 	 * Returns whether the Event Driven Loader is enabled or not.
 	 */
 	virtual bool IsUsingEventDrivenLoader() const;
 
+	/**
+	 * Used by LinkerLoad to decide whether it needs to check its Loader for FEngineVersion::Current().IsCompatibleWith.
+	 * EditorDomain Loaders are selected based on CustomVersions and Schema and do not need to check minor version changes.
+	 */
+	virtual bool NeedsEngineVersionChecks() const { return true; }
+
 #if USE_STABLE_LOCALIZATION_KEYS
 	using FArchiveState::SetLocalizationNamespace;
 	using FArchiveState::GetLocalizationNamespace;
diff --git a/Engine/Source/Runtime/CoreUObject/Private/UObject/LinkerLoad.cpp b/Engine/Source/Runtime/CoreUObject/Private/UObject/LinkerLoad.cpp
--- a/Engine/Source/Runtime/CoreUObject/Private/UObject/LinkerLoad.cpp
+++ b/Engine/Source/Runtime/CoreUObject/Private/UObject/LinkerLoad.cpp
@@ -1242,8 +1242,9 @@ FLinkerLoad::ELinkerStatus FLinkerLoad::SerializePackageFileSummaryInternal()
 		}
 	}
 
+	bool bLoaderVersionCheck = Loader->NeedsEngineVersionChecks();
 	// Don't load packages that are only compatible with an engine version newer than the current one.
-	if (GEnforcePackageCompatibleVersionCheck && !FEngineVersion::Current().IsCompatibleWith(Summary.CompatibleWithEngineVersion))
+	if (bLoaderVersionCheck && GEnforcePackageCompatibleVersionCheck && !FEngineVersion::Current().IsCompatibleWith(Summary.CompatibleWithEngineVersion))
 	{
 		UE_LOG(LogLinker, Warning, TEXT("Asset '%s' has been saved with a newer engine and can't be loaded. CurrentEngineVersion: %s (Licensee=%d). AssetEngineVersion: %s (Licensee=%d)"),
 			*GetDebugName(),
@@ -1259,7 +1260,8 @@ FLinkerLoad::ELinkerStatus FLinkerLoad::SerializePackageFileSummaryInternal()
 	SetUseUnversionedPropertySerialization(bUseUnversionedProperties);
 	Loader->SetUseUnversionedPropertySerialization(bUseUnversionedProperties);
 
-	if (!FPlatformProperties::RequiresCookedData() && !Summary.SavedByEngineVersion.HasChangelist() && FEngineVersion::Current().HasChangelist())
+	if (bLoaderVersionCheck && !FPlatformProperties::RequiresCookedData()
+		&& !Summary.SavedByEngineVersion.HasChangelist() && FEngineVersion::Current().HasChangelist())
 	{
 		// This warning can be disabled in ini with [Core.System] ZeroEngineVersionWarning=False
 		static struct FInitZeroEngineVersionWarning
@@ -1278,7 +1280,8 @@ FLinkerLoad::ELinkerStatus FLinkerLoad::SerializePackageFileSummaryInternal()
 	}
 
 	// Don't load packages that were saved with package version newer than the current one.
-	if ((Summary.GetFileVersionUE() > GPackageFileUEVersion) || (Summary.GetFileVersionLicenseeUE() > GPackageFileLicenseeUEVersion))
+	if (bLoaderVersionCheck && (Summary.GetFileVersionUE() > GPackageFileUEVersion)
+								|| (Summary.GetFileVersionLicenseeUE() > GPackageFileLicenseeUEVersion))
 	{
 		UE_LOG(LogLinker, Warning, TEXT("Unable to load package (%s) PackageVersion %i, MaxExpected %i : LicenseePackageVersion %i, MaxExpected %i."),
 			*GetDebugName(), Summary.GetFileVersionUE(), GPackageFileUEVersion, Summary.GetFileVersionLicenseeUE(), GPackageFileLicenseeUEVersion);
