Fix compile error '&&' within '||'.
#rnx

diff --git a/Engine/Source/Runtime/CoreUObject/Private/UObject/LinkerLoad.cpp b/Engine/Source/Runtime/CoreUObject/Private/UObject/LinkerLoad.cpp
--- a/Engine/Source/Runtime/CoreUObject/Private/UObject/LinkerLoad.cpp
+++ b/Engine/Source/Runtime/CoreUObject/Private/UObject/LinkerLoad.cpp
@@ -1280,8 +1280,8 @@ FLinkerLoad::ELinkerStatus FLinkerLoad::SerializePackageFileSummaryInternal()
 	}
 
 	// Don't load packages that were saved with package version newer than the current one.
-	if (bLoaderVersionCheck && (Summary.GetFileVersionUE() > GPackageFileUEVersion)
-								|| (Summary.GetFileVersionLicenseeUE() > GPackageFileLicenseeUEVersion))
+	if (bLoaderVersionCheck && ((Summary.GetFileVersionUE() > GPackageFileUEVersion)
+								|| (Summary.GetFileVersionLicenseeUE() > GPackageFileLicenseeUEVersion)))
 	{
 		UE_LOG(LogLinker, Warning, TEXT("Unable to load package (%s) PackageVersion %i, MaxExpected %i : LicenseePackageVersion %i, MaxExpected %i."),
 			*GetDebugName(), Summary.GetFileVersionUE(), GPackageFileUEVersion, Summary.GetFileVersionLicenseeUE(), GPackageFileLicenseeUEVersion);
