Add some basic profile data to Mirage to make it easier for people to see when payloads are being pushed to or pulled from the various backends.

#rnx

* VirtualizationManager
- Add method ::IsEnabled which allows the caller to poll if content virtualization is enabled or not.
- Add method ::GetPayloadActivityInfo to return profiling data
- Added a new Profiling namespace containing all of the profiling code.
-- The idea is to try and keep all profiling data contained in the virtualization manager and not require the backends to be aware of it.
- The actual pull/push operations have been moved to new private methods ::TryPushDataToBackend and ::PullDataFromBackend, by limiting the scolpe of where the actual operation occur makes it easier to add the profiling code.
- We use the cooking stats system for recording our profiling data. Currently this adds no real value and we could've implemented our own but longer term this code might get hooked into FCookStatsManager:: CookStatsCallbacks to add it to our telemetry systems and this will be easier to do if it is already in the cooking stats formats.

* SVirtualizationStaticIndicator
- The widget is based on SDDCStatusIndicator for the DDC and is placed just before it in the UI.
- The widget will only be shown if the content virtualization system is enabled, so functions as an easy way to check that (maybe in the future it can be on all the time in which case we'd need the widget to reflect when the systems are disabled)
- Currently shows a green down arrow when a payload has been pulled and a green up arrow when a payload has been pushed.
- The tool tip shown on mouse over will show the total data sizes pulled and pushed from the backends.

diff --git a/Engine/Source/Editor/LevelEditor/Private/SLevelEditor.cpp b/Engine/Source/Editor/LevelEditor/Private/SLevelEditor.cpp
--- a/Engine/Source/Editor/LevelEditor/Private/SLevelEditor.cpp
+++ b/Engine/Source/Editor/LevelEditor/Private/SLevelEditor.cpp
@@ -70,6 +70,7 @@
 #include "Elements/SMInstance/SMInstanceElementId.h"
 #include "Elements/SMInstance/SMInstanceElementLevelEditorSelectionCustomization.h"
 #include "DDC/SDDCStatusIndicator.h"
+#include "Virtualization/SVirtualizationStatusIndicator.h"
 
 #define LOCTEXT_NAMESPACE "SLevelEditor"
 
@@ -1936,7 +1937,13 @@ void SLevelEditor::RegisterStatusBarTools()
 		FToolMenuEntry::InitWidget("DDCStatusIndicator", SNew(SDDCStatusIndicator), FText::GetEmpty(), true, false)
 	);
 
+	TSharedPtr<SWidget> VirtualizationWidget = UE::Virtualization::GetVirtualizationStatusIndicator();
+	if (VirtualizationWidget)
+	{
+		FToolMenuSection& MirageSection = Menu->AddSection("Mirage", FText::GetEmpty(), FToolMenuInsert("DDC", EToolMenuInsertType::Before));
 
+		MirageSection.AddEntry(FToolMenuEntry::InitWidget("VirtualizationStatusIndicator", VirtualizationWidget.ToSharedRef(), FText::GetEmpty(), true, false));
+	}	
 }
 
 void SLevelEditor::AddStandaloneLevelViewport( const TSharedRef<SLevelViewport>& LevelViewport )
diff --git a/Engine/Source/Editor/UnrealEd/Private/Virtualization/SVirtualizationStatusIndicator.cpp b/Engine/Source/Editor/UnrealEd/Private/Virtualization/SVirtualizationStatusIndicator.cpp
new file mode 100644
--- /dev/null
+++ b/Engine/Source/Editor/UnrealEd/Private/Virtualization/SVirtualizationStatusIndicator.cpp
@@ -0,0 +1,190 @@
+// Copyright Epic Games, Inc. All Rights Reserved.
+
+#include "Virtualization/SVirtualizationStatusIndicator.h"
+
+#include "Animation/CurveSequence.h"
+#include "Virtualization/VirtualizationManager.h"
+#include "Widgets/Images/SImage.h"
+#include "Widgets/SBoxPanel.h"
+#include "Widgets/SCompoundWidget.h"
+#include "Widgets/SOverlay.h"
+#include "Widgets/SToolTip.h"
+#include "Widgets/Text/STextBlock.h"
+
+#define LOCTEXT_NAMESPACE "SVirtualizationStatusIndicator"
+
+namespace UE::Virtualization
+{
+
+/** The amount of time (s) that the indication arrows will take to fade in/out */
+static constexpr float ArrowFadeTime = 0.5f;
+
+/** 
+  * The tool tip that will show when the mouse is hovered over SVirtualizationStatusIndicator. 
+  * 
+  * Currently this will display the overall amount of data that content virtualization has 
+  * transfered to and from the backends.
+  */
+class SVirtualizationToolTip : public SCompoundWidget
+{
+	SLATE_BEGIN_ARGS(SVirtualizationToolTip) {}
+	SLATE_END_ARGS()
+
+	void Construct(const FArguments& InArgs);
+
+private:
+	FText GetToolTipText() const;
+};
+
+/** 
+ * A widget used to show the status of content virtualization.
+ * 
+ * A downwards pointing arrow will be displayed when ever a payload is pulled from a backend.
+ * An upwards pointing arrow will be displayed when ever a payload is pushed to a backend.
+ * This gives the user immediately visual feedback as to when the system is in action.
+ * Additional data can be found when hovering the mouse of the widget.
+ */
+class SVirtualizationStatusIndicator : public SCompoundWidget
+{
+	SLATE_BEGIN_ARGS(SVirtualizationStatusIndicator) {}
+	SLATE_END_ARGS()
+
+	void Construct(const FArguments& InArgs);
+
+private:
+
+	EActiveTimerReturnType UpdateLastDataAccess(double InCurrentTime, float InDeltaTime);
+
+	int64 PayloadsPulled = 0;
+	int64 PayloadsPushed = 0;
+
+	FCurveSequence FadePullArrow;
+	FCurveSequence FadePushArrow;
+};
+
+void SVirtualizationToolTip::Construct(const FArguments& InArgs)
+{
+	this->ChildSlot
+	[
+		SNew(SVerticalBox)
+
+		+ SVerticalBox::Slot()
+		.AutoHeight()
+		[
+			SNew(STextBlock)
+			.Text(this, &SVirtualizationToolTip::GetToolTipText)
+		]
+	];
+}
+
+FText SVirtualizationToolTip::GetToolTipText() const
+{
+	using namespace UE::Virtualization;
+	
+	FPayloadActivityInfo Info = FVirtualizationManager::Get().GetPayloadActivityInfo();
+	TStringBuilder<1024> Builder;
+	
+	Builder << TEXT("Total Payload Data Pushed: ") << Info.TotalSizePushed << TEXT(" MB\n");
+	Builder << TEXT("Total Payload Data Pulled: ") << Info.TotalSizePulled << TEXT(" MB");
+	
+	return FText::FromString(Builder.ToString());
+}
+
+void SVirtualizationStatusIndicator::Construct(const FArguments& InArgs)
+{
+	using namespace UE::Virtualization;
+
+	FPayloadActivityInfo Info = FVirtualizationManager::Get().GetPayloadActivityInfo();
+	PayloadsPulled = Info.PayloadsPulled;
+	PayloadsPushed = Info.PayloadsPushed;
+	
+	FadePushArrow = FCurveSequence(0.0f, ArrowFadeTime, ECurveEaseFunction::Linear);
+	FadePullArrow = FCurveSequence(0.0f, ArrowFadeTime, ECurveEaseFunction::Linear);
+
+	const FText MirageLabel = FText::FromString(TEXT("Mirage"));
+
+	this->ChildSlot
+	[
+		SNew(SHorizontalBox)
+		.ToolTip(SNew(SToolTip)[SNew(SVirtualizationToolTip)])
+
+		+ SHorizontalBox::Slot()
+		.AutoWidth()
+		.VAlign(VAlign_Center)
+		.Padding(0, 0, 3, 0)
+		[
+			SNew(SOverlay)
+
+			+ SOverlay::Slot()
+			.HAlign(HAlign_Center)
+			.VAlign(VAlign_Top)
+			.Padding(0, 0, 4, 4)
+			[
+				SNew(SImage)
+				.Image(FAppStyle::Get().GetBrush("Icons.ArrowUp"))
+				.ColorAndOpacity_Lambda([this] { return FLinearColor::Green.CopyWithNewOpacity(FadePushArrow.GetLerp()); })
+			]
+
+			+ SOverlay::Slot()
+			.HAlign(HAlign_Center)
+			.VAlign(VAlign_Bottom)
+			.Padding(4, 4, 0, 0)
+			[
+				SNew(SImage)
+				.Image(FAppStyle::Get().GetBrush("Icons.ArrowDown"))
+				.ColorAndOpacity_Lambda([this] { return FLinearColor::Green.CopyWithNewOpacity(FadePullArrow.GetLerp()); })
+			]
+		]
+
+		+ SHorizontalBox::Slot()
+		.AutoWidth()
+		.VAlign(VAlign_Center)
+		.Padding(0, 0, 10, 0)
+		[
+			SNew(STextBlock)
+			.Text(MirageLabel)
+		]
+	];
+
+	RegisterActiveTimer(ArrowFadeTime, FWidgetActiveTimerDelegate::CreateSP(this, &SVirtualizationStatusIndicator::UpdateLastDataAccess));
+}
+
+EActiveTimerReturnType SVirtualizationStatusIndicator::UpdateLastDataAccess(double InCurrentTime, float InDeltaTime)
+{
+	using namespace UE::Virtualization;
+
+	FPayloadActivityInfo Info = FVirtualizationManager::Get().GetPayloadActivityInfo();
+	
+	// Don't change the fade if it is currently playing, the tick duration of this method is not that accurate
+	// and we don't want to cut out the previous fade too soon or the user might not be able to notice it.
+
+	if (!FadePullArrow.IsPlaying())
+	{
+		FadePullArrow.PlayRelative(this->AsShared(), PayloadsPulled != Info.PayloadsPulled);
+		PayloadsPulled = Info.PayloadsPulled;
+	}
+
+	if (!FadePushArrow.IsPlaying())
+	{
+		FadePushArrow.PlayRelative(this->AsShared(), PayloadsPushed != Info.PayloadsPushed);
+		PayloadsPushed = Info.PayloadsPushed;
+	}
+
+	return EActiveTimerReturnType::Continue;
+}
+
+TSharedPtr<SWidget> GetVirtualizationStatusIndicator()
+{
+	if (FVirtualizationManager::Get().IsEnabled())
+	{
+		return SNew(SVirtualizationStatusIndicator);
+	}
+	else
+	{
+		return nullptr;
+	}
+}
+
+} // namespace UE::Virtualization
+
+#undef LOCTEXT_NAMESPACE
diff --git a/Engine/Source/Editor/UnrealEd/Public/Virtualization/SVirtualizationStatusIndicator.h b/Engine/Source/Editor/UnrealEd/Public/Virtualization/SVirtualizationStatusIndicator.h
new file mode 100644
--- /dev/null
+++ b/Engine/Source/Editor/UnrealEd/Public/Virtualization/SVirtualizationStatusIndicator.h
@@ -0,0 +1,20 @@
+// Copyright Epic Games, Inc. All Rights Reserved.
+
+#pragma once
+
+#include "Templates/SharedPointer.h"
+#include "Widgets/SWidget.h"
+
+namespace UE::Virtualization
+{
+
+/** 
+ * Returns the tool menu widget showing profiling data for content virtualization
+ * in the editor.
+ * 
+ * @return	A valid widget pointer if content virtualization is enabled, and an 
+ *			invalid pointer if it is disabled.
+ */
+UNREALED_API TSharedPtr<SWidget> GetVirtualizationStatusIndicator();
+
+} // namespace UE::Virtualization
diff --git a/Engine/Source/Runtime/CoreUObject/Private/Virtualization/VirtualizationManager.cpp b/Engine/Source/Runtime/CoreUObject/Private/Virtualization/VirtualizationManager.cpp
--- a/Engine/Source/Runtime/CoreUObject/Private/Virtualization/VirtualizationManager.cpp
+++ b/Engine/Source/Runtime/CoreUObject/Private/Virtualization/VirtualizationManager.cpp
@@ -2,12 +2,17 @@
 
 #include "Virtualization/VirtualizationManager.h"
 
+#include "HAL/PlatformTime.h"
 #include "Misc/CommandLine.h"
 #include "Misc/ConfigCacheIni.h"
 #include "Misc/Parse.h"
 #include "Misc/Paths.h"
+#include "Misc/ScopeLock.h"
 #include "Virtualization/IVirtualizationBackend.h"
 
+#include "ProfilingDebugging/CookStats.h"
+#include "Misc/CoreDelegates.h"
+
 namespace UE::Virtualization
 {
 
@@ -92,6 +97,83 @@ TArray<FString> ParseEntries(const FString& Data)
 	return Entries;
 }
 
+/**
+ * Profiling data allowing us to track how payloads are being push/pulled during the lifespan of the process. Note that as all backends are
+ * created at the same time, we don't need to add locked when accessing the maps. In addition FCookStats is thread safe when adding hits/misses
+ * so we don't hasve to worry about that either.
+ * We keep the FCookStats here rather than as a member of IVirtualizationBackend to try and avoid the backends needing to be aware of the data that
+ * we are gathering at all. This way all profiling code is kept to this cpp.
+ */
+namespace Profiling
+{
+#if ENABLE_COOK_STATS
+	TMap<FString, FCookStats::CallStats> PushStats;
+	TMap<FString, FCookStats::CallStats> PullStats;
+
+	void CreateStats(const IVirtualizationBackend& Backend)
+	{
+		PushStats.Add(Backend.GetDebugString());
+		PullStats.Add(Backend.GetDebugString());
+	}
+
+	FCookStats::CallStats& GetPushStats(const IVirtualizationBackend& Backend)
+	{
+		return *PushStats.Find(Backend.GetDebugString());
+	}
+
+	FCookStats::CallStats& GetPullStats(const IVirtualizationBackend& Backend)
+	{
+		return *PullStats.Find(Backend.GetDebugString());
+	}
+
+	void LogStats()
+	{
+		if (PushStats.IsEmpty() && PullStats.IsEmpty())
+		{
+			return; // Early out if we have no data to show at all
+		}
+
+		UE_LOG(LogVirtualization, Log, TEXT("Virtualization ProfileData"));
+
+		if (PushStats.Num() > 0)
+		{
+			UE_LOG(LogVirtualization, Log, TEXT("%-40s|%17s|%12s|%14s|"), TEXT("Pushing Data"), TEXT("TotalSize (MB)"), TEXT("TotalTime(s)"), TEXT("DataRate(MB/S)"));
+
+			for (const auto& Iterator : PushStats)
+			{
+				const double Time = Iterator.Value.GetAccumulatedValueAnyThread(FCookStats::CallStats::EHitOrMiss::Hit, FCookStats::CallStats::EStatType::Cycles) * FPlatformTime::GetSecondsPerCycle();
+				const int64 DataSizeMB = Iterator.Value.GetAccumulatedValueAnyThread(FCookStats::CallStats::EHitOrMiss::Hit, FCookStats::CallStats::EStatType::Bytes) / (1024 * 1024);
+				const double MBPerSecond = DataSizeMB / Time;
+
+				UE_LOG(LogVirtualization, Log, TEXT("%-40.40s|%17" UINT64_FMT "|%12.3f|%14.3f|"),
+					*Iterator.Key,
+					DataSizeMB,
+					Time,
+					MBPerSecond);
+			}
+		}
+
+		if (PullStats.Num() > 0)
+		{
+			UE_LOG(LogVirtualization, Log, TEXT("%-40s|%17s|%12s|%14s|"), TEXT("Pulling Data"), TEXT("TotalSize (MB)"), TEXT("TotalTime(s)"), TEXT("DataRate(MB/S)"));
+
+			for (const auto& Iterator : PullStats)
+			{
+				const double Time = Iterator.Value.GetAccumulatedValueAnyThread(FCookStats::CallStats::EHitOrMiss::Hit, FCookStats::CallStats::EStatType::Cycles) * FPlatformTime::GetSecondsPerCycle();
+				const int64 DataSizeMB = Iterator.Value.GetAccumulatedValueAnyThread(FCookStats::CallStats::EHitOrMiss::Hit, FCookStats::CallStats::EStatType::Bytes) / (1024 * 1024);
+				const double MBPerSecond = DataSizeMB / Time;
+
+				UE_LOG(LogVirtualization, Log, TEXT("%-40.40s|%17" UINT64_FMT "|%12.3f|%14.3f|"),
+					*Iterator.Key,
+					DataSizeMB,
+					Time,
+					MBPerSecond);
+			}
+		}
+	}
+#endif // ENABLE_COOK_STATS
+} //namespace Profiling
+
 FVirtualizationManager& FVirtualizationManager::Get()
 {
 	// TODO: Do we really need to make this a singleton? Easier for prototyping.
@@ -109,6 +191,11 @@ FVirtualizationManager::FVirtualizationManager()
 {
 	UE_LOG(LogVirtualization, Log, TEXT("Virtualization manager created"));
 
+	// Allows us to log the profiling data on process exit. 
+	// TODO: We should just be able to call the logging in the destructor, but 
+	// we need to fix the startup/shutdown ordering of Mirage first.
+	COOK_STAT(FCoreDelegates::OnExit.AddStatic(Profiling::LogStats));
+
 	FConfigFile PlatformEngineIni;
 	if (FConfigCacheIni::LoadLocalIniFile(PlatformEngineIni, TEXT("Engine"), true))
 	{
@@ -141,6 +228,11 @@ FVirtualizationManager::~FVirtualizationManager()
 	UE_LOG(LogVirtualization, Log, TEXT("Virtualization manager destroyed"));
 }
 
+bool FVirtualizationManager::IsEnabled() const
+{
+	return !PushEnabledBackendsArray.IsEmpty();
+}
+
 bool FVirtualizationManager::PushData(const FPayloadId& Id, const FCompressedBuffer& Payload)
 {
 	TRACE_CPUPROFILER_EVENT_SCOPE(FVirtualizationManager::PushData);
@@ -185,7 +277,7 @@ bool FVirtualizationManager::PushData(const FPayloadId& Id, const FCompressedBuf
 	bool bWasPayloadPushed = false;
 	for (IVirtualizationBackend* Backend : PushEnabledBackendsArray)
 	{
-		const EPushResult Result = Backend->PushData(Id, Payload);
+		const EPushResult Result = TryPushDataToBackend(*Backend, Id, Payload) ? EPushResult::Success : EPushResult::Failed;
 
 		UE_CLOG(Result != EPushResult::Failed, LogVirtualization, Verbose, TEXT("[%s] Pushed the payload '%s'"), *Backend->GetDebugString(), *Id.ToString());
 		UE_CLOG(Result == EPushResult::Failed, LogVirtualization, Error, TEXT("[%s] Failed to push the payload '%s'"), *Backend->GetDebugString(), *Id.ToString());
@@ -199,7 +291,7 @@ bool FVirtualizationManager::PushData(const FPayloadId& Id, const FCompressedBuf
 		// that the pulled payload is the same as the original
 		if (bValidateAfterPushOperation && Result != EPushResult::Failed && Backend->SupportsPullOperations())
 		{
-			FCompressedBuffer PulledPayload = Backend->PullData(Id);
+			FCompressedBuffer PulledPayload = PullDataFromBackend(*Backend, Id);
 			checkf(Payload.GetRawHash() == PulledPayload.GetRawHash(), TEXT("[%s] Failed to pull payload '%s' after it was pushed to backend"), 
 				*Backend->GetDebugString(), *Id.ToString());
 		}
@@ -240,7 +332,7 @@ FCompressedBuffer FVirtualizationManager::PullData(const FPayloadId& Id)
 	// (a local cache might want to replicate the data for example)
 	for (IVirtualizationBackend* Backend : PullEnabledBackendsArray)
 	{
-		FCompressedBuffer Payload = Backend->PullData(Id);
+		FCompressedBuffer Payload = PullDataFromBackend(*Backend, Id);
 		if (Payload)
 		{
 			return Payload;
@@ -256,6 +348,27 @@ FCompressedBuffer FVirtualizationManager::PullData(const FPayloadId& Id)
 	return FCompressedBuffer();
 }
 
+FPayloadActivityInfo FVirtualizationManager::GetPayloadActivityInfo() const
+{
+	FPayloadActivityInfo Info;
+
+#if ENABLE_COOK_STATS
+	for (const auto& Iterator : Profiling::PushStats)
+	{
+		Info.PayloadsPushed += Iterator.Value.GetAccumulatedValueAnyThread(FCookStats::CallStats::EHitOrMiss::Hit, FCookStats::CallStats::EStatType::Counter);
+		Info.TotalSizePushed += Iterator.Value.GetAccumulatedValueAnyThread(FCookStats::CallStats::EHitOrMiss::Hit, FCookStats::CallStats::EStatType::Bytes) / (1024 * 1024);
+	}
+
+	for (const auto& Iterator : Profiling::PullStats)
+	{
+		Info.PayloadsPulled += Iterator.Value.GetAccumulatedValueAnyThread(FCookStats::CallStats::EHitOrMiss::Hit, FCookStats::CallStats::EStatType::Counter);
+		Info.TotalSizePulled += Iterator.Value.GetAccumulatedValueAnyThread(FCookStats::CallStats::EHitOrMiss::Hit, FCookStats::CallStats::EStatType::Bytes) / (1024 * 1024);
+	}
+#endif // ENABLE_COOK_STATS
+
+	return Info;
+}
+
 void FVirtualizationManager::ApplySettingsFromConfigFiles(const FConfigFile& PlatformEngineIni)
 {
 	UE_LOG(LogVirtualization, Log, TEXT("Loading virtualization manager settings from config files..."));
@@ -451,7 +564,35 @@ void FVirtualizationManager::AddBackend(IVirtualizationBackend* Backend)
 		PushEnabledBackendsArray.Add(Backend);
 	}
 
+	COOK_STAT(Profiling::CreateStats(*Backend));
+
 	UE_LOG(LogVirtualization, Log, TEXT("Mounted backend: %s"), *Backend->GetDebugString());
 }
 
+bool FVirtualizationManager::TryPushDataToBackend(IVirtualizationBackend& Backend, const FPayloadId& Id, const FCompressedBuffer& Payload)
+{
+	COOK_STAT(FCookStats::FScopedStatsCounter Timer(Profiling::GetPushStats(Backend)));
+	const EPushResult Result = Backend.PushData(Id, Payload);
+
+	if (Result == EPushResult::Success)
+	{
+		COOK_STAT(Timer.AddHit(Payload.GetCompressedSize()));
+	}
+
+	return Result != EPushResult::Failed;
+}
+
+FCompressedBuffer FVirtualizationManager::PullDataFromBackend(IVirtualizationBackend& Backend, const FPayloadId& Id)
+{
+	COOK_STAT(FCookStats::FScopedStatsCounter Timer(Profiling::GetPullStats(Backend)));
+	FCompressedBuffer Payload = Backend.PullData(Id);
+	
+	if (!Payload.IsNull())
+	{
+		COOK_STAT(Timer.AddHit(Payload.GetCompressedSize()));
+	}
+	
+	return Payload;
+}
+
 } // namespace UE::Virtualization
diff --git a/Engine/Source/Runtime/CoreUObject/Public/Virtualization/VirtualizationManager.h b/Engine/Source/Runtime/CoreUObject/Public/Virtualization/VirtualizationManager.h
--- a/Engine/Source/Runtime/CoreUObject/Public/Virtualization/VirtualizationManager.h
+++ b/Engine/Source/Runtime/CoreUObject/Public/Virtualization/VirtualizationManager.h
@@ -74,6 +74,15 @@ namespace UE::Virtualization
 class IVirtualizationBackend;
 class IVirtualizationBackendFactory;
 
+/** Profiling data containing all activity relating to payloads. */
+struct FPayloadActivityInfo
+{
+	int64 PayloadsPulled = 0;
+	int64 PayloadsPushed = 0;
+
+	int64 TotalSizePulled = 0;
+	int64 TotalSizePushed = 0;
+};
 
 /** This is used as a wrapper around the various potential back end implementations. 
 	The calling code shouldn't need to care about which back ends are actually in use. */
@@ -86,23 +95,35 @@ public:
 	FVirtualizationManager();
 	~FVirtualizationManager();
 
+	/** Poll to see if content virtualization is enabled or not. */
+	bool IsEnabled() const;
+
 	/** 
 	 * Push a payload to the virtualization backends.
 	 * 
-	 * @param Id
-	 * @param Payload
-	 * @returns
+	 * @param	Id The identifier of the payload being pushed.
+	 * @param	Payload The payload itself in FCompressedBuffer form, it is assumed
+	 *			that if the buffer is to be compressed that it will have been done
+	 *			by the caller.
+	 * @return	True if at least one backend now contains the payload, otherwise false.
 	 */
 	bool PushData(const FPayloadId& Id, const FCompressedBuffer& Payload);
 
 	/** 
 	 * Pull a payload from the virtualization backends.
 	 *
-	 * @param Id
-	 * @returns
+	 * @param	Id The identifier of the payload being pulled.
+	 * @return	The payload in the form of a FCompressedBuffer. No decompression will
+	 *			be applied to the payload, it is up to the caller if they want to 
+	 *			retain the payload in compressed or uncompressed format.
+	 *			If no backend contained the payload then an empty invalid FCompressedBuffer
+	 *			will be returned.
 	 */
 	FCompressedBuffer PullData(const FPayloadId& Id);
 
+	/** Access profiling info relating to payload activity. Stats will only be collected if ENABLE_COOK_STATS is enabled.*/
+	FPayloadActivityInfo GetPayloadActivityInfo() const;
+
 private:
 
 	void ApplySettingsFromConfigFiles(const FConfigFile& PlatformEngineIni);
@@ -115,6 +136,9 @@ private:
 
 	void AddBackend(class IVirtualizationBackend* Backend);
 
+	bool TryPushDataToBackend(class IVirtualizationBackend& Backend, const FPayloadId& Id, const FCompressedBuffer& Payload);
+	FCompressedBuffer PullDataFromBackend(class IVirtualizationBackend& Backend, const FPayloadId& Id);
+
 	/** Are payloads allowed to be virtualized. Defaults to true. */
 	bool bEnablePayloadPushing;
 
