EditorDomain: Mark that package flag PKG_ReloadingForCooker is transient; it should not be saved to or loaded from the PackageFileSummary.

Add an enum value for transient package flags and clear them all on load from PackageFileSummary in addition to the current clearing on save, so that newly-marked transient flags that were previously saved out can still avoid being set when loading.

#rnx

diff --git a/Engine/Source/Runtime/CoreUObject/Private/UObject/PackageFileSummary.cpp b/Engine/Source/Runtime/CoreUObject/Private/UObject/PackageFileSummary.cpp
--- a/Engine/Source/Runtime/CoreUObject/Private/UObject/PackageFileSummary.cpp
+++ b/Engine/Source/Runtime/CoreUObject/Private/UObject/PackageFileSummary.cpp
@@ -181,13 +181,11 @@ void operator<<(FStructuredArchive::FSlot Slot, FPackageFileSummary& Sum)
 		Record << SA_VALUE(TEXT("FolderName"), Sum.FolderName);
 		Record << SA_VALUE(TEXT("PackageFlags"), Sum.PackageFlags);
 
-#if WITH_EDITOR
 		if (BaseArchive.IsLoading())
 		{
-			// This flag should never be saved and its reused, so we need to make sure it hasn't been loaded.
-			Sum.PackageFlags &= ~PKG_NewlyCreated;
+			// Transient flags should be set to false when saving or loading
+			Sum.PackageFlags &= ~PKG_TransientFlags;
 		}
-#endif // WITH_EDITOR
 
 		if (Sum.PackageFlags & PKG_FilterEditorOnly)
 		{
@@ -404,6 +402,6 @@ void FPackageFileSummary::SetCustomVersionContainer(const FCustomVersionContaine
 
 void FPackageFileSummary::SetPackageFlags(uint32 InPackageFlags)
 {
-	// Package summary flags cannot contain the 'NewlyCreated' nor the 'IsSaving' flag
-	PackageFlags = InPackageFlags & ~(PKG_NewlyCreated | PKG_IsSaving);
+	// Package summary flags should be set to false for all transient flags
+	PackageFlags = InPackageFlags & ~(PKG_TransientFlags);
 }
diff --git a/Engine/Source/Runtime/CoreUObject/Public/UObject/ObjectMacros.h b/Engine/Source/Runtime/CoreUObject/Public/UObject/ObjectMacros.h
--- a/Engine/Source/Runtime/CoreUObject/Public/UObject/ObjectMacros.h
+++ b/Engine/Source/Runtime/CoreUObject/Public/UObject/ObjectMacros.h
@@ -134,6 +134,9 @@ enum EPackageFlags
 	PKG_RuntimeGenerated			= 0x20000000,	///< This package contains elements that are runtime generated, and may not follow standard loading order rules
 	PKG_ReloadingForCooker			= 0x40000000,   ///< This package is reloading in the cooker, try to avoid getting data we will never need. We won't save this package.
 	PKG_FilterEditorOnly			= 0x80000000,	///< Package has editor-only data filtered out
+
+	// Transient Flags are cleared when serializing to or from PackageFileSummary
+	PKG_TransientFlags				= PKG_NewlyCreated | PKG_IsSaving | PKG_ReloadingForCooker,
 };
 
 /** Flag mask that indicates if this package is a package that exists in memory only. */
