Add trace scopes for garbage collection and memory trim

#rnx

diff --git a/Engine/Source/Runtime/Core/Private/HAL/MallocBinned2.cpp b/Engine/Source/Runtime/Core/Private/HAL/MallocBinned2.cpp
--- a/Engine/Source/Runtime/Core/Private/HAL/MallocBinned2.cpp
+++ b/Engine/Source/Runtime/Core/Private/HAL/MallocBinned2.cpp
@@ -1061,6 +1061,7 @@ const TCHAR* FMallocBinned2::GetDescriptiveName()
 void FMallocBinned2::FlushCurrentThreadCache()
 {
 	double StartTimeInner = FPlatformTime::Seconds();
+	TRACE_CPUPROFILER_EVENT_SCOPE(FMallocBinned2::FlushCurrentThreadCache);
 	QUICK_SCOPE_CYCLE_COUNTER(STAT_FMallocBinned2_FlushCurrentThreadCache);
 	FPerThreadFreeBlockLists* Lists = FPerThreadFreeBlockLists::Get();
 
diff --git a/Engine/Source/Runtime/Core/Private/HAL/MallocBinned3.cpp b/Engine/Source/Runtime/Core/Private/HAL/MallocBinned3.cpp
--- a/Engine/Source/Runtime/Core/Private/HAL/MallocBinned3.cpp
+++ b/Engine/Source/Runtime/Core/Private/HAL/MallocBinned3.cpp
@@ -1307,6 +1307,7 @@ const TCHAR* FMallocBinned3::GetDescriptiveName()
 void FMallocBinned3::FlushCurrentThreadCache()
 {
 	double StartTimeInner = FPlatformTime::Seconds();
+	TRACE_CPUPROFILER_EVENT_SCOPE(FMallocBinned3::FlushCurrentThreadCache);
 	QUICK_SCOPE_CYCLE_COUNTER(STAT_FMallocBinned3_FlushCurrentThreadCache);
 	FPerThreadFreeBlockLists* Lists = FPerThreadFreeBlockLists::Get();
 
diff --git a/Engine/Source/Runtime/Core/Private/HAL/UnrealMemory.cpp b/Engine/Source/Runtime/Core/Private/HAL/UnrealMemory.cpp
--- a/Engine/Source/Runtime/Core/Private/HAL/UnrealMemory.cpp
+++ b/Engine/Source/Runtime/Core/Private/HAL/UnrealMemory.cpp
@@ -538,6 +538,7 @@ void FMemory::Trim(bool bTrimThreadCaches)
 		GCreateMalloc();
 		CA_ASSUME(GMalloc != NULL);	// Don't want to assert, but suppress static analysis warnings about potentially NULL GMalloc
 	}
+	TRACE_CPUPROFILER_EVENT_SCOPE(FMemory::Trim);
 	QUICK_SCOPE_CYCLE_COUNTER(STAT_FMemory_Trim);
 	{
 		QUICK_SCOPE_CYCLE_COUNTER(STAT_FMemory_Trim_Broadcast);
diff --git a/Engine/Source/Runtime/CoreUObject/Private/UObject/GarbageCollection.cpp b/Engine/Source/Runtime/CoreUObject/Private/UObject/GarbageCollection.cpp
--- a/Engine/Source/Runtime/CoreUObject/Private/UObject/GarbageCollection.cpp
+++ b/Engine/Source/Runtime/CoreUObject/Private/UObject/GarbageCollection.cpp
@@ -303,6 +303,7 @@ class FAsyncPurge : public FRunnable
 	template <bool bMultithreaded> // Having this template argument lets the compiler strip unnecessary checks
 	bool TickDestroyObjects(bool bUseTimeLimit, float TimeLimit, double StartTime)
 	{
+		TRACE_CPUPROFILER_EVENT_SCOPE(FAsyncPurge::TickDestroyObjects);
 		const int32 TimeLimitEnforcementGranularityForDeletion = 100;
 		int32 ProcessedObjectsCount = 0;
 		bool bFinishedDestroyingObjects = true;
@@ -349,6 +350,7 @@ class FAsyncPurge : public FRunnable
 	/** [GAME THREAD] Destroys objects that are unreachable and couldn't be destroyed on the worker thread */
 	bool TickDestroyGameThreadObjects(bool bUseTimeLimit, float TimeLimit, double StartTime)
 	{
+		TRACE_CPUPROFILER_EVENT_SCOPE(FAsyncPurge::TickDestroyGameThreadObjects);
 		const int32 TimeLimitEnforcementGranularityForDeletion = 100;
 		int32 ProcessedObjectsCount = 0;
 		bool bFinishedDestroyingObjects = true;
@@ -1096,6 +1098,7 @@ class FRealtimeGC : public FGarbageCollectionTracer
 	template <EFastReferenceCollectorOptions CollectorOptions>
 	void PerformReachabilityAnalysisOnObjectsInternal(FGCArrayStruct* ArrayStruct)
 	{
+		TRACE_CPUPROFILER_EVENT_SCOPE(PerformReachabilityAnalysisOnObjectsInternal);
 		FGCReferenceProcessor<CollectorOptions> ReferenceProcessor;
 		// NOTE: we want to run with automatic token stream generation off as it should be already generated at this point,
 		// BUT we want to be ignoring Noop tokens as they're only pointing either at null references or at objects that never get GC'd (native classes)
@@ -1136,6 +1139,7 @@ public:
 	template <bool bParallel, bool bWithClusters>
 	void MarkObjectsAsUnreachable(TArray<UObject*>& ObjectsToSerialize, const EObjectFlags KeepFlags)
 	{
+		TRACE_CPUPROFILER_EVENT_SCOPE(MarkObjectsAsUnreachable);
 		const EInternalObjectFlags FastKeepFlags = EInternalObjectFlags::GarbageCollectionKeepFlags;
 		const int32 MaxNumberOfObjects = GUObjectArray.GetObjectArrayNum() - GUObjectArray.GetFirstGCIndex();
 		const int32 NumThreads = FMath::Max(1, FTaskGraphInterface::Get().GetNumWorkerThreads());
@@ -1153,6 +1157,7 @@ public:
 		// are part of the array so we don't suffer from cache misses as much as we would if we were to check ObjectFlags.
 		ParallelFor(NumThreads, [ObjectsToSerializeArrays, &ClustersToDissolveList, &KeepClusterRefsList, FastKeepFlags, KeepFlags, NumberOfObjectsPerThread, NumThreads, MaxNumberOfObjects](int32 ThreadIndex)
 		{
+			TRACE_CPUPROFILER_EVENT_SCOPE(MarkObjectsAsUnreachableTask);
 			int32 FirstObjectIndex = ThreadIndex * NumberOfObjectsPerThread + GUObjectArray.GetFirstGCIndex();
 			int32 NumObjects = (ThreadIndex < (NumThreads - 1)) ? NumberOfObjectsPerThread : (MaxNumberOfObjects - (NumThreads - 1) * NumberOfObjectsPerThread);
 			int32 LastObjectIndex = FMath::Min(GUObjectArray.GetObjectArrayNum() - 1, FirstObjectIndex + NumObjects - 1);
@@ -1450,10 +1455,6 @@ void IncrementalPurgeGarbage(bool bUseTimeLimit, float TimeLimit)
 #if !UE_WITH_GC
 	return;
 #else
-	SCOPED_NAMED_EVENT(IncrementalPurgeGarbage, FColor::Red);
-	DECLARE_SCOPE_CYCLE_COUNTER(TEXT("IncrementalPurgeGarbage"), STAT_IncrementalPurgeGarbage, STATGROUP_GC);
-	CSV_SCOPED_TIMING_STAT_EXCLUSIVE(GarbageCollection);
-
 	if (GExitPurge)
 	{
 		GObjPurgeIsRequired = true;
@@ -1466,6 +1467,10 @@ void IncrementalPurgeGarbage(bool bUseTimeLimit, float TimeLimit)
 		return;
 	}
 
+	SCOPED_NAMED_EVENT(IncrementalPurgeGarbage, FColor::Red);
+	DECLARE_SCOPE_CYCLE_COUNTER(TEXT("IncrementalPurgeGarbage"), STAT_IncrementalPurgeGarbage, STATGROUP_GC);
+	CSV_SCOPED_TIMING_STAT_EXCLUSIVE(GarbageCollection);
+
 	bool bCompleted = false;
 
 	struct FResetPurgeProgress
@@ -1528,6 +1533,7 @@ static FAutoConsoleVariableRef CVarAdditionalFinishDestroyTimeGC(
 
 bool IncrementalDestroyGarbage(bool bUseTimeLimit, float TimeLimit)
 {
+	TRACE_CPUPROFILER_EVENT_SCOPE(IncrementalDestroyGarbage);
 	const bool bMultithreadedPurge = !ShouldForceSingleThreadedGC() && GMultithreadedDestructionEnabled;
 	if (!GAsyncPurge)
 	{
@@ -1571,60 +1577,64 @@ bool IncrementalDestroyGarbage(bool bUseTimeLimit, float TimeLimit)
 			GObjCurrentPurgeObjectIndexNeedsReset = false;
 		}
 
-		while (GObjCurrentPurgeObjectIndex < GUnreachableObjects.Num())
 		{
-			FUObjectItem* ObjectItem = GUnreachableObjects[GObjCurrentPurgeObjectIndex];
-			checkSlow(ObjectItem);
-
-			//@todo UE - A prefetch was removed here. Re-add it. It wasn't right anyway, since it was ten items ahead and the consoles on have 8 prefetch slots
-
-			check(ObjectItem->IsUnreachable());
+			TRACE_CPUPROFILER_EVENT_SCOPE(ConditionalFinishDestroy);
+			while (GObjCurrentPurgeObjectIndex < GUnreachableObjects.Num())
 			{
-				UObject* Object = static_cast<UObject*>(ObjectItem->Object);
-				// Object should always have had BeginDestroy called on it and never already be destroyed
-				check( Object->HasAnyFlags( RF_BeginDestroyed ) && !Object->HasAnyFlags( RF_FinishDestroyed ) );
+				FUObjectItem* ObjectItem = GUnreachableObjects[GObjCurrentPurgeObjectIndex];
+				checkSlow(ObjectItem);
+
+				//@todo UE - A prefetch was removed here. Re-add it. It wasn't right anyway, since it was ten items ahead and the consoles on have 8 prefetch slots
 
-				// Only proceed with destroying the object if the asynchronous cleanup started by BeginDestroy has finished.
-				if(Object->IsReadyForFinishDestroy())
+				check(ObjectItem->IsUnreachable());
 				{
+					UObject* Object = static_cast<UObject*>(ObjectItem->Object);
+					// Object should always have had BeginDestroy called on it and never already be destroyed
+					check( Object->HasAnyFlags( RF_BeginDestroyed ) && !Object->HasAnyFlags( RF_FinishDestroyed ) );
+
+					// Only proceed with destroying the object if the asynchronous cleanup started by BeginDestroy has finished.
+					if(Object->IsReadyForFinishDestroy())
+					{
 #if PERF_DETAILED_PER_CLASS_GC_STATS
-					// Keep track of how many objects of a certain class we're purging.
-					const FName& ClassName = Object->GetClass()->GetFName();
-					int32 InstanceCount = GClassToPurgeCountMap.FindRef( ClassName );
-					GClassToPurgeCountMap.Add( ClassName, ++InstanceCount );
+						// Keep track of how many objects of a certain class we're purging.
+						const FName& ClassName = Object->GetClass()->GetFName();
+						int32 InstanceCount = GClassToPurgeCountMap.FindRef( ClassName );
+						GClassToPurgeCountMap.Add( ClassName, ++InstanceCount );
 #endif
-					// Send FinishDestroy message.
-					Object->ConditionalFinishDestroy();
-				}
-				else
-				{
-					// The object isn't ready for FinishDestroy to be called yet.  This is common in the
-					// case of a graphics resource that is waiting for the render thread "release fence"
-					// to complete.  Just calling IsReadyForFinishDestroy may begin the process of releasing
-					// a resource, so we don't want to block iteration while waiting on the render thread.
-
-					// Add the object index to our list of objects to revisit after we process everything else
-					GGCObjectsPendingDestruction.Add(Object);
-					GGCObjectsPendingDestructionCount++;
+						// Send FinishDestroy message.
+						Object->ConditionalFinishDestroy();
+					}
+					else
+					{
+						// The object isn't ready for FinishDestroy to be called yet.  This is common in the
+						// case of a graphics resource that is waiting for the render thread "release fence"
+						// to complete.  Just calling IsReadyForFinishDestroy may begin the process of releasing
+						// a resource, so we don't want to block iteration while waiting on the render thread.
+
+						// Add the object index to our list of objects to revisit after we process everything else
+						GGCObjectsPendingDestruction.Add(Object);
+						GGCObjectsPendingDestructionCount++;
+					}
 				}
-			}
 
-			// We've processed the object so increment our global iterator.  It's important to do this before
-			// we test for the time limit so that we don't process the same object again next tick!
-			++GObjCurrentPurgeObjectIndex;
+				// We've processed the object so increment our global iterator.  It's important to do this before
+				// we test for the time limit so that we don't process the same object again next tick!
+				++GObjCurrentPurgeObjectIndex;
 
-			// Only check time limit every so often to avoid calling FPlatformTime::Seconds too often.
-			const bool bPollTimeLimit = ((TimeLimitTimePollCounter++) % TimeLimitEnforcementGranularityForDestroy == 0);
-			if( bUseTimeLimit && bPollTimeLimit && ((FPlatformTime::Seconds() - GCStartTime) > TimeLimit) )
-			{
-				bTimeLimitReached = true;
-				break;
+				// Only check time limit every so often to avoid calling FPlatformTime::Seconds too often.
+				const bool bPollTimeLimit = ((TimeLimitTimePollCounter++) % TimeLimitEnforcementGranularityForDestroy == 0);
+				if( bUseTimeLimit && bPollTimeLimit && ((FPlatformTime::Seconds() - GCStartTime) > TimeLimit) )
+				{
+					bTimeLimitReached = true;
+					break;
+				}
 			}
 		}
 
 		// Have we finished the first round of attempting to call FinishDestroy on unreachable objects?
 		if (GObjCurrentPurgeObjectIndex >= GUnreachableObjects.Num())
 		{
+			TRACE_CPUPROFILER_EVENT_SCOPE(ConditionalFinishDestroyDeferred);
 			double MaxTimeForFinishDestroy = 10.00;
 			bool bFinishDestroyTimeExtended = false;
 			FString FirstObjectNotReadyWhenTimeExtended;
@@ -1780,6 +1790,7 @@ bool IncrementalDestroyGarbage(bool bUseTimeLimit, float TimeLimit)
 	
 	if (GObjFinishDestroyHasBeenRoutedToAllObjects && !bTimeLimitReached)
 	{
+		TRACE_CPUPROFILER_EVENT_SCOPE(DestroyObjects);
 		// Perform actual object deletion.
 		int32 ProcessCount = 0;
 		if (GObjCurrentPurgeObjectIndexNeedsReset)
@@ -1863,6 +1874,7 @@ static FAutoConsoleVariableRef CVarFlushStreamingOnGC(
 
 void GatherUnreachableObjects(bool bForceSingleThreaded)
 {
+	TRACE_CPUPROFILER_EVENT_SCOPE(GatherUnreachableObjects);
 	DECLARE_SCOPE_CYCLE_COUNTER(TEXT("CollectGarbageInternal.GatherUnreachableObjects"), STAT_CollectGarbageInternal_GatherUnreachableObjects, STATGROUP_GC);
 
 	const double StartTime = FPlatformTime::Seconds();
@@ -1881,6 +1893,7 @@ void GatherUnreachableObjects(bool bForceSingleThreaded)
 	// are part of the array so we don't suffer from cache misses as much as we would if we were to check ObjectFlags.
 	ParallelFor(NumThreads, [&ClusterItemsToDestroy, NumberOfObjectsPerThread, NumThreads, MaxNumberOfObjects](int32 ThreadIndex)
 	{
+		TRACE_CPUPROFILER_EVENT_SCOPE(GatherUnreachableObjectsTask);
 		int32 FirstObjectIndex = ThreadIndex * NumberOfObjectsPerThread + (GExitPurge ? 0 : GUObjectArray.GetFirstGCIndex());
 		int32 NumObjects = (ThreadIndex < (NumThreads - 1)) ? NumberOfObjectsPerThread : (MaxNumberOfObjects - (NumThreads - 1) * NumberOfObjectsPerThread);
 		int32 LastObjectIndex = FMath::Min(GUObjectArray.GetObjectArrayNum() - 1, FirstObjectIndex + NumObjects - 1);
@@ -1995,7 +2008,10 @@ void CollectGarbageInternal(EObjectFlags KeepFlags, bool bPerformFullPurge)
 
 	// Route callbacks so we can ensure that we are e.g. not in the middle of loading something by flushing
 	// the async loading, etc...
-	FCoreUObjectDelegates::GetPreGarbageCollectDelegate().Broadcast();
+	{
+		TRACE_CPUPROFILER_EVENT_SCOPE(BroadcastPreGarbageCollect);
+		FCoreUObjectDelegates::GetPreGarbageCollectDelegate().Broadcast();
+	}
 	GLastGCFrame = GFrameCounter;
 
 	{
@@ -2057,7 +2073,10 @@ void CollectGarbageInternal(EObjectFlags KeepFlags, bool bPerformFullPurge)
 		}
 
 		// Fire post-reachability analysis hooks
-		FCoreUObjectDelegates::PostReachabilityAnalysis.Broadcast();
+		{
+			TRACE_CPUPROFILER_EVENT_SCOPE(BroadcastPostReachabilityAnalysis);
+			FCoreUObjectDelegates::PostReachabilityAnalysis.Broadcast();
+		}
 
 		{
 			GatherUnreachableObjects(bForceSingleThreadedGC);
@@ -2076,7 +2095,7 @@ void CollectGarbageInternal(EObjectFlags KeepFlags, bool bPerformFullPurge)
 		// Set flag to indicate that we are relying on a purge to be performed.
 		GObjPurgeIsRequired = true;
 
-		// Perform a full purge by not using a time limit for the incremental purge. The Editor always does a full purge.
+		// Perform a full purge by not using a time limit for the incremental purge.
 		if (bPerformFullPurge)
 		{
 			IncrementalPurgeGarbage(false);
@@ -2095,7 +2114,10 @@ void CollectGarbageInternal(EObjectFlags KeepFlags, bool bPerformFullPurge)
 	}
 
 	// Route callbacks to verify GC assumptions
-	FCoreUObjectDelegates::GetPostGarbageCollect().Broadcast();
+	{
+		TRACE_CPUPROFILER_EVENT_SCOPE(BroadcastPostGarbageCollect);
+		FCoreUObjectDelegates::GetPostGarbageCollect().Broadcast();
+	}
 
 	STAT_ADD_CUSTOMMESSAGE_NAME( STAT_NamedMarker, TEXT( "GarbageCollection - End" ) );
 #endif	// UE_WITH_GC
@@ -2108,11 +2130,15 @@ bool IsIncrementalUnhashPending()
 
 bool UnhashUnreachableObjects(bool bUseTimeLimit, float TimeLimit)
 {
+	TRACE_CPUPROFILER_EVENT_SCOPE(UnhashUnreachableObjects);
 	DECLARE_SCOPE_CYCLE_COUNTER(TEXT("UnhashUnreachableObjects"), STAT_UnhashUnreachableObjects, STATGROUP_GC);
 
 	TGuardValue<bool> GuardObjUnhashUnreachableIsInProgress(GObjUnhashUnreachableIsInProgress, true);
 
-	FCoreUObjectDelegates::PreGarbageCollectConditionalBeginDestroy.Broadcast();
+	{
+		TRACE_CPUPROFILER_EVENT_SCOPE(BroadcastGarbageCollectConditionalBeginDestroy);
+		FCoreUObjectDelegates::PreGarbageCollectConditionalBeginDestroy.Broadcast();
+	}
 
 	// Unhash all unreachable objects.
 	const double StartTime = FPlatformTime::Seconds();
@@ -2121,24 +2147,27 @@ bool UnhashUnreachableObjects(bool bUseTimeLimit, float TimeLimit)
 	int32 TimePollCounter = 0;
 	const bool bFirstIteration = (GUnrechableObjectIndex == 0);
 
-	while (GUnrechableObjectIndex < GUnreachableObjects.Num())
 	{
-		//@todo UE - A prefetch was removed here. Re-add it. It wasn't right anyway, since it was ten items ahead and the consoles on have 8 prefetch slots
-
-		FUObjectItem* ObjectItem = GUnreachableObjects[GUnrechableObjectIndex++];
+		TRACE_CPUPROFILER_EVENT_SCOPE(ConditionalBeginDestroy);
+		while (GUnrechableObjectIndex < GUnreachableObjects.Num())
 		{
-			UObject* Object = static_cast<UObject*>(ObjectItem->Object);
-			FScopedCBDProfile Profile(Object);
-			// Begin the object's asynchronous destruction.
-			Object->ConditionalBeginDestroy();
-		}
+			//@todo UE - A prefetch was removed here. Re-add it. It wasn't right anyway, since it was ten items ahead and the consoles on have 8 prefetch slots
 
-		Items++;
+			FUObjectItem* ObjectItem = GUnreachableObjects[GUnrechableObjectIndex++];
+			{
+				UObject* Object = static_cast<UObject*>(ObjectItem->Object);
+				FScopedCBDProfile Profile(Object);
+				// Begin the object's asynchronous destruction.
+				Object->ConditionalBeginDestroy();
+			}
 
-		const bool bPollTimeLimit = ((TimePollCounter++) % TimeLimitEnforcementGranularityForBeginDestroy == 0);
-		if (bUseTimeLimit && bPollTimeLimit && ((FPlatformTime::Seconds() - StartTime) > TimeLimit))
-		{
-			break;
+			Items++;
+
+			const bool bPollTimeLimit = ((TimePollCounter++) % TimeLimitEnforcementGranularityForBeginDestroy == 0);
+			if (bUseTimeLimit && bPollTimeLimit && ((FPlatformTime::Seconds() - StartTime) > TimeLimit))
+			{
+				break;
+			}
 		}
 	}
 
@@ -2162,7 +2191,10 @@ bool UnhashUnreachableObjects(bool bUseTimeLimit, float TimeLimit)
 		UE_LOG(LogGarbage, Log, TEXT("Starting unhashing unreachable objects (%d objects to unhash)."), GUnreachableObjects.Num());
 	}
 
-	FCoreUObjectDelegates::PostGarbageCollectConditionalBeginDestroy.Broadcast();
+	{
+		TRACE_CPUPROFILER_EVENT_SCOPE(BroadCastPostGarbageCollectConditionalBeginDestroy);
+		FCoreUObjectDelegates::PostGarbageCollectConditionalBeginDestroy.Broadcast();
+	}
 
 	// Return true if time limit has been reached
 	return bTimeLimitReached;
diff --git a/Engine/Source/Runtime/CoreUObject/Private/UObject/UObjectHash.cpp b/Engine/Source/Runtime/CoreUObject/Private/UObject/UObjectHash.cpp
--- a/Engine/Source/Runtime/CoreUObject/Private/UObject/UObjectHash.cpp
+++ b/Engine/Source/Runtime/CoreUObject/Private/UObject/UObjectHash.cpp
@@ -861,6 +861,7 @@ FORCEINLINE static UPackage* UnassignExternalPackageFromObject(FUObjectHashTable
 
 void ShrinkUObjectHashTables()
 {
+	TRACE_CPUPROFILER_EVENT_SCOPE(ShrinkUObjectHashTables);
 	FUObjectHashTables& ThreadHash = FUObjectHashTables::Get();
 	FHashTableLock HashLock(ThreadHash);
 	ThreadHash.ShrinkMaps();
diff --git a/Engine/Source/Runtime/CoreUObject/Public/UObject/FastReferenceCollector.h b/Engine/Source/Runtime/CoreUObject/Public/UObject/FastReferenceCollector.h
--- a/Engine/Source/Runtime/CoreUObject/Public/UObject/FastReferenceCollector.h
+++ b/Engine/Source/Runtime/CoreUObject/Public/UObject/FastReferenceCollector.h
@@ -191,6 +191,7 @@ public:
 	 */
 	void ClearWeakReferences(bool bClearPools)
 	{
+		TRACE_CPUPROFILER_EVENT_SCOPE(ClearWeakReferences);
 		TArray<FGCArrayStruct*> AllArrays;
 		Pool.PopAll(AllArrays);
 		int32 Index = 0;
@@ -398,6 +399,7 @@ private:
 					}
 					else
 					{
+						TRACE_CPUPROFILER_EVENT_SCOPE(FCollectorTaskQueue::WaitEvent);
 						check(WaitEvent);
 						WaitEvent->Wait();
 						FPlatformProcess::ReturnSynchEventToPool(WaitEvent);
@@ -1428,4 +1430,4 @@ public:
 	}
 	// Implement this in your derived class, don't make this virtual as it will affect performance!
 	//FORCEINLINE void HandleTokenStreamObjectReference(FGCArrayStruct& ObjectsToSerializeStruct, UObject* ReferencingObject, UObject*& Object, const int32 TokenIndex, bool bAllowReferenceElimination);
-};
\ No newline at end of file
+};
diff --git a/Engine/Source/Runtime/Engine/Private/AudioThread.cpp b/Engine/Source/Runtime/Engine/Private/AudioThread.cpp
--- a/Engine/Source/Runtime/Engine/Private/AudioThread.cpp
+++ b/Engine/Source/Runtime/Engine/Private/AudioThread.cpp
@@ -203,6 +203,7 @@ static int32 AudioThreadSuspendCount = 0;
 
 void FAudioThread::SuspendAudioThread()
 {
+	TRACE_CPUPROFILER_EVENT_SCOPE(SuspendAudioThread);
 	check(FPlatformTLS::GetCurrentThreadId() == GGameThreadId);
 	check(!GIsAudioThreadSuspended || CVarSuspendAudioThread.GetValueOnGameThread() != 0);
 	if (IsAudioThreadRunning())
diff --git a/Engine/Source/Runtime/Engine/Private/StreamableManager.cpp b/Engine/Source/Runtime/Engine/Private/StreamableManager.cpp
--- a/Engine/Source/Runtime/Engine/Private/StreamableManager.cpp
+++ b/Engine/Source/Runtime/Engine/Private/StreamableManager.cpp
@@ -926,6 +926,7 @@ FStreamableManager::~FStreamableManager()
 
 void FStreamableManager::OnPreGarbageCollect()
 {
+	TRACE_CPUPROFILER_EVENT_SCOPE(FStreamableManager::OnPreGarbageCollect);
 	TSet<FSoftObjectPath> RedirectsToRemove;
 	TArray<FStreamable*> StreamablesToDelete;
 
diff --git a/Engine/Source/Runtime/Engine/Private/Streaming/StreamingManagerTexture.cpp b/Engine/Source/Runtime/Engine/Private/Streaming/StreamingManagerTexture.cpp
--- a/Engine/Source/Runtime/Engine/Private/Streaming/StreamingManagerTexture.cpp
+++ b/Engine/Source/Runtime/Engine/Private/Streaming/StreamingManagerTexture.cpp
@@ -187,6 +187,7 @@ FRenderAssetStreamingManager::~FRenderAssetStreamingManager()
 
 void FRenderAssetStreamingManager::OnPreGarbageCollect()
 {
+	TRACE_CPUPROFILER_EVENT_SCOPE(FRenderAssetStreamingManager::OnPreGarbageCollect);
 #if WITH_EDITORONLY_DATA
 	void PurgeAbandonedDDCHandles();
 	PurgeAbandonedDDCHandles();
diff --git a/Engine/Source/Runtime/Engine/Private/UnrealEngine.cpp b/Engine/Source/Runtime/Engine/Private/UnrealEngine.cpp
--- a/Engine/Source/Runtime/Engine/Private/UnrealEngine.cpp
+++ b/Engine/Source/Runtime/Engine/Private/UnrealEngine.cpp
@@ -1314,6 +1314,7 @@ static FAutoConsoleVariableRef CVarIncrementalGCTimePerFrame(
 
 void UEngine::PreGarbageCollect()
 {
+	TRACE_CPUPROFILER_EVENT_SCOPE(UEngine::PreGarbageCollect);
 	ForEachObjectOfClass(UWorld::StaticClass(), [](UObject* WorldObj)
 	{
 		UWorld* World = CastChecked<UWorld>(WorldObj);
diff --git a/Engine/Source/Runtime/Engine/Private/World.cpp b/Engine/Source/Runtime/Engine/Private/World.cpp
--- a/Engine/Source/Runtime/Engine/Private/World.cpp
+++ b/Engine/Source/Runtime/Engine/Private/World.cpp
@@ -3051,6 +3051,7 @@ void FLevelStreamingGCHelper::CancelUnloadRequest( ULevel* InLevel )
 
 void FLevelStreamingGCHelper::PrepareStreamedOutLevelsForGC()
 {
+	TRACE_CPUPROFILER_EVENT_SCOPE(FLevelStreamingGCHelper::PrepareStreamedOutLevelsForGC);
 	if (LevelsPendingUnload.Num() > 0)
 	{
 		OnGCStreamedOutLevels.Broadcast();
