Downgrade linker warning for missing parent structs

Apply similar logic to functions that are missing their parent implementation
Only for UScriptStructs that are outered to a UClass (e.g. sparse class data) and only when that class is not regenerating on load

diff --git a/Engine/Source/Runtime/CoreUObject/Private/UObject/LinkerLoad.cpp b/Engine/Source/Runtime/CoreUObject/Private/UObject/LinkerLoad.cpp
--- a/Engine/Source/Runtime/CoreUObject/Private/UObject/LinkerLoad.cpp
+++ b/Engine/Source/Runtime/CoreUObject/Private/UObject/LinkerLoad.cpp
@@ -4688,7 +4688,25 @@ UObject* FLinkerLoad::CreateExport( int32 Index )
 				}
 				else
 				{
-					if (!FLinkerLoad::IsKnownMissingPackage(*GetExportFullName(Index)))
+					bool bFailedToLoadGeneratedStruct = false;
+					if( LoadClass->IsChildOf(UScriptStruct::StaticClass()) )
+					{
+						// Similar to functions, in the case of structures that are outered to a class (e.g. generated sparse
+						// class data), it is also possible to legitimately fail to load the parent structure here as while 
+						// we will regenerate the structure on load, it wont appear in the export map until it is re-saved
+						UObject* ObjOuter = IndexToObject(Export.OuterIndex);
+						if (ObjOuter && !Export.bExportLoadFailed)
+						{
+							UClass* StructClass = Cast<UClass>(ObjOuter);
+							if (StructClass && StructClass->ClassGeneratedBy && !StructClass->ClassGeneratedBy->HasAnyFlags(RF_BeingRegenerated))
+							{
+								UE_LOG(LogLinker, Display, TEXT("CreateExport: Failed to load Parent for %s; resaving the parents of %s will remove this message"), *GetExportFullName(Index), *StructClass->ClassGeneratedBy->GetFullName());
+								bFailedToLoadGeneratedStruct = true;
+							}
+						}
+					}
+					
+					if (!bFailedToLoadGeneratedStruct && !FLinkerLoad::IsKnownMissingPackage(*GetExportFullName(Index)))
 					{
 						UE_LOG(LogLinker, Warning, TEXT("CreateExport: Failed to load Parent for %s"), *GetExportFullName(Index));
 					}
