Implement fast path for FCoreRedirects::GetMatchingRedirects when looking for a single redirect type.
Saves ~18% of the time spent in FLinkerLoad::FixupImportMap, which is time sliced in a 5ms window in PIE.

diff --git a/Engine/Source/Runtime/CoreUObject/Private/UObject/CoreRedirects.cpp b/Engine/Source/Runtime/CoreUObject/Private/UObject/CoreRedirects.cpp
--- a/Engine/Source/Runtime/CoreUObject/Private/UObject/CoreRedirects.cpp
+++ b/Engine/Source/Runtime/CoreUObject/Private/UObject/CoreRedirects.cpp
@@ -570,29 +570,51 @@ bool FCoreRedirects::GetMatchingRedirects(ECoreRedirectFlags SearchFlags, const
 	// Look for all redirects that match the given names and flags
 	bool bFound = false;
 	
+	auto SearchRedirect = [&OldObjectName, &bFound, &FoundRedirects](ECoreRedirectFlags PairFlags, const FRedirectNameMap& RedirectNameMap)
+	{
+		if (const TArray<FCoreRedirect>* RedirectsForName = RedirectNameMap.RedirectMap.Find(OldObjectName.GetSearchKey(PairFlags)))
+		{
+			for (const FCoreRedirect& CheckRedirect : *RedirectsForName)
+			{
+				if (CheckRedirect.Matches(PairFlags, OldObjectName))
+				{
+					bFound = true;
+					FoundRedirects.Add(&CheckRedirect);
+				}
+			}
+		}
+	};
+
 	// If we're not explicitly searching for packages or looking for removed things, add the implicit (Type=Package,Category=None) redirects
 	const bool bSearchPackageRedirects = !(SearchFlags & ECoreRedirectFlags::Type_Package) && !(SearchFlags & ECoreRedirectFlags::Category_Removed);
 
-	// Determine list of maps to look over, need to handle being passed multiple types in a bit mask
-	for (const TPair<ECoreRedirectFlags, FRedirectNameMap>& Pair : RedirectTypeMap)
+	if (FMath::CountBits((uint64)SearchFlags) == 1)
 	{
-		ECoreRedirectFlags PairFlags = Pair.Key;
+		// Fast path when looking for a single type
+		if (const FRedirectNameMap* RedirectNameMap = RedirectTypeMap.Find(SearchFlags))
+		{
+			SearchRedirect(SearchFlags, *RedirectNameMap);
+		}
 
-		// We need to check all maps that match the search or package flags
-		if (CheckRedirectFlagsMatch(PairFlags, SearchFlags) || (bSearchPackageRedirects && CheckRedirectFlagsMatch(PairFlags, ECoreRedirectFlags::Type_Package)))
+		if (bSearchPackageRedirects)
 		{
-			const TArray<FCoreRedirect>* RedirectsForName = Pair.Value.RedirectMap.Find(OldObjectName.GetSearchKey(PairFlags));
+			if (const FRedirectNameMap* RedirectNameMap = RedirectTypeMap.Find(ECoreRedirectFlags::Type_Package))
+			{
+				SearchRedirect(SearchFlags, *RedirectNameMap);
+			}
+		}
+	}
+	else
+	{
+		// Determine list of maps to look over, need to handle being passed multiple types in a bit mask
+		for (const TPair<ECoreRedirectFlags, FRedirectNameMap>& Pair : RedirectTypeMap)
+		{
+			ECoreRedirectFlags PairFlags = Pair.Key;
 
-			if (RedirectsForName)
+			// We need to check all maps that match the search or package flags
+			if (CheckRedirectFlagsMatch(PairFlags, SearchFlags) || (bSearchPackageRedirects && CheckRedirectFlagsMatch(PairFlags, ECoreRedirectFlags::Type_Package)))
 			{
-				for (const FCoreRedirect& CheckRedirect : *RedirectsForName)
-				{
-					if (CheckRedirect.Matches(PairFlags, OldObjectName))
-					{
-						bFound = true;
-						FoundRedirects.Add(&CheckRedirect);
-					}
-				}
+				SearchRedirect(PairFlags, Pair.Value);
 			}
 		}
 	}
