Fix error message that was incomplete

#rnx

diff --git a/Engine/Source/Runtime/CoreUObject/Private/UObject/Class.cpp b/Engine/Source/Runtime/CoreUObject/Private/UObject/Class.cpp
--- a/Engine/Source/Runtime/CoreUObject/Private/UObject/Class.cpp
+++ b/Engine/Source/Runtime/CoreUObject/Private/UObject/Class.cpp
@@ -5298,7 +5298,7 @@ void UClass::AddNativeFunction(const WIDECHAR* InName, FNativeFuncPtr InPointer)
 		else
 		{
 			// function was not found, so it's new
-			UE_LOG(LogClass, Log, TEXT("Function %s is new or belongs to a modified."), *InFName.ToString());
+			UE_LOG(LogClass, Log, TEXT("Function %s is new or belongs to a modified class."), *InFName.ToString());
 		}
 	}
 #endif
