Reloading inline bulkdata in the old loader will now raise an ensure rather than log warnings as it should make it easier for people to spot the problem before it hits the new loader and fails to work entirely.

- Since reloading inline data does not work in the old loader we need to make it more obvious to developers when they are doing it in the old loader so that it doesn't get submitted.
- An ensure will cause the process top break in the debugger and is more likely to get attention. It will also provide the callstack for any QA bug report which will make it easier for us to triage to the correct team to fix their calling code than a vague warning about bulkdata.
- Also provided a new define, UE_KEEP_INLINE_RELOADING_CONSISTENT, which when enabled will cause reloading of inline bulkdata to fail (and error) in the old loader as well.
- The ensure can be supressed by setting [Core.System]IgnoreInlineBulkDataReloadEnsures=True in an engine config file, this might be useful for projects that do not care about running from staged data for what ever reason (but supressing the ensure will not prevent the data from failing to load when run with the new loader so this is strictly a 'use at your own risk' option)

diff --git a/Engine/Source/Runtime/CoreUObject/Private/Serialization/BulkData2.cpp b/Engine/Source/Runtime/CoreUObject/Private/Serialization/BulkData2.cpp
--- a/Engine/Source/Runtime/CoreUObject/Private/Serialization/BulkData2.cpp
+++ b/Engine/Source/Runtime/CoreUObject/Private/Serialization/BulkData2.cpp
@@ -26,6 +26,10 @@ IMPLEMENT_TYPE_LAYOUT(FBulkDataBase);
 // If set to 0 then we will pretend that optional data does not exist, useful for testing.
 #define ALLOW_OPTIONAL_DATA 1
 
+// When set to 1 attempting to reload inline data will fail with the old loader in the same way that it fails in
+// the new loader to keep the results consistent.
+#define UE_KEEP_INLINE_RELOADING_CONSISTENT 0
+
 // Handy macro to validate FIoStatus return values
 #define CHECK_IOSTATUS( InIoStatus, InMethodName ) checkf(InIoStatus.IsOk(), TEXT("%s failed: %s"), InMethodName, *InIoStatus.ToString());
 
@@ -65,6 +69,26 @@ namespace
 		}
 	}
 
+	bool ShouldIgnoreInlineDataReloadEnsures()
+	{
+		static struct FIgnoreInlineDataReloadEnsures
+		{
+			bool bEnabled = false;
+
+			FIgnoreInlineDataReloadEnsures()
+			{
+				FConfigFile PlatformEngineIni;
+				FConfigCacheIni::LoadLocalIniFile(PlatformEngineIni, TEXT("Engine"), true, ANSI_TO_TCHAR(FPlatformProperties::IniPlatformName()));
+
+				PlatformEngineIni.GetBool(TEXT("Core.System"), TEXT("IgnoreInlineBulkDataReloadEnsures"), bEnabled);
+
+				UE_LOG(LogSerialization, Display, TEXT("IgnoreInlineDataReloadEnsures: '%s'"), bEnabled ? TEXT("true") : TEXT("false"));
+			}
+		} IgnoreInlineDataReloadEnsures;
+
+		return IgnoreInlineDataReloadEnsures.bEnabled;
+	}
+
 	const FIoFilenameHash FALLBACK_IO_FILENAME_HASH = INVALID_IO_FILENAME_HASH - 1;
 }
 
@@ -761,9 +785,14 @@ void FBulkDataBase::ConditionalSetInlineAlwaysAllowDiscard()
 	// This will cause a bug when using IoStore in any systems that do actually need to reload inline data.
 	// Each project that uses IoStore needs to guarantee that they do not have any systems that need to
 	// reload inline data.
-	//
-	// TODO: We need to make the old loader and new loader consistent on how inline data is treated!
+	// Note that when the define UE_KEEP_INLINE_RELOADING_CONSISTENT is enabled the old loading path should also
+	// not allow the reloading of inline data so that it behaves the same way as the new loading path. So when it
+	// is enabled we do not need to check if the loader is or isn't enabled, we can just set the flag and assume
+	// all inline data should be allowed to be discarded.
+
+#if !UE_KEEP_INLINE_RELOADING_CONSISTENT
 	if (FIoDispatcher::IsInitialized())
+#endif // !UE_KEEP_INLINE_RELOADING_CONSISTENT
 	{
 		SetBulkDataFlags(BULKDATA_AlwaysAllowDiscard);
 	}
@@ -1175,6 +1204,13 @@ bool FBulkDataBase::CanLoadFromDisk() const
 		return true;
 	}
 
+#if UE_KEEP_INLINE_RELOADING_CONSISTENT
+	if (IsInlined())
+	{
+		return false;
+	}
+#endif //UE_KEEP_INLINE_RELOADING_CONSISTENT
+
 	// If this BulkData has a fallback token then it can find it's filepath and load from disk
 	if (Data.Token != InvalidToken)
 	{
@@ -1522,7 +1558,7 @@ bool FBulkDataBase::CanDiscardInternalData() const
 		return true;
 	}
 
-	// IF BULKDATA_AlwaysAllowDiscard has been set then we should always allow the data to 
+	// If BULKDATA_AlwaysAllowDiscard has been set then we should always allow the data to 
 	// be discarded even if it cannot be reloaded again.
 	if ((BulkDataFlags & BULKDATA_AlwaysAllowDiscard) != 0)
 	{
@@ -1535,40 +1571,64 @@ bool FBulkDataBase::CanDiscardInternalData() const
 void FBulkDataBase::LoadDataDirectly(void** DstBuffer)
 {
 	DECLARE_SCOPE_CYCLE_COUNTER(TEXT("FBulkDataBase::LoadDataDirectly"), STAT_UBD_LoadDataDirectly, STATGROUP_Memory);
-	if (!CanLoadFromDisk())
+	
+	// Early out if we have no data to load
+	if (GetBulkDataSize() == 0)
 	{
-		// Only warn if the bulkdata have a valid size
-		UE_CLOG(GetBulkDataSize() > 0, LogSerialization, Warning, TEXT("Attempting to load a BulkData object that cannot be loaded from disk"));
-
-		return; // Early out if there is nothing to load anyway
+		return;
 	}
 
-	if (!IsIoDispatcherEnabled())
+#if !UE_KEEP_INLINE_RELOADING_CONSISTENT	
+	// Note that we only want the ensure to trigger if we have a valid offset (the bulkdata references data from disk)
+	ensureMsgf(!IsInlined() || BulkDataOffset == INDEX_NONE || ShouldIgnoreInlineDataReloadEnsures(), 
+			TEXT(	"Attempting to reload inline BulkData! This operation is not supported by the IoDispatcher and so will eventually stop working."
+					" The calling code should be fixed to retain the inline data in memory and re-use it rather than discard it and then try to reload from disk!"));
+#endif //!UE_KEEP_INLINE_RELOADING_CONSISTENT
+
+	if (!CanLoadFromDisk())
 	{
-		InternalLoadFromPackageResource(DstBuffer);
+		UE_LOG(LogSerialization, Error, TEXT("Attempting to load a BulkData object that cannot be loaded from disk"));
+		return;
 	}
-	else if (IsUsingIODispatcher() && IsInSeparateFile())
+
+	if (IsUsingIODispatcher())
 	{
 		InternalLoadFromIoStore(DstBuffer);
 	}
 	else
 	{
-		// Note that currently this shouldn't be reachable as we should early out due to the ::CanLoadFromDisk check at the start of the method
-		UE_LOG(LogSerialization, Error, TEXT("Attempting to %s when the IoDispatcher is enabled, this operation is not supported!"),
-			IsInlined() ? TEXT("reload inline BulkData") : TEXT("load BulkData from end-of-package-file section"));
+		InternalLoadFromPackageResource(DstBuffer);
 	}
 }
 
 void FBulkDataBase::LoadDataAsynchronously(AsyncCallback&& Callback)
 {
 	DECLARE_SCOPE_CYCLE_COUNTER(TEXT("FBulkDataBase::LoadDataDirectly"), STAT_UBD_LoadDataDirectly, STATGROUP_Memory);
+	
+	// Early out if we have no data to load
+	if (GetBulkDataSize() == 0)
+	{
+		return;
+	}
+
+#if !UE_KEEP_INLINE_RELOADING_CONSISTENT
+	ensureMsgf(!IsInlined() || ShouldIgnoreInlineDataReloadEnsures(), 
+				TEXT(	"Attempting to reload inline BulkData! This operation is not supported by the IoDispatcher and so will eventually stop working."
+						" The calling code should be fixed to retain the inline data in memory and re-use it rather than discard it and then try to reload from disk!"));
+#endif //!UE_KEEP_INLINE_RELOADING_CONSISTENT
+
 	if (!CanLoadFromDisk())
 	{
-		UE_LOG(LogSerialization, Warning, TEXT("Attempting to load a BulkData object that cannot be loaded from disk"));
+		UE_LOG(LogSerialization, Error, TEXT("Attempting to load a BulkData object that cannot be loaded from disk"));
 		return; // Early out if there is nothing to load anyway
 	}
 
-	if (!IsIoDispatcherEnabled())
+	if (IsUsingIODispatcher())
+	{
+		void* DummyPointer = nullptr;
+		InternalLoadFromIoStoreAsync(&DummyPointer, MoveTemp(Callback));
+	}
+	else
 	{
 		Async(EAsyncExecution::ThreadPool, [this, Callback]()
 			{
@@ -1581,17 +1641,6 @@ void FBulkDataBase::LoadDataAsynchronously(AsyncCallback&& Callback)
 				Callback(Status);
 			});
 	}
-	else if (IsUsingIODispatcher() && IsInSeparateFile())
-	{
-		void* DummyPointer = nullptr;
-		InternalLoadFromIoStoreAsync(&DummyPointer, MoveTemp(Callback));
-	}
-	else
-	{
-		// Note that currently this shouldn't be reachable as we should early out due to the ::CanLoadFromDisk check at the start of the method
-		UE_LOG(LogSerialization, Error, TEXT("Attempting to %s when the IoDispatcher is enabled, this operation is not supported!"),
-			IsInlined() ? TEXT("reload inline BulkData") : TEXT("load BulkData from end-of-package-file section"));
-	}
 }
 
 void FBulkDataBase::InternalLoadFromPackageResource(void** DstBuffer)
