Removed superfluous include in preparation for iterating on that header

diff --git a/Engine/Source/Runtime/CoreUObject/Public/CoreUObject.h b/Engine/Source/Runtime/CoreUObject/Public/CoreUObject.h
--- a/Engine/Source/Runtime/CoreUObject/Public/CoreUObject.h
+++ b/Engine/Source/Runtime/CoreUObject/Public/CoreUObject.h
@@ -13,7 +13,6 @@ MONOLITHIC_HEADER_BOILERPLATE()
 #include "UObject/Script.h"
 #include "UObject/ScriptMacros.h"
 #include "UObject/ObjectMacros.h"
-#include "UObject/UObjectAllocator.h"
 #include "UObject/UObjectGlobals.h"
 #include "UObject/UObjectMarks.h"
 #include "UObject/UObjectBase.h"
