Set/map property loading optimization

 ~40x speedup for set property loading, 281ms -> 6.5ms, initial load time for Win64 Test client in an internal project

Slight refactor in preparation of more advanced container allocators and allocating all UObject property memory linearly after UObject itself

diff --git a/Engine/Source/Runtime/CoreUObject/Private/UObject/PropertyMap.cpp b/Engine/Source/Runtime/CoreUObject/Private/UObject/PropertyMap.cpp
--- a/Engine/Source/Runtime/CoreUObject/Private/UObject/PropertyMap.cpp
+++ b/Engine/Source/Runtime/CoreUObject/Private/UObject/PropertyMap.cpp
@@ -331,70 +331,110 @@ void FMapProperty::SerializeItem(FStructuredArchive::FSlot Slot, void* Value, co
 		{
 			CopyValuesInternal(Value, Defaults, 1);
 		}
-		else
-		{
-			MapHelper.EmptyValues();
-		}
 
-		uint8* TempKeyValueStorage = nullptr;
-		ON_SCOPE_EXIT
+		// Delete any explicitly-removed elements
+		int32 NumKeysToRemove = 0;
+		FStructuredArchive::FArray KeysToRemoveArray = Record.EnterArray(SA_FIELD_NAME(TEXT("KeysToRemove")), NumKeysToRemove);
+
+		if (MapHelper.Num() == 0) // Faster loading path when loading into an empty map
 		{
-			if (TempKeyValueStorage)
+			if (NumKeysToRemove)
 			{
+				// Load and discard keys to remove, map is empty
+				void* TempKeyValueStorage = FMemory::Malloc(MapLayout.SetLayout.Size);
+				KeyProp->InitializeValue(TempKeyValueStorage);
+
+				FSerializedPropertyScope SerializedProperty(UnderlyingArchive, KeyProp, this);
+				for (; NumKeysToRemove; --NumKeysToRemove)
+				{
+					KeyProp->SerializeItem(KeysToRemoveArray.EnterElement(), TempKeyValueStorage);
+				}
+
 				KeyProp->DestroyValue(TempKeyValueStorage);
 				FMemory::Free(TempKeyValueStorage);
 			}
-		};
 
-		// Delete any explicitly-removed keys
-		int32 NumKeysToRemove = 0;
-		FStructuredArchive::FArray KeysToRemoveArray = Record.EnterArray(SA_FIELD_NAME(TEXT("KeysToRemove")), NumKeysToRemove);
-		if (NumKeysToRemove)
-		{
-			TempKeyValueStorage = (uint8*)FMemory::Malloc(MapLayout.SetLayout.Size);
-			KeyProp->InitializeValue(TempKeyValueStorage);
+			int32 NumEntries = 0;
+			FStructuredArchive::FArray EntriesArray = Record.EnterArray(SA_FIELD_NAME(TEXT("Entries")), NumEntries);
 
-			FSerializedPropertyScope SerializedProperty(UnderlyingArchive, KeyProp, this);
-			for (; NumKeysToRemove; --NumKeysToRemove)
+			// Empty and reserve then deserialize pairs directly into map memory
+			MapHelper.EmptyValues(NumEntries);
+			for (; NumEntries; --NumEntries)
 			{
-				// Read key into temporary storage
-				KeyProp->SerializeItem(KeysToRemoveArray.EnterElement(), TempKeyValueStorage);
+				FStructuredArchive::FRecord EntryRecord = EntriesArray.EnterElement().EnterRecord();
+				int32 Index = MapHelper.AddDefaultValue_Invalid_NeedsRehash();
 
-				// If the key is in the map, remove it
-				if (uint8* PairPtr = MapHelper.FindMapPairPtrFromHash(TempKeyValueStorage))
 				{
-					MapHelper.RemovePair(PairPtr);
+					FSerializedPropertyScope SerializedProperty(UnderlyingArchive, KeyProp, this);
+					KeyProp->SerializeItem(EntryRecord.EnterField(SA_FIELD_NAME(TEXT("Key"))), MapHelper.GetKeyPtr(Index));
+				}
+				{
+					FSerializedPropertyScope SerializedProperty(UnderlyingArchive, ValueProp, this);
+					ValueProp->SerializeItem(EntryRecord.EnterField(SA_FIELD_NAME(TEXT("Value"))), MapHelper.GetValuePtr(Index));
 				}
 			}
-		}
 
-		int32 NumEntries = 0;
-		FStructuredArchive::FArray EntriesArray = Record.EnterArray(SA_FIELD_NAME(TEXT("Entries")), NumEntries);
-
-		// Allocate temporary key space if we haven't allocated it already above
-		if (NumEntries != 0 && !TempKeyValueStorage)
-		{
-			TempKeyValueStorage = (uint8*)FMemory::Malloc(MapLayout.SetLayout.Size);
-			KeyProp->InitializeValue(TempKeyValueStorage);
+			MapHelper.Rehash();
 		}
-
-		// Read remaining items into container
-		for (; NumEntries; --NumEntries)
+		else // Slower loading path that mutates non-empty map
 		{
-			FStructuredArchive::FRecord EntryRecord = EntriesArray.EnterElement().EnterRecord();
+			uint8* TempKeyValueStorage = nullptr;
+			ON_SCOPE_EXIT
+			{
+				if (TempKeyValueStorage)
+				{
+					KeyProp->DestroyValue(TempKeyValueStorage);
+					FMemory::Free(TempKeyValueStorage);
+				}
+			};
 
-			// Read key into temporary storage
+			if (NumKeysToRemove)
 			{
+				TempKeyValueStorage = (uint8*)FMemory::Malloc(MapLayout.SetLayout.Size);
+				KeyProp->InitializeValue(TempKeyValueStorage);
+
 				FSerializedPropertyScope SerializedProperty(UnderlyingArchive, KeyProp, this);
-				KeyProp->SerializeItem(EntryRecord.EnterField(SA_FIELD_NAME(TEXT("Key"))), TempKeyValueStorage);
+				for (; NumKeysToRemove; --NumKeysToRemove)
+				{
+					// Read key into temporary storage
+					KeyProp->SerializeItem(KeysToRemoveArray.EnterElement(), TempKeyValueStorage);
+
+					// If the key is in the map, remove it
+					if (uint8* PairPtr = MapHelper.FindMapPairPtrFromHash(TempKeyValueStorage))
+					{
+						MapHelper.RemovePair(PairPtr);
+					}
+				}
 			}
 
-			void* ValuePtr = MapHelper.FindOrAdd(TempKeyValueStorage);
+			int32 NumEntries = 0;
+			FStructuredArchive::FArray EntriesArray = Record.EnterArray(SA_FIELD_NAME(TEXT("Entries")), NumEntries);
+
+			// Allocate temporary key space if we haven't allocated it already above
+			if (NumEntries != 0 && !TempKeyValueStorage)
+			{
+				TempKeyValueStorage = (uint8*)FMemory::Malloc(MapLayout.SetLayout.Size);
+				KeyProp->InitializeValue(TempKeyValueStorage);
+			}
 
-			// Deserialize value into hash map-owned memory
+			// Read remaining items into container
+			for (; NumEntries; --NumEntries)
 			{
-				FSerializedPropertyScope SerializedProperty(UnderlyingArchive, ValueProp, this);
-				ValueProp->SerializeItem(EntryRecord.EnterField(SA_FIELD_NAME(TEXT("Value"))), ValuePtr);
+				FStructuredArchive::FRecord EntryRecord = EntriesArray.EnterElement().EnterRecord();
+
+				// Read key into temporary storage
+				{
+					FSerializedPropertyScope SerializedProperty(UnderlyingArchive, KeyProp, this);
+					KeyProp->SerializeItem(EntryRecord.EnterField(SA_FIELD_NAME(TEXT("Key"))), TempKeyValueStorage);
+				}
+
+				void* ValuePtr = MapHelper.FindOrAdd(TempKeyValueStorage);
+
+				// Deserialize value into hash map-owned memory
+				{
+					FSerializedPropertyScope SerializedProperty(UnderlyingArchive, ValueProp, this);
+					ValueProp->SerializeItem(EntryRecord.EnterField(SA_FIELD_NAME(TEXT("Value"))), ValuePtr);
+				}
 			}
 		}
 	}
diff --git a/Engine/Source/Runtime/CoreUObject/Private/UObject/PropertySet.cpp b/Engine/Source/Runtime/CoreUObject/Private/UObject/PropertySet.cpp
--- a/Engine/Source/Runtime/CoreUObject/Private/UObject/PropertySet.cpp
+++ b/Engine/Source/Runtime/CoreUObject/Private/UObject/PropertySet.cpp
@@ -288,70 +288,98 @@ void FSetProperty::SerializeItem(FStructuredArchive::FSlot Slot, void* Value, co
 		{
 			CopyValuesInternal(Value, Defaults, 1);
 		}
-		else
-		{
-			SetHelper.EmptyElements();
-		}
 
-		uint8* TempElementStorage = nullptr;
-		ON_SCOPE_EXIT
+		// Delete any explicitly-removed elements
+		int32 NumElementsToRemove = 0;
+		FStructuredArchive::FArray ElementsToRemoveArray = Record.EnterArray(SA_FIELD_NAME(TEXT("ElementsToRemove")), NumElementsToRemove);
+
+		if (SetHelper.Num() == 0) // Faster loading path when loading elements into an empty set
 		{
-			if (TempElementStorage)
+			if (NumElementsToRemove)
 			{
+				// Load and discard elements to remove, set is empty
+				void* TempElementStorage = FMemory::Malloc(SetLayout.Size);
+				ElementProp->InitializeValue(TempElementStorage);
+
+				FSerializedPropertyScope SerializedProperty(UnderlyingArchive, ElementProp, this);
+				for (; NumElementsToRemove; --NumElementsToRemove)
+				{
+					ElementProp->SerializeItem(ElementsToRemoveArray.EnterElement(), TempElementStorage);
+				}
+
 				ElementProp->DestroyValue(TempElementStorage);
 				FMemory::Free(TempElementStorage);
 			}
-		};
 
-		// Delete any explicitly-removed elements
-		int32 NumElementsToRemove = 0;
-		FStructuredArchive::FArray ElementsToRemoveArray = Record.EnterArray(SA_FIELD_NAME(TEXT("ElementsToRemove")), NumElementsToRemove);
+			int32 Num = 0;
+			FStructuredArchive::FArray ElementsArray = Record.EnterArray(SA_FIELD_NAME(TEXT("Elements")), Num);
 
-		if (NumElementsToRemove)
+			// Empty and reserve then deserialize elements directly into set memory
+			SetHelper.EmptyElements(Num);
+			for (; Num; --Num)
+			{
+				int32 Index = SetHelper.AddDefaultValue_Invalid_NeedsRehash();
+				ElementProp->SerializeItem(ElementsArray.EnterElement(), SetHelper.GetElementPtrWithoutCheck(Index));
+			}
+		}
+		else // Slower loading path that mutates non-empty set
 		{
-			TempElementStorage = (uint8*)FMemory::Malloc(SetLayout.Size);
-			ElementProp->InitializeValue(TempElementStorage);
+			uint8* TempElementStorage = nullptr;
+			ON_SCOPE_EXIT
+			{
+				if (TempElementStorage)
+				{
+					ElementProp->DestroyValue(TempElementStorage);
+					FMemory::Free(TempElementStorage);
+				}
+			};
 
-			FSerializedPropertyScope SerializedProperty(UnderlyingArchive, ElementProp, this);
-			for (; NumElementsToRemove; --NumElementsToRemove)
+			if (NumElementsToRemove)
 			{
-				// Read key into temporary storage
-				ElementProp->SerializeItem(ElementsToRemoveArray.EnterElement(), TempElementStorage);
+				TempElementStorage = (uint8*)FMemory::Malloc(SetLayout.Size);
+				ElementProp->InitializeValue(TempElementStorage);
 
-				// If the key is in the map, remove it
-				const int32 Found = SetHelper.FindElementIndex(TempElementStorage);
-				if (Found != INDEX_NONE)
+				FSerializedPropertyScope SerializedProperty(UnderlyingArchive, ElementProp, this);
+				for (; NumElementsToRemove; --NumElementsToRemove)
 				{
-					SetHelper.RemoveAt(Found);
+					// Read key into temporary storage
+					ElementProp->SerializeItem(ElementsToRemoveArray.EnterElement(), TempElementStorage);
+
+					// If the key is in the map, remove it
+					const int32 Found = SetHelper.FindElementIndex(TempElementStorage);
+					if (Found != INDEX_NONE)
+					{
+						SetHelper.RemoveAt(Found);
+					}
 				}
 			}
-		}
-
-		int32 Num = 0;
-		FStructuredArchive::FArray ElementsArray = Record.EnterArray(SA_FIELD_NAME(TEXT("Elements")), Num);
 
-		// Allocate temporary key space if we haven't allocated it already above
-		if (Num != 0 && !TempElementStorage)
-		{
-			TempElementStorage = (uint8*)FMemory::Malloc(SetLayout.Size);
-			ElementProp->InitializeValue(TempElementStorage);
-		}
+			int32 Num = 0;
+			FStructuredArchive::FArray ElementsArray = Record.EnterArray(SA_FIELD_NAME(TEXT("Elements")), Num);
 
-		FSerializedPropertyScope SerializedProperty(UnderlyingArchive, ElementProp, this);
-		// Read remaining items into container
-		for (; Num; --Num)
-		{
-			// Read key into temporary storage
-			ElementProp->SerializeItem(ElementsArray.EnterElement(), TempElementStorage);
+			// Allocate temporary key space if we haven't allocated it already above
+			if (Num != 0 && !TempElementStorage)
+			{
+				TempElementStorage = (uint8*)FMemory::Malloc(SetLayout.Size);
+				ElementProp->InitializeValue(TempElementStorage);
+			}
 
-			// Add a new entry if the element doesn't currently exist in the set
-			if (SetHelper.FindElementIndex(TempElementStorage) == INDEX_NONE)
+			FSerializedPropertyScope SerializedProperty(UnderlyingArchive, ElementProp, this);
+			// Read remaining items into container
+			for (; Num; --Num)
 			{
-				const int32 NewElementIndex = SetHelper.AddDefaultValue_Invalid_NeedsRehash();
-				uint8* NewElementPtr = SetHelper.GetElementPtrWithoutCheck(NewElementIndex);
+				// Read key into temporary storage
+				ElementProp->SerializeItem(ElementsArray.EnterElement(), TempElementStorage);
 
-				// Copy over deserialized key from temporary storage
-				ElementProp->CopyCompleteValue_InContainer(NewElementPtr, TempElementStorage);
+				// Add a new entry if the element doesn't currently exist in the set
+				if (SetHelper.FindElementIndex(TempElementStorage) == INDEX_NONE)
+				{
+					const int32 NewElementIndex = SetHelper.AddDefaultValue_Invalid_NeedsRehash();
+					uint8* NewElementPtr = SetHelper.GetElementPtrWithoutCheck(NewElementIndex);
+
+					// Copy over deserialized key from temporary storage
+					ElementProp->CopyCompleteValue_InContainer(NewElementPtr, TempElementStorage);
+				}
 			}
 		}
 
