Fixed crash in BulkData when detaching an already detached bulkdata

#issue 313

diff --git a/Engine/Source/Runtime/CoreUObject/Private/Serialization/BulkData.cpp b/Engine/Source/Runtime/CoreUObject/Private/Serialization/BulkData.cpp
--- a/Engine/Source/Runtime/CoreUObject/Private/Serialization/BulkData.cpp
+++ b/Engine/Source/Runtime/CoreUObject/Private/Serialization/BulkData.cpp
@@ -1638,7 +1638,7 @@ bool FUntypedBulkData::RequiresSingleElementSerialization( FArchive& Ar )
 void FUntypedBulkData::DetachFromArchive( FArchive* Ar, bool bEnsureBulkDataIsLoaded )
 {
 	check( Ar );
-	check( Ar == AttachedAr || AttachedAr->IsProxyOf(Ar) );
+	check( Ar == AttachedAr || AttachedAr == nullptr || AttachedAr->IsProxyOf(Ar) );
 	check( LockStatus == LOCKSTATUS_Unlocked );
 
 	// Make sure bulk data is loaded.
@@ -1648,8 +1648,8 @@ void FUntypedBulkData::DetachFromArchive( FArchive* Ar, bool bEnsureBulkDataIsLo
 	}
 
 	// Detach from archive.
-	AttachedAr = NULL;
-	Linker = NULL;
+	AttachedAr = nullptr;
+	Linker = nullptr;
 }
 #endif // WITH_EDITOR
 
