Shrink FObjectImport

diff --git a/Engine/Source/Runtime/CoreUObject/Private/UObject/ObjectResource.cpp b/Engine/Source/Runtime/CoreUObject/Private/UObject/ObjectResource.cpp
--- a/Engine/Source/Runtime/CoreUObject/Private/UObject/ObjectResource.cpp
+++ b/Engine/Source/Runtime/CoreUObject/Private/UObject/ObjectResource.cpp
@@ -242,9 +242,12 @@ void operator<<(FStructuredArchive::FSlot Slot, FObjectTextExport& E)
 
 FObjectImport::FObjectImport()
 	: FObjectResource()
+	, SourceIndex(INDEX_NONE)
 	, bImportPackageHandled(false)
 	, bImportSearchedFor(false)
 	, bImportFailed(false)
+	, XObject(nullptr)
+	, SourceLinker(nullptr)
 {
 }
 
@@ -252,12 +255,12 @@ FObjectImport::FObjectImport(UObject* InObject)
 	: FObjectResource(InObject)
 	, ClassPackage(InObject ? InObject->GetClass()->GetOuter()->GetFName() : NAME_None)
 	, ClassName(InObject ? InObject->GetClass()->GetFName() : NAME_None)
-	, XObject(InObject)
-	, SourceLinker(NULL)
 	, SourceIndex(INDEX_NONE)
 	, bImportPackageHandled(false)
 	, bImportSearchedFor(false)
 	, bImportFailed(false)
+	, XObject(InObject)
+	, SourceLinker(nullptr)
 {
 }
 
@@ -265,12 +268,12 @@ FObjectImport::FObjectImport(UObject* InObject, UClass* InClass)
 	: FObjectResource(InObject)
 	, ClassPackage((InObject && InClass) ? InClass->GetOuter()->GetFName() : NAME_None)
 	, ClassName((InObject && InClass) ? InClass->GetFName() : NAME_None)
-	, XObject(InObject)
-	, SourceLinker(NULL)
 	, SourceIndex(INDEX_NONE)
 	, bImportPackageHandled(false)
 	, bImportSearchedFor(false)
 	, bImportFailed(false)
+	, XObject(InObject)
+	, SourceLinker(nullptr)
 {
 }
 
diff --git a/Engine/Source/Runtime/CoreUObject/Public/UObject/ObjectResource.h b/Engine/Source/Runtime/CoreUObject/Public/UObject/ObjectResource.h
--- a/Engine/Source/Runtime/CoreUObject/Public/UObject/ObjectResource.h
+++ b/Engine/Source/Runtime/CoreUObject/Public/UObject/ObjectResource.h
@@ -445,26 +445,26 @@ struct FObjectImport : public FObjectResource
 #endif
 
 	/**
-	 * The UObject represented by this resource.  Assigned the first time CreateImport is called for this import.
+	 * Index into SourceLinker's ExportMap for the export associated with this import's UObject.
 	 * Transient
 	 */
-	UObject*		XObject;
+	int32             SourceIndex;
+
+	bool			bImportPackageHandled;
+	bool			bImportSearchedFor;
+	bool			bImportFailed;
 
 	/**
-	 * The linker that contains the original FObjectExport resource associated with this import.
+	 * The UObject represented by this resource.  Assigned the first time CreateImport is called for this import.
 	 * Transient
 	 */
-	FLinkerLoad*	SourceLinker;
+	UObject* XObject;
 
 	/**
-	 * Index into SourceLinker's ExportMap for the export associated with this import's UObject.
+	 * The linker that contains the original FObjectExport resource associated with this import.
 	 * Transient
 	 */
-	int32             SourceIndex;
-
-	bool			bImportPackageHandled;
-	bool			bImportSearchedFor;
-	bool			bImportFailed;
+	FLinkerLoad* SourceLinker;
 
 	/**
 	 * Constructors
