VirtualizedBulkData: Support serializing out references to legacy bulk data when saving to the editor domain, to avoid the disk space and cpu costs of upgrading it and saving a copy.
#rnx

diff --git a/Engine/Source/Runtime/CoreUObject/Private/Virtualization/VirtualizedBulkData.cpp b/Engine/Source/Runtime/CoreUObject/Private/Virtualization/VirtualizedBulkData.cpp
--- a/Engine/Source/Runtime/CoreUObject/Private/Virtualization/VirtualizedBulkData.cpp
+++ b/Engine/Source/Runtime/CoreUObject/Private/Virtualization/VirtualizedBulkData.cpp
@@ -96,7 +96,6 @@ FVirtualizedUntypedBulkData& FVirtualizedUntypedBulkData::operator=(const FVirtu
 	OffsetInFile = Other.OffsetInFile;
 	PackagePath = Other.PackagePath;
 	PackageSegment = Other.PackageSegment;
-	bWasKeyGuidDerived = Other.bWasKeyGuidDerived;
 	Flags = Other.Flags;
 	
 	return *this;
@@ -113,6 +112,7 @@ void FVirtualizedUntypedBulkData::CreateFromBulkData(FUntypedBulkData& InBulkDat
 	Reset();
 
 	// We only need to set up the bulkdata/content identifiers if we have a valid payload
+	bool bWasKeyGuidDerived = false;
 	if (InBulkData.GetBulkDataSize() > 0)
 	{
 		if (InGuid.IsValid())
@@ -141,7 +141,7 @@ void FVirtualizedUntypedBulkData::CreateFromBulkData(FUntypedBulkData& InBulkDat
 
 	// Mark that we are actually referencing a payload stored in an old bulkdata
 	// format.
-	EnumAddFlags(Flags,EFlags::ReferencesLegacyFile);
+	EnumAddFlags(Flags, EFlags::ReferencesLegacyFile);
 
 	if (InBulkData.IsStoredCompressedOnDisk())
 	{
@@ -151,6 +151,10 @@ void FVirtualizedUntypedBulkData::CreateFromBulkData(FUntypedBulkData& InBulkDat
 	{
 		EnumAddFlags(Flags, EFlags::DisablePayloadCompression);
 	}
+	if (bWasKeyGuidDerived)
+	{
+		EnumAddFlags(Flags, EFlags::LegacyKeyWasGuidDerived);
+	}
 }
 
 #if UE_VBD_TO_OLD_BULKDATA_PATH
@@ -199,8 +203,7 @@ void FVirtualizedUntypedBulkData::Serialize(FArchive& Ar, UObject* Owner)
 			Ar << Flags;
 			Ar << BulkDataId;
 			Ar << PayloadContentId;
-			Ar << PayloadSize;		
-			Ar << bWasKeyGuidDerived;
+			Ar << PayloadSize;
 			Ar << PackagePath;
 			Ar << PackageSegment;
 			Ar << OffsetInFile;
@@ -239,14 +242,21 @@ void FVirtualizedUntypedBulkData::Serialize(FArchive& Ar, UObject* Owner)
 	else if (Ar.IsPersistent() && !Ar.IsObjectReferenceCollector() && !Ar.ShouldSkipBulkData())
 	{
 		// Check if we need to update the payload id before it is serialized
+		FLinkerSave* LinkerSave = nullptr;
+		bool bKeepLegacyDataByReference = false;
 		if (Ar.IsSaving())
 		{
-			// TODO: Maybe only do this if we are saving with a FLinkerSave (ie package save to disk)
-			UpdateKeyIfNeeded();
-			checkf(bWasKeyGuidDerived == false, TEXT("bWasKeyGuidDerived should be cleared before saving")); // Sanity check
+			LinkerSave = Cast<FLinkerSave>(Ar.GetLinker());
+			// If we're doing a save that can refer to bulk data by reference, and our legacy data format supports it,
+			// keep any legacy data we have referenced rather than stored, to save space and avoid spending time loading it.
+			bKeepLegacyDataByReference = LinkerSave && LinkerSave->bProceduralSave && PackageSegment == EPackageSegment::Header;
+			if (!bKeepLegacyDataByReference)
+			{
+				UpdateKeyIfNeeded();
+			}
 		}
 
-		// Store the position in the archive of the flags incase that we need to update it later
+		// Store the position in the archive of the flags in case we need to update it later
 		const int64 SavedFlagsPos = Ar.Tell();
 		Ar << Flags;
 
@@ -266,8 +276,6 @@ void FVirtualizedUntypedBulkData::Serialize(FArchive& Ar, UObject* Owner)
 		{
 			checkf(Ar.IsCooking() == false, TEXT("FVirtualizedUntypedBulkData::Serialize should not be called during a cook"));
 
-			FLinkerSave* LinkerSave = Cast<FLinkerSave>(Ar.GetLinker());
-
 			const bool bCanAttemptVirtualization = LinkerSave != nullptr;
 
 			if (bCanAttemptVirtualization)
@@ -275,7 +283,7 @@ void FVirtualizedUntypedBulkData::Serialize(FArchive& Ar, UObject* Owner)
 				PushData(); // Note this can change various members if we are going from non-virtualized to virtualized
 			}
 
-			EFlags UpdatedFlags = BuildFlagsForSerialization(Ar);
+			EFlags UpdatedFlags = BuildFlagsForSerialization(Ar, !bKeepLegacyDataByReference);
 
 			// Go back in the archive and update the flags in the archive, we will only apply the updated flags to the current
 			// object later if we detect that the package saved successfully.
@@ -287,7 +295,22 @@ void FVirtualizedUntypedBulkData::Serialize(FArchive& Ar, UObject* Owner)
 			Ar << UpdatedFlags;
 			Ar.Seek(RestorePos);
 
-			if (!IsDataVirtualized())
+			// Write out required extra data if we're saving by reference
+			bool bWriteOutPayload = true;
+			if (EnumHasAnyFlags(UpdatedFlags, EFlags::ReferencesLegacyFile))
+			{
+				check(PackageSegment == EPackageSegment::Header); // This should have been checked before setting bKeepLegacyDataByReference=true
+				FString PackageName = PackagePath.GetPackageName();
+				Ar << PackageName;
+				Ar << OffsetInFile;
+				bWriteOutPayload = false;
+			}
+			else
+			{
+				bWriteOutPayload = !IsDataVirtualized();
+			}
+
+			if (bWriteOutPayload)
 			{
 				// Need to load the payload so that we can write it out
 				FCompressedBuffer PayloadToSerialize = GetDataInternal();
@@ -345,19 +368,9 @@ void FVirtualizedUntypedBulkData::Serialize(FArchive& Ar, UObject* Owner)
 							}
 
 							this->PackagePath = InPackagePath;
-
-							if (!this->PackagePath.IsEmpty())
-							{
-								this->OffsetInFile = PayloadOffset;
-								this->Flags = UpdatedFlags;
-							}
-							else
-							{
-								// If for what ever reason we no longer have a valid package path then reset these values
-								// so that we get better error messages if they are used in the future.
-								this->OffsetInFile = INDEX_NONE;
-								EnumRemoveFlags(this->Flags, EFlags::ReferencesLegacyFile | EFlags::LegacyFileIsCompressed);
-							}
+							check(!this->PackagePath.IsEmpty()); // LinkerSave guarantees a valid PackagePath if we're updating loaded path
+							this->OffsetInFile = PayloadOffset;
+							this->Flags = UpdatedFlags;
 
 							if (CanUnloadData())
 							{
@@ -389,9 +402,15 @@ void FVirtualizedUntypedBulkData::Serialize(FArchive& Ar, UObject* Owner)
 		}
 		else if (Ar.IsLoading())
 		{
-			checkf(!IsReferencingOldBulkData(), TEXT("Virtualized bulkdata was saved with the ReferencesLegacyFile flag!"));
-
-			if (IsDataVirtualized())
+			if (IsReferencingOldBulkData())
+			{
+				FString PackageNameStr;
+				Ar << PackageNameStr;
+				Ar << OffsetInFile;
+				ensure(FPackagePath::TryFromPackageName(*PackageNameStr, this->PackagePath));
+				PackageSegment = EPackageSegment::Header;
+			}
+			else if (IsDataVirtualized())
 			{
 				// We aren't going to use these members so reset them
 				OffsetInFile = INDEX_NONE;
@@ -479,15 +498,32 @@ FCompressedBuffer FVirtualizedUntypedBulkData::LoadFromPackageFile() const
 	FCompressedBuffer PayloadFromDisk;
 
 	// Open a reader to the file
-	FOpenPackageResult Result = IPackageResourceManager::Get().OpenReadPackage(PackagePath, PackageSegment);
-	if (Result.Archive.IsValid() && Result.Format == EPackageFormat::Binary)
+	TUniquePtr<FArchive> BulkArchive;
+	if (!IsReferencingOldBulkData() || PackageSegment != EPackageSegment::Header)
+	{
+		FOpenPackageResult Result = IPackageResourceManager::Get().OpenReadPackage(PackagePath, PackageSegment);
+		if (Result.Format == EPackageFormat::Binary)
+		{
+			BulkArchive = MoveTemp(Result.Archive);
+		}
+	}
+	else
+	{
+		// *this may have been loaded from the EditorDomain, but saved with a reference to the bulk data in the
+		// Workspace Domain file. This was only possible if PackageSegment == Header; we checked that when serializing to the EditorDomain
+		// In that case, we need to use OpenReadExternalResource to access the Workspace Domain file
+		// In the cases where *this was loaded from the WorkspaceDomain, OpenReadExternalResource and OpenReadPackage are identical.
+		BulkArchive = IPackageResourceManager::Get().OpenReadExternalResource(EPackageExternalResource::WorkspaceDomainFile,
+			PackagePath.GetPackageName());
+	}
+	if (BulkArchive.IsValid())
 	{
 		checkf(OffsetInFile != INDEX_NONE, TEXT("Attempting to load '%s' from disk with an invalid OffsetInFile!"), *PackagePath.GetDebugNameWithExtension(PackageSegment));
 		// Move the correct location of the data in the file
-		Result.Archive->Seek(OffsetInFile);
+		BulkArchive->Seek(OffsetInFile);
 
 		// Now we can actually serialize it
-		SerializeData(*Result.Archive, PayloadFromDisk, Flags);
+		SerializeData(*BulkArchive, PayloadFromDisk, Flags);
 	}
 	else
 	{
@@ -636,7 +672,8 @@ void FVirtualizedUntypedBulkData::PushData()
 	// We only need to push if the payload if it actually has data and it is not 
 	// currently virtualized (either we have an updated payload in memory or the 
 	// payload is currently non-virtualized and stored on disk)
-	if (!IsDataVirtualized() && GetPayloadSize() > 0)
+	FVirtualizationManager& VirtualizationManager = FVirtualizationManager::Get();
+	if (!IsDataVirtualized() && GetPayloadSize() > 0 && VirtualizationManager.IsEnabled())
 	{ 
 		TRACE_CPUPROFILER_EVENT_SCOPE(FVirtualizedUntypedBulkData::PushData);
 
@@ -645,16 +682,18 @@ void FVirtualizedUntypedBulkData::PushData()
 		// edited then we should have the payload in memory already and are just accessing a
 		// reference to it.
 
+		UpdateKeyIfNeeded();
 		FCompressedBuffer PayloadToPush = GetDataInternal();
 		// TODO: If the push fails we will end up potentially re-compressing this payload for
 		// serialization, we need a better way to save the results of 'RecompressForSerialization'
 		RecompressForSerialization(PayloadToPush, Flags);
 
-		if (FVirtualizationManager::Get().PushData(PayloadContentId, PayloadToPush))
+		if (VirtualizationManager.PushData(PayloadContentId, PayloadToPush))
 		{
 			EnumAddFlags(Flags, EFlags::IsVirtualized);
 			EnumRemoveFlags(Flags, EFlags::ReferencesLegacyFile | EFlags::LegacyFileIsCompressed);
-			
+			check(!EnumHasAnyFlags(Flags, EFlags::LegacyKeyWasGuidDerived)); // Removed by UpdateKeyIfNeeded
+
 			// Clear members associated with non-virtualized data and release the in-memory
 			// buffer.
 			PackagePath.Empty();
@@ -720,8 +759,7 @@ void FVirtualizedUntypedBulkData::Reset()
 	OffsetInFile = INDEX_NONE;
 	PackagePath.Empty();
 	PackageSegment = EPackageSegment::Header;
-	bWasKeyGuidDerived = false;
-	Flags = EFlags::None;	
+	Flags = EFlags::None;
 }
 
 void FVirtualizedUntypedBulkData::UnloadData()
@@ -810,9 +848,10 @@ void FVirtualizedUntypedBulkData::UpdatePayload(FSharedBuffer InPayload, FName I
 	PayloadSize = (int64)Payload.GetSize();
 	
 	EnumRemoveFlags(Flags,	EFlags::IsVirtualized | 
-							EFlags::DisablePayloadCompression | 
-							EFlags::ReferencesLegacyFile | 
-							EFlags::LegacyFileIsCompressed);
+							EFlags::DisablePayloadCompression |
+							EFlags::ReferencesLegacyFile |
+							EFlags::LegacyFileIsCompressed |
+							EFlags::LegacyKeyWasGuidDerived);
 
 	//TODO: Should we validate now or let FCompressedBuffer do that later?
 	CompressionFormatToUse = InCompressionFormat;
@@ -843,9 +882,9 @@ void FVirtualizedUntypedBulkData::UpdateKeyIfNeeded()
 {
 	// If this was created from old BulkData then the key is generated from an older FGuid, we
 	// should recalculate it based off the payload to keep the key consistent in the future.
-	if (bWasKeyGuidDerived)
+	if (EnumHasAnyFlags(Flags, EFlags::LegacyKeyWasGuidDerived))
 	{
-		checkf(IsDataVirtualized() == false, TEXT("Cannot have a virtualized payload if bWasCreatedFromOldBulkData is still true")); // Sanity check
+		checkf(IsDataVirtualized() == false, TEXT("Cannot have a virtualized payload if loaded from legacy BulkData")); // Sanity check
 
 		// Load the payload from disk (or memory) so that we can hash it
 		FSharedBuffer InPayload = GetDataInternal().Decompress();
@@ -854,8 +893,7 @@ void FVirtualizedUntypedBulkData::UpdateKeyIfNeeded()
 		// Store as the in memory payload, since this method is only called during saving 
 		// we know it will get cleared anyway.
 		Payload = InPayload;
-
-		bWasKeyGuidDerived = false;
+		EnumRemoveFlags(Flags, EFlags::LegacyKeyWasGuidDerived);
 	}
 }
 
@@ -892,7 +930,7 @@ void FVirtualizedUntypedBulkData::RecompressForSerialization(FCompressedBuffer&
 	}
 }
 
-FVirtualizedUntypedBulkData::EFlags FVirtualizedUntypedBulkData::BuildFlagsForSerialization(FArchive& Ar) const
+FVirtualizedUntypedBulkData::EFlags FVirtualizedUntypedBulkData::BuildFlagsForSerialization(FArchive& Ar, bool bUpgradeLegacyData) const
 {
 	if (Ar.IsSaving())
 	{
@@ -913,7 +951,10 @@ FVirtualizedUntypedBulkData::EFlags FVirtualizedUntypedBulkData::BuildFlagsForSe
 			EnumRemoveFlags(UpdatedFlags, EFlags::HasPayloadSidecarFile);
 		}
 
-		EnumRemoveFlags(UpdatedFlags,EFlags::ReferencesLegacyFile | EFlags::LegacyFileIsCompressed);
+		if (bUpgradeLegacyData)
+		{
+			EnumRemoveFlags(UpdatedFlags, EFlags::ReferencesLegacyFile | EFlags::LegacyFileIsCompressed | EFlags::LegacyKeyWasGuidDerived);
+		}
 
 		return UpdatedFlags;
 	}
diff --git a/Engine/Source/Runtime/CoreUObject/Public/Virtualization/VirtualizedBulkData.h b/Engine/Source/Runtime/CoreUObject/Public/Virtualization/VirtualizedBulkData.h
--- a/Engine/Source/Runtime/CoreUObject/Public/Virtualization/VirtualizedBulkData.h
+++ b/Engine/Source/Runtime/CoreUObject/Public/Virtualization/VirtualizedBulkData.h
@@ -173,6 +173,8 @@ private:
 		/** The payload should not have compression applied to it. It is assumed that the payload is already 
 			in some sort of compressed format, see the compression documentation above for more details. */
 		DisablePayloadCompression	= 1 << 4,
+		/** The legacy file being referenced derived its key from guid and it should be replaced with a key-from-hash when saved */
+		LegacyKeyWasGuidDerived		= 1 << 5,
 	};
 
 	/** Used to control what level of error reporting we return from some methods */
@@ -205,7 +207,7 @@ private:
 	void UpdateKeyIfNeeded();
 
 	void RecompressForSerialization(FCompressedBuffer& InOutPayload, EFlags PayloadFlags) const;
-	EFlags BuildFlagsForSerialization(FArchive& Ar) const;
+	EFlags BuildFlagsForSerialization(FArchive& Ar, bool bUpgradeLegacyData) const;
 
 	bool IsDataVirtualized() const { return EnumHasAnyFlags(Flags, EFlags::IsVirtualized); }
 	bool HasPayloadSidecarFile() const { return EnumHasAnyFlags(Flags, EFlags::HasPayloadSidecarFile); }
@@ -237,9 +239,6 @@ private:
 	/** PackageSegment to load with the packagepath (unused if the payload does not come from PackageResourceManager) */
 	EPackageSegment PackageSegment;
 
-	/** Records if the payload key was created from a FGuid via CreateFromBulkData */
-	bool bWasKeyGuidDerived = false;
-
 	/** A 32bit bitfield of flags */
 	EFlags Flags = EFlags::None;
 };
