Repair-style fix for issue where deferred dependency class placeholders could remain in blueprints after loading. This is seen in an example project with complicated circular dependencies intertwining many assets.

diff --git a/Engine/Source/Editor/Kismet/Private/BlueprintCompilationManager.cpp b/Engine/Source/Editor/Kismet/Private/BlueprintCompilationManager.cpp
--- a/Engine/Source/Editor/Kismet/Private/BlueprintCompilationManager.cpp
+++ b/Engine/Source/Editor/Kismet/Private/BlueprintCompilationManager.cpp
@@ -2087,6 +2087,14 @@ void FBlueprintCompilationManagerImpl::ReinstanceBatch(TArray<FReinstancingJob>&
 			{
 				const bool bUseDeltaSerialization = ReinstancingJob.Reinstancer.IsValid() ? ReinstancingJob.Reinstancer->bUseDeltaSerializationToCopyProperties : false;
 				UObject* NewCDO = ReinstancingJob.OldToNew.Value->GetDefaultObject(true);
+
+				UBlueprint* CompiledBlueprint = UBlueprint::GetBlueprintFromClass(ReinstancingJob.OldToNew.Key);
+				if (CompiledBlueprint && CompiledBlueprint->bIsRegeneratingOnLoad)
+				{
+					// This is a catch-all for any deferred dependencies that didn't get resolved during loading/linking (that system is not bulletproof for complex circular dependencies)
+					FBlueprintSupport::RepairDeferredDependenciesInObject(OldCDO);
+				}
+
 				FBlueprintCompileReinstancer::CopyPropertiesForUnrelatedObjects(OldCDO, NewCDO, true, bUseDeltaSerialization);
 
 				if (ReinstancingJob.Compiler.IsValid())
diff --git a/Engine/Source/Runtime/CoreUObject/Private/Blueprint/BlueprintSupport.cpp b/Engine/Source/Runtime/CoreUObject/Private/Blueprint/BlueprintSupport.cpp
--- a/Engine/Source/Runtime/CoreUObject/Private/Blueprint/BlueprintSupport.cpp
+++ b/Engine/Source/Runtime/CoreUObject/Private/Blueprint/BlueprintSupport.cpp
@@ -207,6 +207,64 @@ void FBlueprintSupport::RegisterDeferredDependenciesInStruct(const UStruct* Stru
 #endif // USE_CIRCULAR_DEPENDENCY_LOAD_DEFERRING
 }
 
+void FBlueprintSupport::RepairDeferredDependenciesInObject(UObject* Object)
+{
+	// Go through each property's value on the object and check for placeholders. 
+	//   Try to replace them with the real imported object, if it exists.
+	// This function was created to catch any cases where a deferred dependency fails to resolve during load
+
+	for (TPropertyValueIterator<const FObjectProperty> It(Object->GetClass(), Object); It; ++It)
+	{
+		const FObjectProperty* Property = It.Key();
+		void* PropertyValue = (void*)It.Value();
+		UObject* PropertyValueAsObj = *((UObject**)PropertyValue);
+
+		FLinkerPlaceholderBase* Placeholder = nullptr;
+
+		if (ULinkerPlaceholderExportObject* ValueAsPlaceholderObj = Cast<ULinkerPlaceholderExportObject>(PropertyValueAsObj))
+		{
+			Placeholder = (FLinkerPlaceholderBase*)ValueAsPlaceholderObj;
+		}
+		else if (ULinkerPlaceholderClass* ValueAsPlaceholderClass = Cast<ULinkerPlaceholderClass>(PropertyValueAsObj))
+		{
+			Placeholder = (FLinkerPlaceholderBase*)ValueAsPlaceholderClass;
+		}
+
+		if (Placeholder)
+		{
+			UE_LOG(LogBlueprintSupport, Warning, TEXT("Object %s still has a %s '%s' in property %s. This indicates a failure to resolve every deferred/circular dependency in blueprints."), *Object->GetName(), *PropertyValueAsObj->GetClass()->GetName(), *PropertyValueAsObj->GetName(), *Property->GetName());
+
+			bool bDidRepairStalePlaceholder = false;
+
+			if (!Placeholder->PackageIndex.IsNull())
+			{
+				if (const UPackage* PlaceholderPackage = PropertyValueAsObj->GetPackage())
+				{
+					if (FLinkerLoad* PlaceholderLinker = PlaceholderPackage->LinkerLoad)
+					{
+						int32 const ImportIndex = Placeholder->PackageIndex.ToImport();
+						FObjectImport& Import = PlaceholderLinker->ImportMap[ImportIndex];
+						if ((Import.XObject != nullptr) && (Import.XObject != PropertyValueAsObj))
+						{
+							Property->SetObjectPropertyValue(PropertyValue, Import.XObject);
+							bDidRepairStalePlaceholder = true;
+
+							UE_LOG(LogBlueprintSupport, Display, TEXT("Repaired deferred dependency on object %s: replaced '%s' with '%s'"), *Object->GetName(), *PropertyValueAsObj->GetName(), *Import.XObject->GetName());
+						}
+					}
+				}
+			}
+
+			if (!bDidRepairStalePlaceholder)
+			{
+				UE_LOG(LogBlueprintSupport, Error, TEXT("Failed to repair deferred dependency on object %s (%s). This may lead to blueprint execution problems."), *Object->GetName(), *PropertyValueAsObj->GetName());
+			}
+		}
+	}
+
+}
+
+
 bool FBlueprintSupport::IsInBlueprintPackage(UObject* LoadedObj)
 {
 	UPackage* Pkg = LoadedObj->GetOutermost();
diff --git a/Engine/Source/Runtime/CoreUObject/Public/Blueprint/BlueprintSupport.h b/Engine/Source/Runtime/CoreUObject/Public/Blueprint/BlueprintSupport.h
--- a/Engine/Source/Runtime/CoreUObject/Public/Blueprint/BlueprintSupport.h
+++ b/Engine/Source/Runtime/CoreUObject/Public/Blueprint/BlueprintSupport.h
@@ -100,6 +100,9 @@ struct FBlueprintSupport
 	/** Registers any object properties in this struct with the deferred dependency system */
 	COREUOBJECT_API static void RegisterDeferredDependenciesInStruct(const UStruct* Struct, void* StructData);
 
+	/** Repair function to scan for any deferred dependency placeholders and attempt to resolve them in-place with real assets */
+	COREUOBJECT_API static void RepairDeferredDependenciesInObject(UObject* Object);
+
 	/** Not a particularly fast function. Mostly intended for validation in debug builds. */
 	static bool IsInBlueprintPackage(UObject* LoadedObj);
 
