Fix a bug where bulkdata that was compressed during the content cook was returning the compressed data rather than uncompressed data when the IoDispatcher was active.

- This has gone so long without being noticed as it is not expected that bulkdata will be compressed during a cook, as the pakfiles themselves will be compressed.
- InternalLoadFromIoStore now only loads via the IoDispatcher directly to the destination buffer if the data is in uncompressed format.
- If the data is going to be compressed we let the IoDispatcher allocate and return the compressed buffer and bulkdata will decompress that to the destination buffer instead.

diff --git a/Engine/Source/Runtime/CoreUObject/Private/Serialization/BulkData2.cpp b/Engine/Source/Runtime/CoreUObject/Private/Serialization/BulkData2.cpp
--- a/Engine/Source/Runtime/CoreUObject/Private/Serialization/BulkData2.cpp
+++ b/Engine/Source/Runtime/CoreUObject/Private/Serialization/BulkData2.cpp
@@ -11,6 +11,7 @@
 #include "Misc/Paths.h"
 #include "ProfilingDebugging/LoadTimeTracker.h"
 #include "Serialization/BulkData.h"
+#include "Serialization/LargeMemoryReader.h"
 #include "UObject/LinkerLoad.h"
 #include "UObject/Object.h"
 #include "UObject/Package.h"
@@ -1677,7 +1678,12 @@ void FBulkDataBase::InternalLoadFromIoStore(void** DstBuffer)
 
 	// Set up our options (we only need to set the target)
 	FIoReadOptions Options(BulkDataOffset, BulkDataSize);
-	Options.SetTargetVa(*DstBuffer);
+	
+	// If we do not need to uncompress the data we can load it directly to the destination buffer
+	if (!IsStoredCompressedOnDisk()) 
+	{
+		Options.SetTargetVa(*DstBuffer);
+	}	
 
 	FIoBatch Batch = GetIoDispatcher()->NewBatch();
 	FIoRequest Request = Batch.Read(CreateChunkId(), Options, IoDispatcherPriority_High);
@@ -1687,6 +1693,16 @@ void FBulkDataBase::InternalLoadFromIoStore(void** DstBuffer)
 	BatchCompletedEvent->Wait(); // Blocking wait until all requests in the batch are done
 	FPlatformProcess::ReturnSynchEventToPool(BatchCompletedEvent);
 	CHECK_IOSTATUS(Request.GetResult().Status(), TEXT("FIoRequest"));
+
+	// If the data is compressed we need to decompress it to the destination buffer.
+	// We know it was compressed via FArchive so the easiest way to decompress it 
+	// is to wrap it in a memory reader archive.
+	if (IsStoredCompressedOnDisk())
+	{
+		const FIoBuffer& CompressedBuffer = Request.GetResult().ValueOrDie();
+		FLargeMemoryReader Ar(CompressedBuffer.Data(), (int64)CompressedBuffer.DataSize());
+		Ar.SerializeCompressed(*DstBuffer, GetBulkDataSize(), GetDecompressionFormat(), COMPRESS_NoFlags, false);
+	}
 }
 
 void FBulkDataBase::InternalLoadFromIoStoreAsync(void** DstBuffer, AsyncCallback&& Callback)
