Exposed FInt32Interval to blueprints

diff --git a/Engine/Source/Runtime/CoreUObject/Public/UObject/NoExportTypes.h b/Engine/Source/Runtime/CoreUObject/Public/UObject/NoExportTypes.h
--- a/Engine/Source/Runtime/CoreUObject/Public/UObject/NoExportTypes.h
+++ b/Engine/Source/Runtime/CoreUObject/Public/UObject/NoExportTypes.h
@@ -1509,15 +1509,15 @@ struct FFloatInterval
  * An interval of integers, defined by inclusive min and max values
  * @note This is a mirror of TInterval<int32>, defined in Interval.h
  */
-USTRUCT(noexport)
+USTRUCT(noexport, BlueprintType)
 struct FInt32Interval
 {
 	/** Values must be >= Min */
-	UPROPERTY(EditAnywhere, Category=Interval)
+	UPROPERTY(EditAnywhere, BlueprintReadWrite, Category=Interval)
 	int32 Min;
 
 	/** Values must be <= Max */
-	UPROPERTY(EditAnywhere, Category=Interval)
+	UPROPERTY(EditAnywhere, BlueprintReadWrite, Category=Interval)
 	int32 Max;
 };
 
