Integrating //UE5/Dev-Cooker @ 16678003 to //UE5/Main (Zousar.Shaker-YEG-0943-Quaternary)

diff --git a/Engine/Source/Runtime/CoreUObject/Private/Serialization/AsyncLoading2.cpp b/Engine/Source/Runtime/CoreUObject/Private/Serialization/AsyncLoading2.cpp
--- a/Engine/Source/Runtime/CoreUObject/Private/Serialization/AsyncLoading2.cpp
+++ b/Engine/Source/Runtime/CoreUObject/Private/Serialization/AsyncLoading2.cpp
@@ -6,7 +6,8 @@
 
 #include "Serialization/AsyncLoading2.h"
 #include "Serialization/AsyncPackageLoader.h"
-#include "Serialization/PackageStore.h"
+#include "IO/PackageStore.h"
+#include "Serialization/FilePackageStore.h"
 #include "HAL/PlatformFileManager.h"
 #include "HAL/FileManager.h"
 #include "HAL/Event.h"
@@ -62,6 +63,16 @@
 #include "Async/ParallelFor.h"
 #include "HAL/LowLevelMemStats.h"
 #include "HAL/IPlatformFileOpenLogWrapper.h"
+#include "Modules/ModuleManager.h"
+
+#ifndef WITH_COTF
+#define WITH_COTF !UE_BUILD_SHIPPING
+#endif
+
+#if WITH_COTF
+#include "CookOnTheFly.h"
+#include "Serialization/CookOnTheFlyPackageStore.h"
+#endif
 
 #if UE_BUILD_DEVELOPMENT || UE_BUILD_DEBUG
 PRAGMA_DISABLE_OPTIMIZATION
@@ -85,8 +96,6 @@ FArchive& operator<<(FArchive& Ar, FContainerHeader& ContainerHeader)
 {
 	Ar << ContainerHeader.ContainerId;
 	Ar << ContainerHeader.PackageCount;
-	Ar << ContainerHeader.Names;
-	Ar << ContainerHeader.NameHashes;
 	Ar << ContainerHeader.PackageIds;
 	Ar << ContainerHeader.StoreEntries;
 	Ar << ContainerHeader.CulturePackageMap;
@@ -105,6 +114,7 @@ FArchive& operator<<(FArchive& Ar, FExportBundleEntry& ExportBundleEntry)
 
 FArchive& operator<<(FArchive& Ar, FExportBundleHeader& ExportBundleHeader)
 {
+	Ar << ExportBundleHeader.SerialOffset;
 	Ar << ExportBundleHeader.FirstEntryIndex;
 	Ar << ExportBundleHeader.EntryCount;
 
@@ -374,6 +384,182 @@ struct FExportObject
 	bool bExportLoadFailed = false;
 };
 
+struct FPackageRequest
+{
+	int32 RequestId = -1;
+	int32 Priority = -1;
+	FPackageId PackageId;
+	FPackagePath PackagePath;
+	FName CustomName;
+	TUniquePtr<FLoadPackageAsyncDelegate> PackageLoadedDelegate;
+	FPackageStoreEntryHandle StoreEntryHandle;
+	FPackageRequest* Next = nullptr;
+
+	FName GetUPackageName() const
+	{
+		return CustomName.IsNone() ? PackagePath.GetPackageFName() : CustomName;
+	}
+
+	FName GetDiskPackageName() const
+	{
+		return PackagePath.GetPackageFName();
+	}
+
+	static FPackageRequest Create(const int32 RequestId, const int32 Priority, FPackagePath PackagePath, FName CustomPackageName, TUniquePtr<FLoadPackageAsyncDelegate> PackageLoadedDelegate)
+	{
+		return FPackageRequest
+		{
+			RequestId,
+			Priority,
+			FPackageId::FromName(PackagePath.GetPackageFName()),
+			MoveTemp(PackagePath),
+			CustomPackageName,
+			MoveTemp(PackageLoadedDelegate),
+			FPackageStoreEntryHandle(),
+			nullptr
+		};
+	}
+};
+
+class FPackageRequestQueue
+{
+	struct FPackageRequestList
+	{
+		FPackageRequest* Head = nullptr;
+		FPackageRequest* Tail = nullptr;
+
+		void Add(FPackageRequest* Request)
+		{
+			check(!Request->Next);
+
+			if (!Tail)
+			{
+				check(!Head);
+				Head = Tail = Request;
+			}
+			else
+			{
+				Tail->Next = Request;
+				Tail = Request;
+			}
+		}
+
+		FPackageRequest* Pop()
+		{
+			FPackageRequest* Request = Head;
+			if (Request)
+			{
+				Head = Request->Next;
+				Request->Next = nullptr;
+
+				if (!Head)
+				{
+					Tail = nullptr;
+				}
+			}
+
+			return Request;
+		}
+	};
+
+public:
+	FPackageRequestQueue(IPackageStore& InPackageStore)
+		: PackageStore(InPackageStore)
+	{ }
+
+	void Enqueue(FPackageRequest LoadPackageRequest)
+	{
+		FScopeLock _(&CriticalSection);
+
+		FPackageRequest* Request = Alloc();
+		check(Request->Next == nullptr)
+		*Request = MoveTemp(LoadPackageRequest);
+
+		Request->StoreEntryHandle = PackageStore.GetPackageEntryHandle(Request->PackageId, Request->GetDiskPackageName());
+
+#if WITH_COTF
+		EPackageStoreEntryStatus Status = Request->StoreEntryHandle.Status();
+
+		if (Status == EPackageStoreEntryStatus::Pending)
+		{
+			PendingRequests.Add(Request);
+		}
+		else
+#endif
+		{
+			ReadyRequests.Add(Request);
+		}
+	}
+
+	void Dequeue(TFunction<bool(FPackageRequest&)>&& Fn)
+	{
+		FScopeLock _(&CriticalSection);
+
+#if WITH_COTF
+		for (auto It = PendingRequests.CreateIterator(); It; ++It)
+		{
+			FPackageRequest* PendingRequest = *It;
+			PendingRequest->StoreEntryHandle = PackageStore.GetPackageEntryHandle(PendingRequest->PackageId, PendingRequest->GetDiskPackageName());
+			if (PendingRequest->StoreEntryHandle.Status() != EPackageStoreEntryStatus::Pending)
+			{
+				ReadyRequests.Add(PendingRequest);
+				It.RemoveCurrent();
+			}
+		}
+#endif
+		FPackageRequest* ReadyRequest = ReadyRequests.Pop();
+		while (ReadyRequest)
+		{
+			if (!Fn(*ReadyRequest))
+			{
+				break;
+			}
+
+			ReadyRequest = ReadyRequests.Pop();
+		}
+	}
+
+	void Trim()
+	{
+		FScopeLock _(&CriticalSection);
+
+		if (!ReadyRequests.Head)
+		{
+			ReadyRequests = FPackageRequestList();
+			FreeRequests = FPackageRequestList();
+			Requests.Empty();
+		}
+	}
+
+private:
+	FPackageRequest* Alloc()
+	{
+		if (FPackageRequest* FreeRequest = FreeRequests.Pop())
+		{
+			return FreeRequest;
+		}
+		else
+		{
+			const int32 Index = Requests.Add();
+			return &Requests[Index];
+		}
+	}
+
+	void Free(FPackageRequest* Request)
+	{
+		FreeRequests.Add(Request);
+	}
+
+	IPackageStore& PackageStore;
+	FCriticalSection CriticalSection;
+	TChunkedArray<FPackageRequest, 1024> Requests;
+	FPackageRequestList ReadyRequests;
+	FPackageRequestList FreeRequests;
+#if WITH_COTF
+	TArray<FPackageRequest*> PendingRequests;
+#endif
+};
+
 struct FAsyncPackageDesc2
 {
 	// A unique request id for each external call to LoadPackage
@@ -381,7 +567,7 @@ struct FAsyncPackageDesc2
 	// Package priority
 	int32 Priority;
 	// The package store entry with meta data about the actual disk package
-	const FPackageStoreEntry* StoreEntry;
+	const FPackageStoreEntryHandle StoreEntry;
 	// The disk package id corresponding to the StoreEntry
 	// It is used by the loader for io chunks and to handle ref tracking of loaded packages and import objects.
 	FPackageId DiskPackageId; 
@@ -396,48 +582,6 @@ struct FAsyncPackageDesc2
 	FName CustomPackageName;
 	// Set from the package summary,
 	FName SourcePackageName;
-	/** Delegate called on completion of loading. This delegate can only be created and consumed on the game thread */
-	TUniquePtr<FLoadPackageAsyncDelegate> PackageLoadedDelegate;
-
-	FAsyncPackageDesc2(
-		int32 InRequestID,
-		int32 InPriority,
-		FPackageId InPackageIdToLoad,
-		const FPackageStoreEntry* InStoreEntry,
-		FName InDiskPackageName = FName(),
-		FPackageId InPackageId = FPackageId(),
-		FName InCustomName = FName(),
-		TUniquePtr<FLoadPackageAsyncDelegate>&& InCompletionDelegate = TUniquePtr<FLoadPackageAsyncDelegate>())
-		: RequestID(InRequestID)
-		, Priority(InPriority)
-		, StoreEntry(InStoreEntry)
-		, DiskPackageId(InPackageIdToLoad)
-		, CustomPackageId(InPackageId)
-		, DiskPackageName(InDiskPackageName)
-		, CustomPackageName(InCustomName)
-		, PackageLoadedDelegate(MoveTemp(InCompletionDelegate))
-	{
-	}
-
-	/** This constructor does not modify the package loaded delegate as this is not safe outside the game thread */
-	FAsyncPackageDesc2(const FAsyncPackageDesc2& OldPackage)
-		: RequestID(OldPackage.RequestID)
-		, Priority(OldPackage.Priority)
-		, StoreEntry(OldPackage.StoreEntry)
-		, DiskPackageId(OldPackage.DiskPackageId)
-		, CustomPackageId(OldPackage.CustomPackageId)
-		, DiskPackageName(OldPackage.DiskPackageName)
-		, CustomPackageName(OldPackage.CustomPackageName)
-		, SourcePackageName(OldPackage.SourcePackageName)
-	{
-	}
-
-	/** This constructor will explicitly copy the package loaded delegate and invalidate the old one */
-	FAsyncPackageDesc2(const FAsyncPackageDesc2& OldPackage, TUniquePtr<FLoadPackageAsyncDelegate>&& InPackageLoadedDelegate)
-		: FAsyncPackageDesc2(OldPackage)
-	{
-		PackageLoadedDelegate = MoveTemp(InPackageLoadedDelegate);
-	}
 
 	void SetDiskPackageName(FName SerializedDiskPackageName, FName SerializedSourcePackageName = FName())
 	{
@@ -480,12 +624,46 @@ struct FAsyncPackageDesc2
 		return CustomPackageId.IsValid() ? CustomPackageId : DiskPackageId;
 	}
 
-#if DO_GUARD_SLOW
-	~FAsyncPackageDesc2()
-	{
-		checkSlow(!PackageLoadedDelegate.IsValid() || IsInGameThread());
+	static FAsyncPackageDesc2 FromPackageRequest(
+		int32 RequestID,
+		int32 Priority,
+		FPackageId PackageIdToLoad,
+		const FPackageStoreEntryHandle StoreEntryHandle,
+		FName DiskPackageName,
+		FPackageId PackageId,
+		FName CustomName)
+	{
+		return FAsyncPackageDesc2
+		{
+			RequestID,
+			Priority,
+			StoreEntryHandle,
+			PackageIdToLoad,
+			PackageId,
+			DiskPackageName,
+			CustomName,
+			FName()
+		};
+	}
+
+	static FAsyncPackageDesc2 FromPackageImport(
+		int32 RequestID,
+		int32 Priority,
+		FPackageId PackageIdToLoad,
+		const FPackageStoreEntryHandle StoreEntryHandle)
+	{
+		return FAsyncPackageDesc2
+		{
+			RequestID,
+			Priority,
+			StoreEntryHandle,
+			PackageIdToLoad,
+			FPackageId(),
+			FName(),
+			FName(),
+			FName()
+		};
 	}
-#endif
 };
 
 struct FPublicExport
@@ -503,7 +681,6 @@ using FUnreachablePublicExports = TArray<FUnreachablePublicExport>;
 
 struct FGlobalImportStore
 {
-	FNameMap GlobalNameMap;
 	TMap<FPackageObjectIndex, UObject*> ScriptObjects;
 	TMap<FPackageObjectIndex, FPublicExport> PublicExportObjects;
 	TMap<int32, FPackageObjectIndex> ObjectIndexToPublicExport;
@@ -520,37 +697,13 @@ struct FGlobalImportStore
 
 	void Initialize(FIoDispatcher& IoDispatcher)
 	{
-		{
-			TRACE_CPUPROFILER_EVENT_SCOPE(LoadGlobalNameMap);
-
-			FIoChunkId NamesId = CreateIoChunkId(0, 0, EIoChunkType::LoaderGlobalNames);
-			FIoChunkId HashesId = CreateIoChunkId(0, 0, EIoChunkType::LoaderGlobalNameHashes);
-
-			FIoBatch Batch = IoDispatcher.NewBatch();
-			FIoRequest NameRequest = Batch.Read(NamesId, FIoReadOptions(), IoDispatcherPriority_High);
-			FIoRequest HashRequest = Batch.Read(HashesId, FIoReadOptions(), IoDispatcherPriority_High);
-			FEvent* BatchCompletedEvent = FPlatformProcess::GetSynchEventFromPool();
-			Batch.IssueAndTriggerEvent(BatchCompletedEvent);
-
-			ReserveNameBatch(	IoDispatcher.GetSizeForChunk(NamesId).ValueOrDie(),
-								IoDispatcher.GetSizeForChunk(HashesId).ValueOrDie());
-
-			BatchCompletedEvent->Wait();
-			FPlatformProcess::ReturnSynchEventToPool(BatchCompletedEvent);
-
-			FIoBuffer NameBuffer = NameRequest.GetResult().ConsumeValueOrDie();
-			FIoBuffer HashBuffer = HashRequest.GetResult().ConsumeValueOrDie();
-
-			GlobalNameMap.Load(MakeArrayView(NameBuffer.Data(), NameBuffer.DataSize()), MakeArrayView(HashBuffer.Data(), HashBuffer.DataSize()), FMappedName::EType::Global);
-		}
-
 		{
 			TRACE_CPUPROFILER_EVENT_SCOPE(SetupInitialLoadData);
 
 			FEvent* InitialLoadEvent = FPlatformProcess::GetSynchEventFromPool();
 
 			FIoBatch IoBatch = IoDispatcher.NewBatch();
-			FIoRequest IoRequest = IoBatch.Read(CreateIoChunkId(0, 0, EIoChunkType::LoaderInitialLoadMeta), FIoReadOptions(), IoDispatcherPriority_High);
+			FIoRequest IoRequest = IoBatch.Read(CreateIoChunkId(0, 0, EIoChunkType::ScriptObjects), FIoReadOptions(), IoDispatcherPriority_High);
 			IoBatch.IssueAndTriggerEvent(InitialLoadEvent);
 
 			InitialLoadEvent->Wait();
@@ -558,6 +711,8 @@ struct FGlobalImportStore
 
 			FIoBuffer InitialLoadIoBuffer = IoRequest.GetResult().ConsumeValueOrDie();
 			FLargeMemoryReader InitialLoadArchive(InitialLoadIoBuffer.Data(), InitialLoadIoBuffer.DataSize());
+			FNameMap NameMap;
+			NameMap.Load(InitialLoadArchive, FMappedName::EType::Global);
 			int32 NumScriptObjects = 0;
 			InitialLoadArchive << NumScriptObjects;
 			ScriptObjectEntries = MakeArrayView(reinterpret_cast<const FScriptObjectEntry*>(InitialLoadIoBuffer.Data() + InitialLoadArchive.Tell()), NumScriptObjects);
@@ -567,7 +722,7 @@ struct FGlobalImportStore
 			{
 				const FMappedName& MappedName = FMappedName::FromMinimalName(ScriptObjectEntry.ObjectName);
 				check(MappedName.IsGlobal());
-				ScriptObjectEntry.ObjectName = GlobalNameMap.GetMinimalName(MappedName);
+				ScriptObjectEntry.ObjectName = NameMap.GetMinimalName(MappedName);
 
 				ScriptObjectEntriesMap.Add(ScriptObjectEntry.GlobalIndex, &ScriptObjectEntry);
 			}
@@ -925,15 +1080,11 @@ struct FPackageImportStore
 		, LoadedPackageStore(InLoadedPackageStore)
 		, Desc(InDesc)
 	{
-		TRACE_CPUPROFILER_EVENT_SCOPE(NewPackageImportStore);
-		AddPackageReferences();
 	}
 
 	~FPackageImportStore()
 	{
-		TRACE_CPUPROFILER_EVENT_SCOPE(DeletePackageImportStore);
 		check(ImportMap.Num() == 0);
-		ReleasePackageReferences();
 	}
 
 	inline bool IsValidLocalImportIndex(FPackageIndex LocalIndex)
@@ -1044,9 +1195,10 @@ private:
 		ImportedPackage->AtomicallyClearInternalFlags(EInternalObjectFlags::Async);
 	}
 
-	void AddPackageReferences()
+public:
+	void AddPackageReferences(const TArrayView<const FPackageId>& ImportedPackageIds)
 	{
-		for (const FPackageId& ImportedPackageId : Desc.StoreEntry->ImportedPackages)
+		for (const FPackageId& ImportedPackageId : ImportedPackageIds)
 		{
 			FLoadedPackageRef& PackageRef = LoadedPackageStore.GetPackageRef(ImportedPackageId);
 			if (PackageRef.AddRef())
@@ -1065,9 +1217,9 @@ private:
 		}
 	}
 
-	void ReleasePackageReferences()
+	void ReleasePackageReferences(const TArrayView<const FPackageId>& ImportedPackageIds)
 	{
-		for (const FPackageId& ImportedPackageId : Desc.StoreEntry->ImportedPackages)
+		for (const FPackageId& ImportedPackageId : ImportedPackageIds)
 		{
 			FLoadedPackageRef& PackageRef = LoadedPackageStore.GetPackageRef(ImportedPackageId);
 			if (PackageRef.ReleaseRef(Desc.DiskPackageId, ImportedPackageId))
@@ -1334,6 +1486,7 @@ enum class EAsyncPackageLoadingState2 : uint8
 	ImportPackagesDone,
 	WaitingForIo,
 	ProcessPackageSummary,
+	SetupDependencies,
 	ProcessExportBundles,
 	WaitingForExternalReads,
 	ExportsDone,
@@ -1566,7 +1719,11 @@ uint32 FAsyncLoadingThreadState2::TlsSlot;
 enum EEventLoadNode2 : uint8
 {
 	Package_ProcessSummary,
+	Package_SetupDependencies,
 	Package_ExportsSerialized,
+	Package_BeginPostLoad,
+	Package_PostLoadDone,
+	Package_FinishLoading,
 	Package_NumPhases,
 
 	ExportBundle_Process = 0,
@@ -1577,9 +1734,9 @@ enum EEventLoadNode2 : uint8
 
 struct FAsyncPackageData
 {
-	int32 ExportCount = 0;
-	int32 ExportBundleCount = 0;
-	uint64 ExportBundlesMetaSize = 0;
+	FPackageStoreExportInfo ExportInfo;
+	uint64 ExportBundleHeadersSize = 0;
+	uint64 ExportBundleEntriesSize = 0;
 	uint8* ExportBundlesMetaMemory = nullptr;
 	const FExportBundleHeader* ExportBundleHeaders = nullptr;
 	const FExportBundleEntry* ExportBundleEntries = nullptr;
@@ -1587,6 +1744,8 @@ struct FAsyncPackageData
 	TArrayView<FAsyncPackage2*> ImportedAsyncPackages;
 	TArrayView<FEventLoadNode2> PackageNodes;
 	TArrayView<FEventLoadNode2> ExportBundleNodes;
+	TArrayView<const FPackageId> ImportedPackageIds;
+	TArrayView<const FSHAHash> ShaderMapHashes;
 };
 
 /**
@@ -1690,6 +1849,12 @@ struct FAsyncPackage2
 
 private:
 
+	struct FExportToBundleMapping
+	{
+		FPackageObjectIndex GlobalImportIndex;
+		int32 BundleIndex[FExportBundleEntry::ExportCommandType_Count];
+	};
+
 	/** Checks if all dependencies (imported packages) of this package have been fully loaded */
 	bool AreAllDependenciesFullyLoadedInternal(FAsyncPackage2* Package, TSet<FPackageId>& VisitedPackages, FPackageId& OutPackageId);
 
@@ -1699,11 +1864,13 @@ private:
 	/** Cached async loading thread object this package was created by */
 	FAsyncLoadingThread2& AsyncLoadingThread;
 	FAsyncLoadEventGraphAllocator& GraphAllocator;
+	FPackageImportStore ImportStore;
 	/** Package which is going to have its exports and imports loaded */
 	UPackage* LinkerRoot = nullptr;
 	/** Time load begun. This is NOT the time the load was requested in the case of pending requests.	*/
 	double						LoadStartTime = 0.0;
 	TAtomic<int32> RefCount{ 0 };
+	int32						ProcessedExportBundlesCount = 0;
 	/** Current bundle entry index in the current export bundle */
 	int32						ExportBundleEntryIndex = 0;
 	/** Current index into ExternalReadDependencies array used to spread wating for external reads over several frames			*/
@@ -1727,15 +1894,14 @@ private:
 
 	FIoRequest IoRequest;
 	FIoBuffer IoBuffer;
-	const uint8* CurrentExportDataPtr = nullptr;
 	const uint8* AllExportDataPtr = nullptr;
-	uint64 ExportBundlesSize = 0;
+	const uint8* CurrentExportDataPtr = nullptr;
 	uint32 CookedHeaderSize = 0;
-	uint32 LoadOrder = 0;
 
 	const FExportMapEntry* ExportMap = nullptr;
-	FPackageImportStore ImportStore;
+	TArrayView<const uint8> ArcsData;
 	FNameMap NameMap;
+	TArray<FExportToBundleMapping> ExportToBundleMappings;
 
 public:
 
@@ -1753,7 +1919,11 @@ public:
 
 	static EAsyncPackageState::Type Event_ProcessExportBundle(FAsyncLoadingThreadState2& ThreadState, FAsyncPackage2* Package, int32 ExportBundleIndex);
 	static EAsyncPackageState::Type Event_ProcessPackageSummary(FAsyncLoadingThreadState2& ThreadState, FAsyncPackage2* Package, int32);
+	static EAsyncPackageState::Type Event_SetupDependencies(FAsyncLoadingThreadState2& ThreadState, FAsyncPackage2* Package, int32);
 	static EAsyncPackageState::Type Event_ExportsDone(FAsyncLoadingThreadState2& ThreadState, FAsyncPackage2* Package, int32);
+	static EAsyncPackageState::Type Event_BeginPostLoad(FAsyncLoadingThreadState2& ThreadState, FAsyncPackage2* Package, int32);
+	static EAsyncPackageState::Type Event_PostLoadDone(FAsyncLoadingThreadState2& ThreadState, FAsyncPackage2* Package, int32);
+	static EAsyncPackageState::Type Event_FinishLoading(FAsyncLoadingThreadState2& ThreadState, FAsyncPackage2* Package, int32);
 	static EAsyncPackageState::Type Event_PostLoadExportBundle(FAsyncLoadingThreadState2& ThreadState, FAsyncPackage2* Package, int32 ExportBundleIndex);
 	static EAsyncPackageState::Type Event_DeferredPostLoadExportBundle(FAsyncLoadingThreadState2& ThreadState, FAsyncPackage2* Package, int32 ExportBundleIndex);
 	
@@ -1781,7 +1951,7 @@ public:
 
 private:
 	void CreateNodes(const FAsyncLoadEventSpec* EventSpecs);
-	void SetupSerializedArcs(const uint8* GraphData, uint64 GraphDataSize);
+	void SetupSerializedArcs();
 	void SetupScriptDependencies();
 
 	/**
@@ -1931,22 +2101,20 @@ private:
 	FEvent* ThreadSuspendedEvent;
 	/** [ASYNC/GAME THREAD] Event used to signal that the async loading thread has resumed */
 	FEvent* ThreadResumedEvent;
-	/** [ASYNC/GAME THREAD] List of queued packages to stream */
-	TArray<FAsyncPackageDesc2*> QueuedPackages;
-	/** [ASYNC/GAME THREAD] Package queue critical section */
-	FCriticalSection QueueCritical;
 	TArray<FAsyncPackage2*> LoadedPackagesToProcess;
 	/** [GAME THREAD] Game thread CompletedPackages list */
 	TArray<FAsyncPackage2*> CompletedPackages;
 	/** [ASYNC/GAME THREAD] Packages to be deleted from async thread */
 	TQueue<FAsyncPackage2*, EQueueMode::Spsc> DeferredDeletePackages;
 	
-	struct FQueuedFailedPackageCallback
+	struct FFailedPackageRequest
 	{
+		int32 RequestID = INDEX_NONE;
 		FName PackageName;
 		TUniquePtr<FLoadPackageAsyncDelegate> Callback;
 	};
-	TArray<FQueuedFailedPackageCallback> QueuedFailedPackageCallbacks;
+	TArray<FFailedPackageRequest> FailedPackageRequests;
+	FCriticalSection FailedPackageRequestsCritical;
 
 	FCriticalSection AsyncPackagesCritical;
 	/** Packages in active loading with GetAsyncPackageId() as key */
@@ -1983,6 +2151,7 @@ private:
 	TUniquePtr<IPackageStore> PackageStore;
 	FLoadedPackageStore LoadedPackageStore;
 	FGlobalImportStore GlobalImportStore;
+	TUniquePtr<FPackageRequestQueue> PackageRequestQueue;
 
 	/** Initial load pending CDOs */
 	TMap<UClass*, TArray<FEventLoadNode2*>> PendingCDOs;
@@ -1991,7 +2160,7 @@ private:
 	{
 		bool operator<(const FBundleIoRequest& Other) const
 		{
-			return Package->LoadOrder < Other.Package->LoadOrder;
+			return Package->Data.ExportInfo.LoadOrder < Other.Package->Data.ExportInfo.LoadOrder;
 		}
 
 		FAsyncPackage2* Package;
@@ -2120,15 +2289,10 @@ public:
 	}
 
 	void UpdatePackagePriority(FAsyncPackage2* Package, int32 NewPriority);
-	
-	FAsyncPackage2* FindOrInsertPackage(FAsyncPackageDesc2* InDesc, bool& bInserted);
 
-	/**
-	* [ASYNC/GAME THREAD] Queues a package for streaming.
-	*
-	* @param Package package descriptor.
-	*/
-	void QueuePackage(FAsyncPackageDesc2& Package);
+	FAsyncPackageDesc2 CreatePackageDesc(FPackageRequest& Request, bool& bIsMissingPackage);
+	FAsyncPackage2* FindOrInsertPackage(FAsyncPackageDesc2& InDesc, bool& bInserted, TUniquePtr<FLoadPackageAsyncDelegate>&& PackageLoadedDelegate = TUniquePtr<FLoadPackageAsyncDelegate>());
+	void QueueMissingPackage(FAsyncPackageDesc2& PackageDesc, TUniquePtr<FLoadPackageAsyncDelegate>&& LoadPackageAsyncDelegate);
 
 	/**
 	* [ASYNC* THREAD] Loads all packages
@@ -2236,7 +2400,7 @@ public:
 	/**
 	 * [ASYNC/GAME THREAD] Removes a request ID from the list of pending requests
 	 */
-	void RemovePendingRequests(TArray<int32, TInlineAllocator<2>>& RequestIDs)
+	void RemovePendingRequests(TArrayView<int32> RequestIDs)
 	{
 		FScopeLock Lock(&PendingRequestsCritical);
 		for (int32 ID : RequestIDs)
@@ -2266,7 +2430,6 @@ private:
 	void SuspendWorkers();
 	void ResumeWorkers();
 
-	void LoadGlobalNameMap();
 	void LazyInitializeFromLoadPackage();
 	void FinalizeInitialLoad();
 
@@ -2320,40 +2483,57 @@ private:
 		UE_ASYNC_PACKAGE_DEBUG(Desc);
 		checkf(Desc.StoreEntry, TEXT("No package store entry for package %s"), *Desc.DiskPackageName.ToString());
 
+		FPackageStoreEntry PackageStoreEntry = PackageStore->GetPackageEntry(Desc.StoreEntry);
+
 		FAsyncPackageData Data;
-		Data.ExportCount = Desc.StoreEntry->ExportCount;
-		Data.ExportBundleCount = Desc.StoreEntry->ExportBundleCount;
+		Data.ExportInfo = PackageStoreEntry.ExportInfo;
 
-		const int32 ExportBundleNodeCount = Data.ExportBundleCount * EEventLoadNode2::ExportBundle_NumPhases;
-		const int32 ImportedPackageCount = Desc.StoreEntry->ImportedPackages.Num();
+		const int32 ExportBundleNodeCount = Data.ExportInfo.ExportBundleCount * EEventLoadNode2::ExportBundle_NumPhases;
+		const int32 ImportedPackageCount = PackageStoreEntry.ImportedPackageIds.Num();
+		const int32 ShaderMapHashesCount = PackageStoreEntry.ShaderMapHashes.Num();
 		const int32 NodeCount = EEventLoadNode2::Package_NumPhases + ExportBundleNodeCount;
 
-		const uint64 ExportBundleHeadersSize = sizeof(FExportBundleHeader) * Data.ExportBundleCount;
-		const uint64 ExportBundleEntriesSize = sizeof(FExportBundleEntry) * Data.ExportCount * FExportBundleEntry::ExportCommandType_Count;
-		Data.ExportBundlesMetaSize = ExportBundleHeadersSize + ExportBundleEntriesSize;
+		Data.ExportBundleHeadersSize = sizeof(FExportBundleHeader) * Data.ExportInfo.ExportBundleCount;
+		Data.ExportBundleEntriesSize = sizeof(FExportBundleEntry) * Data.ExportInfo.ExportCount * FExportBundleEntry::ExportCommandType_Count;
+		const uint64 ExportBundlesMetaSize = Data.ExportBundleHeadersSize + Data.ExportBundleEntriesSize;
 
 		const uint64 AsyncPackageMemSize = Align(sizeof(FAsyncPackage2), 8);
-		const uint64 ExportBundlesMetaMemSize = Align(Data.ExportBundlesMetaSize, 8);
-		const uint64 ExportsMemSize = Align(sizeof(FExportObject) * Data.ExportCount, 8);
+		const uint64 ExportBundlesMetaMemSize = Align(ExportBundlesMetaSize, 8);
+		const uint64 ExportsMemSize = Align(sizeof(FExportObject) * Data.ExportInfo.ExportCount, 8);
 		const uint64 ImportedPackagesMemSize = Align(sizeof(FAsyncPackage2*) * ImportedPackageCount, 8);
 		const uint64 PackageNodesMemSize = Align(sizeof(FEventLoadNode2) * NodeCount, 8);
+		const uint64 ImportedPackageIdsMemSize = Align(sizeof(FPackageId) * ImportedPackageCount, 8);
+		const uint64 ShaderMapHashesMemSize = Align(sizeof(FSHAHash) * ShaderMapHashesCount, 8);
 		const uint64 MemoryBufferSize =
 			AsyncPackageMemSize +
 			ExportBundlesMetaMemSize +
 			ExportsMemSize +
 			ImportedPackagesMemSize +
-			PackageNodesMemSize;
+			PackageNodesMemSize +
+			ImportedPackageIdsMemSize +
+			ShaderMapHashesMemSize;
 
 		uint8* MemoryBuffer = reinterpret_cast<uint8*>(FMemory::Malloc(MemoryBufferSize));
 
-		Data.ExportBundlesMetaMemory = MemoryBuffer + AsyncPackageMemSize;
-		Data.ExportBundleHeaders = reinterpret_cast<const FExportBundleHeader*>(Data.ExportBundlesMetaMemory);
-		Data.ExportBundleEntries = reinterpret_cast<const FExportBundleEntry*>(Data.ExportBundleHeaders + Data.ExportBundleCount);
+		uint8* DataPtr = MemoryBuffer + AsyncPackageMemSize;
 
-		Data.Exports = MakeArrayView(reinterpret_cast<FExportObject*>(MemoryBuffer + AsyncPackageMemSize + ExportBundlesMetaMemSize), Data.ExportCount);
-		Data.ImportedAsyncPackages = MakeArrayView(reinterpret_cast<FAsyncPackage2**>(MemoryBuffer + AsyncPackageMemSize + ExportBundlesMetaMemSize + ExportsMemSize), 0);
-		Data.PackageNodes = MakeArrayView(reinterpret_cast<FEventLoadNode2*>(MemoryBuffer + AsyncPackageMemSize + ExportBundlesMetaMemSize + ExportsMemSize + ImportedPackagesMemSize), NodeCount);
+		Data.ExportBundlesMetaMemory = DataPtr;
+		Data.ExportBundleHeaders = reinterpret_cast<const FExportBundleHeader*>(Data.ExportBundlesMetaMemory);
+		Data.ExportBundleEntries = reinterpret_cast<const FExportBundleEntry*>(Data.ExportBundleHeaders + Data.ExportInfo.ExportBundleCount);
+
+		DataPtr += ExportBundlesMetaMemSize;
+		Data.Exports = MakeArrayView(reinterpret_cast<FExportObject*>(DataPtr), Data.ExportInfo.ExportCount);
+		DataPtr += ExportsMemSize;
+		Data.ImportedAsyncPackages = MakeArrayView(reinterpret_cast<FAsyncPackage2**>(DataPtr), 0);
+		DataPtr += ImportedPackagesMemSize;
+		Data.PackageNodes = MakeArrayView(reinterpret_cast<FEventLoadNode2*>(DataPtr), NodeCount);
 		Data.ExportBundleNodes = MakeArrayView(&Data.PackageNodes[EEventLoadNode2::Package_NumPhases], ExportBundleNodeCount);
+		DataPtr += PackageNodesMemSize;
+		Data.ImportedPackageIds = MakeArrayView(reinterpret_cast<const FPackageId*>(DataPtr), ImportedPackageCount);
+		FMemory::Memcpy((void*)Data.ImportedPackageIds.GetData(), PackageStoreEntry.ImportedPackageIds.GetData(), sizeof(FPackageId) * ImportedPackageCount);
+		DataPtr += ImportedPackageIdsMemSize;
+		Data.ShaderMapHashes = MakeArrayView(reinterpret_cast<const FSHAHash*>(DataPtr), ShaderMapHashesCount);
+		FMemory::Memcpy((void*)Data.ShaderMapHashes.GetData(), PackageStoreEntry.ShaderMapHashes.GetData(), sizeof(FSHAHash) * ShaderMapHashesCount);
 
 		ExistingAsyncPackagesCounter.Increment();
 		return new (MemoryBuffer) FAsyncPackage2(Desc, Data, *this, GraphAllocator, EventSpecs.GetData());
@@ -2451,8 +2631,25 @@ void FAsyncLoadingThread2::InitializeLoading()
 	FBulkDataBase::SetIoDispatcher(&IoDispatcher);
 #endif
 
-	//TODO: Move construction of package store from the loader
-	PackageStore = MakeFilePackageStore(IoDispatcher);
+#if WITH_COTF
+	if (IsRunningCookOnTheFly())
+	{
+		UE::Cook::ICookOnTheFlyModule& CookOnTheFlyModule = FModuleManager::LoadModuleChecked<UE::Cook::ICookOnTheFlyModule>(TEXT("CookOnTheFly"));
+		PackageStore = MakeCookOnTheFlyPackageStore(
+			IoDispatcher,
+			CookOnTheFlyModule.GetServerConnection(),
+			[this]()
+			{
+				AltZenaphore.NotifyOne();
+			});
+	}
+	else
+#endif
+	{
+		PackageStore = MakeFilePackageStore(IoDispatcher);
+	}
+	
+	PackageRequestQueue = MakeUnique<FPackageRequestQueue>(*PackageStore);
 
 	FPackageName::DoesPackageExistOverride().BindLambda([this](FName PackageName)
 	{
@@ -2465,19 +2662,6 @@ void FAsyncLoadingThread2::InitializeLoading()
 	UE_LOG(LogStreaming, Display, TEXT("AsyncLoading2 - Initialized"));
 }
 
-void FAsyncLoadingThread2::QueuePackage(FAsyncPackageDesc2& Package)
-{
-	//TRACE_CPUPROFILER_EVENT_SCOPE(QueuePackage);
-	UE_ASYNC_PACKAGE_DEBUG(Package);
-	checkf(Package.StoreEntry, TEXT("No package store entry for package %s"), *Package.DiskPackageName.ToString());
-	{
-		FScopeLock QueueLock(&QueueCritical);
-		++QueuedPackagesCounter;
-		QueuedPackages.Add(new FAsyncPackageDesc2(Package, MoveTemp(Package.PackageLoadedDelegate)));
-	}
-	AltZenaphore.NotifyOne();
-}
-
 void FAsyncLoadingThread2::UpdatePackagePriority(FAsyncPackage2* Package, int32 NewPriority)
 {
 	TRACE_CPUPROFILER_EVENT_SCOPE(UpdatePackagePriority);
@@ -2485,37 +2669,37 @@ void FAsyncLoadingThread2::UpdatePackagePriority(FAsyncPackage2* Package, int32
 	Package->IoRequest.UpdatePriority(NewPriority);
 }
 
-FAsyncPackage2* FAsyncLoadingThread2::FindOrInsertPackage(FAsyncPackageDesc2* Desc, bool& bInserted)
+FAsyncPackage2* FAsyncLoadingThread2::FindOrInsertPackage(FAsyncPackageDesc2& Desc, bool& bInserted, TUniquePtr<FLoadPackageAsyncDelegate>&& PackageLoadedDelegate)
 {
 	TRACE_CPUPROFILER_EVENT_SCOPE(FindOrInsertPackage);
 	FAsyncPackage2* Package = nullptr;
 	bInserted = false;
 	{
 		FScopeLock LockAsyncPackages(&AsyncPackagesCritical);
-		Package = AsyncPackageLookup.FindRef(Desc->GetAsyncPackageId());
+		Package = AsyncPackageLookup.FindRef(Desc.GetAsyncPackageId());
 		if (!Package)
 		{
-			Package = CreateAsyncPackage(*Desc);
-			checkf(Package, TEXT("Failed to create async package %s"), *Desc->DiskPackageName.ToString());
+			Package = CreateAsyncPackage(Desc);
+			checkf(Package, TEXT("Failed to create async package %s"), *Desc.DiskPackageName.ToString());
 			Package->AddRef();
 			ActiveAsyncPackagesCounter.Increment();
-			AsyncPackageLookup.Add(Desc->GetAsyncPackageId(), Package);
+			AsyncPackageLookup.Add(Desc.GetAsyncPackageId(), Package);
 			bInserted = true;
 		}
 		else
 		{
-			if (Desc->RequestID > 0)
+			if (Desc.RequestID > 0)
 			{
-				Package->AddRequestID(Desc->RequestID);
+				Package->AddRequestID(Desc.RequestID);
 			}
-			if (Desc->Priority > Package->Desc.Priority)
+			if (Desc.Priority > Package->Desc.Priority)
 			{
-				UpdatePackagePriority(Package, Desc->Priority);
+				UpdatePackagePriority(Package, Desc.Priority);
 			}
 		}
-		if (Desc->PackageLoadedDelegate.IsValid())
+		if (PackageLoadedDelegate.IsValid())
 		{
-			Package->AddCompletionCallback(MoveTemp(Desc->PackageLoadedDelegate));
+			Package->AddCompletionCallback(MoveTemp(PackageLoadedDelegate));
 		}
 	}
 	return Package;
@@ -2527,64 +2711,64 @@ bool FAsyncLoadingThread2::CreateAsyncPackagesFromQueue(FAsyncLoadingThreadState
 	
 	bool bPackagesCreated = false;
 	const int32 TimeSliceGranularity = ThreadState.UseTimeLimit() ? 4 : MAX_int32;
-	TArray<FAsyncPackageDesc2*> QueueCopy;
 
-	do
+	for (;;)
 	{
+		int32 NumDequeued = 0;
+		PackageRequestQueue->Dequeue([this, &NumDequeued, TimeSliceGranularity](FPackageRequest& Request)
 		{
-			QueueCopy.Reset();
-			FScopeLock QueueLock(&QueueCritical);
-
-			const int32 NumPackagesToCopy = FMath::Min(TimeSliceGranularity, QueuedPackages.Num());
-			if (NumPackagesToCopy > 0)
-			{
-				QueueCopy.Append(QueuedPackages.GetData(), NumPackagesToCopy);
-				QueuedPackages.RemoveAt(0, NumPackagesToCopy, false);
-			}
-			else
-			{
-				break;
-			}
-		}
-
-		for (FAsyncPackageDesc2* PackageDesc : QueueCopy)
-		{
-			bool bInserted;
-			FAsyncPackage2* Package = FindOrInsertPackage(PackageDesc, bInserted);
-			checkf(Package, TEXT("Failed to find or insert imported package %s"), *PackageDesc->DiskPackageName.ToString());
+			bool bIsMissingPackage;
+			FAsyncPackageDesc2 PackageDesc = CreatePackageDesc(Request, bIsMissingPackage);
 
-			if (bInserted)
+			if (bIsMissingPackage)
 			{
-				UE_ASYNC_PACKAGE_LOG(Verbose, *PackageDesc, TEXT("CreateAsyncPackages: AddPackage"),
-					TEXT("Start loading package."));
+				QueueMissingPackage(PackageDesc, MoveTemp(Request.PackageLoadedDelegate));
 			}
 			else
 			{
-				UE_ASYNC_PACKAGE_LOG_VERBOSE(Verbose, *PackageDesc, TEXT("CreateAsyncPackages: UpdatePackage"),
-					TEXT("Package is alreay being loaded."));
-			}
+				bool bInserted;
+				FAsyncPackage2* Package = FindOrInsertPackage(PackageDesc, bInserted, MoveTemp(Request.PackageLoadedDelegate));
+				checkf(Package, TEXT("Failed to find or insert imported package %s"), *PackageDesc.DiskPackageName.ToString());
 
-			--QueuedPackagesCounter;
-			if (Package)
-			{
+				if (bInserted)
 				{
-					TRACE_CPUPROFILER_EVENT_SCOPE(ImportPackages);
-					Package->ImportPackagesRecursive();
+					UE_ASYNC_PACKAGE_LOG(Verbose, PackageDesc, TEXT("CreateAsyncPackages: AddPackage"),
+						TEXT("Start loading package."));
 				}
-
-				if (bInserted)
+				else
 				{
-					Package->StartLoading();
+					UE_ASYNC_PACKAGE_LOG_VERBOSE(Verbose, PackageDesc, TEXT("CreateAsyncPackages: UpdatePackage"),
+						TEXT("Package is alreay being loaded."));
 				}
 
-				StartBundleIoRequests();
+				--QueuedPackagesCounter;
+				if (Package)
+				{
+					{
+						TRACE_CPUPROFILER_EVENT_SCOPE(ImportPackages);
+						Package->ImportPackagesRecursive();
+					}
+
+					if (bInserted)
+					{
+						Package->StartLoading();
+					}
+
+					StartBundleIoRequests();
+				}
 			}
-			delete PackageDesc;
-		}
 
-		bPackagesCreated |= QueueCopy.Num() > 0;
-	} while (!ThreadState.IsTimeLimitExceeded(TEXT("CreateAsyncPackagesFromQueue")));
+			return ++NumDequeued < TimeSliceGranularity;
+		});
+		
+		bPackagesCreated |= NumDequeued > 0;
 
+		if (!NumDequeued || ThreadState.IsTimeLimitExceeded(TEXT("CreateAsyncPackagesFromQueue")))
+		{
+			break;
+		}
+	}
+	
 	return bPackagesCreated;
 }
 
@@ -2597,10 +2781,10 @@ void FAsyncLoadingThread2::AddBundleIoRequest(FAsyncPackage2* Package)
 
 void FAsyncLoadingThread2::BundleIoRequestCompleted(FAsyncPackage2* Package)
 {
-	check(PendingBundleIoRequestsTotalSize >= Package->ExportBundlesSize)
-	PendingBundleIoRequestsTotalSize -= Package->ExportBundlesSize;
+	check(PendingBundleIoRequestsTotalSize >= Package->Data.ExportInfo.ExportBundlesSize)
+	PendingBundleIoRequestsTotalSize -= Package->Data.ExportInfo.ExportBundlesSize;
 	TRACE_COUNTER_SET(AsyncLoadingPendingIoRequestsSize, PendingBundleIoRequestsTotalSize);
-	CSV_CUSTOM_STAT_DEFINED(FrameCompletedBundleLoadsKB, float((double)Package->ExportBundlesSize / 1024.0), ECsvCustomStatOp::Accumulate);
+	CSV_CUSTOM_STAT_DEFINED(FrameCompletedBundleLoadsKB, float((double)Package->Data.ExportInfo.ExportBundlesSize / 1024.0), ECsvCustomStatOp::Accumulate);
 	if (WaitingIoRequests.Num())
 	{
 		StartBundleIoRequests();
@@ -2610,17 +2794,17 @@ void FAsyncLoadingThread2::BundleIoRequestCompleted(FAsyncPackage2* Package)
 void FAsyncLoadingThread2::StartBundleIoRequests()
 {
 	TRACE_CPUPROFILER_EVENT_SCOPE(StartBundleIoRequests);
-	const uint64 MaxPendingRequestsSize = uint64(GAsyncLoadingMaxPendingRequestsSizeMB) << 20;
+	static const uint64 MaxPendingRequestsSize = IsRunningCookOnTheFly() ? MAX_uint64 : uint64(GAsyncLoadingMaxPendingRequestsSizeMB) << 20;
 	FIoBatch IoBatch = IoDispatcher.NewBatch();
 	while (WaitingIoRequests.Num())
 	{
 		FBundleIoRequest& BundleIoRequest = WaitingIoRequests.HeapTop();
 		FAsyncPackage2* Package = BundleIoRequest.Package;
-		if (PendingBundleIoRequestsTotalSize > 0 && PendingBundleIoRequestsTotalSize + Package->ExportBundlesSize > MaxPendingRequestsSize)
+		if (PendingBundleIoRequestsTotalSize > 0 && PendingBundleIoRequestsTotalSize + Package->Data.ExportInfo.ExportBundlesSize > MaxPendingRequestsSize)
 		{
 			break;
 		}
-		PendingBundleIoRequestsTotalSize += Package->ExportBundlesSize;
+		PendingBundleIoRequestsTotalSize += Package->Data.ExportInfo.ExportBundlesSize;
 		WaitingIoRequests.HeapPop(BundleIoRequest, false);
 
 		FIoReadOptions ReadOptions;
@@ -2645,7 +2829,7 @@ void FAsyncLoadingThread2::StartBundleIoRequests()
 			Package->GetPackageNode(EEventLoadNode2::Package_ProcessSummary).ReleaseBarrier();
 		});
 
-		if (Package->Desc.StoreEntry->ShaderMapHashes.Num())
+		if (!Package->Data.ShaderMapHashes.IsEmpty())
 		{
 			TRACE_CPUPROFILER_EVENT_SCOPE(StartShaderMapRequests);
 			auto ReadShaderMapFunc = [Package, &IoBatch](const FIoChunkId& ChunkId, FGraphEventRef GraphEvent)
@@ -2662,7 +2846,7 @@ void FAsyncLoadingThread2::StartBundleIoRequests()
 						Package->GetPackageNode(Package_ExportsSerialized).ReleaseBarrier();
 					});
 			};
-			FCoreDelegates::PreloadPackageShaderMaps.ExecuteIfBound(MakeArrayView(Package->Desc.StoreEntry->ShaderMapHashes.Data(), Package->Desc.StoreEntry->ShaderMapHashes.Num()), ReadShaderMapFunc);
+			FCoreDelegates::PreloadPackageShaderMaps.ExecuteIfBound(Package->Data.ShaderMapHashes, ReadShaderMapFunc);
 		}
 	}
 	IoBatch.Issue();
@@ -3021,42 +3205,71 @@ FUObjectSerializeContext* FAsyncPackage2::GetSerializeContext()
 	return FUObjectThreadContext::Get().GetSerializeContext();
 }
 
-void FAsyncPackage2::SetupSerializedArcs(const uint8* GraphData, uint64 GraphDataSize)
+void FAsyncPackage2::SetupSerializedArcs()
 {
-	TRACE_CPUPROFILER_EVENT_SCOPE(SetupSerializedArcs);
-
-	FSimpleArchive GraphArchive(GraphData, GraphDataSize);
-	int32 ImportedPackagesCount;
-	GraphArchive << ImportedPackagesCount;
-	for (int32 ImportedPackageIndex = 0; ImportedPackageIndex < ImportedPackagesCount; ++ImportedPackageIndex)
+	TRACE_CPUPROFILER_EVENT_SCOPE(SetupExternalArcs);
+
+	FSimpleArchive ArcsArchive(ArcsData.GetData(), ArcsData.Num());
+	int32 InternalArcsCount;
+	ArcsArchive << InternalArcsCount;
+	for (int32 InternalArcIndex = 0; InternalArcIndex < InternalArcsCount; ++InternalArcIndex)
+	{
+		int32 FromExportBundleIndex;
+		ArcsArchive << FromExportBundleIndex;
+		int32 ToExportBundleIndex;
+		ArcsArchive << ToExportBundleIndex;
+		uint32 FromNodeIndexBase = FromExportBundleIndex * EEventLoadNode2::ExportBundle_NumPhases;
+		uint32 ToNodeIndexBase = ToExportBundleIndex * EEventLoadNode2::ExportBundle_NumPhases;
+		for (int32 Phase = 0; Phase < EEventLoadNode2::ExportBundle_NumPhases; ++Phase)
+		{
+			uint32 ToNodeIndex = ToNodeIndexBase + Phase;
+			uint32 FromNodeIndex = FromNodeIndexBase + Phase;
+			Data.ExportBundleNodes[ToNodeIndex].DependsOn(&Data.ExportBundleNodes[FromNodeIndex]);
+		}
+	}
+	for (const FAsyncPackage2* ImportedPackage : Data.ImportedAsyncPackages)
 	{
-		FPackageId ImportedPackageId;
 		int32 ExternalArcCount;
-		GraphArchive << ImportedPackageId;
-		GraphArchive << ExternalArcCount;
+		ArcsArchive << ExternalArcCount;
+
+		int32 PreviousFromExportBundleIndex = -1;
+		int32 PreviousToExportBundleIndex = -1;
 
-		FAsyncPackage2* ImportedPackage = AsyncLoadingThread.GetAsyncPackage(ImportedPackageId);
 		for (int32 ExternalArcIndex = 0; ExternalArcIndex < ExternalArcCount; ++ExternalArcIndex)
 		{
-			int32 FromExportBundleIndex;
+			int32 FromImportIndex;
+			uint8 FromCommandType;
 			int32 ToExportBundleIndex;
-			GraphArchive << FromExportBundleIndex;
-			GraphArchive << ToExportBundleIndex;
+			ArcsArchive << FromImportIndex;
+			ArcsArchive << FromCommandType;
+			ArcsArchive << ToExportBundleIndex;
 			if (ImportedPackage)
 			{
-				FromExportBundleIndex = FromExportBundleIndex == MAX_uint32
-					? ImportedPackage->Data.ExportBundleCount - 1
-					: FromExportBundleIndex;
-
-				check(FromExportBundleIndex < ImportedPackage->Data.ExportBundleCount);
-				check(ToExportBundleIndex < Data.ExportBundleCount);
-				uint32 FromNodeIndexBase = FromExportBundleIndex * EEventLoadNode2::ExportBundle_NumPhases;
-				uint32 ToNodeIndexBase = ToExportBundleIndex * EEventLoadNode2::ExportBundle_NumPhases;
-				for (int32 Phase = 0; Phase < EEventLoadNode2::ExportBundle_NumPhases; ++Phase)
+				check(FromImportIndex < ImportStore.ImportMap.Num());
+				check(FromCommandType < FExportBundleEntry::ExportCommandType_Count);
+				check(ToExportBundleIndex < Data.ExportInfo.ExportBundleCount);
+				
+				FPackageObjectIndex GlobalImportIndex = ImportStore.ImportMap[FromImportIndex];
+				for (const FExportToBundleMapping& ExportToBundleMapping : ImportedPackage->ExportToBundleMappings)
 				{
-					uint32 ToNodeIndex = ToNodeIndexBase + Phase;
-					uint32 FromNodeIndex = FromNodeIndexBase + Phase;
-					Data.ExportBundleNodes[ToNodeIndex].DependsOn(&ImportedPackage->Data.ExportBundleNodes[FromNodeIndex]);
+					if (ExportToBundleMapping.GlobalImportIndex == GlobalImportIndex)
+					{
+						int32 FromExportBundleIndex = ExportToBundleMapping.BundleIndex[FromCommandType];
+						if (PreviousFromExportBundleIndex != FromExportBundleIndex || PreviousToExportBundleIndex != ToExportBundleIndex)
+						{
+							PreviousFromExportBundleIndex = FromExportBundleIndex;
+							PreviousToExportBundleIndex = ToExportBundleIndex;
+							uint32 FromNodeIndexBase = FromExportBundleIndex * EEventLoadNode2::ExportBundle_NumPhases;
+							uint32 ToNodeIndexBase = ToExportBundleIndex * EEventLoadNode2::ExportBundle_NumPhases;
+							for (int32 Phase = 0; Phase < EEventLoadNode2::ExportBundle_NumPhases; ++Phase)
+							{
+								uint32 ToNodeIndex = ToNodeIndexBase + Phase;
+								uint32 FromNodeIndex = FromNodeIndexBase + Phase;
+								Data.ExportBundleNodes[ToNodeIndex].DependsOn(&ImportedPackage->Data.ExportBundleNodes[FromNodeIndex]);
+							}
+						}
+						break;
+					}
 				}
 			}
 		}
@@ -3151,7 +3364,7 @@ void FAsyncPackage2::ImportPackagesRecursive()
 	}
 	check(AsyncPackageLoadingState == EAsyncPackageLoadingState2::NewPackage);
 
-	const int32 ImportedPackageCount = Desc.StoreEntry->ImportedPackages.Num();
+	const int32 ImportedPackageCount = Data.ImportedPackageIds.Num();
 	if (!ImportedPackageCount)
 	{
 		AsyncPackageLoadingState = EAsyncPackageLoadingState2::ImportPackagesDone;
@@ -3165,27 +3378,30 @@ void FAsyncPackage2::ImportPackagesRecursive()
 	int32 ImportedPackageIndex = 0;
 
 	IPackageStore& PackageStore = AsyncLoadingThread.GetPackageStore();
-	for (const FPackageId& ImportedPackageId : Desc.StoreEntry->ImportedPackages)
+	Data.ImportedAsyncPackages = MakeArrayView(Data.ImportedAsyncPackages.GetData(), Data.ImportedPackageIds.Num());
+	for (const FPackageId& ImportedPackageId : Data.ImportedPackageIds)
 	{
 		FLoadedPackageRef& PackageRef = ImportStore.LoadedPackageStore.GetPackageRef(ImportedPackageId);
 		if (PackageRef.AreAllPublicExportsLoaded())
 		{
+			Data.ImportedAsyncPackages[ImportedPackageIndex++] = nullptr;
 			continue;
 		}
 
-		const FPackageStoreEntry* ImportedPackageEntry = PackageStore.GetPackageEntry(ImportedPackageId);
+		const FPackageStoreEntryHandle ImportedPackageEntry = PackageStore.GetPackageEntryHandle(ImportedPackageId);
 
 		if (!ImportedPackageEntry)
 		{
 			UE_ASYNC_PACKAGE_LOG(Warning, Desc, TEXT("ImportPackages: SkipPackage"),
 				TEXT("Skipping non mounted imported package with id '0x%llX'"), ImportedPackageId.Value());
 			PackageRef.SetIsMissingPackage();
+			Data.ImportedAsyncPackages[ImportedPackageIndex++] = nullptr;
 			continue;
 		}
 
-		FAsyncPackageDesc2 PackageDesc(INDEX_NONE, Desc.Priority, ImportedPackageId, ImportedPackageEntry);
+		FAsyncPackageDesc2 PackageDesc = FAsyncPackageDesc2::FromPackageImport(INDEX_NONE, Desc.Priority, ImportedPackageId, ImportedPackageEntry);
 		bool bInserted;
-		FAsyncPackage2* ImportedPackage = AsyncLoadingThread.FindOrInsertPackage(&PackageDesc, bInserted);
+		FAsyncPackage2* ImportedPackage = AsyncLoadingThread.FindOrInsertPackage(PackageDesc, bInserted);
 
 		checkf(ImportedPackage, TEXT("Failed to find or insert imported package with id '0x%llX'"), ImportedPackageId.Value());
 		TRACE_LOADTIME_ASYNC_PACKAGE_IMPORT_DEPENDENCY(this, ImportedPackage);
@@ -3201,9 +3417,11 @@ void FAsyncPackage2::ImportPackagesRecursive()
 				TEXT("Imported package is already being loaded."));
 		}
 		ImportedPackage->AddRef();
-		check(ImportedPackageIndex == Data.ImportedAsyncPackages.Num());
-		Data.ImportedAsyncPackages = MakeArrayView(Data.ImportedAsyncPackages.GetData(), ImportedPackageIndex + 1);
 		Data.ImportedAsyncPackages[ImportedPackageIndex++] = ImportedPackage;
+		GetPackageNode(Package_SetupDependencies).DependsOn(&ImportedPackage->GetPackageNode(Package_ProcessSummary));
+		GetPackageNode(Package_BeginPostLoad).DependsOn(&ImportedPackage->GetPackageNode(Package_ExportsSerialized));
+		GetPackageNode(Package_FinishLoading).DependsOn(&ImportedPackage->GetPackageNode(Package_PostLoadDone));
+
 		if (bInserted)
 		{
 			ImportedPackage->ImportPackagesRecursive();
@@ -3253,18 +3471,11 @@ EAsyncPackageState::Type FAsyncPackage2::Event_ProcessPackageSummary(FAsyncLoadi
 
 		const uint8* PackageSummaryData = Package->IoBuffer.Data();
 		const FPackageSummary* PackageSummary = reinterpret_cast<const FPackageSummary*>(PackageSummaryData);
-		const uint8* GraphData = PackageSummaryData + PackageSummary->GraphDataOffset;
-		const uint64 PackageSummarySize = GraphData + PackageSummary->GraphDataSize - PackageSummaryData;
-
-		if (PackageSummary->NameMapNamesSize)
 		{
 			TRACE_CPUPROFILER_EVENT_SCOPE(LoadPackageNameMap);
-			const uint8* NameMapNamesData = PackageSummaryData + PackageSummary->NameMapNamesOffset;
-			const uint8* NameMapHashesData = PackageSummaryData + PackageSummary->NameMapHashesOffset;
-			Package->NameMap.Load(
-				TArrayView<const uint8>(NameMapNamesData, PackageSummary->NameMapNamesSize),
-				TArrayView<const uint8>(NameMapHashesData, PackageSummary->NameMapHashesSize),
-				FMappedName::EType::Package);
+			TArrayView<const uint8> NameMapView(PackageSummaryData + sizeof(FPackageSummary), PackageSummary->HeaderSize - sizeof(FPackageSummary));
+			FMemoryReaderView NameMapReader(NameMapView);
+			Package->NameMap.Load(NameMapReader, FMappedName::EType::Package);
 		}
 
 		{
@@ -3288,30 +3499,62 @@ EAsyncPackageState::Type FAsyncPackage2::Event_ProcessPackageSummary(FAsyncLoadi
 				reinterpret_cast<const FPackageObjectIndex*>(PackageSummaryData + PackageSummary->ImportMapOffset),
 				(PackageSummary->ExportMapOffset - PackageSummary->ImportMapOffset) / sizeof(FPackageObjectIndex));
 		Package->ExportMap = reinterpret_cast<const FExportMapEntry*>(PackageSummaryData + PackageSummary->ExportMapOffset);
-		
-		FMemory::Memcpy(Package->Data.ExportBundlesMetaMemory, PackageSummaryData + PackageSummary->ExportBundlesOffset, Package->Data.ExportBundlesMetaSize);
+
+		uint64 ExportBundleHeadersOffset = PackageSummary->GraphDataOffset;
+		uint64 ArcsDataOffset = ExportBundleHeadersOffset + Package->Data.ExportBundleHeadersSize;
+		uint64 ArcsDataSize = PackageSummary->HeaderSize - ArcsDataOffset;
+		Package->ArcsData = TArrayView<const uint8>(PackageSummaryData + ArcsDataOffset, ArcsDataSize);
+		FMemory::Memcpy(Package->Data.ExportBundlesMetaMemory, PackageSummaryData + ExportBundleHeadersOffset, Package->Data.ExportBundleHeadersSize);
+		FMemory::Memcpy(Package->Data.ExportBundlesMetaMemory + Package->Data.ExportBundleHeadersSize, PackageSummaryData + PackageSummary->ExportBundleEntriesOffset, Package->Data.ExportBundleEntriesSize);
 
 		Package->CreateUPackage(PackageSummary);
-		Package->SetupSerializedArcs(GraphData, PackageSummary->GraphDataSize);
 
-		Package->AllExportDataPtr = PackageSummaryData + PackageSummarySize;
-		Package->CurrentExportDataPtr = Package->AllExportDataPtr;
+		Package->AllExportDataPtr = PackageSummaryData + PackageSummary->HeaderSize;
+
+		Package->ExportToBundleMappings.SetNum(Package->Data.ExportInfo.ExportCount);
+		for (int32 ExportBundleIndex = 0, ExportBundleCount = Package->Data.ExportInfo.ExportBundleCount; ExportBundleIndex < ExportBundleCount; ++ExportBundleIndex)
+		{
+			const FExportBundleHeader* ExportBundle = Package->Data.ExportBundleHeaders + ExportBundleIndex;
+			const FExportBundleEntry* BundleEntries = Package->Data.ExportBundleEntries + ExportBundle->FirstEntryIndex;
+			const FExportBundleEntry* BundleEntry = BundleEntries + Package->ExportBundleEntryIndex;
+			const FExportBundleEntry* BundleEntryEnd = BundleEntries + ExportBundle->EntryCount;
+			while (BundleEntry < BundleEntryEnd)
+			{
+				const FExportMapEntry& ExportMapEntry = Package->ExportMap[BundleEntry->LocalExportIndex];
+				FExportToBundleMapping& ExportToBundleMapping = Package->ExportToBundleMappings[BundleEntry->LocalExportIndex];
+				ExportToBundleMapping.GlobalImportIndex = ExportMapEntry.GlobalImportIndex;
+				ExportToBundleMapping.BundleIndex[BundleEntry->CommandType] = ExportBundleIndex;
+				++BundleEntry;
+			}
+		}
 
-		TRACE_LOADTIME_PACKAGE_SUMMARY(Package, Package->Desc.DiskPackageName, PackageSummarySize, Package->ImportStore.ImportMap.Num(), Package->Data.ExportCount);
+		TRACE_LOADTIME_PACKAGE_SUMMARY(Package, Package->Desc.DiskPackageName, PackageSummary->HeaderSize, Package->ImportStore.ImportMap.Num(), Package->Data.ExportInfo.ExportCount);
 	}
 
+	Package->AsyncPackageLoadingState = EAsyncPackageLoadingState2::SetupDependencies;
+	Package->GetPackageNode(Package_SetupDependencies).ReleaseBarrier();
+	return EAsyncPackageState::Complete;
+}
+
+EAsyncPackageState::Type FAsyncPackage2::Event_SetupDependencies(FAsyncLoadingThreadState2& ThreadState, FAsyncPackage2* Package, int32)
+{
+	TRACE_CPUPROFILER_EVENT_SCOPE(Event_SetupDependencies);
+	check(Package->AsyncPackageLoadingState == EAsyncPackageLoadingState2::SetupDependencies);
+	
 	if (GIsInitialLoad)
 	{
 		Package->SetupScriptDependencies();
 	}
-	Package->GetExportBundleNode(ExportBundle_Process, 0).ReleaseBarrier();
-
-	check(Package->AsyncPackageLoadingState == EAsyncPackageLoadingState2::ProcessPackageSummary);
+	Package->SetupSerializedArcs();
 	Package->AsyncPackageLoadingState = EAsyncPackageLoadingState2::ProcessExportBundles;
+	for (int32 ExportBundleIndex = 0, ExportBundleCount = Package->Data.ExportInfo.ExportBundleCount; ExportBundleIndex < ExportBundleCount; ++ExportBundleIndex)
+	{
+		Package->GetExportBundleNode(ExportBundle_Process, ExportBundleIndex).ReleaseBarrier();
+	}
 	return EAsyncPackageState::Complete;
 }
 
-EAsyncPackageState::Type FAsyncPackage2::Event_ProcessExportBundle(FAsyncLoadingThreadState2& ThreadState, FAsyncPackage2* Package, int32 ExportBundleIndex)
+EAsyncPackageState::Type FAsyncPackage2::Event_ProcessExportBundle(FAsyncLoadingThreadState2& ThreadState, FAsyncPackage2* Package, int32 InExportBundleIndex)
 {
 	TRACE_CPUPROFILER_EVENT_SCOPE(Event_ProcessExportBundle);
 	UE_ASYNC_PACKAGE_DEBUG(Package->Desc);
@@ -3343,11 +3586,16 @@ EAsyncPackageState::Type FAsyncPackage2::Event_ProcessExportBundle(FAsyncLoading
 #endif
 	};
 
-	check(ExportBundleIndex < Package->Data.ExportBundleCount);
+	check(InExportBundleIndex < Package->Data.ExportInfo.ExportBundleCount);
 	
 	if (!Package->bLoadHasFailed)
 	{
+		const FExportBundleHeader* ExportBundle = Package->Data.ExportBundleHeaders + InExportBundleIndex;
 		const uint64 AllExportDataSize = Package->IoBuffer.DataSize() - (Package->AllExportDataPtr - Package->IoBuffer.Data());
+		if (Package->ExportBundleEntryIndex == 0)
+		{
+			Package->CurrentExportDataPtr = Package->AllExportDataPtr + ExportBundle->SerialOffset;
+		}
 		FExportArchive Ar(Package->AllExportDataPtr, Package->CurrentExportDataPtr, AllExportDataSize);
 		{
 			Ar.SetUEVer(Package->LinkerRoot->LinkerPackageVersion);
@@ -3372,7 +3620,6 @@ EAsyncPackageState::Type FAsyncPackage2::Event_ProcessExportBundle(FAsyncLoading
 			Ar.ExportMap = Package->ExportMap;
 			Ar.ExternalReadDependencies = &Package->ExternalReadDependencies;
 		}
-		const FExportBundleHeader* ExportBundle = Package->Data.ExportBundleHeaders + ExportBundleIndex;
 		const FExportBundleEntry* BundleEntries = Package->Data.ExportBundleEntries + ExportBundle->FirstEntryIndex;
 		const FExportBundleEntry* BundleEntry = BundleEntries + Package->ExportBundleEntryIndex;
 		const FExportBundleEntry* BundleEntryEnd = BundleEntries + ExportBundle->EntryCount;
@@ -3432,12 +3679,14 @@ EAsyncPackageState::Type FAsyncPackage2::Event_ProcessExportBundle(FAsyncLoading
 	
 	Package->ExportBundleEntryIndex = 0;
 
-	if (ExportBundleIndex + 1 < Package->Data.ExportBundleCount)
+	if (Package->ProcessedExportBundlesCount == 0)
 	{
-		Package->GetExportBundleNode(ExportBundle_Process, ExportBundleIndex + 1).ReleaseBarrier();
+		Package->AsyncLoadingThread.BundleIoRequestCompleted(Package);
 	}
-	else
+
+	if (++Package->ProcessedExportBundlesCount == Package->Data.ExportInfo.ExportBundleCount)
 	{
+		Package->ProcessedExportBundlesCount = 0;
 		Package->ImportStore.ImportMap = TArrayView<const FPackageObjectIndex>();
 		Package->IoBuffer = FIoBuffer();
 
@@ -3455,11 +3704,6 @@ EAsyncPackageState::Type FAsyncPackage2::Event_ProcessExportBundle(FAsyncLoading
 		}
 	}
 
-	if (ExportBundleIndex == 0)
-	{
-		Package->AsyncLoadingThread.BundleIoRequestCompleted(Package);
-	}
-
 	return EAsyncPackageState::Complete;
 }
 
@@ -3818,18 +4062,42 @@ EAsyncPackageState::Type FAsyncPackage2::Event_ExportsDone(FAsyncLoadingThreadSt
 		PackageRef.SetAllPublicExportsLoaded();
 	}
 
-	if (Package->Desc.StoreEntry->ShaderMapHashes.Num())
+	if (!Package->Data.ShaderMapHashes.IsEmpty())
 	{
 		TRACE_CPUPROFILER_EVENT_SCOPE(ReleasePreloadedShaderMaps);
-		FCoreDelegates::ReleasePreloadedPackageShaderMaps.ExecuteIfBound(MakeArrayView(Package->Desc.StoreEntry->ShaderMapHashes.Data(), Package->Desc.StoreEntry->ShaderMapHashes.Num()));
+		FCoreDelegates::ReleasePreloadedPackageShaderMaps.ExecuteIfBound(Package->Data.ShaderMapHashes);
 	}
 
 	Package->AsyncPackageLoadingState = EAsyncPackageLoadingState2::PostLoad;
-	Package->GetExportBundleNode(EEventLoadNode2::ExportBundle_PostLoad, 0).ReleaseBarrier();
+	Package->GetPackageNode(Package_BeginPostLoad).ReleaseBarrier(&ThreadState);
 	return EAsyncPackageState::Complete;
 }
 
-EAsyncPackageState::Type FAsyncPackage2::Event_PostLoadExportBundle(FAsyncLoadingThreadState2& ThreadState, FAsyncPackage2* Package, int32 ExportBundleIndex)
+EAsyncPackageState::Type FAsyncPackage2::Event_BeginPostLoad(FAsyncLoadingThreadState2& ThreadState, FAsyncPackage2* Package, int32)
+{
+	TRACE_CPUPROFILER_EVENT_SCOPE(Event_BeginPostLoad);
+	for (int32 ExportBundleIndex = 0, ExportBundleCount = Package->Data.ExportInfo.ExportBundleCount; ExportBundleIndex < ExportBundleCount; ++ExportBundleIndex)
+	{
+		Package->GetExportBundleNode(EEventLoadNode2::ExportBundle_PostLoad, ExportBundleIndex).ReleaseBarrier();
+	}
+	return EAsyncPackageState::Complete;
+}
+
+EAsyncPackageState::Type FAsyncPackage2::Event_PostLoadDone(FAsyncLoadingThreadState2& ThreadState, FAsyncPackage2* Package, int32)
+{
+	TRACE_CPUPROFILER_EVENT_SCOPE(Event_PostLoadDone);
+	Package->GetPackageNode(Package_FinishLoading).ReleaseBarrier(&ThreadState);
+	return EAsyncPackageState::Complete;
+}
+
+EAsyncPackageState::Type FAsyncPackage2::Event_FinishLoading(FAsyncLoadingThreadState2& ThreadState, FAsyncPackage2* Package, int32)
+{
+	TRACE_CPUPROFILER_EVENT_SCOPE(Event_FinishLoading);
+	Package->AsyncLoadingThread.LoadedPackagesToProcess.Add(Package);
+	return EAsyncPackageState::Complete;
+}
+
+EAsyncPackageState::Type FAsyncPackage2::Event_PostLoadExportBundle(FAsyncLoadingThreadState2& ThreadState, FAsyncPackage2* Package, int32 InExportBundleIndex)
 {
 	TRACE_CPUPROFILER_EVENT_SCOPE(Event_PostLoad);
 	UE_ASYNC_PACKAGE_DEBUG(Package->Desc);
@@ -3859,7 +4127,7 @@ EAsyncPackageState::Type FAsyncPackage2::Event_PostLoadExportBundle(FAsyncLoadin
 		}
 	}*/
 	
-	check(ExportBundleIndex < Package->Data.ExportBundleCount);
+	check(InExportBundleIndex < Package->Data.ExportInfo.ExportBundleCount);
 
 	EAsyncPackageState::Type LoadingState = EAsyncPackageState::Complete;
 
@@ -3876,7 +4144,7 @@ EAsyncPackageState::Type FAsyncPackage2::Event_PostLoadExportBundle(FAsyncLoadin
 		const bool bAsyncPostLoadEnabled = FAsyncLoadingThreadSettings::Get().bAsyncPostLoadEnabled;
 		const bool bIsMultithreaded = Package->AsyncLoadingThread.IsMultithreaded();
 
-		const FExportBundleHeader* ExportBundle = Package->Data.ExportBundleHeaders + ExportBundleIndex;
+		const FExportBundleHeader* ExportBundle = Package->Data.ExportBundleHeaders + InExportBundleIndex;
 		const FExportBundleEntry* BundleEntries = Package->Data.ExportBundleEntries + ExportBundle->FirstEntryIndex;
 		const FExportBundleEntry* BundleEntry = BundleEntries + Package->ExportBundleEntryIndex;
 		const FExportBundleEntry* BundleEntryEnd = BundleEntries + ExportBundle->EntryCount;
@@ -3934,12 +4202,9 @@ EAsyncPackageState::Type FAsyncPackage2::Event_PostLoadExportBundle(FAsyncLoadin
 
 	Package->ExportBundleEntryIndex = 0;
 
-	if (ExportBundleIndex + 1 < Package->Data.ExportBundleCount)
-	{
-		Package->GetExportBundleNode(ExportBundle_PostLoad, ExportBundleIndex + 1).ReleaseBarrier();
-	}
-	else
+	if (++Package->ProcessedExportBundlesCount == Package->Data.ExportInfo.ExportBundleCount)
 	{
+		Package->ProcessedExportBundlesCount = 0;
 		if (Package->LinkerRoot && !Package->bLoadHasFailed)
 		{
 			UE_ASYNC_PACKAGE_LOG(Verbose, Package->Desc, TEXT("AsyncThread: FullyLoaded"),
@@ -3951,13 +4216,16 @@ EAsyncPackageState::Type FAsyncPackage2::Event_PostLoadExportBundle(FAsyncLoadin
 
 		check(Package->AsyncPackageLoadingState == EAsyncPackageLoadingState2::PostLoad);
 		Package->AsyncPackageLoadingState = EAsyncPackageLoadingState2::DeferredPostLoad;
-		Package->GetExportBundleNode(ExportBundle_DeferredPostLoad, 0).ReleaseBarrier();
+		for (int32 ExportBundleIndex = 0, ExportBundleCount = Package->Data.ExportInfo.ExportBundleCount; ExportBundleIndex < ExportBundleCount; ++ExportBundleIndex)
+		{
+			Package->GetExportBundleNode(ExportBundle_DeferredPostLoad, ExportBundleIndex).ReleaseBarrier();
+		}
 	}
 
 	return EAsyncPackageState::Complete;
 }
 
-EAsyncPackageState::Type FAsyncPackage2::Event_DeferredPostLoadExportBundle(FAsyncLoadingThreadState2& ThreadState, FAsyncPackage2* Package, int32 ExportBundleIndex)
+EAsyncPackageState::Type FAsyncPackage2::Event_DeferredPostLoadExportBundle(FAsyncLoadingThreadState2& ThreadState, FAsyncPackage2* Package, int32 InExportBundleIndex)
 {
 	SCOPE_CYCLE_COUNTER(STAT_FAsyncPackage_PostLoadObjectsGameThread);
 	TRACE_CPUPROFILER_EVENT_SCOPE(Event_DeferredPostLoad);
@@ -3966,7 +4234,7 @@ EAsyncPackageState::Type FAsyncPackage2::Event_DeferredPostLoadExportBundle(FAsy
 
 	FAsyncPackageScope2 PackageScope(Package);
 
-	check(ExportBundleIndex < Package->Data.ExportBundleCount);
+	check(InExportBundleIndex < Package->Data.ExportInfo.ExportBundleCount);
 	EAsyncPackageState::Type LoadingState = EAsyncPackageState::Complete;
 
 	if (Package->bLoadHasFailed)
@@ -3979,7 +4247,7 @@ EAsyncPackageState::Type FAsyncPackage2::Event_DeferredPostLoadExportBundle(FAsy
 		TGuardValue<bool> GuardIsRoutingPostLoad(PackageScope.ThreadContext.IsRoutingPostLoad, true);
 		FAsyncLoadingTickScope2 InAsyncLoadingTick(Package->AsyncLoadingThread);
 
-		const FExportBundleHeader* ExportBundle = Package->Data.ExportBundleHeaders + ExportBundleIndex;
+		const FExportBundleHeader* ExportBundle = Package->Data.ExportBundleHeaders + InExportBundleIndex;
 		const FExportBundleEntry* BundleEntries = Package->Data.ExportBundleEntries + ExportBundle->FirstEntryIndex;
 		const FExportBundleEntry* BundleEntry = BundleEntries + Package->ExportBundleEntryIndex;
 		const FExportBundleEntry* BundleEntryEnd = BundleEntries + ExportBundle->EntryCount;
@@ -4029,15 +4297,12 @@ EAsyncPackageState::Type FAsyncPackage2::Event_DeferredPostLoadExportBundle(FAsy
 
 	Package->ExportBundleEntryIndex = 0;
 
-	if (ExportBundleIndex + 1 < Package->Data.ExportBundleCount)
-	{
-		Package->GetExportBundleNode(ExportBundle_DeferredPostLoad, ExportBundleIndex + 1).ReleaseBarrier();
-	}
-	else
+	if (++Package->ProcessedExportBundlesCount == Package->Data.ExportInfo.ExportBundleCount)
 	{
+		Package->ProcessedExportBundlesCount = 0;
 		check(Package->AsyncPackageLoadingState == EAsyncPackageLoadingState2::DeferredPostLoad);
 		Package->AsyncPackageLoadingState = EAsyncPackageLoadingState2::DeferredPostLoadDone;
-		Package->AsyncLoadingThread.LoadedPackagesToProcess.Add(Package);
+		Package->GetPackageNode(Package_PostLoadDone).ReleaseBarrier(&ThreadState);
 	}
 
 	return EAsyncPackageState::Complete;
@@ -4051,7 +4316,7 @@ FEventLoadNode2& FAsyncPackage2::GetPackageNode(EEventLoadNode2 Phase)
 
 FEventLoadNode2& FAsyncPackage2::GetExportBundleNode(EEventLoadNode2 Phase, uint32 ExportBundleIndex)
 {
-	check(ExportBundleIndex < uint32(Data.ExportBundleCount));
+	check(ExportBundleIndex < uint32(Data.ExportInfo.ExportBundleCount));
 	uint32 ExportBundleNodeIndex = ExportBundleIndex * EEventLoadNode2::ExportBundle_NumPhases + Phase;
 	return Data.ExportBundleNodes[ExportBundleNodeIndex];
 }
@@ -4113,9 +4378,15 @@ EAsyncPackageState::Type FAsyncLoadingThread2::ProcessAsyncLoadingFromGameThread
 
 			if (QueuedPackagesCounter)
 			{
-				CreateAsyncPackagesFromQueue(ThreadState);
-				OutPackagesProcessed++;
-				break;
+				if (CreateAsyncPackagesFromQueue(ThreadState))
+				{
+					OutPackagesProcessed++;
+					break;
+				}
+				else
+				{
+					return EAsyncPackageState::TimeOut;
+				}
 			}
 
 			bool bPopped = false;
@@ -4167,7 +4438,7 @@ EAsyncPackageState::Type FAsyncLoadingThread2::ProcessAsyncLoadingFromGameThread
 
 bool FAsyncPackage2::AreAllDependenciesFullyLoadedInternal(FAsyncPackage2* Package, TSet<FPackageId>& VisitedPackages, FPackageId& OutPackageId)
 {
-	for (const FPackageId& ImportedPackageId : Package->Desc.StoreEntry->ImportedPackages)
+	for (const FPackageId& ImportedPackageId : Package->Data.ImportedPackageIds)
 	{
 		if (VisitedPackages.Contains(ImportedPackageId))
 		{
@@ -4256,7 +4527,7 @@ EAsyncPackageState::Type FAsyncLoadingThread2::ProcessLoadedPackagesFromGameThre
 			bool bHasClusterObjects = false;
 			TArray<UObject*> CDODefaultSubobjects;
 			// Clear async loading flags (we still want RF_Async, but EInternalObjectFlags::AsyncLoading can be cleared)
-			for (int32 FinalizeIndex = 0; FinalizeIndex < Package->Data.ExportCount; ++FinalizeIndex)
+			for (int32 FinalizeIndex = 0, ExportCount = Package->Data.ExportInfo.ExportCount; FinalizeIndex < ExportCount; ++FinalizeIndex)
 			{
 				const FExportObject& Export = Package->Data.Exports[FinalizeIndex];
 				if (Export.bFiltered | Export.bExportLoadFailed)
@@ -4362,13 +4633,22 @@ EAsyncPackageState::Type FAsyncLoadingThread2::ProcessLoadedPackagesFromGameThre
 			UE_ASYNC_PACKAGE_LOG(Verbose, Package->Desc, TEXT("GameThread: LoadCompleted"),
 				TEXT("All loading of package is done, and the async package and load request will be deleted."));
 		}
-
-		bLocalDidSomething |= QueuedFailedPackageCallbacks.Num() > 0;
-		for (FQueuedFailedPackageCallback& QueuedFailedPackageCallback : QueuedFailedPackageCallbacks)
+		
 		{
-			QueuedFailedPackageCallback.Callback->ExecuteIfBound(QueuedFailedPackageCallback.PackageName, nullptr, EAsyncLoadingResult::Failed);
+			TArray<FFailedPackageRequest> LocalFailedPackageRequests;
+			{
+				FScopeLock _(&FailedPackageRequestsCritical);
+				Swap(LocalFailedPackageRequests, FailedPackageRequests);
+			}
+
+			bLocalDidSomething |= LocalFailedPackageRequests.Num() > 0;
+			for (FFailedPackageRequest& FailedPackageRequest : LocalFailedPackageRequests)
+			{
+				FailedPackageRequest.Callback->ExecuteIfBound(FailedPackageRequest.PackageName, nullptr, EAsyncLoadingResult::Failed);
+				RemovePendingRequests(TArrayView<int32>(&FailedPackageRequest.RequestID, 1));
+				--QueuedPackagesCounter;
+			}
 		}
-		QueuedFailedPackageCallbacks.Empty();
 
 		bLocalDidSomething |= CompletedPackages.Num() > 0;
 		for (int32 PackageIndex = 0; PackageIndex < CompletedPackages.Num(); ++PackageIndex)
@@ -4517,7 +4797,11 @@ FAsyncLoadingThread2::FAsyncLoadingThread2(FIoDispatcher& InIoDispatcher)
 
 	EventSpecs.AddDefaulted(EEventLoadNode2::Package_NumPhases + EEventLoadNode2::ExportBundle_NumPhases);
 	EventSpecs[EEventLoadNode2::Package_ProcessSummary] = { &FAsyncPackage2::Event_ProcessPackageSummary, &EventQueue, false };
+	EventSpecs[EEventLoadNode2::Package_SetupDependencies] = { &FAsyncPackage2::Event_SetupDependencies, &EventQueue, false };
 	EventSpecs[EEventLoadNode2::Package_ExportsSerialized] = { &FAsyncPackage2::Event_ExportsDone, &EventQueue, true };
+	EventSpecs[EEventLoadNode2::Package_BeginPostLoad] = { &FAsyncPackage2::Event_BeginPostLoad, &EventQueue, true };
+	EventSpecs[EEventLoadNode2::Package_PostLoadDone] = { &FAsyncPackage2::Event_PostLoadDone, &MainThreadEventQueue, true };
+	EventSpecs[EEventLoadNode2::Package_FinishLoading] = { &FAsyncPackage2::Event_FinishLoading, &MainThreadEventQueue, true };
 
 	EventSpecs[EEventLoadNode2::Package_NumPhases + EEventLoadNode2::ExportBundle_Process] = { &FAsyncPackage2::Event_ProcessExportBundle, &EventQueue, false };
 	EventSpecs[EEventLoadNode2::Package_NumPhases + EEventLoadNode2::ExportBundle_PostLoad] = { &FAsyncPackage2::Event_PostLoadExportBundle, &EventQueue, false };
@@ -5098,12 +5382,10 @@ FAsyncPackage2::FAsyncPackage2(
 , ImportStore(AsyncLoadingThread.GetPackageStore(), AsyncLoadingThread.GlobalImportStore, AsyncLoadingThread.LoadedPackageStore, Desc)
 {
 	TRACE_LOADTIME_NEW_ASYNC_PACKAGE(this);
+	ImportStore.AddPackageReferences(Data.ImportedPackageIds);
 	AddRequestID(Desc.RequestID);
 
-	ExportBundlesSize = Desc.StoreEntry->ExportBundlesSize;
-	LoadOrder = Desc.StoreEntry->LoadOrder;
-
-	ConstructedObjects.Reserve(Data.ExportCount + 1); // +1 for UPackage, may grow dynamically beoynd that
+	ConstructedObjects.Reserve(Data.ExportInfo.ExportCount + 1); // +1 for UPackage, may grow dynamically beoynd that
 
 	for (FExportObject& Export : Data.Exports)
 	{
@@ -5121,7 +5403,7 @@ void FAsyncPackage2::CreateNodes(const FAsyncLoadEventSpec* EventSpecs)
 		new (&Data.PackageNodes[Phase]) FEventLoadNode2(EventSpecs + Phase, this, -1, BarrierCount);
 	}
 
-	for (int32 ExportBundleIndex = 0; ExportBundleIndex < Data.ExportBundleCount; ++ExportBundleIndex)
+	for (int32 ExportBundleIndex = 0; ExportBundleIndex < Data.ExportInfo.ExportBundleCount; ++ExportBundleIndex)
 	{
 		uint32 NodeIndex = EEventLoadNode2::ExportBundle_NumPhases * ExportBundleIndex;
 		for (int32 Phase = 0; Phase < EEventLoadNode2::ExportBundle_NumPhases; ++Phase)
@@ -5136,6 +5418,8 @@ FAsyncPackage2::~FAsyncPackage2()
 	TRACE_LOADTIME_DESTROY_ASYNC_PACKAGE(this);
 	UE_ASYNC_PACKAGE_LOG(Verbose, Desc, TEXT("AsyncThread: Deleted"), TEXT("Package deleted."));
 
+	ImportStore.ReleasePackageReferences(Data.ImportedPackageIds);
+
 	checkf(RefCount == 0, TEXT("RefCount is not 0 when deleting package %s"),
 		*Desc.DiskPackageName.ToString());
 
@@ -5162,7 +5446,10 @@ void FAsyncPackage2::ClearImportedPackages()
 	TRACE_CPUPROFILER_EVENT_SCOPE(ClearImportedPackages);
 	for (FAsyncPackage2* ImportedAsyncPackage : Data.ImportedAsyncPackages)
 	{
-		ImportedAsyncPackage->ReleaseRef();
+		if (ImportedAsyncPackage)
+		{
+			ImportedAsyncPackage->ReleaseRef();
+		}
 	}
 	Data.ImportedAsyncPackages = MakeArrayView(Data.ImportedAsyncPackages.GetData(), 0);
 }
@@ -5403,7 +5690,7 @@ EAsyncPackageState::Type FAsyncPackage2::ProcessExternalReads(EExternalReadActio
 
 EAsyncPackageState::Type FAsyncPackage2::CreateClusters(FAsyncLoadingThreadState2& ThreadState)
 {
-	while (DeferredClusterIndex < Data.ExportCount && !ThreadState.IsTimeLimitExceeded(TEXT("CreateClusters")))
+	while (DeferredClusterIndex < Data.ExportInfo.ExportCount && !ThreadState.IsTimeLimitExceeded(TEXT("CreateClusters")))
 	{
 		const FExportObject& Export = Data.Exports[DeferredClusterIndex++];
 
@@ -5413,7 +5700,7 @@ EAsyncPackageState::Type FAsyncPackage2::CreateClusters(FAsyncLoadingThreadState
 		}
 	}
 
-	return DeferredClusterIndex == Data.ExportCount ? EAsyncPackageState::Complete : EAsyncPackageState::TimeOut;
+	return DeferredClusterIndex == Data.ExportInfo.ExportCount ? EAsyncPackageState::Complete : EAsyncPackageState::TimeOut;
 }
 
 void FAsyncPackage2::FinishUPackage()
@@ -5487,18 +5774,41 @@ int32 FAsyncLoadingThread2::LoadPackage(const FPackagePath& InPackagePath, FName
 	TRACE_CPUPROFILER_EVENT_SCOPE(LoadPackage);
 
 	LazyInitializeFromLoadPackage();
+	
+	const FName PackageName = InCustomName.IsNone() ? InPackagePath.GetPackageFName() : InCustomName;
+	if (FCoreDelegates::OnAsyncLoadPackage.IsBound())
+	{
+		FCoreDelegates::OnAsyncLoadPackage.Broadcast(PackageName.ToString());
+	}
+
+	// Generate new request ID and add it immediately to the global request list (it needs to be there before we exit
+	// this function, otherwise it would be added when the packages are being processed on the async thread).
+	const int32 RequestId = IAsyncPackageLoader::GetNextRequestId();
+	TRACE_LOADTIME_BEGIN_REQUEST(RequestId);
+	AddPendingRequest(RequestId);
 
-	int32 RequestID = INDEX_NONE;
+	// Allocate delegate on Game Thread, it is not safe to copy delegates by value on other threads
+	TUniquePtr<FLoadPackageAsyncDelegate> CompletionDelegate = InCompletionDelegate.IsBound()
+		? MakeUnique<FLoadPackageAsyncDelegate>(MoveTemp(InCompletionDelegate))
+		: TUniquePtr<FLoadPackageAsyncDelegate>();
 
-	// happy path where all inputs are actual package names
-	const FName Name = InCustomName.IsNone() ? InPackagePath.GetPackageFName() : InCustomName;
-	FName DiskPackageName = InPackagePath.GetPackageFName();
+	PackageRequestQueue->Enqueue(FPackageRequest::Create(RequestId, InPackagePriority, InPackagePath, InCustomName, MoveTemp(CompletionDelegate)));
+	++QueuedPackagesCounter;
+
+	AltZenaphore.NotifyOne();
+
+	return RequestId;
+}
+
+FAsyncPackageDesc2 FAsyncLoadingThread2::CreatePackageDesc(FPackageRequest& Request, bool& bIsMissingPackage)
+{
+	const FName Name = Request.GetUPackageName();
+	FName DiskPackageName = Request.GetDiskPackageName();
 	bool bHasCustomPackageName = Name != DiskPackageName;
 
-	// Verify PackageToLoadName, or fixup to handle any input string that can be converted to a long package name.
-	FPackageId DiskPackageId = FPackageId::FromName(DiskPackageName);
-	const FPackageStoreEntry* StoreEntry = PackageStore->GetPackageEntry(DiskPackageId);
-	if (!StoreEntry)
+	FPackageId DiskPackageId = Request.PackageId;
+	FPackageStoreEntryHandle StoreEntry = Request.StoreEntryHandle;
+	if (!StoreEntry.IsValid())
 	{
 		FString PackageNameStr = DiskPackageName.ToString();
 		if (!FPackageName::IsValidLongPackageName(PackageNameStr))
@@ -5508,7 +5818,7 @@ int32 FAsyncLoadingThread2::LoadPackage(const FPackagePath& InPackagePath, FName
 			{
 				DiskPackageName = *NewPackageNameStr;
 				DiskPackageId = FPackageId::FromName(DiskPackageName);
-				StoreEntry = PackageStore->GetPackageEntry(DiskPackageId);
+				StoreEntry = PackageStore->GetPackageEntryHandle(DiskPackageId);
 				bHasCustomPackageName &= Name != DiskPackageName;
 			}
 		}
@@ -5522,7 +5832,7 @@ int32 FAsyncLoadingThread2::LoadPackage(const FPackagePath& InPackagePath, FName
 	if (bHasCustomPackageName)
 	{
 		FPackageId PackageId = FPackageId::FromName(Name);
-		if (!PackageStore->GetPackageEntry(PackageId))
+		if (!PackageStore->GetPackageEntryHandle(PackageId))
 		{
 			FString PackageNameStr = Name.ToString();
 			if (FPackageName::IsValidLongPackageName(PackageNameStr))
@@ -5536,7 +5846,7 @@ int32 FAsyncLoadingThread2::LoadPackage(const FPackagePath& InPackagePath, FName
 				if (FPackageName::TryConvertFilenameToLongPackageName(PackageNameStr, NewPackageNameStr))
 				{
 					PackageId = FPackageId::FromName(FName(*NewPackageNameStr));
-					if (!PackageStore->GetPackageEntry(PackageId))
+					if (!PackageStore->GetPackageEntryHandle(PackageId))
 					{
 						CustomPackageName = *NewPackageNameStr;
 						CustomPackageId = PackageId;
@@ -5557,7 +5867,7 @@ int32 FAsyncLoadingThread2::LoadPackage(const FPackagePath& InPackagePath, FName
 	check(CustomPackageId.IsValid() == !CustomPackageName.IsNone());
 
 	bool bCustomNameIsValid = (!bHasCustomPackageName && CustomPackageName.IsNone()) || (bHasCustomPackageName && !CustomPackageName.IsNone());
-	bool bDiskPackageIdIsValid = !!StoreEntry;
+	bool bDiskPackageIdIsValid = StoreEntry.IsValid();
 	if (!bDiskPackageIdIsValid)
 	{
 		// While there is an active load request for (InName=/Temp/PackageABC_abc, InPackageToLoadFrom=/Game/PackageABC), then allow these requests too:
@@ -5578,36 +5888,11 @@ int32 FAsyncLoadingThread2::LoadPackage(const FPackagePath& InPackagePath, FName
 		}
 	}
 
-	if (bDiskPackageIdIsValid && bCustomNameIsValid)
-	{
-		if (FCoreDelegates::OnAsyncLoadPackage.IsBound())
-		{
-			FCoreDelegates::OnAsyncLoadPackage.Broadcast(Name.ToString());
-		}
-
-		// Generate new request ID and add it immediately to the global request list (it needs to be there before we exit
-		// this function, otherwise it would be added when the packages are being processed on the async thread).
-		RequestID = IAsyncPackageLoader::GetNextRequestId();
-		TRACE_LOADTIME_BEGIN_REQUEST(RequestID);
-		AddPendingRequest(RequestID);
-
-		// Allocate delegate on Game Thread, it is not safe to copy delegates by value on other threads
-		TUniquePtr<FLoadPackageAsyncDelegate> CompletionDelegatePtr;
-		if (InCompletionDelegate.IsBound())
-		{
-			CompletionDelegatePtr.Reset(new FLoadPackageAsyncDelegate(InCompletionDelegate));
-		}
-
-#if !UE_BUILD_SHIPPING
-		if (FileOpenLogWrapper)
-		{
-			FileOpenLogWrapper->AddPackageToOpenLog(*DiskPackageName.ToString());
-		}
-#endif
-
-		// Add new package request
-		FAsyncPackageDesc2 PackageDesc(RequestID, InPackagePriority, DiskPackageId, StoreEntry, DiskPackageName, CustomPackageId, CustomPackageName, MoveTemp(CompletionDelegatePtr));
+	FAsyncPackageDesc2 PackageDesc = FAsyncPackageDesc2::FromPackageRequest(Request.RequestId, Request.Priority, DiskPackageId, StoreEntry, DiskPackageName, CustomPackageId, CustomPackageName);
 
+	bool bIsValid = bDiskPackageIdIsValid && bCustomNameIsValid;
+	if (bIsValid)
+	{
 		// Fixup for redirected packages since the slim StoreEntry itself has been stripped from both package names and package ids
 		FPackageId RedirectedDiskPackageId = PackageStore->GetRedirectedPackageId(DiskPackageId);
 		if (RedirectedDiskPackageId.IsValid())
@@ -5616,46 +5901,54 @@ int32 FAsyncLoadingThread2::LoadPackage(const FPackagePath& InPackagePath, FName
 			PackageDesc.SourcePackageName = PackageDesc.DiskPackageName;
 			PackageDesc.DiskPackageName = FName();
 		}
+	}
+	
+	bIsMissingPackage = !bIsValid || StoreEntry.Status() == EPackageStoreEntryStatus::Missing;
 
-		QueuePackage(PackageDesc);
+	return MoveTemp(PackageDesc);
+}
 
-		UE_ASYNC_PACKAGE_LOG(Verbose, PackageDesc, TEXT("LoadPackage: QueuePackage"), TEXT("Package added to pending queue."));
-	}
-	else
+void FAsyncLoadingThread2::QueueMissingPackage(FAsyncPackageDesc2& PackageDesc, TUniquePtr<FLoadPackageAsyncDelegate>&& PackageLoadedDelegate)
+{
+	const bool bInvalidPackageName = PackageDesc.StoreEntry.IsValid();
+	const FName FailedPackageName = bInvalidPackageName ? PackageDesc.CustomPackageName : PackageDesc.DiskPackageName;
+
+	static TSet<FName> SkippedPackages;
+	bool bIsAlreadySkipped = false;
+
+	SkippedPackages.Add(FailedPackageName, &bIsAlreadySkipped);
+
+	if (!bIsAlreadySkipped)
 	{
-		static TSet<FName> SkippedPackages;
-		bool bIsAlreadySkipped = false;
-		if (!StoreEntry)
+		if (bInvalidPackageName)
 		{
-			SkippedPackages.Add(DiskPackageName, &bIsAlreadySkipped);
-			if (!bIsAlreadySkipped)
-			{
-				UE_LOG(LogStreaming, Warning,
-					TEXT("LoadPackage: SkipPackage: %s (0x%llX) - The package to load does not exist on disk or in the loader"),
-					*DiskPackageName.ToString(), FPackageId::FromName(DiskPackageName).ValueForDebugging());
-			}
+			UE_LOG(LogStreaming, Warning,
+				TEXT("LoadPackage: SkipPackage: %s (0x%llX) - The package name is invalid"),
+				*FailedPackageName.ToString(), PackageDesc.CustomPackageId.ValueForDebugging());
 		}
 		else
 		{
-			SkippedPackages.Add(Name, &bIsAlreadySkipped);
-			if (!bIsAlreadySkipped)
-			{
-				UE_LOG(LogStreaming, Warning,
-					TEXT("LoadPackage: SkipPackage: %s (0x%llX) - The package name is invalid"),
-					*Name.ToString(), FPackageId::FromName(Name).ValueForDebugging());
-			}
+			UE_LOG(LogStreaming, Warning,
+				TEXT("LoadPackage: SkipPackage: %s (0x%llX) - The package to load does not exist on disk or in the loader"),
+				*FailedPackageName.ToString(), PackageDesc.DiskPackageId.ValueForDebugging());
 		}
+	}
 
-		if (InCompletionDelegate.IsBound())
+	if (PackageLoadedDelegate.IsValid())
+	{
+		FScopeLock LockMissingPackages(&FailedPackageRequestsCritical);
+		FailedPackageRequests.Add(FFailedPackageRequest
 		{
-			// Queue completion callback and execute at next process loaded packages call to maintain behavior compatibility with old loader
-			FQueuedFailedPackageCallback& QueuedFailedPackageCallback = QueuedFailedPackageCallbacks.AddDefaulted_GetRef();
-			QueuedFailedPackageCallback.PackageName = Name;
-			QueuedFailedPackageCallback.Callback.Reset(new FLoadPackageAsyncDelegate(InCompletionDelegate));
-		}
+			PackageDesc.RequestID,
+			FailedPackageName,
+			MoveTemp(PackageLoadedDelegate)
+		});
+	}
+	else
+	{
+		RemovePendingRequests(TArrayView<int32>(&PackageDesc.RequestID, 1));
+		--QueuedPackagesCounter;
 	}
-
-	return RequestID;
 }
 
 EAsyncPackageState::Type FAsyncLoadingThread2::ProcessLoadingFromGameThread(FAsyncLoadingThreadState2& ThreadState, bool bUseTimeLimit, bool bUseFullTimeLimit, float TimeLimit)
@@ -5689,6 +5982,7 @@ void FAsyncLoadingThread2::FlushLoading(int32 RequestId)
 		FCoreDelegates::OnAsyncLoadingFlush.Broadcast();
 
 		double StartTime = FPlatformTime::Seconds();
+		double LogFlushTime = StartTime;
 
 		// Flush async loaders by not using a time limit. Needed for e.g. garbage collection.
 		{
@@ -5706,6 +6000,13 @@ void FAsyncLoadingThread2::FlushLoading(int32 RequestId)
 					// Update the heartbeat and sleep. If we're not multithreading, the heartbeat is updated after each package has been processed
 					FThreadHeartBeat::Get().HeartBeat();
 					FPlatformProcess::SleepNoStats(0.0001f);
+
+					// Flush logging when runing cook-on-the-fly and waiting for packages
+					if (IsRunningCookOnTheFly() && FPlatformTime::Seconds() - LogFlushTime > 1.0)
+					{
+						GLog->FlushThreadedLogs();
+						LogFlushTime = FPlatformTime::Seconds();
+					}
 				}
 
 				// push stats so that we don't overflow number of tags per thread during blocking loading
diff --git a/Engine/Source/Runtime/CoreUObject/Private/Serialization/AsyncPackageLoader.cpp b/Engine/Source/Runtime/CoreUObject/Private/Serialization/AsyncPackageLoader.cpp
--- a/Engine/Source/Runtime/CoreUObject/Private/Serialization/AsyncPackageLoader.cpp
+++ b/Engine/Source/Runtime/CoreUObject/Private/Serialization/AsyncPackageLoader.cpp
@@ -7,7 +7,7 @@
 #include "Serialization/EditorPackageLoader.h"
 #include "UObject/GCObject.h"
 #include "UObject/LinkerLoad.h"
-#include "UObject/PackageId.h"
+#include "IO/PackageId.h"
 #include "Misc/CoreDelegates.h"
 #include "Misc/PackageName.h"
 #include "Misc/PathViews.h"
diff --git a/Engine/Source/Runtime/CoreUObject/Private/Serialization/FilePackageStore.cpp b/Engine/Source/Runtime/CoreUObject/Private/Serialization/FilePackageStore.cpp
--- a/Engine/Source/Runtime/CoreUObject/Private/Serialization/FilePackageStore.cpp
+++ b/Engine/Source/Runtime/CoreUObject/Private/Serialization/FilePackageStore.cpp
@@ -1,6 +1,7 @@
 // Copyright Epic Games, Inc. All Rights Reserved.
 
-#include "Serialization/PackageStore.h"
+#include "Serialization/FilePackageStore.h"
+#include "IO/PackageStore.h"
 #include "Serialization/MappedName.h"
 #include "Serialization/AsyncLoading2.h"
 #include "Serialization/LargeMemoryReader.h"
@@ -40,14 +41,33 @@ public:
 
 	virtual bool DoesPackageExist(FPackageId PackageId) override
 	{
-		return nullptr != GetPackageEntry(PackageId);
+		FScopeLock Lock(&PackageNameMapsCritical);
+		return nullptr != StoreEntriesMap.Find(PackageId);
 	}
 
-	virtual const FPackageStoreEntry* GetPackageEntry(FPackageId PackageId) override
+	virtual FPackageStoreEntryHandle GetPackageEntryHandle(FPackageId PackageId, const FName& PackageName) override
 	{
 		FScopeLock Lock(&PackageNameMapsCritical);
-		FPackageStoreEntry* Entry = StoreEntriesMap.FindRef(PackageId);
-		return Entry;
+		const uint64 Handle = uint64(StoreEntriesMap.FindRef(PackageId));
+		return FPackageStoreEntryHandle::Create(Handle, Handle ? EPackageStoreEntryStatus::Ok : EPackageStoreEntryStatus::Missing);
+	}
+	
+	virtual FPackageStoreEntry GetPackageEntry(FPackageStoreEntryHandle Handle) override
+	{
+		const FFilePackageStoreEntry* Entry = reinterpret_cast<const FFilePackageStoreEntry*>(Handle.Value());
+
+		return FPackageStoreEntry
+		{
+			FPackageStoreExportInfo
+			{
+				Entry->ExportBundlesSize,
+				Entry->ExportCount,
+				Entry->ExportBundleCount,
+				Entry->LoadOrder
+			},
+			MakeArrayView(Entry->ImportedPackages.Data(), Entry->ImportedPackages.Num()),
+			MakeArrayView(Entry->ShaderMapHashes.Data(), Entry->ShaderMapHashes.Num())
+		};
 	}
 
 	virtual FPackageId GetRedirectedPackageId(FPackageId PackageId) override
@@ -137,32 +157,24 @@ private:
 					FContainerHeader ContainerHeader;
 					Ar << ContainerHeader;
 
-					const bool bHasContainerLocalNameMap = ContainerHeader.Names.Num() > 0;
-					if (bHasContainerLocalNameMap)
-					{
-						TRACE_CPUPROFILER_EVENT_SCOPE(LoadContainerNameMap);
-						LoadedContainer.ContainerNameMap.Reset(new FNameMap());
-						LoadedContainer.ContainerNameMap->Load(ContainerHeader.Names, ContainerHeader.NameHashes, FMappedName::EType::Container);
-					}
-
 					LoadedContainer.PackageCount = ContainerHeader.PackageCount;
 					LoadedContainer.StoreEntries = MoveTemp(ContainerHeader.StoreEntries);
 					{
 						TRACE_CPUPROFILER_EVENT_SCOPE(AddPackages);
 						FScopeLock Lock(&PackageNameMapsCritical);
 
-						TArrayView<FPackageStoreEntry> StoreEntries(reinterpret_cast<FPackageStoreEntry*>(LoadedContainer.StoreEntries.GetData()), LoadedContainer.PackageCount);
+						TArrayView<FFilePackageStoreEntry> StoreEntries(reinterpret_cast<FFilePackageStoreEntry*>(LoadedContainer.StoreEntries.GetData()), LoadedContainer.PackageCount);
 
 						int32 Index = 0;
 						StoreEntriesMap.Reserve(StoreEntriesMap.Num() + LoadedContainer.PackageCount);
-						for (FPackageStoreEntry& ContainerEntry : StoreEntries)
+						for (FFilePackageStoreEntry& StoreEntry : StoreEntries)
 						{
 							const FPackageId& PackageId = ContainerHeader.PackageIds[Index];
 
-							FPackageStoreEntry*& GlobalEntry = StoreEntriesMap.FindOrAdd(PackageId);
+							FFilePackageStoreEntry*& GlobalEntry = StoreEntriesMap.FindOrAdd(PackageId);
 							if (!GlobalEntry)
 							{
-								GlobalEntry = &ContainerEntry;
+								GlobalEntry = &StoreEntry;
 							}
 							++Index;
 						}
@@ -240,9 +252,9 @@ private:
 			const FPackageId& SourceId = It.Key();
 			const FPackageId& RedirectId = It.Value();
 			check(RedirectId.IsValid());
-			FPackageStoreEntry* RedirectEntry = StoreEntriesMap.FindRef(RedirectId);
+			FFilePackageStoreEntry* RedirectEntry = StoreEntriesMap.FindRef(RedirectId);
 			check(RedirectEntry);
-			FPackageStoreEntry*& PackageEntry = StoreEntriesMap.FindOrAdd(SourceId);
+			FFilePackageStoreEntry*& PackageEntry = StoreEntriesMap.FindOrAdd(SourceId);
 			if (RedirectEntry)
 			{
 				PackageEntry = RedirectEntry;
@@ -251,7 +263,7 @@ private:
 
 		for (auto It = StoreEntriesMap.CreateIterator(); It; ++It)
 		{
-			FPackageStoreEntry* StoreEntry = It.Value();
+			FFilePackageStoreEntry* StoreEntry = It.Value();
 
 			for (FPackageId& ImportedPackageId : StoreEntry->ImportedPackages)
 			{
@@ -266,7 +278,7 @@ private:
 	struct FLoadedContainer
 	{
 		TUniquePtr<FNameMap> ContainerNameMap;
-		TArray<uint8> StoreEntries; //FPackageStoreEntry[PackageCount];
+		TArray<uint8> StoreEntries; //FFilePackageStoreEntry[PackageCount];
 		uint32 PackageCount = 0;
 		bool bValid = false;
 	};
@@ -278,7 +290,7 @@ private:
 
 	FCriticalSection PackageNameMapsCritical;
 
-	TMap<FPackageId, FPackageStoreEntry*> StoreEntriesMap;
+	TMap<FPackageId, FFilePackageStoreEntry*> StoreEntriesMap;
 	TMap<FPackageId, FPackageId> RedirectsPackageMap;
 	TSet<FPackageId> TargetRedirectIds;
 	int32 NextCustomPackageIndex = 0;
diff --git a/Engine/Source/Runtime/CoreUObject/Private/Serialization/MappedName.cpp b/Engine/Source/Runtime/CoreUObject/Private/Serialization/MappedName.cpp
--- a/Engine/Source/Runtime/CoreUObject/Private/Serialization/MappedName.cpp
+++ b/Engine/Source/Runtime/CoreUObject/Private/Serialization/MappedName.cpp
@@ -11,8 +11,8 @@ FArchive& operator<<(FArchive& Ar, FMappedName& MappedName)
 	return Ar;
 }
 
-void FNameMap::Load(TArrayView<const uint8> NameBuffer, TArrayView<const uint8> HashBuffer, FMappedName::EType InNameMapType)
+void FNameMap::Load(FArchive& Ar, FMappedName::EType InNameMapType)
 {
-	LoadNameBatch(NameEntries, NameBuffer, HashBuffer);
+	NameEntries = LoadNameBatch(Ar);
 	NameMapType = InNameMapType;
 }
diff --git a/Engine/Source/Runtime/CoreUObject/Private/UObject/Linker.cpp b/Engine/Source/Runtime/CoreUObject/Private/UObject/Linker.cpp
--- a/Engine/Source/Runtime/CoreUObject/Private/UObject/Linker.cpp
+++ b/Engine/Source/Runtime/CoreUObject/Private/UObject/Linker.cpp
@@ -9,10 +9,12 @@
 #include "Misc/PackageName.h"
 #include "Misc/CommandLine.h"
 #include "UObject/Package.h"
+#include "UObject/SavePackage.h"
 #include "Templates/Casts.h"
 #include "UObject/UnrealType.h"
 #include "UObject/LinkerLoad.h"
 #include "Misc/SecureHash.h"
+#include "Internationalization/GatherableTextData.h"
 #include "Logging/TokenizedMessage.h"
 #include "Logging/MessageLog.h"
 #include "Misc/UObjectToken.h"
diff --git a/Engine/Source/Runtime/CoreUObject/Private/UObject/SavePackage.cpp b/Engine/Source/Runtime/CoreUObject/Private/UObject/SavePackage.cpp
--- a/Engine/Source/Runtime/CoreUObject/Private/UObject/SavePackage.cpp
+++ b/Engine/Source/Runtime/CoreUObject/Private/UObject/SavePackage.cpp
@@ -68,6 +68,7 @@
 #include "Misc/ScopeExit.h"
 #include "Misc/PackageAccessTracking.h"
 #include "Misc/PackageAccessTrackingOps.h"
+#include "IO/PackageStoreWriter.h"
 
 #if ENABLE_COOK_STATS
 #include "ProfilingDebugging/ScopedTimers.h"
@@ -4515,10 +4516,13 @@ FSavePackageResultStruct UPackage::Save(UPackage* InOuter, UObject* Base, EObjec
 									PackageInfo.LooseFilePath	= Filename;
 									PackageInfo.HeaderSize		= HeaderSize;
 
-									SavePackageContext->PackageStoreWriter->WritePackage(PackageInfo, IoBuffer, Linker->FileRegions);
+									FPackageId PackageId = FPackageId::FromName(PackageInfo.PackageName);
+									PackageInfo.ChunkId = CreateIoChunkId(PackageId.Value(), 0, EIoChunkType::ExportBundleData);
+
+									SavePackageContext->PackageStoreWriter->WritePackageData(PackageInfo, IoBuffer, Linker->FileRegions);
 								}
 								else
-								{
+								{		
 									EAsyncWriteOptions WriteOptions(EAsyncWriteOptions::WriteFileToDisk);
 									if (bComputeHash)
 									{
@@ -4764,12 +4768,6 @@ bool UPackage::SavePackage(UPackage* InOuter, UObject* Base, EObjectFlags TopLev
 	return Result == ESavePackageResult::Success;
 }
 
-//////////////////////////////////////////////////////////////////////////
-// TODO: this should go elsewhere, this file is big enough as it is already
-//
-
-IPackageStoreWriter::~IPackageStoreWriter() = default;
-
 FSavePackageContext::~FSavePackageContext()
 {
 	delete PackageStoreWriter;
diff --git a/Engine/Source/Runtime/CoreUObject/Private/UObject/SavePackage/SavePackageUtilities.cpp b/Engine/Source/Runtime/CoreUObject/Private/UObject/SavePackage/SavePackageUtilities.cpp
--- a/Engine/Source/Runtime/CoreUObject/Private/UObject/SavePackage/SavePackageUtilities.cpp
+++ b/Engine/Source/Runtime/CoreUObject/Private/UObject/SavePackage/SavePackageUtilities.cpp
@@ -31,6 +31,7 @@
 #include "UObject/UObjectHash.h"
 #include "UObject/UObjectThreadContext.h"
 #include "Virtualization/VirtualizedBulkData.h"
+#include "IO/PackageStoreWriter.h"
 
 DEFINE_LOG_CATEGORY(LogSavePackage);
 UE_TRACE_CHANNEL_DEFINE(SaveTimeChannel);
@@ -2485,21 +2486,26 @@ ESavePackageResult SaveBulkData(FLinkerSave* Linker, const UPackage* InOuter, co
 
 			IPackageStoreWriter::FBulkDataInfo BulkInfo;
 			BulkInfo.PackageName = InOuter->GetFName();
+
+			FPackageId PackageId = FPackageId::FromName(BulkInfo.PackageName);
 				
 			if (BulkArchive->TotalSize())
 			{
+				BulkInfo.ChunkId = CreateIoChunkId(PackageId.Value(), 0, EIoChunkType::BulkData);
 				BulkInfo.BulkdataType = IPackageStoreWriter::FBulkDataInfo::Standard;
 				BulkInfo.LooseFilePath = FPaths::ChangeExtension(Filename, LexToString(EPackageExtension::BulkDataDefault));
 				SavePackageContext->PackageStoreWriter->WriteBulkdata(BulkInfo, AddSizeAndConvertToIoBuffer(BulkArchive.Get()), BulkArchive->FileRegions);
 			}
 			if (OptionalBulkArchive->TotalSize())
 			{
+				BulkInfo.ChunkId = CreateIoChunkId(PackageId.Value(), 0, EIoChunkType::OptionalBulkData);
 				BulkInfo.BulkdataType = IPackageStoreWriter::FBulkDataInfo::Optional;
 				BulkInfo.LooseFilePath = FPaths::ChangeExtension(Filename, LexToString(EPackageExtension::BulkDataOptional));
 				SavePackageContext->PackageStoreWriter->WriteBulkdata(BulkInfo, AddSizeAndConvertToIoBuffer(OptionalBulkArchive.Get()), OptionalBulkArchive->FileRegions);
 			}
 			if (MappedBulkArchive->TotalSize())
 			{
+				BulkInfo.ChunkId = CreateIoChunkId(PackageId.Value(), 0, EIoChunkType::MemoryMappedBulkData);
 				BulkInfo.BulkdataType = IPackageStoreWriter::FBulkDataInfo::Mmap;
 				BulkInfo.LooseFilePath = FPaths::ChangeExtension(Filename, LexToString(EPackageExtension::BulkDataMemoryMapped));
 				SavePackageContext->PackageStoreWriter->WriteBulkdata(BulkInfo, AddSizeAndConvertToIoBuffer(MappedBulkArchive.Get()), MappedBulkArchive->FileRegions);
diff --git a/Engine/Source/Runtime/CoreUObject/Private/UObject/SavePackage2.cpp b/Engine/Source/Runtime/CoreUObject/Private/UObject/SavePackage2.cpp
--- a/Engine/Source/Runtime/CoreUObject/Private/UObject/SavePackage2.cpp
+++ b/Engine/Source/Runtime/CoreUObject/Private/UObject/SavePackage2.cpp
@@ -39,6 +39,7 @@
 #include "UObject/SavePackage/SavePackageUtilities.h"
 #include "UObject/UObjectGlobals.h"
 #include "UObject/UObjectHash.h"
+#include "IO/PackageStoreWriter.h"
 
 #if ENABLE_COOK_STATS
 #include "ProfilingDebugging/ScopedTimers.h"
@@ -1632,15 +1633,30 @@ ESavePackageResult WriteExports(FStructuredArchive::FRecord& StructuredArchiveRo
 
 ESavePackageResult WriteAdditionalExportFiles(FSaveContext& SaveContext)
 {
+	FSavePackageContext* SavePackageContext = SaveContext.GetSavePackageContext();
+	IPackageStoreWriter* PackageStoreWriter = SavePackageContext ? SavePackageContext->PackageStoreWriter : nullptr;
+
 	if (SaveContext.IsCooking() && SaveContext.AdditionalFilesFromExports.Num() > 0)
 	{
-		const bool bWriteFileToDisk = !SaveContext.IsDiffing();
+		bool bWriteFileToDisk = !SaveContext.IsDiffing();
 		const bool bComputeHash = SaveContext.IsComputeHash();
 		for (FLargeMemoryWriter& Writer : SaveContext.AdditionalFilesFromExports)
 		{
 			const int64 Size = Writer.TotalSize();
 			SaveContext.TotalPackageSizeUncompressed += Size;
 
+			if (PackageStoreWriter)
+			{
+				IPackageStoreWriter::FAdditionalFileInfo FileInfo;
+				FileInfo.PackageName = SaveContext.GetPackage()->GetFName();
+				FileInfo.Filename = *Writer.GetArchiveName();
+
+				FIoBuffer FileData(FIoBuffer::Wrap, Writer.GetData(), Size);
+
+				const bool bWrittenToPackageStore = PackageStoreWriter->WriteAdditionalFile(FileInfo, FileData);
+				bWriteFileToDisk &= !bWrittenToPackageStore;
+			}
+
 			if (bComputeHash || bWriteFileToDisk)
 			{
 				FLargeMemoryPtr DataPtr(Writer.ReleaseOwnership());
@@ -1836,7 +1852,10 @@ ESavePackageResult FinalizeFile(FStructuredArchive::FRecord& StructuredArchiveRo
 				PackageInfo.LooseFilePath = SaveContext.GetFilename();
 				PackageInfo.HeaderSize = HeaderSize;
 
-				SavePackageContext->PackageStoreWriter->WritePackage(PackageInfo, IoBuffer, Linker->FileRegions);
+				FPackageId PackageId = FPackageId::FromName(PackageInfo.PackageName);
+				PackageInfo.ChunkId = CreateIoChunkId(PackageId.Value(), 0, EIoChunkType::ExportBundleData);
+
+				SavePackageContext->PackageStoreWriter->WritePackageData(PackageInfo, IoBuffer, Linker->FileRegions);
 			}
 			else
 			{
diff --git a/Engine/Source/Runtime/CoreUObject/Private/UObject/UObjectBase.cpp b/Engine/Source/Runtime/CoreUObject/Private/UObject/UObjectBase.cpp
--- a/Engine/Source/Runtime/CoreUObject/Private/UObject/UObjectBase.cpp
+++ b/Engine/Source/Runtime/CoreUObject/Private/UObject/UObjectBase.cpp
@@ -1044,9 +1044,7 @@ void UObjectBaseInit()
 	// FPlatformProperties::RequiresCookedData() will be false. Please note that GIsEditor and FApp::IsGame() are not valid at this point.
 	if (FPlatformProperties::RequiresCookedData())
 	{
-		FString Value;
-		bool bIsCookOnTheFly = FParse::Value(FCommandLine::Get(), TEXT("-filehostip="), Value);
-		if (bIsCookOnTheFly)
+		if (IsRunningCookOnTheFly())
 		{
 			GCreateGCClusters = false;
 		}
diff --git a/Engine/Source/Runtime/CoreUObject/Public/Serialization/AsyncLoading2.h b/Engine/Source/Runtime/CoreUObject/Public/Serialization/AsyncLoading2.h
--- a/Engine/Source/Runtime/CoreUObject/Public/Serialization/AsyncLoading2.h
+++ b/Engine/Source/Runtime/CoreUObject/Public/Serialization/AsyncLoading2.h
@@ -8,7 +8,7 @@
 
 #include "CoreMinimal.h"
 #include "UObject/ObjectResource.h"
-#include "UObject/PackageId.h"
+#include "IO/PackageId.h"
 #include "Serialization/Archive.h"
 #include "Serialization/MappedName.h"
 #include "IO/IoContainerId.h"
@@ -25,8 +25,6 @@ struct FContainerHeader
 {
 	FIoContainerId ContainerId;
 	uint32 PackageCount = 0;
-	TArray<uint8> Names;
-	TArray<uint8> NameHashes;
 	TArray<FPackageId> PackageIds;
 	TArray<uint8> StoreEntries; //FPackageStoreEntry[PackageCount]
 	FCulturePackageMap CulturePackageMap;
@@ -150,19 +148,15 @@ enum class EExportFilterFlags : uint8
  */
 struct FPackageSummary
 {
+	uint32 HeaderSize;
 	FMappedName Name;
 	FMappedName SourceName;
 	uint32 PackageFlags;
 	uint32 CookedHeaderSize;
-	int32 NameMapNamesOffset;
-	int32 NameMapNamesSize;
-	int32 NameMapHashesOffset;
-	int32 NameMapHashesSize;
 	int32 ImportMapOffset;
 	int32 ExportMapOffset;
-	int32 ExportBundlesOffset;
+	int32 ExportBundleEntriesOffset;
 	int32 GraphDataOffset;
-	int32 GraphDataSize;
 	int32 Pad = 0;
 };
 
@@ -188,6 +182,7 @@ struct FExportBundleEntry
  */
 struct FExportBundleHeader
 {
+	uint64 SerialOffset;
 	uint32 FirstEntryIndex;
 	uint32 EntryCount;
 
diff --git a/Engine/Source/Runtime/CoreUObject/Public/Serialization/BulkData.h b/Engine/Source/Runtime/CoreUObject/Public/Serialization/BulkData.h
--- a/Engine/Source/Runtime/CoreUObject/Public/Serialization/BulkData.h
+++ b/Engine/Source/Runtime/CoreUObject/Public/Serialization/BulkData.h
@@ -13,7 +13,7 @@
 #include "BulkDataCommon.h"
 #include "BulkData2.h"
 #if WITH_IOSTORE_IN_EDITOR
-#include "UObject/PackageId.h"
+#include "IO/PackageId.h"
 #endif
 
 #if WITH_EDITOR == 0 && WITH_EDITORONLY_DATA == 0 //Runtime
diff --git a/Engine/Source/Runtime/CoreUObject/Public/Serialization/CookOnTheFlyPackageStore.h b/Engine/Source/Runtime/CoreUObject/Public/Serialization/CookOnTheFlyPackageStore.h
new file mode 100644
--- /dev/null
+++ b/Engine/Source/Runtime/CoreUObject/Public/Serialization/CookOnTheFlyPackageStore.h
@@ -0,0 +1,67 @@
+// Copyright Epic Games, Inc. All Rights Reserved.
+
+#pragma once
+
+#include "CoreMinimal.h"
+#include "IO/PackageStore.h"
+
+class FIoDispatcher;
+
+namespace UE { namespace Cook
+{
+	class ICookOnTheFlyServerConnection;
+}}
+
+namespace UE { namespace PackageStore { namespace Messaging
+{
+
+struct FPackageStoreData
+{
+	TArray<FPackageStoreEntryResource> CookedPackages;
+	TArray<FPackageId> FailedPackages;
+	int32 TotalCookedPackages = 0;
+	int32 TotalFailedPackages = 0;
+
+	COREUOBJECT_API friend FArchive& operator<<(FArchive& Ar, FPackageStoreData& PackageStoreData);
+};
+
+struct FCookPackageRequest
+{
+	FName PackageName;
+	
+	COREUOBJECT_API friend FArchive& operator<<(FArchive& Ar, FCookPackageRequest& Request);
+};
+
+struct FCookPackageResponse
+{
+	EPackageStoreEntryStatus Status = EPackageStoreEntryStatus::None;
+
+	COREUOBJECT_API friend FArchive& operator<<(FArchive& Ar, FCookPackageResponse& Response);
+};
+
+struct FPackagesCookedMessage
+{
+	FPackageStoreData PackageStoreData;
+
+	friend FArchive& operator<<(FArchive& Ar, FPackagesCookedMessage& PackagesCookedMessage)
+	{
+		return Ar << PackagesCookedMessage.PackageStoreData;
+	}
+};
+
+struct FGetCookedPackagesResponse
+{
+	FPackageStoreData PackageStoreData;
+	
+	friend FArchive& operator<<(FArchive& Ar, FGetCookedPackagesResponse& GetCookedPackagesResponse)
+	{
+		return Ar << GetCookedPackagesResponse.PackageStoreData;
+	}
+};
+
+}}} // namesapce UE::PackageStore::Messaging
+
+TUniquePtr<IPackageStore> MakeCookOnTheFlyPackageStore(
+	FIoDispatcher& IoDispatcher,
+	UE::Cook::ICookOnTheFlyServerConnection& CookOnTheFlyServerConnection,
+	TFunction<void()>&& EntriesAddedCallback);
diff --git a/Engine/Source/Runtime/CoreUObject/Public/Serialization/FilePackageStore.h b/Engine/Source/Runtime/CoreUObject/Public/Serialization/FilePackageStore.h
new file mode 100644
--- /dev/null
+++ b/Engine/Source/Runtime/CoreUObject/Public/Serialization/FilePackageStore.h
@@ -0,0 +1,50 @@
+// Copyright Epic Games, Inc. All Rights Reserved.
+
+#pragma once
+
+#include "CoreMinimal.h"
+#include "IO/PackageId.h"
+
+class IPackageStore;
+class FIoDispatcher;
+
+/**
+ * Package store entry array view.
+ */
+template<typename T>
+class TFilePackageStoreEntryCArrayView
+{
+	const uint32 ArrayNum = 0;
+	const uint32 OffsetToDataFromThis = 0;
+
+public:
+	inline uint32 Num() const						{ return ArrayNum; }
+
+	inline const T* Data() const					{ return (T*)((char*)this + OffsetToDataFromThis); }
+	inline T* Data()								{ return (T*)((char*)this + OffsetToDataFromThis); }
+
+	inline const T* begin() const					{ return Data(); }
+	inline T* begin()								{ return Data(); }
+
+	inline const T* end() const						{ return Data() + ArrayNum; }
+	inline T* end()									{ return Data() + ArrayNum; }
+
+	inline const T& operator[](uint32 Index) const	{ return Data()[Index]; }
+	inline T& operator[](uint32 Index)				{ return Data()[Index]; }
+};
+
+/**
+ * File based package store entry
+ */
+struct FFilePackageStoreEntry
+{
+	uint64 ExportBundlesSize;
+	int32 ExportCount;
+	int32 ExportBundleCount;
+	uint32 LoadOrder;
+	uint32 Pad;
+	TFilePackageStoreEntryCArrayView<FPackageId> ImportedPackages;
+	TFilePackageStoreEntryCArrayView<FSHAHash> ShaderMapHashes;
+};
+
+TUniquePtr<IPackageStore> MakeFilePackageStore(FIoDispatcher& IoDispatcher);
diff --git a/Engine/Source/Runtime/CoreUObject/Public/Serialization/MappedName.h b/Engine/Source/Runtime/CoreUObject/Public/Serialization/MappedName.h
--- a/Engine/Source/Runtime/CoreUObject/Public/Serialization/MappedName.h
+++ b/Engine/Source/Runtime/CoreUObject/Public/Serialization/MappedName.h
@@ -105,7 +105,7 @@ public:
 		return NameEntries.Num();
 	}
 
-	void Load(TArrayView<const uint8> NameBuffer, TArrayView<const uint8> HashBuffer, FMappedName::EType NameMapType);
+	void Load(FArchive& Ar, FMappedName::EType NameMapType);
 
 	FName GetName(const FMappedName& MappedName) const
 	{
diff --git a/Engine/Source/Runtime/CoreUObject/Public/UObject/Linker.h b/Engine/Source/Runtime/CoreUObject/Public/UObject/Linker.h
--- a/Engine/Source/Runtime/CoreUObject/Public/UObject/Linker.h
+++ b/Engine/Source/Runtime/CoreUObject/Public/UObject/Linker.h
@@ -4,14 +4,14 @@
 
 #include "CoreMinimal.h"
 #include "UObject/ObjectResource.h"
-#include "Internationalization/GatherableTextData.h"
 #include "UObject/PackageFileSummary.h"
 #include "UObject/LinkerInstancingContext.h"
-#include "UObject/SavePackage.h"
 #include "Templates/RefCounting.h"
 
 class FReferenceCollector;
 class FPackagePath;
+struct FPackageSaveInfo;
+struct FGatherableTextData;
 
 DECLARE_LOG_CATEGORY_EXTERN(LogLinker, Log, All);
 
@@ -623,7 +623,7 @@ COREUOBJECT_API void ResetLinkerExports(UPackage* InPackage);
  * Can be called after the package has finished loading.
  * Flushes async loading.
  */
-COREUOBJECT_API void ResetLoaders( UObject* InOuter );
+COREUOBJECT_API void ResetLoaders(UObject* InOuter);
 
 /** Deletes all linkers that have finished loading */
 COREUOBJECT_API void DeleteLoaders();
