Launch on to platform crashes on open Assertion failed: IsAligned(PropertyDest, TheCppStructOps->GetAlignment())
#fix actually allow for 32-byte aligned UObject allocations, allocate by size and aligment rather than aligned size and added development/debug build verification of resulting mem adress

diff --git a/Engine/Source/Runtime/CoreUObject/Private/UObject/UObjectAllocator.cpp b/Engine/Source/Runtime/CoreUObject/Private/UObject/UObjectAllocator.cpp
--- a/Engine/Source/Runtime/CoreUObject/Private/UObject/UObjectAllocator.cpp
+++ b/Engine/Source/Runtime/CoreUObject/Private/UObject/UObjectAllocator.cpp
@@ -51,10 +51,10 @@ void FUObjectAllocator::BootMessage()
  */
 UObjectBase* FUObjectAllocator::AllocateUObject(int32 Size, int32 Alignment, bool bAllowPermanent)
 {
-	// Force alignment to 16 bytes
-	Alignment = 16;
+	// Force alignment to minimal of 16 bytes
+	Alignment = FMath::Max(16, Alignment);
 	int32 AlignedSize = Align( Size, Alignment );
-	UObjectBase* Result = NULL;
+	UObjectBase* Result = nullptr;
 
 	const bool bPlaceInPerm = bAllowPermanent && (Align(PermanentObjectPoolTail,Alignment) + Size) <= (PermanentObjectPool + PermanentObjectPoolSize);
 	if (bAllowPermanent && !bPlaceInPerm)
@@ -80,8 +80,13 @@ UObjectBase* FUObjectAllocator::AllocateUObject(int32 Size, int32 Alignment, boo
 	else
 	{
 		// Allocate new memory of the appropriate size and alignment.
-		Result = (UObjectBase*)FMemory::Malloc( AlignedSize );
+		Result = (UObjectBase*)FMemory::Malloc( Size, Alignment );
 	}
+
+#if !UE_BUILD_TEST && !UE_BUILD_SHIPPING
+	checkf(IsAligned(Result, Alignment), TEXT("Allocated memory address does not match requirement of %d byte alignment for size %d"), Alignment, Size);
+#endif
+
 	return Result;
 }
 
