Prevent -inf values from showing up for the MBps entries in Mirage stats logging code.

- If the time for the selection is 0.0 then just use 0.0 as the result, don't divide.

diff --git a/Engine/Source/Runtime/CoreUObject/Private/Virtualization/VirtualizationManager.cpp b/Engine/Source/Runtime/CoreUObject/Private/Virtualization/VirtualizationManager.cpp
--- a/Engine/Source/Runtime/CoreUObject/Private/Virtualization/VirtualizationManager.cpp
+++ b/Engine/Source/Runtime/CoreUObject/Private/Virtualization/VirtualizationManager.cpp
@@ -143,13 +143,13 @@ namespace Profiling
 			{
 				const double Time = Iterator.Value.GetAccumulatedValueAnyThread(FCookStats::CallStats::EHitOrMiss::Hit, FCookStats::CallStats::EStatType::Cycles) * FPlatformTime::GetSecondsPerCycle();
 				const int64 DataSizeMB = Iterator.Value.GetAccumulatedValueAnyThread(FCookStats::CallStats::EHitOrMiss::Hit, FCookStats::CallStats::EStatType::Bytes) / (1024 * 1024);
-				const double MBPerSecond = DataSizeMB / Time;
+				const double MBps = Time != 0.0 ? (DataSizeMB / Time) : 0.0;
 
 				UE_LOG(LogVirtualization, Log, TEXT("%-40.40s|%17" UINT64_FMT "|%12.3f|%14.3f|"),
 					*Iterator.Key,
 					DataSizeMB,
 					Time,
-					MBPerSecond);
+					MBps);
 			}
 		}
 
@@ -161,13 +161,13 @@ namespace Profiling
 			{
 				const double Time = Iterator.Value.GetAccumulatedValueAnyThread(FCookStats::CallStats::EHitOrMiss::Hit, FCookStats::CallStats::EStatType::Cycles) * FPlatformTime::GetSecondsPerCycle();
 				const int64 DataSizeMB = Iterator.Value.GetAccumulatedValueAnyThread(FCookStats::CallStats::EHitOrMiss::Hit, FCookStats::CallStats::EStatType::Bytes) / (1024 * 1024);
-				const double MBPerSecond = DataSizeMB / Time;
+				const double MBps = Time != 0.0 ? (DataSizeMB / Time) : 0.0;
 
 				UE_LOG(LogVirtualization, Log, TEXT("%-40.40s|%17" UINT64_FMT "|%12.3f|%14.3f|"),
 					*Iterator.Key,
 					DataSizeMB,
 					Time,
-					MBPerSecond);
+					MBps);
 			}
 		}
 	}
