Matinee: Misc cleanup (code comment changes, etc.)

diff --git a/Engine/Config/BaseEditorPerProjectUserSettings.ini b/Engine/Config/BaseEditorPerProjectUserSettings.ini
--- a/Engine/Config/BaseEditorPerProjectUserSettings.ini
+++ b/Engine/Config/BaseEditorPerProjectUserSettings.ini
@@ -501,27 +501,6 @@ QualityLevel=0
 NumUnusedLocalCores=1
 ShowLightingBuildInfo=false
 
-
-[MatineeCreateMovieOptions]
-CloseEditor=false
-CaptureResolutionIndex = 0;
-CaptureResolutionFPS = 30;
-CaptureTypeIndex = 0;
-Compress=false
-CinematicMode=true
-DisableMovement=true
-DisableTurning=true
-HidePlayer=true
-DisableInput=true
-HideHUD=true
-
-
-[Matinee]
-Hide3DTracks=false
-ZoomToScrubPos=false
-ShowCurveEd=false
-
-
 [/Script/UnrealEd.PhysicsAssetEditorOptions]
 AngularSnap=15.0
 LinearSnap=2.0
diff --git a/Engine/Plugins/Runtime/ApexDestruction/Source/ApexDestructionEditor/Private/SDestructibleMeshEditorViewport.cpp b/Engine/Plugins/Runtime/ApexDestruction/Source/ApexDestructionEditor/Private/SDestructibleMeshEditorViewport.cpp
--- a/Engine/Plugins/Runtime/ApexDestruction/Source/ApexDestructionEditor/Private/SDestructibleMeshEditorViewport.cpp
+++ b/Engine/Plugins/Runtime/ApexDestruction/Source/ApexDestructionEditor/Private/SDestructibleMeshEditorViewport.cpp
@@ -421,7 +421,7 @@ void FDestructibleMeshEditorViewportClient::ImportFBXChunks()
 	{
 		bOpened = DesktopPlatform->OpenFileDialog(
 			NULL, 
-			NSLOCTEXT("UnrealEd", "ImportMatineeSequence", "Import UnrealMatinee Sequence").ToString(),
+			NSLOCTEXT("UnrealEd", "ImportDestructableMeshFBXSequence", "Import FBX File").ToString(),
 			*(FEditorDirectories::Get().GetLastDirectory(ELastDirectory::GENERIC_IMPORT)),
 			TEXT(""),
 			TEXT("FBX document|*.fbx"),
diff --git a/Engine/Shaders/Private/GammaCorrection.usf b/Engine/Shaders/Private/GammaCorrection.usf
--- a/Engine/Shaders/Private/GammaCorrection.usf
+++ b/Engine/Shaders/Private/GammaCorrection.usf
@@ -45,7 +45,7 @@ void MainPS(float2 UV : TEXCOORD0, out float4 OutColor : SV_Target0)
 	half4 LinearColor = Texture2DSample(SceneColorTexture, SceneColorTextureSampler, UV);
 	half3 LDRColor = TonemapAndGammaCorrect(LinearColor.rgb);
 	
-	// blend with custom LDR color, used for Fade track in Matinee
+	// blend with custom LDR color, used for Fade track in Cinematics
 	LDRColor = lerp(LDRColor * ColorScale, OverlayColor.rgb, OverlayColor.a);
 	
 	// RETURN_COLOR not needed unless writing to SceneColor
diff --git a/Engine/Shaders/Private/PostProcessCombineLUTs.usf b/Engine/Shaders/Private/PostProcessCombineLUTs.usf
--- a/Engine/Shaders/Private/PostProcessCombineLUTs.usf
+++ b/Engine/Shaders/Private/PostProcessCombineLUTs.usf
@@ -369,7 +369,7 @@ float4 CombineLUTsCommon(float2 InUV, uint InLayerIndex)
 	// apply math color correction on top to texture based solution
 	FilmColor = ColorCorrection( FilmColor );
 
-	// blend with custom LDR color, used for Fade track in Matinee
+	// blend with custom LDR color, used for Fade track in Cinematics
 	float3 FilmColorNoGamma = lerp( FilmColor * ColorScale, OverlayColor.rgb, OverlayColor.a );
 	// Apply Fade track to linear outputs also
 	GradedColor = lerp(GradedColor * ColorScale, OverlayColor.rgb, OverlayColor.a);
diff --git a/Engine/Source/Runtime/ApplicationCore/Private/Linux/LinuxApplication.cpp b/Engine/Source/Runtime/ApplicationCore/Private/Linux/LinuxApplication.cpp
--- a/Engine/Source/Runtime/ApplicationCore/Private/Linux/LinuxApplication.cpp
+++ b/Engine/Source/Runtime/ApplicationCore/Private/Linux/LinuxApplication.cpp
@@ -1597,7 +1597,7 @@ void FLinuxApplication::RemoveRevertFocusWindow(SDL_HWindow HWnd)
 			{
 				ActivateWindow(Window->GetParent());
 			}
-			// Was the deleted window a Blueprint, Cascade, Matinee etc. window?
+			// Was the deleted window a Blueprint, Cascade etc. window?
 			else if (Window->IsNotificationWindow())
 			{
 				// Do not revert focus if the root window of the destroyed window is another one.
diff --git a/Engine/Source/Runtime/CoreUObject/Public/UObject/ObjectMacros.h b/Engine/Source/Runtime/CoreUObject/Public/UObject/ObjectMacros.h
--- a/Engine/Source/Runtime/CoreUObject/Public/UObject/ObjectMacros.h
+++ b/Engine/Source/Runtime/CoreUObject/Public/UObject/ObjectMacros.h
@@ -403,7 +403,7 @@ enum EPropertyFlags : uint64
 	CPF_IsPlainOldData					= 0x0000000040000000,	///< If this is set, then the property can be memcopied instead of CopyCompleteValue / CopySingleValue
 	CPF_RepSkip							= 0x0000000080000000,	///< Not replicated. For non replicated properties in replicated structs 
 	CPF_RepNotify						= 0x0000000100000000,	///< Notify actors when a property is replicated
-	CPF_Interp							= 0x0000000200000000,	///< interpolatable property for use with matinee
+	CPF_Interp							= 0x0000000200000000,	///< interpolatable property for use with cinematics
 	CPF_NonTransactional				= 0x0000000400000000,	///< Property isn't transacted
 	CPF_EditorOnly						= 0x0000000800000000,	///< Property should only be loaded in the editor
 	CPF_NoDestructor					= 0x0000001000000000,	///< No destructor
@@ -943,7 +943,7 @@ namespace UP
 		/// Skip replication (only for struct members and parameters in service request functions).
 		NotReplicated,
 
-		/// Interpolatable property for use with matinee. Always user-settable in the editor.
+		/// Interpolatable property for use with cinematics. Always user-settable in the editor.
 		Interp,
 
 		/// Property isn't transacted.
diff --git a/Engine/Source/Runtime/Engine/Public/GameDelegates.h b/Engine/Source/Runtime/Engine/Public/GameDelegates.h
--- a/Engine/Source/Runtime/Engine/Public/GameDelegates.h
+++ b/Engine/Source/Runtime/Engine/Public/GameDelegates.h
@@ -68,9 +68,6 @@ public:
 	// Called when ending playing a map
 	DEFINE_GAME_DELEGATE_TYPED(EndPlayMapDelegate, FSimpleMulticastDelegate);
 
-	// Called when a matinee is canceled 
-	DEFINE_GAME_DELEGATE_TYPED(MatineeCancelledDelegate, FSimpleMulticastDelegate);
-
 	// Called when a pending connection has been lost 
 	DEFINE_GAME_DELEGATE_TYPED(PendingConnectionLostDelegate, FPendingConnectionLostDelegate);
 
diff --git a/Engine/Source/Runtime/Launch/Private/LaunchEngineLoop.cpp b/Engine/Source/Runtime/Launch/Private/LaunchEngineLoop.cpp
--- a/Engine/Source/Runtime/Launch/Private/LaunchEngineLoop.cpp
+++ b/Engine/Source/Runtime/Launch/Private/LaunchEngineLoop.cpp
@@ -3777,9 +3777,7 @@ int32 FEngineLoop::PreInitPostStartupScreen(const TCHAR* CmdLine)
 
 	FEmbeddedCommunication::ForceTick(8);
 
-	FString MatineeName;
-
-	if(FParse::Param(FCommandLine::Get(),TEXT("DUMPMOVIE")) || FParse::Value(FCommandLine::Get(), TEXT("-MATINEESSCAPTURE="), MatineeName))
+	if(FParse::Param(FCommandLine::Get(),TEXT("DUMPMOVIE")))
 	{
 		// -1: remain on
 		GIsDumpingMovie = -1;
diff --git a/Engine/Source/Runtime/Slate/Private/Framework/Application/SlateApplication.cpp b/Engine/Source/Runtime/Slate/Private/Framework/Application/SlateApplication.cpp
--- a/Engine/Source/Runtime/Slate/Private/Framework/Application/SlateApplication.cpp
+++ b/Engine/Source/Runtime/Slate/Private/Framework/Application/SlateApplication.cpp
@@ -1639,8 +1639,7 @@ void FSlateApplication::ThrottleApplicationBasedOnMouseMovement()
 		const double TimeSinceLastMouseMove = CurrentTime - LastMouseMoveTime;
 		if( TimeSinceLastMouseMove < TimeToThrottleAfterMouseStops )
 		{
-			// Only throttle if a Slate window is currently active.  If a Wx window (such as Matinee) is
-			// being used, we don't want to throttle
+			// Only throttle if a Slate window is currently active.
 			if( this->GetActiveTopLevelWindow().IsValid() )
 			{
 				// Only throttle if the user hasn't pressed a button in awhile
