Remove dependency to CookOnTheFly module from CoreUObject and Launch when compiling programs

diff --git a/Engine/Source/Runtime/Core/Public/Misc/Build.h b/Engine/Source/Runtime/Core/Public/Misc/Build.h
--- a/Engine/Source/Runtime/Core/Public/Misc/Build.h
+++ b/Engine/Source/Runtime/Core/Public/Misc/Build.h
@@ -441,7 +441,7 @@
 #endif
 
 #ifndef WITH_COTF
-	#define WITH_COTF ((WITH_ENGINE) && !(IS_PROGRAM || WITH_EDITOR || UE_SERVER || UE_BUILD_SHIPPING))
+	#define WITH_COTF ((WITH_ENGINE) && !(IS_PROGRAM || UE_BUILD_SHIPPING))
 #endif
 
 // Controls if the config system can stores configs for other platforms than the running one
diff --git a/Engine/Source/Runtime/CoreUObject/CoreUObject.Build.cs b/Engine/Source/Runtime/CoreUObject/CoreUObject.Build.cs
--- a/Engine/Source/Runtime/CoreUObject/CoreUObject.Build.cs
+++ b/Engine/Source/Runtime/CoreUObject/CoreUObject.Build.cs
@@ -21,9 +21,13 @@ public class CoreUObject : ModuleRules
 
 		PublicDependencyModuleNames.Add("Core");
         PublicDependencyModuleNames.Add("TraceLog");
-		PublicDependencyModuleNames.Add("CookOnTheFly");
 
 		PrivateDependencyModuleNames.Add("Projects");
         PrivateDependencyModuleNames.Add("Json");
+		
+		if (Target.Configuration != UnrealTargetConfiguration.Shipping && Target.Type != TargetType.Program)
+		{
+			PublicDependencyModuleNames.Add("CookOnTheFly");
+		}
 	}
 }
diff --git a/Engine/Source/Runtime/CoreUObject/Private/Serialization/CookOnTheFlyPackageStore.cpp b/Engine/Source/Runtime/CoreUObject/Private/Serialization/CookOnTheFlyPackageStore.cpp
--- a/Engine/Source/Runtime/CoreUObject/Private/Serialization/CookOnTheFlyPackageStore.cpp
+++ b/Engine/Source/Runtime/CoreUObject/Private/Serialization/CookOnTheFlyPackageStore.cpp
@@ -1,13 +1,13 @@
 // Copyright Epic Games, Inc. All Rights Reserved.
 
+#if WITH_COTF
+
 #include "Serialization/CookOnTheFlyPackageStore.h"
 #include "CookOnTheFly.h"
 #include "HAL/PlatformTime.h"
 #include "Containers/ChunkedArray.h"
 #include "Serialization/MemoryReader.h"
 
-PRAGMA_DISABLE_OPTIMIZATION
-
 namespace UE { namespace PackageStore { namespace Messaging
 {
 
@@ -325,4 +325,4 @@ TUniquePtr<IPackageStore> MakeCookOnTheFlyPackageStore(
 	return MakeUnique<FCookOnTheFlyPackageStore>(IoDispatcher, CookOnTheFlyServerConnection, MoveTemp(EntriesAddedCallback));
 }
 
-PRAGMA_ENABLE_OPTIMIZATION
+#endif // WITH_COTF
diff --git a/Engine/Source/Runtime/CoreUObject/Public/Serialization/CookOnTheFlyPackageStore.h b/Engine/Source/Runtime/CoreUObject/Public/Serialization/CookOnTheFlyPackageStore.h
--- a/Engine/Source/Runtime/CoreUObject/Public/Serialization/CookOnTheFlyPackageStore.h
+++ b/Engine/Source/Runtime/CoreUObject/Public/Serialization/CookOnTheFlyPackageStore.h
@@ -2,6 +2,8 @@
 
 #pragma once
 
+#if WITH_COTF
+
 #include "CoreMinimal.h"
 #include "IO/PackageStore.h"
 
@@ -65,3 +67,5 @@ TUniquePtr<IPackageStore> MakeCookOnTheFlyPackageStore(
 	FIoDispatcher& IoDispatcher,
 	UE::Cook::ICookOnTheFlyServerConnection& CookOnTheFlyServerConnection,
 	TFunction<void()>&& EntriesAddedCallback);
+
+#endif // WITH_COTF
diff --git a/Engine/Source/Runtime/Launch/Launch.Build.cs b/Engine/Source/Runtime/Launch/Launch.Build.cs
--- a/Engine/Source/Runtime/Launch/Launch.Build.cs
+++ b/Engine/Source/Runtime/Launch/Launch.Build.cs
@@ -35,7 +35,6 @@ public class Launch : ModuleRules
 				"Overlay",
 				"PreLoadScreen",
 				"InstallBundleManager",
-				"CookOnTheFly"
 			});
 
 		// Set a macro allowing us to switch between debuggame/development configuration
@@ -264,6 +263,11 @@ public class Launch : ModuleRules
 		if (Target.Configuration != UnrealTargetConfiguration.Shipping)
 		{
 			PrivateDependencyModuleNames.Add("StorageServerClient");
+
+			if (Target.Type != TargetType.Program)
+			{
+				PublicDependencyModuleNames.Add("CookOnTheFly");
+			}
 		}
 
 		if(Target.LinkType == TargetLinkType.Monolithic && !Target.bFormalBuild)
