VirtualizedBulkData: fix incorrect serialized PayloadIds for editordomain versions of non-virtual payloads that virtualize on save; the incorrect paylodIds were noticeable by the "Failed to pull virtual data with guid" error when loading the saved package.
#rnx

diff --git a/Engine/Source/Runtime/CoreUObject/Private/Virtualization/VirtualizedBulkData.cpp b/Engine/Source/Runtime/CoreUObject/Private/Virtualization/VirtualizedBulkData.cpp
--- a/Engine/Source/Runtime/CoreUObject/Private/Virtualization/VirtualizedBulkData.cpp
+++ b/Engine/Source/Runtime/CoreUObject/Private/Virtualization/VirtualizedBulkData.cpp
@@ -212,7 +212,6 @@ void FVirtualizedUntypedBulkData::Serialize(FArchive& Ar, UObject* Owner)
 	}
 	else if (Ar.IsPersistent() && !Ar.IsObjectReferenceCollector() && !Ar.ShouldSkipBulkData())
 	{
-		// Check if we need to update the payload id before it is serialized
 		FLinkerSave* LinkerSave = nullptr;
 		bool bKeepLegacyDataByReference = false;
 		if (Ar.IsSaving())
@@ -225,6 +224,12 @@ void FVirtualizedUntypedBulkData::Serialize(FArchive& Ar, UObject* Owner)
 			{
 				UpdateKeyIfNeeded();
 			}
+
+			const bool bCanAttemptVirtualization = LinkerSave != nullptr;
+			if (bCanAttemptVirtualization)
+			{
+				PushData(); // Note this can change various members if we are going from non-virtualized to virtualized
+			}
 		}
 
 		// Store the position in the archive of the flags in case we need to update it later
@@ -247,13 +252,6 @@ void FVirtualizedUntypedBulkData::Serialize(FArchive& Ar, UObject* Owner)
 		{
 			checkf(Ar.IsCooking() == false, TEXT("FVirtualizedUntypedBulkData::Serialize should not be called during a cook"));
 
-			const bool bCanAttemptVirtualization = LinkerSave != nullptr;
-
-			if (bCanAttemptVirtualization)
-			{
-				PushData(); // Note this can change various members if we are going from non-virtualized to virtualized
-			}
-
 			EFlags UpdatedFlags = BuildFlagsForSerialization(Ar, !bKeepLegacyDataByReference);
 
 			// Go back in the archive and update the flags in the archive, we will only apply the updated flags to the current
