GarbageCollection: Improve GC reachability analysis wall time by letting the GT do work using a ParallelFor instead of waiting.

#rnx

diff --git a/Engine/Source/Runtime/CoreUObject/Public/UObject/FastReferenceCollector.h b/Engine/Source/Runtime/CoreUObject/Public/UObject/FastReferenceCollector.h
--- a/Engine/Source/Runtime/CoreUObject/Public/UObject/FastReferenceCollector.h
+++ b/Engine/Source/Runtime/CoreUObject/Public/UObject/FastReferenceCollector.h
@@ -534,20 +534,11 @@ public:
 			}
 			else
 			{
-				FGraphEventArray ChunkTasks;
-
-				int32 NumThreads = FTaskGraphInterface::Get().GetNumWorkerThreads();
-				int32 NumBackgroundThreads = ENamedThreads::bHasBackgroundThreads ? NumThreads : 0;
-				ENamedThreads::Type NormalThreadName = ENamedThreads::AnyNormalThreadNormalTask;
-				ENamedThreads::Type BackgroundThreadName = ENamedThreads::AnyBackgroundThreadNormalTask;
-
-				FPlatformProcess::ModifyThreadAssignmentForUObjectReferenceCollector(NumThreads, NumBackgroundThreads, NormalThreadName, BackgroundThreadName);
-
+				int32 NumThreads = FTaskGraphInterface::Get().GetNumWorkerThreads() + 1; // +1 for GT
 				// Avoid going over 26 threads as it may cause race conditions in TLockFreePointerListUnordered
-				int32 NumTasks = FMath::Min(NumThreads + NumBackgroundThreads, 26);
+				int32 NumTasks = FMath::Min(NumThreads, 26);
 
 				check(NumTasks > 0);
-				ChunkTasks.Empty(NumTasks);
 				int32 NumPerChunk = ObjectsToCollectReferencesFor.Num() / NumTasks;
 				int32 StartIndex = 0;
 				for (int32 Chunk = 0; Chunk < NumTasks; Chunk++)
@@ -560,24 +551,10 @@ public:
 					StartIndex += NumPerChunk;
 				}
 
-				// This is only enabled in editor for now because it has not been tested client mode neither with little/big cores.
-				// If this gets activated for non-editor, please also activate the FGarbageCollectorStarvationTest test in client mode.
-#if WITH_EDITOR
 				// Use ParallelFor to ensure we're making progress in this thread instead of just waiting for tasks to complete.
 				// This is especially important for good game-thread latency when the taskgraph is already filled and to avoid
 				// potential deadlock with workers making use of FGCScopeGuard.
 				ParallelForTemplate(NumTasks, [this](int32) { TaskQueue.DoTask(); }, EParallelForFlags::Unbalanced);
-#else
-				for (int32 Chunk = 0; Chunk < NumTasks; Chunk++)
-				{
-					extern CORE_API int32 GUseNewTaskBackend;
-					ENamedThreads::Type ThreadPriority = GUseNewTaskBackend ? ENamedThreads::AnyHiPriThreadHiPriTask : ((Chunk >= NumThreads) ? BackgroundThreadName : NormalThreadName);
-					ChunkTasks.Add(TGraphTask< FCollectorTaskProcessorTask >::CreateTask().ConstructAndDispatchWhenReady(TaskQueue, ThreadPriority));
-				}
-
-				QUICK_SCOPE_CYCLE_COUNTER(STAT_GC_Subtask_Wait);
-				FTaskGraphInterface::Get().WaitUntilTasksComplete(ChunkTasks, ENamedThreads::GameThread_Local);
-#endif
 				TaskQueue.CheckDone();
 			}
 		}
