Remove the usage of modal message dialog in package saving
Always try to find the likely culprit in case of illegal references

diff --git a/Engine/Source/Runtime/CoreUObject/Private/UObject/SavePackage.cpp b/Engine/Source/Runtime/CoreUObject/Private/UObject/SavePackage.cpp
--- a/Engine/Source/Runtime/CoreUObject/Private/UObject/SavePackage.cpp
+++ b/Engine/Source/Runtime/CoreUObject/Private/UObject/SavePackage.cpp
@@ -2954,7 +2954,7 @@ FSavePackageResultStruct UPackage::Save(UPackage* InOuter, UObject* Base, EObjec
 					const FText Message = FText::Format( NSLOCTEXT("Core", "LinkedToObjectsInOtherMap_FindCulpritQ", "Can't save {FileName}: Graph is linked to object(s) in external map.\nExternal Object(s):\n{ObjectNames}  \nTry to find the chain of references to that object (may take some time)?"), Args );
 
 					FString CulpritString = TEXT( "Unknown" );
-					bool bFindCulprit = IsRunningCommandlet() || (FMessageDialog::Open(EAppMsgType::YesNo, Message) == EAppReturnType::Yes);
+					bool bFindCulprit = true;
 					if (bFindCulprit)
 					{
 						FString Referencer;
@@ -3014,7 +3014,7 @@ FSavePackageResultStruct UPackage::Save(UPackage* InOuter, UObject* Base, EObjec
 					const FText Message = FText::Format( NSLOCTEXT("Core", "LinkedToPrivateObjectsInOtherPackage_FindCulpritQ", "Can't save {FileName}: Graph is linked to private object(s) in an external package.\nExternal Object(s):\n{ObjectNames}  \nTry to find the chain of references to that object (may take some time)?"), Args );
 
 					FString CulpritString = TEXT( "Unknown" );
-					bool bFindCulprit = IsRunningCommandlet() || (FMessageDialog::Open(EAppMsgType::YesNo, Message) == EAppReturnType::Yes);
+					bool bFindCulprit = true;
 					if (bFindCulprit)
 					{
 						FString Referencer;
diff --git a/Engine/Source/Runtime/CoreUObject/Private/UObject/SavePackage2.cpp b/Engine/Source/Runtime/CoreUObject/Private/UObject/SavePackage2.cpp
--- a/Engine/Source/Runtime/CoreUObject/Private/UObject/SavePackage2.cpp
+++ b/Engine/Source/Runtime/CoreUObject/Private/UObject/SavePackage2.cpp
@@ -465,15 +465,11 @@ ESavePackageResult ValidateIllegalReferences(FSaveContext& SaveContext, TArray<U
 		const FText Message = FText::Format(NSLOCTEXT("Core", "LinkedToObjectsInOtherMap_FindCulpritQ", "Can't save {FileName}: Graph is linked to object(s) in external map.\nExternal Object(s):\n{ObjectNames}  \nTry to find the chain of references to that object (may take some time)?"), Args);
 
 		FString CulpritString = TEXT("Unknown");
-		bool bFindCulprit = IsRunningCommandlet() || (FMessageDialog::Open(EAppMsgType::YesNo, Message) == EAppReturnType::Yes);
-		if (bFindCulprit)
+		FString Referencer;
+		SavePackageUtilities::FindMostLikelyCulprit(ObjectsInOtherMaps, MostLikelyCulprit, Referencer);
+		if (MostLikelyCulprit != nullptr)
 		{
-			FString Referencer;
-			SavePackageUtilities::FindMostLikelyCulprit(ObjectsInOtherMaps, MostLikelyCulprit, Referencer);
-			if (MostLikelyCulprit != nullptr)
-			{
-				CulpritString = FString::Printf(TEXT("%s (%s)"), *MostLikelyCulprit->GetFullName(), *Referencer);
-			}
+			CulpritString = FString::Printf(TEXT("%s (%s)"), *MostLikelyCulprit->GetFullName(), *Referencer);
 		}
 
 		FString ErrorMessage = FString::Printf(TEXT("Can't save %s: Graph is linked to object %s in external map"), SaveContext.GetFilename(), *CulpritString);
@@ -520,15 +516,11 @@ ESavePackageResult ValidateIllegalReferences(FSaveContext& SaveContext, TArray<U
 		const FText Message = FText::Format(NSLOCTEXT("Core", "LinkedToPrivateObjectsInOtherPackage_FindCulpritQ", "Can't save {FileName}: Graph is linked to private object(s) in an external package.\nExternal Object(s):\n{ObjectNames}  \nTry to find the chain of references to that object (may take some time)?"), Args);
 
 		FString CulpritString = TEXT("Unknown");
-		bool bFindCulprit = IsRunningCommandlet() || (FMessageDialog::Open(EAppMsgType::YesNo, Message) == EAppReturnType::Yes);
-		if (bFindCulprit)
-		{
-			FString Referencer;
-			SavePackageUtilities::FindMostLikelyCulprit(PrivateObjects, MostLikelyCulprit, Referencer, &SaveContext);
-			CulpritString = FString::Printf(TEXT("%s (%s)"),
-				(MostLikelyCulprit != nullptr) ? *MostLikelyCulprit->GetFullName() : TEXT("(unknown culprit)"),
-				*Referencer);
-		}
+		FString Referencer;
+		SavePackageUtilities::FindMostLikelyCulprit(PrivateObjects, MostLikelyCulprit, Referencer, &SaveContext);
+		CulpritString = FString::Printf(TEXT("%s (%s)"),
+			(MostLikelyCulprit != nullptr) ? *MostLikelyCulprit->GetFullName() : TEXT("(unknown culprit)"),
+			*Referencer);
 
 		if (SaveContext.IsGenerateSaveError())
 		{
