Core: Changed the verbosity of the unintialized reflected property check to an error for types in engine modules or UObject* properties

diff --git a/Engine/Source/Runtime/CoreUObject/Private/UObject/Class.cpp b/Engine/Source/Runtime/CoreUObject/Private/UObject/Class.cpp
--- a/Engine/Source/Runtime/CoreUObject/Private/UObject/Class.cpp
+++ b/Engine/Source/Runtime/CoreUObject/Private/UObject/Class.cpp
@@ -3512,16 +3512,29 @@ int32 FStructUtils::AttemptToFindUninitializedScriptStructMembers()
 			{
 				if (!UninitializedPropertiesZeroed.Contains(Property))
 				{
+					ELogVerbosity::Type Verbosity = ELogVerbosity::Display;
+
 					++UninitializedScriptStructMemberCount;
-					if (Property->IsA<FObjectPropertyBase>())
+
+					UPackage* ScriptPackage = ScriptStruct->GetOutermost();
+					FName ScriptModuleName = FPackageName::GetShortFName(ScriptPackage->GetName());
+					IModuleInterface* ScriptModule = FModuleManager::Get().GetModule(ScriptModuleName);
+					const bool bIsEngineModule = (ScriptModule != nullptr) && !ScriptModule->IsGameModule();
+
+					if (bIsEngineModule)
 					{
-						++UninitializedObjectPropertyCount;
-						UE_LOG(LogClass, Warning, TEXT("%s %s%s::%s is not initialized properly.%s"), *Property->GetClass()->GetName(), ScriptStruct->GetPrefixCPP(), *ScriptStruct->GetName(), *Property->GetNameCPP(), *GetStructLocation(ScriptStruct));
+						Verbosity = ELogVerbosity::Error;
 					}
-					else
+
+					if (Property->IsA<FObjectPropertyBase>())
 					{
-						UE_LOG(LogClass, Display, TEXT("%s %s%s::%s is not initialized properly.%s"), *Property->GetClass()->GetName(), ScriptStruct->GetPrefixCPP(), *ScriptStruct->GetName(), *Property->GetNameCPP(), *GetStructLocation(ScriptStruct));
+						++UninitializedObjectPropertyCount;
+						Verbosity = ELogVerbosity::Error;
 					}
+
+#if !NO_LOGGING
+					FMsg::Logf(__FILE__, __LINE__, LogClass.GetCategoryName(), Verbosity, TEXT("%s %s%s::%s is not initialized properly.%s"), *Property->GetClass()->GetName(), ScriptStruct->GetPrefixCPP(), *ScriptStruct->GetName(), *Property->GetNameCPP(), *GetStructLocation(ScriptStruct));
+#endif
 				}
 			}
 		}
