[Insights]
  - Fixed the trace of scope cycle counters for UObject to be enabled everywhere where FScopeCycleCounterUObject is used switced by Stats::IsThreadCollectingData().
  - Fixed UObject::CallFunction to trace UObject name for Insights only when GVerboseScriptStats is on.

diff --git a/Engine/Source/Runtime/Core/Public/Stats/Stats2.h b/Engine/Source/Runtime/Core/Public/Stats/Stats2.h
--- a/Engine/Source/Runtime/Core/Public/Stats/Stats2.h
+++ b/Engine/Source/Runtime/Core/Public/Stats/Stats2.h
@@ -2325,7 +2325,7 @@ namespace Stats
 	FORCEINLINE bool IsThreadCollectingData()
 	{
 #if STATS
-		return FThreadStats::IsCollectingData();
+		return UE_TRACE_CHANNELEXPR_IS_ENABLED(CpuChannel) || FThreadStats::IsCollectingData();
 #else
 		return GCycleStatsShouldEmitNamedEvents > 0;
 #endif
diff --git a/Engine/Source/Runtime/CoreUObject/Private/UObject/ScriptCore.cpp b/Engine/Source/Runtime/CoreUObject/Private/UObject/ScriptCore.cpp
--- a/Engine/Source/Runtime/CoreUObject/Private/UObject/ScriptCore.cpp
+++ b/Engine/Source/Runtime/CoreUObject/Private/UObject/ScriptCore.cpp
@@ -957,12 +957,12 @@ IMPLEMENT_VM_FUNCTION(EX_CallMath, execCallMathFunction);
 void UObject::CallFunction( FFrame& Stack, RESULT_DECL, UFunction* Function )
 {
 #if PER_FUNCTION_SCRIPT_STATS
-	const bool bShouldTrackFunction = (Stack.DepthCounter <= GMaxFunctionStatDepth) && (Stats::IsThreadCollectingData() || UE_TRACE_CHANNELEXPR_IS_ENABLED(CpuChannel));
+	const bool bShouldTrackFunction = (Stack.DepthCounter <= GMaxFunctionStatDepth) && Stats::IsThreadCollectingData();
 	FScopeCycleCounterUObject FunctionScope(bShouldTrackFunction ? Function : nullptr);
 #endif // PER_FUNCTION_SCRIPT_STATS
 
 #if STATS || ENABLE_STATNAMEDEVENTS
-	const bool bShouldTrackObject = (GVerboseScriptStats && Stats::IsThreadCollectingData()) || UE_TRACE_CHANNELEXPR_IS_ENABLED(CpuChannel);
+	const bool bShouldTrackObject = GVerboseScriptStats && Stats::IsThreadCollectingData();
 	FScopeCycleCounterUObject ContextScope(bShouldTrackObject ? this : nullptr);
 #endif
 
