Core: Made the verbosity of the unintialized reflected property check configurable via ini
- Issues with UObject* properties are upgraded from Warning to Error
- Other types remain the same for both engine and project modules (Display), but engine will soon change to Error as well

Settings can be overridden for either project or engine modules in DefaultEngine.ini (engine defaults are in BaseEngine.ini), e.g.,

[CoreUObject.UninitializedScriptStructMembersCheck]
EngineModuleReflectedUninitializedPropertyVerbosity=Error
ProjectModuleReflectedUninitializedPropertyVerbosity=Warning

diff --git a/Engine/Config/BaseEngine.ini b/Engine/Config/BaseEngine.ini
--- a/Engine/Config/BaseEngine.ini
+++ b/Engine/Config/BaseEngine.ini
@@ -72,6 +72,11 @@ SmallerMemoryBucket_MinGB=6
 ; if SmallestMemoryBucket_MinGB is set, then the engine will throw up an error and exit(0) if the device has less
 SmallestMemoryBucket_MinGB=0
 
+[CoreUObject.UninitializedScriptStructMembersCheck]
+EngineModuleReflectedUninitializedPropertyVerbosity=Display
+ProjectModuleReflectedUninitializedPropertyVerbosity=Display
+ObjectReferenceReflectedUninitializedPropertyVerbosity=Warning
+
 [/Script/Engine.Engine]
 ConsoleClassName=/Script/Engine.Console
 GameViewportClientClassName=/Script/Engine.GameViewportClient
diff --git a/Engine/Source/Runtime/Core/Private/Logging/LogVerbosity.cpp b/Engine/Source/Runtime/Core/Private/Logging/LogVerbosity.cpp
--- a/Engine/Source/Runtime/Core/Private/Logging/LogVerbosity.cpp
+++ b/Engine/Source/Runtime/Core/Private/Logging/LogVerbosity.cpp
@@ -24,4 +24,45 @@ const TCHAR* ToString(ELogVerbosity::Type Verbosity)
 		return TEXT("VeryVerbose");
 	}
 	return TEXT("UnknownVerbosity");
+}
+
+CORE_API ELogVerbosity::Type ParseLogVerbosityFromString(const FString& VerbosityString)
+{
+	if (VerbosityString == TEXT("NoLogging"))
+	{
+		return ELogVerbosity::NoLogging;
+	}
+	else if (VerbosityString == TEXT("Fatal"))
+	{
+		return ELogVerbosity::Fatal;
+	}
+	else if (VerbosityString == TEXT("Error"))
+	{
+		return ELogVerbosity::Error;
+	}
+	else if (VerbosityString == TEXT("Warning"))
+	{
+		return ELogVerbosity::Warning;
+	}
+	else if (VerbosityString == TEXT("Display"))
+	{
+		return ELogVerbosity::Display;
+	}
+	else if (VerbosityString == TEXT("Log"))
+	{
+		return ELogVerbosity::Log;
+	}
+	else if (VerbosityString == TEXT("Verbose"))
+	{
+		return ELogVerbosity::Verbose;
+	}
+	else if (VerbosityString == TEXT("VeryVerbose"))
+	{
+		return ELogVerbosity::VeryVerbose;
+	}
+	else
+	{
+		// An unknown value is treated as log
+		return ELogVerbosity::Log;
+	}
 }
\ No newline at end of file
diff --git a/Engine/Source/Runtime/Core/Public/Logging/LogVerbosity.h b/Engine/Source/Runtime/Core/Public/Logging/LogVerbosity.h
--- a/Engine/Source/Runtime/Core/Public/Logging/LogVerbosity.h
+++ b/Engine/Source/Runtime/Core/Public/Logging/LogVerbosity.h
@@ -67,4 +67,12 @@ static_assert(!(ELogVerbosity::VerbosityMask & ELogVerbosity::BreakOnLog), "Bad
  * @param Verbosity verbosity enum
  * @returns String representation of the verbosity enum
  */
-CORE_API const TCHAR* ToString(ELogVerbosity::Type Verbosity);
\ No newline at end of file
+CORE_API const TCHAR* ToString(ELogVerbosity::Type Verbosity);
+
+/**
+ * Converts a string to verbosity 
+ * @param VerbosityString verbosity in string form (e.g., "Log", "Display", or "Fatal")
+ * @returns Parsed verbosity from the string (or Log if the string was invalid)
+ */
+extern CORE_API ELogVerbosity::Type ParseLogVerbosityFromString(const FString& VerbosityString);
+
diff --git a/Engine/Source/Runtime/CoreUObject/Private/UObject/Class.cpp b/Engine/Source/Runtime/CoreUObject/Private/UObject/Class.cpp
--- a/Engine/Source/Runtime/CoreUObject/Private/UObject/Class.cpp
+++ b/Engine/Source/Runtime/CoreUObject/Private/UObject/Class.cpp
@@ -1,9 +1,5 @@
 // Copyright Epic Games, Inc. All Rights Reserved.
 
-/*=============================================================================
-	UnClass.cpp: Object class implementation.
-=============================================================================*/
-
 #include "UObject/Class.h"
 #include "HAL/ThreadSafeBool.h"
 #include "HAL/LowLevelMemTracker.h"
@@ -3455,6 +3451,61 @@ int32 FStructUtils::AttemptToFindUninitializedScriptStructMembers()
 		return StructLocation;
 	};
 
+	auto DetermineIfModuleIsEngine = [](const UScriptStruct* ScriptStruct) -> bool
+	{
+		UPackage* ScriptPackage = ScriptStruct->GetOutermost();
+		const FName ScriptModuleName = FPackageName::GetShortFName(ScriptPackage->GetName());
+
+		FModuleStatus ScriptModuleStatus;
+		if (FModuleManager::Get().QueryModule(ScriptModuleName, /*out*/ ScriptModuleStatus))
+		{
+			const bool bIsProjectModule = ScriptModuleStatus.FilePath.StartsWith(FPaths::ConvertRelativePathToFull(FPaths::ProjectDir()));
+			return !bIsProjectModule;
+		}
+		else
+		{
+			// Default to project if we can't determine it (shouldn't ever happen tho)
+			return false;
+		}
+	};
+
+	struct FScriptStructSettings
+	{
+		ELogVerbosity::Type ProjectVerbosity = ELogVerbosity::Display;
+		ELogVerbosity::Type EngineVerbosity = ELogVerbosity::Display;
+		ELogVerbosity::Type PointerVerbosity = ELogVerbosity::Warning;
+
+		FScriptStructSettings()
+		{
+			{
+				FString ProjectSettingString;
+				if (GConfig->GetString(TEXT("CoreUObject.UninitializedScriptStructMembersCheck"), TEXT("ProjectModuleReflectedUninitializedPropertyVerbosity"), ProjectSettingString, GEngineIni))
+				{
+					ProjectVerbosity = ParseLogVerbosityFromString(ProjectSettingString);
+				}
+			}
+
+			{
+				FString EngineSettingString;
+				if (GConfig->GetString(TEXT("CoreUObject.UninitializedScriptStructMembersCheck"), TEXT("EngineModuleReflectedUninitializedPropertyVerbosity"), EngineSettingString, GEngineIni))
+				{
+					EngineVerbosity = ParseLogVerbosityFromString(EngineSettingString);
+				}
+			}
+
+			//@TODO: Remove this eventually and change the default to Error (or maybe even Fatal)
+			{
+				FString PointerSettingString;
+				if (GConfig->GetString(TEXT("CoreUObject.UninitializedScriptStructMembersCheck"), TEXT("ObjectReferenceReflectedUninitializedPropertyVerbosity"), PointerSettingString, GEngineIni))
+				{
+					PointerVerbosity = ParseLogVerbosityFromString(PointerSettingString);
+				}
+			}
+		}
+	};
+
+	static FScriptStructSettings Settings;
+
 	int32 UninitializedScriptStructMemberCount = 0;
 	int32 UninitializedObjectPropertyCount = 0;
 	UScriptStruct* TestUninitializedScriptStructMembersTestStruct = TBaseStructure<FTestUninitializedScriptStructMembersTest>::Get();
@@ -3513,15 +3564,18 @@ int32 FStructUtils::AttemptToFindUninitializedScriptStructMembers()
 				if (!UninitializedPropertiesZeroed.Contains(Property))
 				{
 					++UninitializedScriptStructMemberCount;
+
+					ELogVerbosity::Type Verbosity = DetermineIfModuleIsEngine(ScriptStruct) ? Settings.EngineVerbosity : Settings.ProjectVerbosity;
+
 					if (Property->IsA<FObjectPropertyBase>())
 					{
 						++UninitializedObjectPropertyCount;
-						UE_LOG(LogClass, Warning, TEXT("%s %s%s::%s is not initialized properly.%s"), *Property->GetClass()->GetName(), ScriptStruct->GetPrefixCPP(), *ScriptStruct->GetName(), *Property->GetNameCPP(), *GetStructLocation(ScriptStruct));
-					}
-					else
-					{
-						UE_LOG(LogClass, Display, TEXT("%s %s%s::%s is not initialized properly.%s"), *Property->GetClass()->GetName(), ScriptStruct->GetPrefixCPP(), *ScriptStruct->GetName(), *Property->GetNameCPP(), *GetStructLocation(ScriptStruct));
+						Verbosity = FMath::Min(Verbosity, Settings.PointerVerbosity);
 					}
+
+#if !NO_LOGGING
+					FMsg::Logf(__FILE__, __LINE__, LogClass.GetCategoryName(), Verbosity, TEXT("%s %s%s::%s is not initialized properly.%s"), *Property->GetClass()->GetName(), ScriptStruct->GetPrefixCPP(), *ScriptStruct->GetName(), *Property->GetNameCPP(), *GetStructLocation(ScriptStruct));
+#endif
 				}
 			}
 		}
