Fix for optimized map and set properties serialization when there are no defaults and the container is already populated.

diff --git a/Engine/Source/Runtime/CoreUObject/Private/UObject/PropertyMap.cpp b/Engine/Source/Runtime/CoreUObject/Private/UObject/PropertyMap.cpp
--- a/Engine/Source/Runtime/CoreUObject/Private/UObject/PropertyMap.cpp
+++ b/Engine/Source/Runtime/CoreUObject/Private/UObject/PropertyMap.cpp
@@ -336,7 +336,7 @@ void FMapProperty::SerializeItem(FStructuredArchive::FSlot Slot, void* Value, co
 		int32 NumKeysToRemove = 0;
 		FStructuredArchive::FArray KeysToRemoveArray = Record.EnterArray(SA_FIELD_NAME(TEXT("KeysToRemove")), NumKeysToRemove);
 
-		if (MapHelper.Num() == 0) // Faster loading path when loading into an empty map
+		if (!Defaults || MapHelper.Num() == 0) // Faster loading path when loading into an empty map
 		{
 			if (NumKeysToRemove)
 			{
diff --git a/Engine/Source/Runtime/CoreUObject/Private/UObject/PropertySet.cpp b/Engine/Source/Runtime/CoreUObject/Private/UObject/PropertySet.cpp
--- a/Engine/Source/Runtime/CoreUObject/Private/UObject/PropertySet.cpp
+++ b/Engine/Source/Runtime/CoreUObject/Private/UObject/PropertySet.cpp
@@ -293,7 +293,7 @@ void FSetProperty::SerializeItem(FStructuredArchive::FSlot Slot, void* Value, co
 		int32 NumElementsToRemove = 0;
 		FStructuredArchive::FArray ElementsToRemoveArray = Record.EnterArray(SA_FIELD_NAME(TEXT("ElementsToRemove")), NumElementsToRemove);
 
-		if (SetHelper.Num() == 0) // Faster loading path when loading elements into an empty set
+		if (!Defaults || SetHelper.Num() == 0) // Faster loading path when loading elements into an empty set
 		{
 			if (NumElementsToRemove)
 			{
