Generalize IPackageResourceManager::OnEndLoad into FCoreUObjectDelegates::OnEndLoadPackage for use by the BulkDataRegistry, and add UPackage::bHasBeenEndLoaded for BulkDataRegistry to know when a package has finished its initial load serialization.
Add OnEndLoad calls for Async loaded packages, since the editor sometimes uses async loads.
#rnx

diff --git a/Engine/Source/Editor/UnrealEd/Private/EditorDomain/BulkDataRegistryEditorDomain.cpp b/Engine/Source/Editor/UnrealEd/Private/EditorDomain/BulkDataRegistryEditorDomain.cpp
--- a/Engine/Source/Editor/UnrealEd/Private/EditorDomain/BulkDataRegistryEditorDomain.cpp
+++ b/Engine/Source/Editor/UnrealEd/Private/EditorDomain/BulkDataRegistryEditorDomain.cpp
@@ -53,10 +53,12 @@ FBulkDataRegistryEditorDomain::FBulkDataRegistryEditorDomain()
 	SharedDataLock = new FTaskSharedDataLock();
 	// We piggyback on the BulkDataRegistry hook to set the pointer to tunnel in the pointer to the EditorBuildInputResolver as well
 	SetGlobalBuildInputResolver(&UE::DerivedData::FEditorBuildInputResolver::Get());
+	FCoreUObjectDelegates::OnEndLoadPackage.AddRaw(this, &FBulkDataRegistryEditorDomain::OnEndLoadPackage);
 }
 
 FBulkDataRegistryEditorDomain::~FBulkDataRegistryEditorDomain()
 {
+	FCoreUObjectDelegates::OnEndLoadPackage.RemoveAll(this);
 	SetGlobalBuildInputResolver(nullptr);
 
 	TMap<FGuid, FUpdatingPayload> LocalUpdatingPayloads;
@@ -114,6 +116,10 @@ void FBulkDataRegistryEditorDomain::Register(UPackage* Owner, const UE::Virtuali
 	Registry.Add(BulkData.GetIdentifier(), FRegisteredBulk(MoveTemp(CopyBulk), PackageNameToUpdate));
 }
 
+void FBulkDataRegistryEditorDomain::OnEndLoadPackage(TConstArrayView<UPackage*> LoadedPackages)
+{
+}
+
 void FBulkDataRegistryEditorDomain::OnExitMemory(const UE::Virtualization::FVirtualizedUntypedBulkData& BulkData)
 {
 	const FGuid& Key = BulkData.GetIdentifier();
diff --git a/Engine/Source/Editor/UnrealEd/Private/EditorDomain/BulkDataRegistryEditorDomain.h b/Engine/Source/Editor/UnrealEd/Private/EditorDomain/BulkDataRegistryEditorDomain.h
--- a/Engine/Source/Editor/UnrealEd/Private/EditorDomain/BulkDataRegistryEditorDomain.h
+++ b/Engine/Source/Editor/UnrealEd/Private/EditorDomain/BulkDataRegistryEditorDomain.h
@@ -137,6 +137,7 @@ public:
 private:
 	void AddPendingPackageBulkData(FName PackageName, const UE::Virtualization::FVirtualizedUntypedBulkData& BulkData);
 	void AddPendingPackageBulkDataInsideLock(FName PackageName, const UE::Virtualization::FVirtualizedUntypedBulkData& BulkData);
+	void OnEndLoadPackage(TConstArrayView<UPackage*> LoadedPackages);
 	void PollPendingPackages(bool bWaitForCooldown);
 	void AddTempLoadedPayload(const FGuid& RegistryKey, uint64 PayloadSize);
 	void PruneTempLoadedPayloads();
diff --git a/Engine/Source/Editor/UnrealEd/Private/EditorDomain/EditorDomain.cpp b/Engine/Source/Editor/UnrealEd/Private/EditorDomain/EditorDomain.cpp
--- a/Engine/Source/Editor/UnrealEd/Private/EditorDomain/EditorDomain.cpp
+++ b/Engine/Source/Editor/UnrealEd/Private/EditorDomain/EditorDomain.cpp
@@ -25,6 +25,7 @@
 #include "Templates/RefCounting.h"
 #include "Templates/UniquePtr.h"
 #include "UObject/PackageResourceManagerFile.h"
+#include "UObject/UObjectGlobals.h"
 #include "UObject/UObjectIterator.h"
 
 DEFINE_LOG_CATEGORY(LogEditorDomain);
@@ -86,11 +87,13 @@ FEditorDomain::FEditorDomain()
 	{
 		OnPostEngineInit();
 	}
+	FCoreUObjectDelegates::OnEndLoadPackage.AddRaw(this, &FEditorDomain::OnEndLoadPackage);
 }
 
 FEditorDomain::~FEditorDomain()
 {
 	FScopeLock ScopeLock(&Locks->Lock);
+	FCoreUObjectDelegates::OnEndLoadPackage.RemoveAll(this);
 	FCoreDelegates::OnPostEngineInit.RemoveAll(this);
 	Locks->Owner = nullptr;
 	AssetRegistry = nullptr;
@@ -216,7 +219,7 @@ void FEditorDomain::MarkNeedsLoadFromWorkspace(const FPackagePath& PackagePath,
 	{
 		SaveClient->RequestSave(PackagePath);
 	}
-	// Otherwise, we will note the need for save in OnEndLoad
+	// Otherwise, we will note the need for save in OnEndLoadPackage
 
 }
 
@@ -419,7 +422,7 @@ void FEditorDomain::Tick(float DeltaTime)
 	}
 }
 
-void FEditorDomain::OnEndLoad(TConstArrayView<UPackage*> LoadedPackages)
+void FEditorDomain::OnEndLoadPackage(TConstArrayView<UPackage*> LoadedPackages)
 {
 	if (bExternalSave)
 	{
diff --git a/Engine/Source/Editor/UnrealEd/Private/EditorDomain/EditorDomain.h b/Engine/Source/Editor/UnrealEd/Private/EditorDomain/EditorDomain.h
--- a/Engine/Source/Editor/UnrealEd/Private/EditorDomain/EditorDomain.h
+++ b/Engine/Source/Editor/UnrealEd/Private/EditorDomain/EditorDomain.h
@@ -75,7 +75,6 @@ public:
 		FPackageSegmentStatVisitor Callback) override;
 	virtual void IteratePackagesStatInLocalOnlyDirectory(FStringView RootDir, 
 		FPackageSegmentStatVisitor Callback) override;
-	virtual void OnEndLoad(TConstArrayView<UPackage*> LoadedPackages) override;
 
 	// FTickableEditorObject interface
 	virtual void Tick(float DeltaTime) override;
@@ -84,6 +83,7 @@ public:
 
 	// EditorDomain interface
 	void PrecachePackageDigest(FName PackageName);
+	void OnEndLoadPackage(TConstArrayView<UPackage*> LoadedPackages);
 
 private:
 	/**
diff --git a/Engine/Source/Runtime/CoreUObject/Private/Serialization/AsyncLoading.cpp b/Engine/Source/Runtime/CoreUObject/Private/Serialization/AsyncLoading.cpp
--- a/Engine/Source/Runtime/CoreUObject/Private/Serialization/AsyncLoading.cpp
+++ b/Engine/Source/Runtime/CoreUObject/Private/Serialization/AsyncLoading.cpp
@@ -6817,6 +6817,15 @@ void FAsyncPackage::CallCompletionCallbacks(bool bInternal, EAsyncLoadingResult:
 	checkSlow(bInternal || !IsInAsyncLoadingThread());
 
 	UPackage* LoadedPackage = (!bLoadHasFailed) ? LinkerRoot : nullptr;
+#if WITH_EDITOR
+	if (!bInternal && GIsEditor && LoadedPackage)
+	{
+		// Call the global delegate for package endloads, which is a kind of external completioncallback,
+		// and set the bHasBeenLoaded flag that is used to check which packages have reached this state
+		LoadedPackage->bHasBeenEndLoaded = true;
+		FCoreUObjectDelegates::OnEndLoadPackage.Broadcast(TConstArrayView<UPackage*> { LoadedPackage });
+	}
+#endif
 	for (FCompletionCallback& CompletionCallback : CompletionCallbacks)
 	{
 		if (CompletionCallback.bIsInternal == bInternal && !CompletionCallback.bCalled)
diff --git a/Engine/Source/Runtime/CoreUObject/Private/Serialization/AsyncLoading2.cpp b/Engine/Source/Runtime/CoreUObject/Private/Serialization/AsyncLoading2.cpp
--- a/Engine/Source/Runtime/CoreUObject/Private/Serialization/AsyncLoading2.cpp
+++ b/Engine/Source/Runtime/CoreUObject/Private/Serialization/AsyncLoading2.cpp
@@ -5734,6 +5734,15 @@ void FAsyncPackage2::CallCompletionCallbacks(EAsyncLoadingResult::Type LoadingRe
 	checkSlow(IsInGameThread());
 
 	UPackage* LoadedPackage = (!bLoadHasFailed) ? LinkerRoot : nullptr;
+#if WITH_EDITOR
+	if (GIsEditor && LoadedPackage)
+	{
+		// Call the global delegate for package endloads, which is a kind of completioncallback,
+		// and set the bHasBeenLoaded flag that is used to check which packages have reached this state
+		LoadedPackage->bHasBeenEndLoaded = true;
+		FCoreUObjectDelegates::OnEndLoadPackage.Broadcast(TConstArrayView<UPackage*> { LoadedPackage });
+	}
+#endif
 	for (FCompletionCallback& CompletionCallback : CompletionCallbacks)
 	{
 		CompletionCallback->ExecuteIfBound(Desc.GetUPackageName(), LoadedPackage, LoadingResult);
diff --git a/Engine/Source/Runtime/CoreUObject/Private/UObject/UObjectGlobals.cpp b/Engine/Source/Runtime/CoreUObject/Private/UObject/UObjectGlobals.cpp
--- a/Engine/Source/Runtime/CoreUObject/Private/UObject/UObjectGlobals.cpp
+++ b/Engine/Source/Runtime/CoreUObject/Private/UObject/UObjectGlobals.cpp
@@ -148,6 +148,7 @@ FCoreUObjectDelegates::FOnObjectModified FCoreUObjectDelegates::OnObjectModified
 FCoreUObjectDelegates::FOnObjectTransacted FCoreUObjectDelegates::OnObjectTransacted;
 FCoreUObjectDelegates::FOnObjectsReplaced FCoreUObjectDelegates::OnObjectsReplaced;
 FCoreUObjectDelegates::FOnAssetLoaded FCoreUObjectDelegates::OnAssetLoaded;
+FCoreUObjectDelegates::FOnEndLoadPackage FCoreUObjectDelegates::OnEndLoadPackage;
 PRAGMA_DISABLE_DEPRECATION_WARNINGS
 FCoreUObjectDelegates::FOnObjectSaved FCoreUObjectDelegates::OnObjectSaved;
 PRAGMA_ENABLE_DEPRECATION_WARNINGS
@@ -1420,14 +1421,17 @@ UPackage* LoadPackageInternal(UPackage* InOuter, const FPackagePath& PackagePath
 		Result->SetFlags(RF_WasLoaded);
 	}
 
-	if (GIsEditor)
+#if WITH_EDITOR
+	if (GIsEditor && !IsInAsyncLoadingThread() && GGameThreadLoadCounter == 0)
 	{
-		if (IsInGameThread() && !IsInAsyncLoadingThread() && GGameThreadLoadCounter == 0)
+		// check(IsInGameThread() was called above, but we still need to test !IsInAsyncLoadingThread to exclude that callsite when the engine is single-threaded
+		for (UPackage* LoadedPackage : LoadedPackages)
 		{
-			IPackageResourceManager::Get().OnEndLoad(LoadedPackages);
+			LoadedPackage->bHasBeenEndLoaded = true;
 		}
+		FCoreUObjectDelegates::OnEndLoadPackage.Broadcast(LoadedPackages);
 	}
-
+#endif
 	return Result;
 }
 
diff --git a/Engine/Source/Runtime/CoreUObject/Public/UObject/Package.h b/Engine/Source/Runtime/CoreUObject/Public/UObject/Package.h
--- a/Engine/Source/Runtime/CoreUObject/Public/UObject/Package.h
+++ b/Engine/Source/Runtime/CoreUObject/Public/UObject/Package.h
@@ -176,6 +176,11 @@ public:
 	/** Whether this package has been fully loaded (aka had all it's exports created) at some point.															*/
 	mutable uint8 bHasBeenFullyLoaded:1;
 
+#if WITH_EDITORONLY_DATA
+	/** Set to true after serialization and postload and before returning from LoadPackage or calling load completion delegate, for loaded packages.			*/
+	uint8 bHasBeenEndLoaded : 1;
+#endif
+
 	/**
 	 * Whether this package can be imported, i.e. its package name is a package that exists on disk.
 	 * Note: This includes all normal packages where the Name matches the FileName
diff --git a/Engine/Source/Runtime/CoreUObject/Public/UObject/PackageResourceManager.h b/Engine/Source/Runtime/CoreUObject/Public/UObject/PackageResourceManager.h
--- a/Engine/Source/Runtime/CoreUObject/Public/UObject/PackageResourceManager.h
+++ b/Engine/Source/Runtime/CoreUObject/Public/UObject/PackageResourceManager.h
@@ -268,9 +268,6 @@ public:
 	 */
 	COREUOBJECT_API void IteratePackagesStatInLocalOnlyDirectory(FStringView RootDir, FPackagePathStatVisitor Callback);
 
-	/** Hook for notification of packages loaded during an outermost BeginLoad/EndLoad */
-	COREUOBJECT_API virtual void OnEndLoad(TConstArrayView<UPackage*> LoadedPackages) {}
-
 	/**
 	 * Call during engine startup to choose the appropriate PackageResourceManager for the current configuration and
 	 * construct it and bind it to be returned from Get
diff --git a/Engine/Source/Runtime/CoreUObject/Public/UObject/UObjectGlobals.h b/Engine/Source/Runtime/CoreUObject/Public/UObject/UObjectGlobals.h
--- a/Engine/Source/Runtime/CoreUObject/Public/UObject/UObjectGlobals.h
+++ b/Engine/Source/Runtime/CoreUObject/Public/UObject/UObjectGlobals.h
@@ -2388,6 +2388,10 @@ struct COREUOBJECT_API FCoreUObjectDelegates
 	DECLARE_MULTICAST_DELEGATE_OneParam(FOnAssetLoaded, UObject*);
 	static FOnAssetLoaded OnAssetLoaded;
 
+	/** Callback when packages end loading in LoadPackage or AsyncLoadPackage. All packages loaded recursively due to imports are included in the single call of the explicitly-loaded package. */
+	DECLARE_MULTICAST_DELEGATE_OneParam(FOnEndLoadPackage, TConstArrayView<UPackage*>);
+	static FOnEndLoadPackage OnEndLoadPackage;
+
 	/** Delegate used by SavePackage() to create the package backup */
 	DECLARE_DELEGATE_RetVal_OneParam(bool, FAutoPackageBackupDelegate, const UPackage&);
 	static FAutoPackageBackupDelegate AutoPackageBackupDelegate;
