First pass header changes in support of Engine+ShooterGame TObjectPtr refresh.

diff --git a/Engine/Source/Developer/RigVMDeveloper/Public/RigVMModel/RigVMVariableDescription.h b/Engine/Source/Developer/RigVMDeveloper/Public/RigVMModel/RigVMVariableDescription.h
--- a/Engine/Source/Developer/RigVMDeveloper/Public/RigVMModel/RigVMVariableDescription.h
+++ b/Engine/Source/Developer/RigVMDeveloper/Public/RigVMModel/RigVMVariableDescription.h
@@ -1,4 +1,4 @@
-﻿// Copyright Epic Games, Inc. All Rights Reserved.
+// Copyright Epic Games, Inc. All Rights Reserved.
 
 #pragma once
 
@@ -119,7 +119,8 @@ public:
 		return PinType;
 	}
 
-	static bool CPPTypeFromPinType(const FEdGraphPinType& InPinType, FString& OutCPPType, UObject*& OutCPPTypeObject)
+	template <typename T>
+	static bool CPPTypeFromPinType(const FEdGraphPinType& InPinType, FString& OutCPPType, T& OutCPPTypeObject)
 	{
 		FString Prefix = "";
 		FString Suffix = "";
diff --git a/Engine/Source/Runtime/CoreUObject/Public/UObject/ObjectPtr.h b/Engine/Source/Runtime/CoreUObject/Public/UObject/ObjectPtr.h
--- a/Engine/Source/Runtime/CoreUObject/Public/UObject/ObjectPtr.h
+++ b/Engine/Source/Runtime/CoreUObject/Public/UObject/ObjectPtr.h
@@ -667,6 +667,30 @@ FORCEINLINE const TObjectPtr<UObject>& FObjectPtr::ToTObjectPtr() const
 	return *static_cast<const TObjectPtr<UObject>*>(this);
 }
 
+/** Swap variants between TObjectPtr<T> and raw pointer to T */
+template <typename T>
+inline void Swap(TObjectPtr<T>& A, T*& B)
+{
+	Swap((T*&)A, B);
+}
+template <typename T>
+inline void Swap(T*& A, TObjectPtr<T>& B)
+{
+	Swap(A, (T*&)B);
+}
+
+/** Swap variants between TArray<TObjectPtr<T>> and TArray<T*> */
+template <typename T>
+inline void Swap(TArray<TObjectPtr<T>>& A, TArray<T*>& B)
+{
+	Swap(ToRawPtrTArrayUnsafe(A), B);
+}
+template <typename T>
+inline void Swap(TArray<T*>& A, TArray<TObjectPtr<T>>& B)
+{
+	Swap(A, ToRawPtrTArrayUnsafe(B));
+}
+
 /** Exchange variants between TObjectPtr<T> and raw pointer to T */
 template <typename T>
 inline void Exchange(TObjectPtr<T>& A, T*& B)
