Changes to ifdef some TObjectPtr and FObjectHandle code to ensure we're not absorbing overhead in non-late-resolve and non-access-tracking builds.

diff --git a/Engine/Source/Runtime/CoreUObject/Public/UObject/ObjectHandle.h b/Engine/Source/Runtime/CoreUObject/Public/UObject/ObjectHandle.h
--- a/Engine/Source/Runtime/CoreUObject/Public/UObject/ObjectHandle.h
+++ b/Engine/Source/Runtime/CoreUObject/Public/UObject/ObjectHandle.h
@@ -330,6 +330,7 @@ inline FObjectHandle MakeObjectHandle(UObject* Object) { return Object; }
 
 inline UObject* ResolveObjectHandleNoRead(FObjectHandle& Handle)
 {
+#if UE_WITH_OBJECT_HANDLE_LATE_RESOLVE
 	FObjectHandle LocalHandle = Handle;
 	if (IsObjectHandleResolved(LocalHandle))
 	{
@@ -343,13 +344,20 @@ inline UObject* ResolveObjectHandleNoRead(FObjectHandle& Handle)
 		Handle = LocalHandle;
 		return ResolvedObject;
 	}
+#else
+	return ReadObjectHandlePointerNoCheck(Handle);
+#endif
 }
 
 inline UObject* ResolveObjectHandle(FObjectHandle& Handle)
 {
+#if UE_WITH_OBJECT_HANDLE_LATE_RESOLVE || UE_WITH_OBJECT_HANDLE_TRACKING
 	UObject* ResolvedObject = ResolveObjectHandleNoRead(Handle);
 	ObjectHandle_Private::OnHandleRead(ResolvedObject);
 	return ResolvedObject;
+#else
+	return ReadObjectHandlePointerNoCheck(Handle);
+#endif
 }
 
 inline UClass* ResolveObjectHandleClass(FObjectHandle Handle)
diff --git a/Engine/Source/Runtime/CoreUObject/Public/UObject/ObjectPtr.h b/Engine/Source/Runtime/CoreUObject/Public/UObject/ObjectPtr.h
--- a/Engine/Source/Runtime/CoreUObject/Public/UObject/ObjectPtr.h
+++ b/Engine/Source/Runtime/CoreUObject/Public/UObject/ObjectPtr.h
@@ -533,11 +533,13 @@ FORCEINLINE T* ToRawPtr(T* Ptr)
 template <typename T, SIZE_T Size>
 FORCEINLINE T** ToRawPtrArrayUnsafe(TObjectPtr<T>(&ArrayOfPtr)[Size])
 {
+#if UE_WITH_OBJECT_HANDLE_LATE_RESOLVE || UE_WITH_OBJECT_HANDLE_TRACKING
 	for (TObjectPtr<T>& Item : ArrayOfPtr)
 	{
 		// NOTE: Relying on the fact that the TObjectPtr will cache the resolved pointer in place after calling Get.
 		Item.Get();
 	}
+#endif
 	return reinterpret_cast<T**>(ArrayOfPtr);
 }
 
@@ -575,7 +577,7 @@ struct TContainerElementTypeCompatibility<TObjectPtr<T>>
 	UE_OBJPTR_DEPRECATED(5.0, "Reinterpretation between ranges of one type to another type is deprecated.")
 	static void ReinterpretRange(IterBeginType Iter, IterEndType IterEnd, OperatorType Operator = [](IterBeginType& InIt) -> decltype(auto) { return *InIt; })
 	{
-#if UE_WITH_OBJECT_HANDLE_LATE_RESOLVE
+#if UE_WITH_OBJECT_HANDLE_LATE_RESOLVE || UE_WITH_OBJECT_HANDLE_TRACKING
 		while (Iter != IterEnd)
 		{
 			Operator(Iter).Get();
@@ -599,7 +601,7 @@ struct TContainerElementTypeCompatibility<const TObjectPtr<T>>
 	UE_OBJPTR_DEPRECATED(5.0, "Reinterpretation between ranges of one type to another type is deprecated.")
 	static void ReinterpretRange(IterBeginType Iter, IterEndType IterEnd, OperatorType Operator = [](IterBeginType& InIt) -> const TObjectPtr<T>& { return *InIt; })
 	{
-#if UE_WITH_OBJECT_HANDLE_LATE_RESOLVE
+#if UE_WITH_OBJECT_HANDLE_LATE_RESOLVE || UE_WITH_OBJECT_HANDLE_TRACKING
 		while (Iter != IterEnd)
 		{
 			Operator(Iter).Get();
