Fix similar issue with IncludeAssetBundlesName
Original fix from Roubert.Rouhani via UDN

diff --git a/Engine/Source/Runtime/CoreUObject/Private/UObject/SoftObjectPath.cpp b/Engine/Source/Runtime/CoreUObject/Private/UObject/SoftObjectPath.cpp
--- a/Engine/Source/Runtime/CoreUObject/Private/UObject/SoftObjectPath.cpp
+++ b/Engine/Source/Runtime/CoreUObject/Private/UObject/SoftObjectPath.cpp
@@ -749,7 +749,7 @@ bool FSoftObjectPathThreadContext::GetSerializationOptions(FName& OutPackageName
 		bEditorOnly = Archive->IsEditorOnlyPropertyOnTheStack();
 
 		static FName UntrackedName = TEXT("Untracked");
-		if (CurrentProperty && CurrentProperty->HasMetaData(UntrackedName))
+		if (CurrentProperty && CurrentProperty->GetOwnerProperty()->HasMetaData(UntrackedName))
 		{
 			// Property has the Untracked metadata, so set to never collect references
 			CurrentCollectType = ESoftObjectPathCollectType::NeverCollect;
diff --git a/Engine/Source/Runtime/Engine/Private/AssetManager.cpp b/Engine/Source/Runtime/Engine/Private/AssetManager.cpp
--- a/Engine/Source/Runtime/Engine/Private/AssetManager.cpp
+++ b/Engine/Source/Runtime/Engine/Private/AssetManager.cpp
@@ -4369,7 +4369,7 @@ void UAssetManager::InitializeAssetBundlesFromMetadata_Recursive(const UStruct*
 		}
 		else if (const FObjectProperty* ObjectProperty = CastField<FObjectProperty>(Property))
 		{
-			if (ObjectProperty->PropertyFlags & CPF_InstancedReference || ObjectProperty->HasMetaData(IncludeAssetBundlesName))
+			if (ObjectProperty->PropertyFlags & CPF_InstancedReference || ObjectProperty->GetOwnerProperty()->HasMetaData(IncludeAssetBundlesName))
 			{
 				UObject* const* ObjectPtr = reinterpret_cast<UObject* const*>(PropertyValue);
 				if (ObjectPtr && *ObjectPtr)
