BulkDataRegistry:
Change texture and meshdescription BulkData guids to be unique.
#rnx

diff --git a/Engine/Source/Runtime/Core/Public/UObject/UE5MainStreamObjectVersion.h b/Engine/Source/Runtime/Core/Public/UObject/UE5MainStreamObjectVersion.h
--- a/Engine/Source/Runtime/Core/Public/UObject/UE5MainStreamObjectVersion.h
+++ b/Engine/Source/Runtime/Core/Public/UObject/UE5MainStreamObjectVersion.h
@@ -150,7 +150,7 @@ struct CORE_API FUE5MainStreamObjectVersion
 
 		// Removed the AtmosphericFog component with conversion to SkyAtmosphere component
 		RemovedAtmosphericFog,
-		
+
 		// The SkyAtmosphere now light up the heightfog by default, and by default the height fog has a black color.
 		SkyAtmosphereAffectsHeightFogWithBetterDefault,
 
@@ -163,6 +163,9 @@ struct CORE_API FUE5MainStreamObjectVersion
 		// UEdGraphPin serializes SourceIndex
 		EdGraphPinSourceIndex,
 
+		// Change texture bulkdatas to have unique guids
+		VirtualizedBulkDataHaveUniqueGuids,
+
 		// -----<new versions can be added above this line>-------------------------------------------------
 		VersionPlusOne,
 		LatestVersion = VersionPlusOne - 1
diff --git a/Engine/Source/Runtime/CoreUObject/Private/Virtualization/VirtualizedBulkData.cpp b/Engine/Source/Runtime/CoreUObject/Private/Virtualization/VirtualizedBulkData.cpp
--- a/Engine/Source/Runtime/CoreUObject/Private/Virtualization/VirtualizedBulkData.cpp
+++ b/Engine/Source/Runtime/CoreUObject/Private/Virtualization/VirtualizedBulkData.cpp
@@ -198,6 +198,31 @@ void FVirtualizedUntypedBulkData::Unregister()
 #endif
 }
 
+static FGuid CreateUniqueGuid(const FGuid& NonUniqueGuid, const UObject* Owner, const TCHAR* DebugName)
+{
+	if (NonUniqueGuid.IsValid() && Owner)
+	{
+		TStringBuilder<256> PathName;
+		Owner->GetPathName(nullptr, PathName);
+		FBlake3 Builder;
+		Builder.Update(&NonUniqueGuid, sizeof(NonUniqueGuid));
+		Builder.Update(PathName.GetData(), PathName.Len() * sizeof(*PathName.GetData()));
+		FBlake3Hash Hash = Builder.Finalize();
+		// We use the first 16 bytes of the FIoHash to create the guid, there is
+		// no specific reason why these were chosen, we could take any pattern or combination
+		// of bytes.
+		uint32* HashBytes = (uint32*)Hash.GetBytes();
+		return FGuid(HashBytes[0], HashBytes[1], HashBytes[2], HashBytes[3]);
+	}
+	else
+	{
+		UE_LOG(LogVirtualization, Warning,
+			TEXT("CreateFromBulkData recieved an invalid FGuid. A temporary one will be generated until the package is next re-saved! Package: '%s'"),
+			DebugName);
+		return FGuid::NewGuid();
+	}
+}
+
 void FVirtualizedUntypedBulkData::CreateFromBulkData(FUntypedBulkData& InBulkData, const FGuid& InGuid, UObject* Owner)
 {
 	TRACE_CPUPROFILER_EVENT_SCOPE(FVirtualizedUntypedBulkData::CreateFromBulkData);
@@ -212,20 +237,8 @@ void FVirtualizedUntypedBulkData::CreateFromBulkData(FUntypedBulkData& InBulkDat
 	bool bWasKeyGuidDerived = false;
 	if (InBulkData.GetBulkDataSize() > 0)
 	{
-		if (InGuid.IsValid())
-		{
-			BulkDataId = InGuid;
-		}
-		else
-		{
-			UE_LOG(LogVirtualization, Warning,
-				TEXT("CreateFromBulkData recieved an invalid FGuid. A temporary one will be generated until the package is next re-saved! Package: '%s'"),
-				*InBulkData.GetPackagePath().GetDebugName());
-
-			BulkDataId = FGuid::NewGuid();
-		}
-
-		PayloadContentId = FPayloadId(InGuid);
+		BulkDataId = CreateUniqueGuid(InGuid, Owner, *InBulkData.GetPackagePath().GetDebugName());
+		PayloadContentId = FPayloadId(BulkDataId);
 		bWasKeyGuidDerived = true;
 	}
 	
@@ -255,7 +268,17 @@ void FVirtualizedUntypedBulkData::CreateFromBulkData(FUntypedBulkData& InBulkDat
 	Register(Owner);
 }
 
-void FVirtualizedUntypedBulkData::Serialize(FArchive& Ar, UObject* Owner)
+void FVirtualizedUntypedBulkData::CreateLegacyUniqueIdentifier(UObject* Owner)
+{
+	if (BulkDataId.IsValid())
+	{
+		Unregister();
+		BulkDataId = CreateUniqueGuid(BulkDataId, Owner, TEXT("Unknown"));
+		Register(Owner);
+	}
+}
+
+void FVirtualizedUntypedBulkData::Serialize(FArchive& Ar, UObject* Owner, bool bAllowRegister)
 {
 	TRACE_CPUPROFILER_EVENT_SCOPE(FVirtualizedUntypedBulkData::Serialize);
 
@@ -528,7 +551,10 @@ void FVirtualizedUntypedBulkData::Serialize(FArchive& Ar, UObject* Owner)
 					Payload = CompressedPayload.Decompress();
 				}
 			}
-			Register(Owner);
+			if (bAllowRegister)
+			{
+				Register(Owner);
+			}
 		}
 	}
 }
diff --git a/Engine/Source/Runtime/CoreUObject/Public/Virtualization/VirtualizedBulkData.h b/Engine/Source/Runtime/CoreUObject/Public/Virtualization/VirtualizedBulkData.h
--- a/Engine/Source/Runtime/CoreUObject/Public/Virtualization/VirtualizedBulkData.h
+++ b/Engine/Source/Runtime/CoreUObject/Public/Virtualization/VirtualizedBulkData.h
@@ -67,6 +67,8 @@ public:
 	 *					This MUST remain the same between sessions so that the payloads key remains consistent!
 	 */
 	void CreateFromBulkData(FUntypedBulkData& BulkData, const FGuid& Guid, UObject* Owner);
+	/** Fix legacy content that created the Id from non-unique Guids. */
+	void CreateLegacyUniqueIdentifier(UObject* Owner);
 
 	/** 
 	 * Used to serialize the bulkdata to/from a FArchive
@@ -75,8 +77,10 @@ public:
 	 * @param Owner	The UObject that contains the bulkdata object, if this is a nullptr then the bulkdata will
 	 *				assume that it must serialize the payload immediately to memory as it will not be able to
 	 *				identify it's package path.
+	 * @param bAllowRegistry Legacy parameter to skip registration when loading BulkData we know we will need to
+	 *				modify the identifier of. Should always be true for non-legacy serialization.
 	 */
-	void Serialize(FArchive& Ar, UObject* Owner);
+	void Serialize(FArchive& Ar, UObject* Owner, bool bAllowRegister=true);
 
 	/** Reset to a truly empty state */
 	void Reset();
diff --git a/Engine/Source/Runtime/Engine/Private/Texture.cpp b/Engine/Source/Runtime/Engine/Private/Texture.cpp
--- a/Engine/Source/Runtime/Engine/Private/Texture.cpp
+++ b/Engine/Source/Runtime/Engine/Private/Texture.cpp
@@ -468,13 +468,22 @@ void UTexture::Serialize(FArchive& Ar)
 		FWriteScopeLock BulkDataExclusiveScope(Source.BulkDataLock.Get());
 #endif
 
-		if (Ar.IsLoading() && Ar.CustomVer(FUE5MainStreamObjectVersion::GUID) < FUE5MainStreamObjectVersion::TextureSourceVirtualization)
+		if (Ar.IsLoading() && Ar.CustomVer(FUE5MainStreamObjectVersion::GUID) < FUE5MainStreamObjectVersion::VirtualizedBulkDataHaveUniqueGuids)
 		{
-			FByteBulkData TempBulkData;
-			TempBulkData.Serialize(Ar, this);
+			if (Ar.CustomVer(FUE5MainStreamObjectVersion::GUID) < FUE5MainStreamObjectVersion::TextureSourceVirtualization)
+			{
+				FByteBulkData TempBulkData;
+				TempBulkData.Serialize(Ar, this);
+
+				FGuid LegacyPersistentId = Source.GetId();
+				Source.BulkData.CreateFromBulkData(TempBulkData, LegacyPersistentId, this);
+			}
+			else
+			{
+				Source.BulkData.Serialize(Ar, this, false /* bAllowRegister */);
+				Source.BulkData.CreateLegacyUniqueIdentifier(this);
+			}
 
-			FGuid LegacyPersistentId = Source.GetId();
-			Source.BulkData.CreateFromBulkData(TempBulkData, LegacyPersistentId, this);
 		}
 		else
 		{
diff --git a/Engine/Source/Runtime/MeshDescription/Private/MeshDescription.cpp b/Engine/Source/Runtime/MeshDescription/Private/MeshDescription.cpp
--- a/Engine/Source/Runtime/MeshDescription/Private/MeshDescription.cpp
+++ b/Engine/Source/Runtime/MeshDescription/Private/MeshDescription.cpp
@@ -2006,11 +2006,19 @@ void FMeshDescriptionBulkData::Serialize( FArchive& Ar, UObject* Owner )
 
 	bool bSerializedOldDataTypes = false;
 	FByteBulkData TempBulkData;
-	if (Ar.IsLoading() && Ar.CustomVer(FUE5MainStreamObjectVersion::GUID) < FUE5MainStreamObjectVersion::MeshDescriptionVirtualization)
+	if (Ar.IsLoading() && Ar.CustomVer(FUE5MainStreamObjectVersion::GUID) < FUE5MainStreamObjectVersion::VirtualizedBulkDataHaveUniqueGuids)
 	{
-		// Serialize the old BulkData format and mark that we require a conversion after the guid has been serialized
-		TempBulkData.Serialize(Ar, Owner);
-		bSerializedOldDataTypes = true;
+		if (Ar.CustomVer(FUE5MainStreamObjectVersion::GUID) < FUE5MainStreamObjectVersion::MeshDescriptionVirtualization)
+		{
+			// Serialize the old BulkData format and mark that we require a conversion after the guid has been serialized
+			TempBulkData.Serialize(Ar, Owner);
+			bSerializedOldDataTypes = true;
+		}
+		else
+		{
+			BulkData.Serialize(Ar, Owner, false /* bAllowRegister */);
+			BulkData.CreateLegacyUniqueIdentifier(Owner);
+		}
 	}
 	else
 	{
