Improved consistency of how FScriptArrayHelper handles zero counts

The following functions would gracefully handle a count of zero:
 - EmptyAndAddValues
 - EmptyAndAddUninitializedValues

The following functions would assert if given a count of zero:
 - AddUninitializedValues
 - AddValues
 - InsertValues
 - RemoveValues
 - ClearValues

This change makes all of the above functions gracefully handle a count of zero, which addresses some editor crashes that could easily occur in user scripts when using unreal.Array in Python.

diff --git a/Engine/Source/Runtime/CoreUObject/Public/UObject/UnrealType.h b/Engine/Source/Runtime/CoreUObject/Public/UObject/UnrealType.h
--- a/Engine/Source/Runtime/CoreUObject/Public/UObject/UnrealType.h
+++ b/Engine/Source/Runtime/CoreUObject/Public/UObject/UnrealType.h
@@ -3417,10 +3417,7 @@ public:
 		check(Count>=0);
 		checkSlow(Num() >= 0); 
 		EmptyValues(Count);
-		if (Count)
-		{
-			AddValues(Count);
-		}
+		AddValues(Count);
 	}
 	/**
 	*	Empty the array, then add uninitialized values to a given size.
@@ -3431,10 +3428,7 @@ public:
 		check(Count>=0);
 		checkSlow(Num() >= 0); 
 		EmptyValues(Count);
-		if (Count)
-		{
-			AddUninitializedValues(Count);
-		}
+		AddUninitializedValues(Count);
 	}
 	/**
 	*	Expand the array, if needed, so that the given index is valid
@@ -3496,7 +3490,7 @@ public:
 	**/
 	int32 AddUninitializedValues(int32 Count)
 	{
-		check(Count>0);
+		check(Count>=0);
 		checkSlow(Num() >= 0);
 		const int32 OldNum = WithScriptArray([this, Count](auto* Array) { return Array->Add(Count, ElementSize); });
 		return OldNum;
@@ -3516,7 +3510,7 @@ public:
 	**/
 	void InsertValues( int32 Index, int32 Count = 1)
 	{
-		check(Count>0);
+		check(Count>=0);
 		check(Index>=0 && Index <= Num());
 		WithScriptArray([this, Index, Count](auto* Array) { Array->Insert(Index, Count, ElementSize); });
 		ConstructItems(Index, Count);
@@ -3545,7 +3539,7 @@ public:
 	**/
 	void RemoveValues(int32 Index, int32 Count = 1)
 	{
-		check(Count>0);
+		check(Count>=0);
 		check(Index>=0 && Index + Count <= Num());
 		DestructItems(Index, Count);
 		WithScriptArray([this, Index, Count](auto* Array) { Array->Remove(Index, Count, ElementSize); });
@@ -3558,7 +3552,7 @@ public:
 	**/
 	void ClearValues(int32 Index, int32 Count = 1)
 	{
-		check(Count>0);
+		check(Count>=0);
 		check(Index>=0);
 		ClearItems(Index, Count);
 	}
@@ -3633,20 +3627,23 @@ private:
 	**/
 	void ConstructItems(int32 Index, int32 Count)
 	{
-		checkSlow(Count > 0);
+		checkSlow(Count >= 0);
 		checkSlow(Index >= 0); 
 		checkSlow(Index <= Num());
 		checkSlow(Index + Count <= Num());
-		uint8 *Dest = GetRawPtr(Index);
-		if (InnerProperty->PropertyFlags & CPF_ZeroConstructor)
-		{
-			FMemory::Memzero(Dest, Count * ElementSize);
-		}
-		else
+		if (Count > 0)
 		{
-			for (int32 LoopIndex = 0 ; LoopIndex < Count; LoopIndex++, Dest += ElementSize)
+			uint8* Dest = GetRawPtr(Index);
+			if (InnerProperty->PropertyFlags & CPF_ZeroConstructor)
+			{
+				FMemory::Memzero(Dest, Count * ElementSize);
+			}
+			else
 			{
-				InnerProperty->InitializeValue(Dest);
+				for (int32 LoopIndex = 0; LoopIndex < Count; LoopIndex++, Dest += ElementSize)
+				{
+					InnerProperty->InitializeValue(Dest);
+				}
 			}
 		}
 	}
@@ -3659,14 +3656,17 @@ private:
 	{
 		if (!(InnerProperty->PropertyFlags & (CPF_IsPlainOldData | CPF_NoDestructor)))
 		{
-			checkSlow(Count > 0);
+			checkSlow(Count >= 0);
 			checkSlow(Index >= 0); 
 			checkSlow(Index < Num());
 			checkSlow(Index + Count <= Num());
-			uint8 *Dest = GetRawPtr(Index);
-			for (int32 LoopIndex = 0 ; LoopIndex < Count; LoopIndex++, Dest += ElementSize)
+			if (Count > 0)
 			{
-				InnerProperty->DestroyValue(Dest);
+				uint8* Dest = GetRawPtr(Index);
+				for (int32 LoopIndex = 0; LoopIndex < Count; LoopIndex++, Dest += ElementSize)
+				{
+					InnerProperty->DestroyValue(Dest);
+				}
 			}
 		}
 	}
@@ -3677,20 +3677,23 @@ private:
 	**/
 	void ClearItems(int32 Index, int32 Count)
 	{
-		checkSlow(Count > 0);
+		checkSlow(Count >= 0);
 		checkSlow(Index >= 0); 
 		checkSlow(Index < Num());
 		checkSlow(Index + Count <= Num());
-		uint8 *Dest = GetRawPtr(Index);
-		if ((InnerProperty->PropertyFlags & (CPF_ZeroConstructor | CPF_NoDestructor)) == (CPF_ZeroConstructor | CPF_NoDestructor))
+		if (Count > 0)
 		{
-			FMemory::Memzero(Dest, Count * ElementSize);
-		}
-		else
-		{
-			for (int32 LoopIndex = 0; LoopIndex < Count; LoopIndex++, Dest += ElementSize)
+			uint8* Dest = GetRawPtr(Index);
+			if ((InnerProperty->PropertyFlags & (CPF_ZeroConstructor | CPF_NoDestructor)) == (CPF_ZeroConstructor | CPF_NoDestructor))
+			{
+				FMemory::Memzero(Dest, Count * ElementSize);
+			}
+			else
 			{
-				InnerProperty->ClearValue(Dest);
+				for (int32 LoopIndex = 0; LoopIndex < Count; LoopIndex++, Dest += ElementSize)
+				{
+					InnerProperty->ClearValue(Dest);
+				}
 			}
 		}
 	}
