IPackageStoreWriter: Update SavePackage to allow saving non-cooked packages into a PackageStoreWriter.
#rnx

diff --git a/Engine/Source/Editor/UnrealEd/Private/EditorDomain/EditorDomainUtils.cpp b/Engine/Source/Editor/UnrealEd/Private/EditorDomain/EditorDomainUtils.cpp
--- a/Engine/Source/Editor/UnrealEd/Private/EditorDomain/EditorDomainUtils.cpp
+++ b/Engine/Source/Editor/UnrealEd/Private/EditorDomain/EditorDomainUtils.cpp
@@ -12,6 +12,7 @@
 #include "DerivedDataCacheRecord.h"
 #include "Editor.h"
 #include "HAL/FileManager.h"
+#include "IO/PackageStoreWriter.h"
 #include "Memory/SharedBuffer.h"
 #include "Misc/FileHelper.h"
 #include "Misc/PackagePath.h"
@@ -20,6 +21,7 @@
 #include "UObject/CoreRedirects.h"
 #include "UObject/ObjectVersion.h"
 #include "UObject/Package.h"
+#include "UObject/SavePackage.h"
 #include "UObject/UObjectHash.h"
 
 namespace UE
@@ -225,6 +227,122 @@ UE::DerivedData::FRequest RequestEditorDomainPackage(const FPackagePath& Package
 		SkipFlags | UE::DerivedData::ECachePolicy::QueryLocal, CachePriority, MoveTemp(Callback));
 }
 
+/** Stores data from SavePackage in accessible fields */
+class FMemoryPackageStoreWriter : public IPackageStoreWriter
+{
+public:
+	// IPackageStoreWriter
+	virtual void WritePackageData(const FPackageInfo& Info, const FIoBuffer& PackageData, const TArray<FFileRegion>& FileRegions) override;
+	virtual void WriteBulkdata(const FBulkDataInfo& Info, const FIoBuffer& BulkData, const TArray<FFileRegion>& FileRegions) override;
+	virtual void Flush() override
+	{
+	}
+	virtual void BeginPackage(const FBeginPackageInfo& Info) override;
+	virtual void CommitPackage(const FCommitPackageInfo& Info) override
+	{
+	}
+
+	// IPackageStoreWriter cooking and accessor interface: not implemented in this writer
+	bool WriteAdditionalFile(const FAdditionalFileInfo& Info, const FIoBuffer& FileData) override
+	{
+		checkNoEntry();
+		return false;
+	}
+	virtual void BeginCook(const FCookInfo& Info) override
+	{
+		checkNoEntry();
+	}
+	virtual void EndCook() override
+	{
+		checkNoEntry();
+	}
+	virtual void GetEntries(TFunction<void(TArrayView<const FPackageStoreEntryResource>)>&&) override
+	{
+		checkNoEntry();
+	}
+	DECLARE_DERIVED_EVENT(FMemoryPackageStoreWriter, IPackageStoreWriter::FCommitEvent, FCommitEvent);
+	virtual FCommitEvent& OnCommit() override
+	{
+		checkNoEntry();
+		return CommitEvent;
+	}
+	virtual void GetCookedPackages(TArray<FCookedPackageInfo>& OutCookedPackages) override
+	{
+		checkNoEntry();
+	}
+	virtual void RemoveCookedPackages(TArrayView<const FName> PackageNamesToRemove) override
+	{
+		checkNoEntry();
+	}
+
+	// FMemoryPackageStoreWriter interface
+	uint64 GetHeaderSize() const { return HeaderSize; }
+	FIoBuffer GetHeaderAndExports() const { return HeaderAndExports; }
+	/**
+	 * BulkDatas in this array are views into the FMemoryPackageStoreWriter.
+	 * Callers should call MakeOwned if they make copies that need to outlive the FMemoryPackageStoreWriter.
+	 */
+	const TArray<FIoBuffer>& GetBulkDatas() const { return BulkDataRegions; }
+
+private:
+	void SetPackageName(FName InPackageName);
+
+	FIoBuffer HeaderAndExports;
+	TArray<FIoBuffer> BulkDataArchives;
+	TArray<FIoBuffer> BulkDataRegions;
+	FCommitEvent CommitEvent;
+	FName PackageName;
+	uint64 HeaderSize = 0;
+};
+
+void FMemoryPackageStoreWriter::SetPackageName(FName InPackageName)
+{
+	if (PackageName.IsNone())
+	{
+		PackageName = InPackageName;
+	}
+	else
+	{
+		checkf(PackageName == InPackageName, TEXT("FMemoryPackageStoreWriter received different PackageNames in WritePackageData and WriteBulkdata."));
+	}
+}
+
+void FMemoryPackageStoreWriter::BeginPackage(const FBeginPackageInfo& Info)
+{
+	SetPackageName(Info.PackageName);
+}
+
+void FMemoryPackageStoreWriter::WritePackageData(const FPackageInfo& Info, const FIoBuffer& PackageData, const TArray<FFileRegion>& FileRegions)
+{
+	for (const FFileRegion& FileRegion : FileRegions)
+	{
+		checkf(FileRegion.Type == EFileRegionType::None, TEXT("FMemoryPackageStoreWriter does not currently support FileRegion types other than None."));
+	}
+	SetPackageName(Info.PackageName);
+	// Info.LooseFilePath is ignored
+	HeaderSize = Info.HeaderSize;
+	// Info.ChunkId is ignored
+	HeaderAndExports = PackageData;
+}
+
+void FMemoryPackageStoreWriter::WriteBulkdata(const FBulkDataInfo& Info, const FIoBuffer& BulkData, const TArray<FFileRegion>& FileRegions)
+{
+	SetPackageName(Info.PackageName);
+	// Info.LooseFilePath is ignored
+	// Info.ChunkId is ignored
+	checkf(Info.BulkdataType == IPackageStoreWriter::FBulkDataInfo::Standard, TEXT("MemoryPackageStoreWriter does not currently support BulkData types other than Standard."));
+
+	FIoBuffer& BulkDataOwner = BulkDataArchives.Add_GetRef(BulkData);
+	const uint8* BulkDataStart = BulkData.Data();
+	uint64 BulkDataLen = BulkData.DataSize();
+	for (const FFileRegion& FileRegion : FileRegions)
+	{
+		checkf(FileRegion.Type == EFileRegionType::None, TEXT("FMemoryPackageStoreWriter does not currently support FileRegion types other than None."));
+		checkf(FileRegion.Offset + FileRegion.Length <= BulkDataLen, TEXT("FileRegions in WriteBulkdata were outside of the range of the BulkData's size."));
+		BulkDataRegions.Add(FIoBuffer(BulkDataStart + FileRegion.Offset, FileRegion.Length, BulkDataOwner));
+	}
+}
+
 bool TrySavePackage(UPackage* Package)
 {
 	FString ErrorMessage;
@@ -240,12 +358,12 @@ bool TrySavePackage(UPackage* Package)
 		return false;
 	}
 
-	// EDITOR_DOMAIN_TODO: Need to extend SavePackage to allow saving to an archive
-	FString TempFilename = FPaths::Combine(FPaths::ProjectIntermediateDir(), FGuid::NewGuid().ToString());
-	ON_SCOPE_EXIT{ IFileManager::Get().Delete(*TempFilename); };
-
-	uint32 SaveFlags = SAVE_NoError | // Do not crash the SaveServer on an error
-		SAVE_BulkDataByReference; // EditorDomain saves reference bulkdata from the WorkspaceDomain rather than duplicating it
+	uint32 SaveFlags = SAVE_NoError // Do not crash the SaveServer on an error
+		| SAVE_BulkDataByReference	// EditorDomain saves reference bulkdata from the WorkspaceDomain rather than duplicating it
+		| SAVE_Async				// SavePackage support for PackageStoreWriter is only implemented with SAVE_Async
+		// EDITOR_DOMAIN_TODO: Add a a save flag that specifies the creation of a deterministic guid
+		// | SAVE_KeepGUID;			// Prevent indeterminism by keeping the Guid
+		;
 
 	bool bEditorDomainSaveUnversioned = false;
 	GConfig->GetBool(TEXT("CookSettings"), TEXT("EditorDomainSaveUnversioned"), bEditorDomainSaveUnversioned, GEditorIni);
@@ -261,7 +379,7 @@ bool TrySavePackage(UPackage* Package)
 			UClass* Class = Object ? Object->GetClass() : nullptr;
 			if (Class && Class->HasAnyClassFlags(CLASS_CompiledFromBlueprint))
 			{
-				// TODO: Revisit this once we track package schemas
+				// EDITOR_DOMAIN_TODO: Revisit this once we track package schemas
 				// Packages with Blueprint class instances can not be saved unversioned,
 				// as the Blueprint class's layout can change during the editor's lifetime,
 				// and we don't currently have a way to keep track of the changing package schema
@@ -271,21 +389,26 @@ bool TrySavePackage(UPackage* Package)
 		SaveFlags |= bSaveUnversioned ? SAVE_Unversioned : 0;
 	}
 
-	FSavePackageResultStruct Result = GEditor->Save(Package, nullptr, RF_Standalone, *TempFilename, GError,
-		nullptr /* Conform */, false /* bForceByteSwapping */, true /* bWarnOfLongFilename */, SaveFlags);
+	FMemoryPackageStoreWriter* PackageStoreWriter = new FMemoryPackageStoreWriter();
+	FSavePackageContext SavePackageContext(nullptr /* TargetPlatform */, PackageStoreWriter, false /* bInForceLegacyOffsets */);
+	FSavePackageResultStruct Result = GEditor->Save(Package, nullptr, RF_Standalone, TEXT("EditorDomainPackageStoreWriter"),
+		GError, nullptr /* Conform */, false /* bForceByteSwapping */, true /* bWarnOfLongFilename */,
+		SaveFlags, nullptr /* TargetPlatform */, FDateTime::MinValue(), false /* bSlowTask */,
+		nullptr /* DiffMap */, &SavePackageContext);
 	if (Result.Result != ESavePackageResult::Success)
 	{
 		return false;
 	}
 
+	// EDITOR_DOMAIN_TODO: Convert FIoBuffer to use FSharedBuffer and refer to the SharedBuffers rather than copying
 	FSharedBuffer PackageBuffer;
 	{
 		TArray64<uint8> TempBytes;
-		bool bBytesRead = FFileHelper::LoadFileToArray(TempBytes, *TempFilename);
-		IFileManager::Get().Delete(*TempFilename);
-		if (!bBytesRead)
+		const FIoBuffer& MainBuffer = PackageStoreWriter->GetHeaderAndExports();
+		TempBytes.Append(MainBuffer .Data(), MainBuffer.DataSize());
+		for (const FIoBuffer& BulkBuffer : PackageStoreWriter->GetBulkDatas())
 		{
-			return false;
+			TempBytes.Append(BulkBuffer.Data(), BulkBuffer.DataSize());
 		}
 		PackageBuffer = MakeSharedBufferFromArray(MoveTemp(TempBytes));
 	}
diff --git a/Engine/Source/Runtime/CoreUObject/Private/UObject/SavePackage.cpp b/Engine/Source/Runtime/CoreUObject/Private/UObject/SavePackage.cpp
--- a/Engine/Source/Runtime/CoreUObject/Private/UObject/SavePackage.cpp
+++ b/Engine/Source/Runtime/CoreUObject/Private/UObject/SavePackage.cpp
@@ -4442,6 +4442,7 @@ FSavePackageResultStruct UPackage::Save(UPackage* InOuter, UObject* Base, EObjec
 					}
 
 					// Compress the temporarily file to destination.
+					IPackageStoreWriter* PackageStoreWriter = SavePackageContext ? SavePackageContext->PackageStoreWriter : nullptr;
 					if (bSaveAsync)
 					{	
 						FString NewPathToSave = NewPath;
@@ -4487,50 +4488,52 @@ FSavePackageResultStruct UPackage::Save(UPackage* InOuter, UObject* Base, EObjec
 							FLargeMemoryWriter* Writer = static_cast<FLargeMemoryWriter*>(Linker->Saver);
 							const int64 DataSize = Writer->TotalSize();
 
-							if (IsEventDrivenLoaderEnabledInCookedBuilds() && Linker->IsCooking())
+							if (PackageStoreWriter)
 							{
 								checkf(AdditionalOutputFiles.IsEmpty(), TEXT("Saving additional output files during cooking is not supported!"));
 
 								if (!(SaveFlags & SAVE_DiffCallstack))  // Avoid double counting the package size if SAVE_DiffCallstack flag is set and DiffSettings.bSaveForDiff == true.
 									TotalPackageSizeUncompressed += DataSize;
 
-								if (SavePackageContext != nullptr && SavePackageContext->PackageStoreWriter != nullptr)
-								{
-									FIoBuffer IoBuffer(FIoBuffer::AssumeOwnership, Writer->ReleaseOwnership(), DataSize);
+								FIoBuffer IoBuffer(FIoBuffer::AssumeOwnership, Writer->ReleaseOwnership(), DataSize);
 
-									if (bComputeHash)
+								if (bComputeHash)
+								{
+									FIoBuffer InnerBuffer(IoBuffer.Data(), IoBuffer.DataSize(), IoBuffer);
+									SavePackageUtilities::IncrementOutstandingAsyncWrites();
+									AsyncWriteAndHashSequence.AddWork([InnerBuffer = MoveTemp(InnerBuffer)](FMD5& State)
 									{
-										FIoBuffer InnerBuffer(IoBuffer.Data(), IoBuffer.DataSize(), IoBuffer);
-										SavePackageUtilities::IncrementOutstandingAsyncWrites();
-										AsyncWriteAndHashSequence.AddWork([InnerBuffer = MoveTemp(InnerBuffer)](FMD5& State)
-										{
-											State.Update(InnerBuffer.Data(), InnerBuffer.DataSize());
-											SavePackageUtilities::DecrementOutstandingAsyncWrites();
-										});
-									}
+										State.Update(InnerBuffer.Data(), InnerBuffer.DataSize());
+										SavePackageUtilities::DecrementOutstandingAsyncWrites();
+									});
+								}
 
-									const int32 HeaderSize = Linker->Summary.TotalHeaderSize;
+								const int32 HeaderSize = Linker->Summary.TotalHeaderSize;
 
-									IPackageStoreWriter::FPackageInfo PackageInfo;
-									PackageInfo.PackageName		= InOuter->GetFName();
-									PackageInfo.LooseFilePath	= Filename;
-									PackageInfo.HeaderSize		= HeaderSize;
+								IPackageStoreWriter::FPackageInfo PackageInfo;
+								PackageInfo.PackageName = InOuter->GetFName();
+								PackageInfo.LooseFilePath = Filename;
+								PackageInfo.HeaderSize = HeaderSize;
 
-									FPackageId PackageId = FPackageId::FromName(PackageInfo.PackageName);
-									PackageInfo.ChunkId = CreateIoChunkId(PackageId.Value(), 0, EIoChunkType::ExportBundleData);
+								FPackageId PackageId = FPackageId::FromName(PackageInfo.PackageName);
+								PackageInfo.ChunkId = CreateIoChunkId(PackageId.Value(), 0, EIoChunkType::ExportBundleData);
 
-									SavePackageContext->PackageStoreWriter->WritePackageData(PackageInfo, IoBuffer, Linker->FileRegions);
-								}
-								else
-								{		
-									EAsyncWriteOptions WriteOptions(EAsyncWriteOptions::WriteFileToDisk);
-									if (bComputeHash)
-									{
-										WriteOptions |= EAsyncWriteOptions::ComputeHash;
-									}
+								PackageStoreWriter->WritePackageData(PackageInfo, IoBuffer, Linker->FileRegions);
+							}
+							else if (IsEventDrivenLoaderEnabledInCookedBuilds() && Linker->IsCooking())
+							{
+								checkf(AdditionalOutputFiles.IsEmpty(), TEXT("Saving additional output files during cooking is not supported!"));
 
-									SavePackageUtilities::AsyncWriteFileWithSplitExports(AsyncWriteAndHashSequence, FLargeMemoryPtr(Writer->ReleaseOwnership()), DataSize, Linker->Summary.TotalHeaderSize, *NewPathToSave, WriteOptions, Linker->FileRegions);
+								if (!(SaveFlags & SAVE_DiffCallstack))  // Avoid double counting the package size if SAVE_DiffCallstack flag is set and DiffSettings.bSaveForDiff == true.
+									TotalPackageSizeUncompressed += DataSize;
+
+								EAsyncWriteOptions WriteOptions(EAsyncWriteOptions::WriteFileToDisk);
+								if (bComputeHash)
+								{
+									WriteOptions |= EAsyncWriteOptions::ComputeHash;
 								}
+
+								SavePackageUtilities::AsyncWriteFileWithSplitExports(AsyncWriteAndHashSequence, FLargeMemoryPtr(Writer->ReleaseOwnership()), DataSize, Linker->Summary.TotalHeaderSize, *NewPathToSave, WriteOptions, Linker->FileRegions);
 							}
 							else
 							{
@@ -4554,6 +4557,12 @@ FSavePackageResultStruct UPackage::Save(UPackage* InOuter, UObject* Base, EObjec
 								}	
 							}
 						}
+#if WITH_EDITOR
+						else
+						{
+							checkf(!PackageStoreWriter, TEXT("PackageStoreWriter is not currently supported when diffing."));
+						}
+#endif
 						Linker->CloseAndDestroySaver();
 
 						delete StructuredArchive;
@@ -4563,6 +4572,7 @@ FSavePackageResultStruct UPackage::Save(UPackage* InOuter, UObject* Base, EObjec
 					// Move the temporary file.
 					else
 					{
+						checkf(!PackageStoreWriter, TEXT("PackageStoreWriter is not currently supported with synchronous writes."));
 						checkf(TempFilename.IsSet(), TEXT("The package should've been saved to a tmp file first!"));
 
 						// When saving in text format we will have two temp files, so we need to manually delete the non-textbased one
diff --git a/Engine/Source/Runtime/CoreUObject/Private/UObject/SavePackage/SavePackageUtilities.cpp b/Engine/Source/Runtime/CoreUObject/Private/UObject/SavePackage/SavePackageUtilities.cpp
--- a/Engine/Source/Runtime/CoreUObject/Private/UObject/SavePackage/SavePackageUtilities.cpp
+++ b/Engine/Source/Runtime/CoreUObject/Private/UObject/SavePackage/SavePackageUtilities.cpp
@@ -2223,8 +2223,8 @@ ESavePackageResult SaveBulkData(FLinkerSave* Linker, const UPackage* InOuter, co
 {
 	// Now we write all the bulkdata that is supposed to be at the end of the package
 	// and fix up the offset
-	const int64 StartOfBulkDataArea = Linker->Tell();
-	Linker->Summary.BulkDataStartOffset = StartOfBulkDataArea;
+	IPackageStoreWriter* PackageStoreWriter = SavePackageContext != nullptr ? SavePackageContext->PackageStoreWriter : nullptr;
+	Linker->Summary.BulkDataStartOffset = Linker->Tell();
 
 	if (Linker->BulkDataToAppend.Num() == 0)
 	{
@@ -2267,9 +2267,12 @@ ESavePackageResult SaveBulkData(FLinkerSave* Linker, const UPackage* InOuter, co
 	} CookerUseSeparateSegments;
 	const bool bSeparateSegmentsEnabled = CookerUseSeparateSegments.bEnable && Linker->IsCooking();
 
-	if (bSeparateSegmentsEnabled)
+	if (PackageStoreWriter || bSeparateSegmentsEnabled)
 	{
 		BulkArchive.Reset(new FLargeMemoryWriterWithRegions);
+	}
+	if (bSeparateSegmentsEnabled)
+	{
 		OptionalBulkArchive.Reset(new FLargeMemoryWriterWithRegions);
 		MappedBulkArchive.Reset(new FLargeMemoryWriterWithRegions);
 	}
@@ -2277,6 +2280,7 @@ ESavePackageResult SaveBulkData(FLinkerSave* Linker, const UPackage* InOuter, co
 
 	bool bAlignBulkData = false;
 	bool bUseFileRegions = false;
+	bool bOneBulkPerRegion = false;
 	int64 BulkDataAlignment = 0;
 
 	if (TargetPlatform)
@@ -2285,6 +2289,12 @@ ESavePackageResult SaveBulkData(FLinkerSave* Linker, const UPackage* InOuter, co
 		bUseFileRegions = TargetPlatform->SupportsFeature(ETargetPlatformFeatures::CookFileRegionMetadata);
 		BulkDataAlignment = TargetPlatform->GetMemoryMappingAlignment();
 	}
+	else if (PackageStoreWriter)
+	{
+		bUseFileRegions = true;
+		bOneBulkPerRegion = true;
+	}
+
 	if (bRequestSaveByReference)
 	{
 		const TCHAR* FailureReason = nullptr;
@@ -2335,11 +2345,12 @@ ESavePackageResult SaveBulkData(FLinkerSave* Linker, const UPackage* InOuter, co
 			EndOfPackage,
 			SeparateSegment,
 			Reference,
+			SeparateArchiveAtEndOfPackage,
 		};
 		ESaveLocation SaveLocation = ESaveLocation::Invalid;
-		if (SaveLocation == ESaveLocation::Invalid && bRequestSaveByReference)
+		auto CanSaveByReference = [](FUntypedBulkData* BulkData)
 		{
-			if (BulkData->GetBulkDataOffsetInFile() != INDEX_NONE &&
+			return BulkData->GetBulkDataOffsetInFile() != INDEX_NONE &&
 				// We don't support yet loading from a separate file
 				!BulkData->IsInSeparateFile() &&
 				// It is possible to have a BulkData marked as optional without putting it into a separate file, and we
@@ -2351,17 +2362,21 @@ ESavePackageResult SaveBulkData(FLinkerSave* Linker, const UPackage* InOuter, co
 				// Inline or end-of-package-file data can only be loaded from the workspace domain package file if the
 				// archive used by the bulk data was actually from the package file; BULKDATA_LazyLoadable is set by
 				// Serialize iff that is the case										
-				(BulkData->GetBulkDataFlags() & BULKDATA_LazyLoadable) 
-				)
-			{
-				SaveLocation = ESaveLocation::Reference;
-			}
+				(BulkData->GetBulkDataFlags() & BULKDATA_LazyLoadable);
+		};
+		if (bRequestSaveByReference && CanSaveByReference(BulkData))
+		{
+			SaveLocation = ESaveLocation::Reference;
 		}
-		if (SaveLocation == ESaveLocation::Invalid && bSeparateSegmentsEnabled)
+		else if (bSeparateSegmentsEnabled)
 		{
 			SaveLocation = ESaveLocation::SeparateSegment;
 		}
-		if (SaveLocation == ESaveLocation::Invalid)
+		else if (PackageStoreWriter)
+		{
+			SaveLocation = ESaveLocation::SeparateArchiveAtEndOfPackage;
+		}
+		else
 		{
 			SaveLocation = ESaveLocation::EndOfPackage;
 		}
@@ -2411,6 +2426,13 @@ ESavePackageResult SaveBulkData(FLinkerSave* Linker, const UPackage* InOuter, co
 					TargetRegions = &BulkArchive->FileRegions;
 				}
 			}
+			else if (SaveLocation == ESaveLocation::SeparateArchiveAtEndOfPackage)
+			{
+				BulkDataFlags |= BULKDATA_NoOffsetFixUp;
+				TargetArchive = BulkArchive.Get();
+				TargetRegions = &BulkArchive->FileRegions;
+			}
+
 			check(TargetArchive && TargetRegions);
 
 			// Pad archive for proper alignment for memory mapping
@@ -2447,11 +2469,25 @@ ESavePackageResult SaveBulkData(FLinkerSave* Linker, const UPackage* InOuter, co
 			BulkData->SerializeBulkData(*TargetArchive, BulkData->Lock(LOCK_READ_ONLY), static_cast<EBulkDataFlags>(BulkDataFlags));
 			BulkData->Unlock();
 			const int64 BulkEndOffset = TargetArchive->Tell();
-
-			BulkDataOffsetInFile = (BulkDataFlags & BULKDATA_NoOffsetFixUp) == 0 ? BulkStartOffset - StartOfBulkDataArea : BulkStartOffset;
+			if (SaveLocation == ESaveLocation::SeparateArchiveAtEndOfPackage)
+			{
+				// BulkDatas will be read from a CompositeArchive that holds an exports segment followed by a bulkdata segment.
+				// We record offsets in the bulkdata segment relative to the beginning of the CompositeArchive.
+				// We have to do this to distinguish them from inline bulkdatas in the exports segment
+				check((BulkDataFlags& BULKDATA_NoOffsetFixUp) != 0); // offset fixups are not supported with this SaveLocation
+				BulkDataOffsetInFile = BulkStartOffset + Linker->Summary.BulkDataStartOffset;
+			}
+			else if ((BulkDataFlags & BULKDATA_NoOffsetFixUp) == 0)
+			{
+				BulkDataOffsetInFile = BulkStartOffset - Linker->Summary.BulkDataStartOffset;
+			}
+			else
+			{
+				BulkDataOffsetInFile = BulkStartOffset;
+			}
 			BulkDataSizeOnDisk = BulkEndOffset - BulkStartOffset;
 
-			if (bUseFileRegions && BulkDataStorageInfo.BulkDataFileRegionType != EFileRegionType::None && BulkDataSizeOnDisk > 0)
+			if ((bOneBulkPerRegion || (bUseFileRegions && BulkDataStorageInfo.BulkDataFileRegionType != EFileRegionType::None)) && BulkDataSizeOnDisk > 0)
 			{
 				TargetRegions->Add(FFileRegion(BulkStartOffset, BulkDataSizeOnDisk, BulkDataStorageInfo.BulkDataFileRegionType));
 			}
@@ -2472,19 +2508,16 @@ ESavePackageResult SaveBulkData(FLinkerSave* Linker, const UPackage* InOuter, co
 		if (Linker->bUpdatingLoadedPath)
 		{
 			BulkData->SetFlagsFromDiskWrittenValues(static_cast<EBulkDataFlags>(BulkDataFlags), BulkDataOffsetInFile,
-				BulkDataSizeOnDisk, StartOfBulkDataArea);
+				BulkDataSizeOnDisk, Linker->Summary.BulkDataStartOffset);
 		}
 #endif
 	}
 
 	if (BulkArchive)
 	{
-		check(OptionalBulkArchive);
-		check(MappedBulkArchive);
-
 		const bool bWriteBulkToDisk = !bDiffing;
 
-		if (SavePackageContext != nullptr && SavePackageContext->PackageStoreWriter != nullptr && bWriteBulkToDisk)
+		if (PackageStoreWriter && bWriteBulkToDisk)
 		{
 			auto AddSizeAndConvertToIoBuffer = [&TotalPackageSizeUncompressed](FLargeMemoryWriter* Writer)
 			{
@@ -2503,28 +2536,28 @@ ESavePackageResult SaveBulkData(FLinkerSave* Linker, const UPackage* InOuter, co
 				BulkInfo.ChunkId = CreateIoChunkId(PackageId.Value(), 0, EIoChunkType::BulkData);
 				BulkInfo.BulkdataType = IPackageStoreWriter::FBulkDataInfo::Standard;
 				BulkInfo.LooseFilePath = FPaths::ChangeExtension(Filename, LexToString(EPackageExtension::BulkDataDefault));
-				SavePackageContext->PackageStoreWriter->WriteBulkdata(BulkInfo, AddSizeAndConvertToIoBuffer(BulkArchive.Get()), BulkArchive->FileRegions);
+				PackageStoreWriter->WriteBulkdata(BulkInfo, AddSizeAndConvertToIoBuffer(BulkArchive.Get()), BulkArchive->FileRegions);
 			}
-			if (OptionalBulkArchive->TotalSize())
+			if (OptionalBulkArchive && OptionalBulkArchive->TotalSize())
 			{
 				BulkInfo.ChunkId = CreateIoChunkId(PackageId.Value(), 0, EIoChunkType::OptionalBulkData);
 				BulkInfo.BulkdataType = IPackageStoreWriter::FBulkDataInfo::Optional;
 				BulkInfo.LooseFilePath = FPaths::ChangeExtension(Filename, LexToString(EPackageExtension::BulkDataOptional));
-				SavePackageContext->PackageStoreWriter->WriteBulkdata(BulkInfo, AddSizeAndConvertToIoBuffer(OptionalBulkArchive.Get()), OptionalBulkArchive->FileRegions);
+				PackageStoreWriter->WriteBulkdata(BulkInfo, AddSizeAndConvertToIoBuffer(OptionalBulkArchive.Get()), OptionalBulkArchive->FileRegions);
 			}
-			if (MappedBulkArchive->TotalSize())
+			if (MappedBulkArchive && MappedBulkArchive->TotalSize())
 			{
 				BulkInfo.ChunkId = CreateIoChunkId(PackageId.Value(), 0, EIoChunkType::MemoryMappedBulkData);
 				BulkInfo.BulkdataType = IPackageStoreWriter::FBulkDataInfo::Mmap;
 				BulkInfo.LooseFilePath = FPaths::ChangeExtension(Filename, LexToString(EPackageExtension::BulkDataMemoryMapped));
-				SavePackageContext->PackageStoreWriter->WriteBulkdata(BulkInfo, AddSizeAndConvertToIoBuffer(MappedBulkArchive.Get()), MappedBulkArchive->FileRegions);
+				PackageStoreWriter->WriteBulkdata(BulkInfo, AddSizeAndConvertToIoBuffer(MappedBulkArchive.Get()), MappedBulkArchive->FileRegions);
 			}
 		}
 		else
 		{
 			auto WriteBulkData = [&](FLargeMemoryWriterWithRegions* Archive, const TCHAR* BulkFileExtension)
 			{
-				if (const int64 DataSize = Archive->TotalSize())
+				if (const int64 DataSize = Archive ? Archive->TotalSize() : 0)
 				{
 					TotalPackageSizeUncompressed += DataSize;
 
diff --git a/Engine/Source/Runtime/CoreUObject/Private/UObject/SavePackage2.cpp b/Engine/Source/Runtime/CoreUObject/Private/UObject/SavePackage2.cpp
--- a/Engine/Source/Runtime/CoreUObject/Private/UObject/SavePackage2.cpp
+++ b/Engine/Source/Runtime/CoreUObject/Private/UObject/SavePackage2.cpp
@@ -1771,6 +1771,8 @@ ESavePackageResult FinalizeFile(FStructuredArchive::FRecord& StructuredArchiveRo
 		ResetLoadersForSave(SaveContext.GetPackage(), SaveContext.GetFilename());
 	}
 
+	FSavePackageContext* SavePackageContext = SaveContext.GetSavePackageContext();
+	IPackageStoreWriter* PackageStoreWriter = SavePackageContext ? SavePackageContext->PackageStoreWriter : nullptr;
 	if (SaveContext.IsSaveAsync())
 	{
 		FString PathToSave = SaveContext.GetFilename();
@@ -1821,8 +1823,7 @@ ESavePackageResult FinalizeFile(FStructuredArchive::FRecord& StructuredArchiveRo
 
 			checkf(SaveContext.IsCooking() == false || SaveContext.AdditionalPackageFiles.IsEmpty(), TEXT("Saving additional output files during cooking is not currently supported! (%s)"), *PathToSave);
 
-			FSavePackageContext* SavePackageContext = SaveContext.GetSavePackageContext();
-			if (SavePackageContext != nullptr && SavePackageContext->PackageStoreWriter != nullptr && SaveContext.IsCooking())
+			if (PackageStoreWriter)
 			{
 				FIoBuffer IoBuffer(FIoBuffer::AssumeOwnership, Writer->ReleaseOwnership(), DataSize);
 
@@ -1847,7 +1848,7 @@ ESavePackageResult FinalizeFile(FStructuredArchive::FRecord& StructuredArchiveRo
 				FPackageId PackageId = FPackageId::FromName(PackageInfo.PackageName);
 				PackageInfo.ChunkId = CreateIoChunkId(PackageId.Value(), 0, EIoChunkType::ExportBundleData);
 
-				SavePackageContext->PackageStoreWriter->WritePackageData(PackageInfo, IoBuffer, Linker->FileRegions);
+				PackageStoreWriter->WritePackageData(PackageInfo, IoBuffer, Linker->FileRegions);
 			}
 			else
 			{
@@ -1874,9 +1875,15 @@ ESavePackageResult FinalizeFile(FStructuredArchive::FRecord& StructuredArchiveRo
 			}
 			SaveContext.CloseLinkerArchives();
 		}
+		else
+		{
+			checkf(!PackageStoreWriter, TEXT("PackageStoreWriter is not currently supported when diffing."));
+		}
 	}
 	else
 	{
+		checkf(!PackageStoreWriter, TEXT("PackageStoreWriter is not currently supported with synchronous writes."));
+
 		// Destroy archives used for saving, closing file handle.
 		if (SaveContext.CloseLinkerArchives() == false)
 		{
