Rework how the virtualization backends are defined.

#rnx

- When a package is saved we now push it's payloads to local storage backends (in memory/local file cache etc)
- In a future piece of work, when a package is submitted we will push it's payloads to the persistent storage backends (jupiter/source control etc)
- When pulling a payload we will first search the local storage backends for it and then go to the slower persistent storage backends if it was not found.
-- Once a payload has been found when being pulled we will immediately cache it to the local storage backends so that subsequent pulls will be faster.
-- At some point a ZenStorage backend will be added which will most likely be the primary local storage backend.

diff --git a/Engine/Config/BaseEngine.ini b/Engine/Config/BaseEngine.ini
--- a/Engine/Config/BaseEngine.ini
+++ b/Engine/Config/BaseEngine.ini
@@ -1057,17 +1057,20 @@ ValidateAfterPushOperation=False
 
 ; A backend graph with no backends
 [ContentVirtualizationBackendGraph_None]
-Hierarchy=()
+LocalStorageHierarchy=()
+PersistentStorageHierarchy=()
 
 ; A simple backend graph that will store virtualized data locally, suitable for testing.
 [ContentVirtualizationBackendGraph_LocalCache]
-Hierarchy=(Entry=LocalCache)
-LocalCache=(Type=FileSystem, Path="%GAMESAVEDDIR%VirtualizedPayloads/")
+LocalStorageHierarchy=(Entry=LocalFileCache)
+PersistentStorageHierarchy=()
+LocalFileCache=(Type=FileSystem, Path="%GAMESAVEDDIR%VirtualizedPayloads/")
 
 ; A backend graph that will connect to a locally hosted Jupiter service along with a local file cache, suitable for testing.
 [ContentVirtualizationBackendGraph_JupiterLocalHost]
-Hierarchy=(Entry=LocalCache, Entry=JupiterCache)
-LocalCache=(Type=FileSystem, Path="%GAMESAVEDDIR%VirtualizedPayloads/")
+LocalStorageHierarchy=(Entry=LocalFileCache)
+PersistentStorageHierarchy=(JupiterCache)
+LocalFileCache=(Type=FileSystem, Path="%GAMESAVEDDIR%VirtualizedPayloads/")
 JupiterCache=(Type=Jupiter, Host="http://localhost", Namespace="ue.mirage", ChunkSize=524288)
 
 [/Script/Engine.StreamingSettings]
diff --git a/Engine/Source/Runtime/CoreUObject/Private/Virtualization/VirtualizationManager.cpp b/Engine/Source/Runtime/CoreUObject/Private/Virtualization/VirtualizationManager.cpp
--- a/Engine/Source/Runtime/CoreUObject/Private/Virtualization/VirtualizationManager.cpp
+++ b/Engine/Source/Runtime/CoreUObject/Private/Virtualization/VirtualizationManager.cpp
@@ -50,9 +50,9 @@ private:
 };
 
 /* Utility function for building up a lookup table of all available IBackendFactory interfaces*/
-TMap<FName, IVirtualizationBackendFactory*> FindBackendFactories()
+FVirtualizationManager::FRegistedFactories FindBackendFactories()
 {
-	TMap<FName, IVirtualizationBackendFactory*> BackendFactories;
+	FVirtualizationManager::FRegistedFactories BackendFactories;
 
 	TArray<IVirtualizationBackendFactory*> FactoriesArray = IModularFeatures::Get().GetModularFeatureImplementations<IVirtualizationBackendFactory>(FName("VirtualizationBackendFactory"));
 	for (IVirtualizationBackendFactory* FactoryInterface : FactoriesArray)
@@ -215,25 +215,22 @@ FVirtualizationManager::FVirtualizationManager()
 FVirtualizationManager::~FVirtualizationManager()
 {
 	UE_LOG(LogVirtualization, Log, TEXT("Destroying backends"));
+	
+	LocalCachableBackends.Empty();
+	PersistentStorageBackends.Empty();
+	PullEnabledBackends.Empty();
 
-	for (IVirtualizationBackend* Backend : AllBackendsArray)
-	{
-		delete Backend;
-	}
-
-	AllBackendsArray.Empty();
-	PullEnabledBackendsArray.Empty();
-	PushEnabledBackendsArray.Empty();
+	AllBackends.Empty(); // This will delete all backends and beyond this point all references to them are invalid
 
 	UE_LOG(LogVirtualization, Log, TEXT("Virtualization manager destroyed"));
 }
 
 bool FVirtualizationManager::IsEnabled() const
 {
-	return !PushEnabledBackendsArray.IsEmpty();
+	return !AllBackends.IsEmpty();
 }
 
-bool FVirtualizationManager::PushData(const FPayloadId& Id, const FCompressedBuffer& Payload)
+bool FVirtualizationManager::PushData(const FPayloadId& Id, const FCompressedBuffer& Payload, EStorageType StorageType)
 {
 	TRACE_CPUPROFILER_EVENT_SCOPE(FVirtualizationManager::PushData);
 
@@ -247,7 +244,7 @@ bool FVirtualizationManager::PushData(const FPayloadId& Id, const FCompressedBuf
 	FConditionalScopeLock _(&ForceSingleThreadedCS, bForceSingleThreaded);
 
 	// Early out if there are no backends or if the pushing of payloads has been disabled
-	if (PushEnabledBackendsArray.IsEmpty() || bEnablePayloadPushing == false)
+	if (!IsEnabled() || bEnablePayloadPushing == false)
 	{
 		return false;
 	}
@@ -274,8 +271,11 @@ bool FVirtualizationManager::PushData(const FPayloadId& Id, const FCompressedBuf
 	// TODO: Note that all push operations are currently synchronous, probably 
 	// should change to async at some point, although this makes handling failed
 	// pushed much more difficult.
+
 	bool bWasPayloadPushed = false;
-	for (IVirtualizationBackend* Backend : PushEnabledBackendsArray)
+	FBackendArray& Backends = StorageType == EStorageType::Local ? LocalCachableBackends : PersistentStorageBackends;
+
+	for (IVirtualizationBackend* Backend : Backends)
 	{
 		const EPushResult Result = TryPushDataToBackend(*Backend, Id, Payload) ? EPushResult::Success : EPushResult::Failed;
 
@@ -313,7 +313,7 @@ FCompressedBuffer FVirtualizationManager::PullData(const FPayloadId& Id)
 		return FCompressedBuffer();
 	}
 
-	if (PullEnabledBackendsArray.IsEmpty())
+	if (PullEnabledBackends.IsEmpty())
 	{
 		// TODO: See below, should errors here be fatal?
 		UE_LOG(LogVirtualization, Error, TEXT("Payload '%s' failed to be pulled as there are no backends mounted!'"), *Id.ToString());
@@ -328,13 +328,12 @@ FCompressedBuffer FVirtualizationManager::PullData(const FPayloadId& Id)
 
 	FConditionalScopeLock _(&ForceSingleThreadedCS, bForceSingleThreaded);
 
-	// TODO: Once a payload is found, other backends should probably be notified 
-	// (a local cache might want to replicate the data for example)
-	for (IVirtualizationBackend* Backend : PullEnabledBackendsArray)
+	for (IVirtualizationBackend* Backend : PullEnabledBackends)
 	{
 		FCompressedBuffer Payload = PullDataFromBackend(*Backend, Id);
 		if (Payload)
 		{
+			CachePayload(Id, Payload, Backend);
 			return Payload;
 		}
 	}
@@ -462,14 +461,21 @@ void FVirtualizationManager::MountBackends()
 {
 	UE_LOG(LogVirtualization, Log, TEXT("Mounting virtualization backends..."));
 
-	TMap<FName, IVirtualizationBackendFactory*> FactoryLookupTable = FindBackendFactories();
+	const FRegistedFactories FactoryLookupTable = FindBackendFactories();
 	UE_LOG(LogVirtualization, Verbose, TEXT("Found %d backend factories"), FactoryLookupTable.Num());
 
 	const TCHAR* GraphName = *BackendGraphName;
-	const TCHAR* HierarchyKey = TEXT("Hierarchy");
 
 	UE_LOG(LogVirtualization, Log, TEXT("Using backend graph: '%s'"), GraphName);
 
+	// It is important to parse the local storage hierarchy first so those backends will show up before the
+	// persistent storage backends in 'PullEnabledBackends'.
+	ParseHierarchy(GraphName, TEXT("LocalStorageHierarchy"), FactoryLookupTable, LocalCachableBackends);
+	ParseHierarchy(GraphName, TEXT("PersistentStorageHierarchy"), FactoryLookupTable, PersistentStorageBackends);
+}
+
+void FVirtualizationManager::ParseHierarchy(const TCHAR* GraphName, const TCHAR* HierarchyKey, const FRegistedFactories& FactoryLookupTable, FBackendArray& PushArray)
+{
 	FString HierarchyData;
 	if (!GConfig->GetString(GraphName, HierarchyKey, HierarchyData, GEngineIni))
 	{
@@ -483,15 +489,15 @@ void FVirtualizationManager::MountBackends()
 
 	TArray<FString> Entries = ParseEntries(HierarchyData);
 
-	UE_LOG(LogVirtualization, Log, TEXT("The backend graph hierarchy has %d entries"), Entries.Num());
+	UE_LOG(LogVirtualization, Log, TEXT("The backend graph hierarchy '%s' has %d entries"), HierarchyKey, Entries.Num());
 
 	for (const FString& Entry : Entries)
 	{
-		CreateBackend(GraphName, Entry, FactoryLookupTable);
+		CreateBackend(GraphName, Entry, FactoryLookupTable, PushArray);
 	}
 }
 
-bool FVirtualizationManager::CreateBackend(const TCHAR* GraphName, const FString& ConfigEntryName, TMap<FName, IVirtualizationBackendFactory*>& FactoryLookupTable)
+bool FVirtualizationManager::CreateBackend(const TCHAR* GraphName, const FString& ConfigEntryName, const FRegistedFactories& FactoryLookupTable, FBackendArray& PushArray)
 {
 	// All failures in this method are considered fatal, however it still returns true/false in case we decide
 	// to be more forgiving in the future.
@@ -511,11 +517,11 @@ bool FVirtualizationManager::CreateBackend(const TCHAR* GraphName, const FString
 		FString Cmdine = BackendData.RightChop(BackendData.Find(BackendType) + BackendType.Len());
 		Cmdine.RemoveFromEnd(TEXT(")"));
 
-		IVirtualizationBackendFactory** FactoryPtr = FactoryLookupTable.Find(FName(BackendType));
+		UE::Virtualization::IVirtualizationBackendFactory* const* FactoryPtr = FactoryLookupTable.Find(FName(BackendType));
 		if (FactoryPtr != nullptr && *FactoryPtr != nullptr)
 		{
 			IVirtualizationBackendFactory* Factory = *FactoryPtr;
-			IVirtualizationBackend* Backend = Factory->CreateInstance(ConfigEntryName);
+			TUniquePtr<IVirtualizationBackend> Backend = Factory->CreateInstance(ConfigEntryName);
 
 			if (Backend == nullptr)
 			{
@@ -526,12 +532,11 @@ bool FVirtualizationManager::CreateBackend(const TCHAR* GraphName, const FString
 			
 			if (Backend->Initialize(Cmdine))
 			{
-				AddBackend(Backend);
+				AddBackend(MoveTemp(Backend), PushArray);
 			}
 			else
 			{
 				UE_LOG(LogVirtualization, Fatal, TEXT("Backend '%s' reported errors when initializing"), *ConfigEntryName);
-				delete Backend;
 				return false;
 			}
 		}
@@ -550,25 +555,49 @@ bool FVirtualizationManager::CreateBackend(const TCHAR* GraphName, const FString
 	return true;
 }
 
-void FVirtualizationManager::AddBackend(IVirtualizationBackend* Backend)
+void FVirtualizationManager::AddBackend(TUniquePtr<IVirtualizationBackend> Backend, FBackendArray& PushArray)
 {
-	checkf(!AllBackendsArray.Contains(Backend), TEXT("Adding the same virtualization backend (%s) multiple times!"), *Backend->GetDebugString());
+	checkf(!AllBackends.Contains(Backend), TEXT("Adding the same virtualization backend (%s) multiple times!"), *Backend->GetDebugString());
+
+	// Move ownership of the backend to AllBackends
+	AllBackends.Add(MoveTemp(Backend));
 
-	AllBackendsArray.Add(Backend);
+	// Get a reference pointer to use in the other backend arrays
+	IVirtualizationBackend* BackendRef = AllBackends.Last().Get();
 
-	if (Backend->SupportsPullOperations())
+	if (BackendRef->SupportsPullOperations())
 	{
-		PullEnabledBackendsArray.Add(Backend);
+		PullEnabledBackends.Add(BackendRef);
 	}
 
-	if (Backend->SupportsPushOperations())
+	if (BackendRef->SupportsPushOperations())
 	{
-		PushEnabledBackendsArray.Add(Backend);
+		PushArray.Add(BackendRef);
 	}
 
-	COOK_STAT(Profiling::CreateStats(*Backend));
+	COOK_STAT(Profiling::CreateStats(*BackendRef));
 
-	UE_LOG(LogVirtualization, Log, TEXT("Mounted backend: %s"), *Backend->GetDebugString());
+	UE_LOG(LogVirtualization, Log, TEXT("Mounted backend: %s"), *BackendRef->GetDebugString());
+}
+
+void FVirtualizationManager::CachePayload(const FPayloadId& Id, const FCompressedBuffer& Payload, const IVirtualizationBackend* BackendSource)
+{
+	TRACE_CPUPROFILER_EVENT_SCOPE(FVirtualizationManager::CachePayload);
+
+	// We start caching at the first (assumed to be fastest) local cache backend. 
+	for (IVirtualizationBackend* BackendToCache : LocalCachableBackends)
+	{
+		if (BackendToCache == BackendSource)
+		{
+			return; // No point going past BackendSource
+		}
+
+		EPushResult Result = BackendToCache->PushData(Id, Payload);
+		UE_CLOG(Result == EPushResult::Failed, LogVirtualization, Warning,
+			TEXT("Failed to cache payload '%s' to backend '%s'"),
+			*Id.ToString(),
+			*BackendToCache->GetDebugString());
+	}
 }
 
 bool FVirtualizationManager::TryPushDataToBackend(IVirtualizationBackend& Backend, const FPayloadId& Id, const FCompressedBuffer& Payload)
diff --git a/Engine/Source/Runtime/CoreUObject/Private/Virtualization/VirtualizedBulkData.cpp b/Engine/Source/Runtime/CoreUObject/Private/Virtualization/VirtualizedBulkData.cpp
--- a/Engine/Source/Runtime/CoreUObject/Private/Virtualization/VirtualizedBulkData.cpp
+++ b/Engine/Source/Runtime/CoreUObject/Private/Virtualization/VirtualizedBulkData.cpp
@@ -842,7 +842,8 @@ void FVirtualizedUntypedBulkData::PushData()
 		// serialization, we need a better way to save the results of 'RecompressForSerialization'
 		RecompressForSerialization(PayloadToPush, Flags);
 
-		if (VirtualizationManager.PushData(PayloadContentId, PayloadToPush))
+		// TODO: We could make this a config option?
+		if (VirtualizationManager.PushData(PayloadContentId, PayloadToPush, EStorageType::Local))
 		{
 			EnumAddFlags(Flags, EFlags::IsVirtualized);
 			EnumRemoveFlags(Flags, EFlags::ReferencesLegacyFile | EFlags::LegacyFileIsCompressed);
diff --git a/Engine/Source/Runtime/CoreUObject/Public/Virtualization/IVirtualizationBackend.h b/Engine/Source/Runtime/CoreUObject/Public/Virtualization/IVirtualizationBackend.h
--- a/Engine/Source/Runtime/CoreUObject/Public/Virtualization/IVirtualizationBackend.h
+++ b/Engine/Source/Runtime/CoreUObject/Public/Virtualization/IVirtualizationBackend.h
@@ -6,6 +6,7 @@
 #include "Containers/StringView.h"
 #include "Features/IModularFeature.h"
 #include "Features/IModularFeatures.h"
+#include "Templates/UniquePtr.h"
 
 namespace UE::Virtualization
 {
@@ -126,7 +127,7 @@ public:
 	 * @param ConfigName	The name given to the back end in the config ini file
 	 * @return A new backend instance
 	 */
-	virtual IVirtualizationBackend* CreateInstance(FStringView ConfigName) = 0;
+	virtual TUniquePtr<IVirtualizationBackend> CreateInstance(FStringView ConfigName) = 0;
 
 	/** Returns the name used to identify the type in config ini files */
 	virtual FName GetName() = 0;
@@ -148,7 +149,7 @@ public:
 		BackendClass##Factory() { IModularFeatures::Get().RegisterModularFeature(FName("VirtualizationBackendFactory"), this); }\
 		virtual ~BackendClass##Factory() { IModularFeatures::Get().UnregisterModularFeature(FName("VirtualizationBackendFactory"), this); } \
 	private: \
-		virtual IVirtualizationBackend* CreateInstance(FStringView ConfigName) override { return new BackendClass(ConfigName); } \
+		virtual TUniquePtr<IVirtualizationBackend> CreateInstance(FStringView ConfigName) override { return MakeUnique<BackendClass>(ConfigName); } \
 		virtual FName GetName() override { return FName(#ConfigName); } \
 	}; \
 	static BackendClass##Factory BackendClass##Factory##Instance;
diff --git a/Engine/Source/Runtime/CoreUObject/Public/Virtualization/VirtualizationManager.h b/Engine/Source/Runtime/CoreUObject/Public/Virtualization/VirtualizationManager.h
--- a/Engine/Source/Runtime/CoreUObject/Public/Virtualization/VirtualizationManager.h
+++ b/Engine/Source/Runtime/CoreUObject/Public/Virtualization/VirtualizationManager.h
@@ -5,6 +5,7 @@
 #include "Compression/CompressedBuffer.h"
 #include "HAL/CriticalSection.h"
 #include "Logging/LogMacros.h"
+#include "Templates/UniquePtr.h"
 #include "Virtualization/PayloadId.h"
 
 // TODO: Do we want to keep this header public for UE5 release? If the only interaction should be
@@ -45,7 +46,6 @@
 
 namespace UE::Virtualization
 {
-
 class IVirtualizationBackend;
 class IVirtualizationBackendFactory;
 
@@ -66,11 +66,23 @@ struct FPayloadActivityInfo
 	FActivity Push;
 };
 
+/** Describes the type of storage to use for a given action */
+enum EStorageType
+{
+	/** Store in the local cache backends, this can be called from any thread */
+	Local = 0,
+	/** Store in the persistent backends, this can only be called from the game thread due to limitations with ISourceControlModule. */
+	Persistent
+};
+
 /** This is used as a wrapper around the various potential back end implementations. 
 	The calling code shouldn't need to care about which back ends are actually in use. */
 class COREUOBJECT_API FVirtualizationManager
 {
 public:
+	using FRegistedFactories = TMap<FName, IVirtualizationBackendFactory*>;
+	using FBackendArray = TArray<IVirtualizationBackend*>;
+
 	/** Singleton access */
 	static FVirtualizationManager& Get();
 
@@ -79,17 +91,19 @@ public:
 
 	/** Poll to see if content virtualization is enabled or not. */
 	bool IsEnabled() const;
-
+	
 	/** 
 	 * Push a payload to the virtualization backends.
 	 * 
-	 * @param	Id The identifier of the payload being pushed.
-	 * @param	Payload The payload itself in FCompressedBuffer form, it is assumed
-	 *			that if the buffer is to be compressed that it will have been done
-	 *			by the caller.
+	 * @param	Id			The identifier of the payload being pushed.
+	 * @param	Payload		The payload itself in FCompressedBuffer form, it is 
+	 *						assumed that if the buffer is to be compressed that
+	 *						it will have been done by the caller.
+	 * @param	StorageType	The type of storage to push the payload to, see EStorageType
+	 *						for details.
 	 * @return	True if at least one backend now contains the payload, otherwise false.
 	 */
-	bool PushData(const FPayloadId& Id, const FCompressedBuffer& Payload);
+	bool PushData(const FPayloadId& Id, const FCompressedBuffer& Payload, EStorageType StorageType);
 
 	/** 
 	 * Pull a payload from the virtualization backends.
@@ -107,19 +121,22 @@ public:
 	FPayloadActivityInfo GetPayloadActivityInfo() const;
 
 private:
-
+	
 	void ApplySettingsFromConfigFiles(const FConfigFile& PlatformEngineIni);
 	void ApplySettingsFromCmdline();
 
 	void ApplyDebugSettingsFromConfigFiles(const FConfigFile& PlatformEngineIni);
 
 	void MountBackends();
-	bool CreateBackend(const TCHAR* GraphName, const FString& ConfigEntryName, TMap<FName, IVirtualizationBackendFactory*>& FactoryLookupTable);
+	void ParseHierarchy(const TCHAR* GraphName, const TCHAR* HierarchyKey, const FRegistedFactories& FactoryLookupTable, FBackendArray& PushArray);
+	bool CreateBackend(const TCHAR* GraphName, const FString& ConfigEntryName, const FRegistedFactories& FactoryLookupTable, FBackendArray& PushArray);
+
+	void AddBackend(TUniquePtr<IVirtualizationBackend> Backend, FBackendArray& PushArray);
 
-	void AddBackend(class IVirtualizationBackend* Backend);
+	void CachePayload(const FPayloadId& Id, const FCompressedBuffer& Payload, const IVirtualizationBackend* BackendSource);
 
-	bool TryPushDataToBackend(class IVirtualizationBackend& Backend, const FPayloadId& Id, const FCompressedBuffer& Payload);
-	FCompressedBuffer PullDataFromBackend(class IVirtualizationBackend& Backend, const FPayloadId& Id);
+	bool TryPushDataToBackend(IVirtualizationBackend& Backend, const FPayloadId& Id, const FCompressedBuffer& Payload);
+	FCompressedBuffer PullDataFromBackend(IVirtualizationBackend& Backend, const FPayloadId& Id);
 
 	/** Are payloads allowed to be virtualized. Defaults to true. */
 	bool bEnablePayloadPushing;
@@ -150,14 +167,20 @@ private:
 	/** The critical section used to force single threaded access if bForceSingleThreaded is true */
 	FCriticalSection ForceSingleThreadedCS;
 
-	/** Collection of all the instantiated backends */
-	TArray<IVirtualizationBackend*> AllBackendsArray;
+	/** All of the backends that were mounted during graph creation */
+	TArray<TUniquePtr<IVirtualizationBackend>> AllBackends;
 
-	/** Collection of all the instantiated backends that support pull operations*/
-	TArray<IVirtualizationBackend*> PullEnabledBackendsArray;
+	/** Backends used for caching operations (must support push operations). */
+	FBackendArray LocalCachableBackends; 
 
-	/** Collection of all the instantiated backends that support push operations*/
-	TArray<IVirtualizationBackend*> PushEnabledBackendsArray;
+	/** Backends used for persistent storage operations (must support push operations). */
+	FBackendArray PersistentStorageBackends; 
+
+	/** 
+	 * The hierarchy of backends to pull from, this is assumed to be ordered from fastest to slowest
+	 * and can contain a mixture of local cacheable and persistent backends 
+	 */
+	FBackendArray PullEnabledBackends;
 };
 
 } // namespace UE::Virtualization
