Compile time improvement: restore removal of LinkerSave.h from Package.h that was accidentally added back.
#rnx

diff --git a/Engine/Source/Runtime/CoreUObject/Public/UObject/Package.h b/Engine/Source/Runtime/CoreUObject/Public/UObject/Package.h
--- a/Engine/Source/Runtime/CoreUObject/Public/UObject/Package.h
+++ b/Engine/Source/Runtime/CoreUObject/Public/UObject/Package.h
@@ -7,7 +7,6 @@
 #include "UObject/UObjectGlobals.h"
 #include "UObject/Object.h"
 #include "IO/PackageId.h"
-#include "UObject/LinkerSave.h"
 #include "Misc/Guid.h"
 #include "Misc/OutputDeviceError.h"
 #include "Misc/ObjectThumbnail.h"
