Replacing the UClass::StaticRegistrationInfo method with a extern function that won't be exported from the DLL if the class is marked as exported.

#rnx

diff --git a/Engine/Source/Programs/UnrealHeaderTool/Private/CodeGenerator.cpp b/Engine/Source/Programs/UnrealHeaderTool/Private/CodeGenerator.cpp
--- a/Engine/Source/Programs/UnrealHeaderTool/Private/CodeGenerator.cpp
+++ b/Engine/Source/Programs/UnrealHeaderTool/Private/CodeGenerator.cpp
@@ -2132,11 +2132,16 @@ void FNativeClassHeaderGenerator::ExportNativeGeneratedInitCode(FOutputDevice& O
 		GeneratedClassRegisterFunctionText.Logf(TEXT("\t};\r\n"));
 		GeneratedClassRegisterFunctionText.Log(*StaticDefinitions);
 
+		if (!bIsDynamic)
+		{
+			GeneratedClassRegisterFunctionText.Logf(TEXT("\textern FClassRegistrationInfo& Z_Registration_Info_UClass_%s();\r\n"), *ClassNameCPP);
+		}
+
 		GeneratedClassRegisterFunctionText.Logf(TEXT("\tUClass* %s\r\n"), *SingletonName);
 		GeneratedClassRegisterFunctionText.Logf(TEXT("\t{\r\n"));
 		if (!bIsDynamic)
 		{
-			GeneratedClassRegisterFunctionText.Logf(TEXT("\t\tstatic UClass*& OuterClass = %s::StaticRegistrationInfo().OuterSingleton;\r\n"), *ClassNameCPP);
+			GeneratedClassRegisterFunctionText.Logf(TEXT("\t\tstatic UClass*& OuterClass = Z_Registration_Info_UClass_%s().OuterSingleton;\r\n"), *ClassNameCPP);
 			GeneratedClassRegisterFunctionText.Logf(TEXT("\t\tif (!OuterClass)\r\n"));
 		}
 		else
diff --git a/Engine/Source/Runtime/CoreUObject/Private/UObject/CoreNative.cpp b/Engine/Source/Runtime/CoreUObject/Private/UObject/CoreNative.cpp
--- a/Engine/Source/Runtime/CoreUObject/Private/UObject/CoreNative.cpp
+++ b/Engine/Source/Runtime/CoreUObject/Private/UObject/CoreNative.cpp
@@ -64,9 +64,10 @@ IMPLEMENT_MODULE( FCoreUObjectModule, CoreUObject );
 // if we are not using compiled in natives, we still need this as a base class for intrinsics
 #if !USE_COMPILED_IN_NATIVES
 COREUOBJECT_API class UClass* Z_Construct_UClass_UObject();
+extern FClassRegistrationInfo& Z_Registration_Info_UClass_UObject();
 UClass* Z_Construct_UClass_UObject()
 {
-	static UClass* OuterClass = UObject::StaticRegistrationInfo().OuterSingleton;
+	static UClass* OuterClass = Z_Registration_Info_UClass_UObject().OuterSingleton;
 	if (!OuterClass)
 	{
 		OuterClass = UObject::StaticClass();
diff --git a/Engine/Source/Runtime/CoreUObject/Public/UObject/ObjectMacros.h b/Engine/Source/Runtime/CoreUObject/Public/UObject/ObjectMacros.h
--- a/Engine/Source/Runtime/CoreUObject/Public/UObject/ObjectMacros.h
+++ b/Engine/Source/Runtime/CoreUObject/Public/UObject/ObjectMacros.h
@@ -1599,8 +1599,6 @@ public: \
 	{ \
 		return TStaticCastFlags; \
 	} \
-	/** Returns the static registration structure for the class */ \
-	static FClassRegistrationInfo& StaticRegistrationInfo(); \
 	/** For internal use only; use StaticConstructObject() to create new objects. */ \
 	inline void* operator new(const size_t InSize, EInternal InInternalOnly, UObject* InOuter = (UObject*)GetTransientPackage(), FName InName = NAME_None, EObjectFlags InSetFlags = RF_NoFlags) \
 	{ \
@@ -1743,16 +1741,16 @@ public: \
 
 // Register a class at startup time.
 #define IMPLEMENT_CLASS(TClass, TClassCrc) \
-	FClassRegistrationInfo& TClass::StaticRegistrationInfo() \
+	FClassRegistrationInfo& Z_Registration_Info_UClass_##TClass() \
 	{ \
 		static FClassRegistrationInfo info; \
 		return info; \
 	} \
 	/* Do not change the AutoInitialize_ without changing LC_SymbolPatterns */ \
-	static FRegisterCompiledInInfo AutoInitialize_##TClass(&Z_Construct_UClass_##TClass, TClass::StaticClass, TClass::StaticPackage(), TEXT(#TClass), TClass::StaticRegistrationInfo(), CONSTUCT_RELOAD_VERSION_INFO(FClassReloadVersionInfo, sizeof(TClass), TClassCrc)); \
+	static FRegisterCompiledInInfo AutoInitialize_##TClass(&Z_Construct_UClass_##TClass, TClass::StaticClass, TClass::StaticPackage(), TEXT(#TClass), Z_Registration_Info_UClass_##TClass(), CONSTUCT_RELOAD_VERSION_INFO(FClassReloadVersionInfo, sizeof(TClass), TClassCrc)); \
 	UClass* TClass::GetPrivateStaticClass() \
 	{ \
-		static UClass*& PrivateStaticClass = StaticRegistrationInfo().InnerSingleton; \
+		static UClass*& PrivateStaticClass = Z_Registration_Info_UClass_##TClass().InnerSingleton; \
 		if (!PrivateStaticClass) \
 		{ \
 			/* this could be handled with templates, but we want it external to avoid code bloat */ \
@@ -1779,6 +1777,7 @@ public: \
 // Used for intrinsics, this sets up the boiler plate, plus an initialization singleton, which can create properties and GC tokens
 #define IMPLEMENT_INTRINSIC_CLASS(TClass, TRequiredAPI, TSuperClass, TSuperRequiredAPI, TPackage, InitCode) \
 	TRequiredAPI UClass* Z_Construct_UClass_##TClass(); \
+	extern FClassRegistrationInfo& Z_Registration_Info_UClass_##TClass(); \
 	struct Z_Construct_UClass_##TClass##_Statics \
 	{ \
 		static UClass* Construct() \
@@ -1795,7 +1794,7 @@ public: \
 	}; \
 	UClass* Z_Construct_UClass_##TClass() \
 	{ \
-		static UClass* Class = TClass::StaticRegistrationInfo().OuterSingleton; \
+		static UClass* Class = Z_Registration_Info_UClass_##TClass().OuterSingleton; \
 		if (!Class) \
 		{ \
 			Class = Z_Construct_UClass_##TClass##_Statics::Construct();\
