Allow Cast to gracefully handle an interface instance that's not a UObject

This allows you to test whether a particular interface instance (implemented as a UInterface) has been implemented by a UObject, or whether it's been implemented by a standard C++ type. This is useful if you have an interface that _may_ be implemented by a UObject, but isn't required to be.

diff --git a/Engine/Source/Programs/UnrealHeaderTool/Private/CodeGenerator.cpp b/Engine/Source/Programs/UnrealHeaderTool/Private/CodeGenerator.cpp
--- a/Engine/Source/Programs/UnrealHeaderTool/Private/CodeGenerator.cpp
+++ b/Engine/Source/Programs/UnrealHeaderTool/Private/CodeGenerator.cpp
@@ -2874,7 +2874,7 @@ void FNativeClassHeaderGenerator::ExportClassFromSourceFileInner(
 			{
 				// Note: This used to be declared as a pure virtual function, but it was changed here in order to allow the Blueprint nativization process
 				// to detect C++ interface classes that explicitly declare pure virtual functions via type traits. This code will no longer trigger that check.
-				InterfaceBoilerplate.Logf(TEXT("\tvirtual UObject* _getUObject() const { check(0 && \"Missing required implementation.\"); return nullptr; }\r\n"));
+				InterfaceBoilerplate.Logf(TEXT("\tvirtual UObject* _getUObject() const { return nullptr; }\r\n"));
 			}
 
 			if (ClassHasReplicatedProperties(ClassDef))
diff --git a/Engine/Source/Runtime/CoreUObject/Public/Templates/Casts.h b/Engine/Source/Runtime/CoreUObject/Public/Templates/Casts.h
--- a/Engine/Source/Runtime/CoreUObject/Public/Templates/Casts.h
+++ b/Engine/Source/Runtime/CoreUObject/Public/Templates/Casts.h
@@ -199,16 +199,15 @@ struct TCastImpl<From, To, ECastType::InterfaceToUObject>
 {
 	FORCEINLINE static To* DoCast( From* Src )
 	{
-		To* Result = nullptr;
 		if (Src)
 		{
 			UObject* Obj = Src->_getUObject();
-			if (Obj->IsA<To>())
+			if (Obj && Obj->IsA<To>())
 			{
-				Result = (To*)Obj;
+				return (To*)Obj;
 			}
 		}
-		return Result;
+		return nullptr;
 	}
 
 	FORCEINLINE static To* DoCastCheckedWithoutTypeCheck( From* Src )
@@ -258,12 +257,26 @@ struct TCastImpl<From, To, ECastType::InterfaceToInterface>
 {
 	FORCEINLINE static To* DoCast( From* Src )
 	{
-		return Src ? (To*)Src->_getUObject()->GetInterfaceAddress(To::UClassType::StaticClass()) : nullptr;
+		if (Src)
+		{
+			if (UObject* Obj = Src->_getUObject())
+			{
+				return (To*)Obj->GetInterfaceAddress(To::UClassType::StaticClass());
+			}
+		}
+		return nullptr;
 	}
 
 	FORCEINLINE static To* DoCastCheckedWithoutTypeCheck( From* Src )
 	{
-		return Src ? (To*)Src->_getUObject()->GetInterfaceAddress(To::UClassType::StaticClass()) : nullptr;
+		if (Src)
+		{
+			if (UObject* Obj = Src->_getUObject())
+			{
+				return (To*)Obj->GetInterfaceAddress(To::UClassType::StaticClass());
+			}
+		}
+		return nullptr;
 	}
 };
 
