Fix issue where SavePackage2 warns about a deprecated and unreferenced object being unsaveable, in SavePackage1 this would be silently ignored
This version changes both the logic and comment to handle abstract and deprecated outers, but not warn for the base object
It may be more theoretically correct to try and warn for directly referenced deprecated base objects, but that is not easy to detect in SavePackage 2 and this change matches the historical behavior

diff --git a/Engine/Source/Runtime/CoreUObject/Private/UObject/SavePackage/SaveContext.cpp b/Engine/Source/Runtime/CoreUObject/Private/UObject/SavePackage/SaveContext.cpp
--- a/Engine/Source/Runtime/CoreUObject/Private/UObject/SavePackage/SaveContext.cpp
+++ b/Engine/Source/Runtime/CoreUObject/Private/UObject/SavePackage/SaveContext.cpp
@@ -20,16 +20,17 @@ bool FSaveContext::IsUnsaveable(UObject* InObject) const
 	UObject* Obj = InObject;
 	while (Obj)
 	{
-		// if the object class is abstract, has been marked as deprecated, there is a newer version that exist or the class is maked transient, then the object is unsaveable
-		// @note: Although object instance of transient class should definitely be unsaveable, it results in discrepencies with the old save algorithm and currently load problems
+		// if the object class is abstract, has been marked as deprecated, there is a newer version that exist, or the class is marked transient, then the object is unsaveable
+		// @note: Although object instances of a transient class should definitely be unsaveable, it results in discrepancies with the old save algorithm and currently load problems
 		if (Obj->GetClass()->HasAnyClassFlags(CLASS_Abstract | CLASS_Deprecated | CLASS_NewerVersionExists /*| CLASS_Transient*/) 
 			&& !Obj->HasAnyFlags(RF_ClassDefaultObject))
 		{
+			// only warn if the base object is fine but the outer is invalid. If an object is itself unsaveable, the old behavior is to ignore it
 			if (!InObject->IsPendingKill() 
 				&& InObject->GetOutermost() == GetPackage()
-				&& Obj->GetClass()->HasAnyClassFlags(CLASS_Deprecated))
+				&& Obj != InObject)
 			{
-				UE_LOG(LogSavePackage, Warning, TEXT("%s has a deprecated, abstract or transient class outer %s, so it will not be saved"), *InObject->GetFullName(), *Obj->GetFullName());
+				UE_LOG(LogSavePackage, Warning, TEXT("%s has a deprecated or abstract class outer %s, so it will not be saved"), *InObject->GetFullName(), *Obj->GetFullName());
 			}
 
 			// there used to be a check for reference if the class had the CLASS_HasInstancedReference,
