Only print Mirage profiling data to the log file if there is relevant data to show.

#rnx

- Previously we would print the profiling data to the log if we have entries in either PushStats or PullStats but and entry in one of those maps only proves that a backend has been mounted, not if a Push or Pull has actually been carried out.
- Now we iterate over the maps and check that we actually have at least one entry before logging anything.
- This will help keep the logs clean in cases where Mirage is enabled but irrelevant, such as a commandlet that doesn't actually load a package etc.
- Added some improvements to the formatting so that it looks more like the DDC profiling output.
-- Line break before the profiling info
-- Use of == to mark different sections
-- Use of --- to clearly mark the headers vs data

diff --git a/Engine/Source/Runtime/CoreUObject/Private/Virtualization/VirtualizationManager.cpp b/Engine/Source/Runtime/CoreUObject/Private/Virtualization/VirtualizationManager.cpp
--- a/Engine/Source/Runtime/CoreUObject/Private/Virtualization/VirtualizationManager.cpp
+++ b/Engine/Source/Runtime/CoreUObject/Private/Virtualization/VirtualizationManager.cpp
@@ -100,7 +100,7 @@ TArray<FString> ParseEntries(const FString& Data)
 /**
  * Profiling data allowing us to track how payloads are being push/pulled during the lifespan of the process. Note that as all backends are
  * created at the same time, we don't need to add locked when accessing the maps. In addition FCookStats is thread safe when adding hits/misses
- * so we don't hasve to worry about that either.
+ * so we don't have to worry about that either.
  * We keep the FCookStats here rather than as a member of IVirtualizationBackend to try and avoid the backends needing to be aware of the data that
  * we are gathering at all. This way all profiling code is kept to this cpp.
  */
@@ -126,18 +126,45 @@ namespace Profiling
 		return *PullStats.Find(Backend.GetDebugString());
 	}
 
+	/** Returns true if we have gathered any profiling data at all */
+	bool HasProfilingData()
+	{
+		auto HasAccumulatedData = [](const TMap<FString, FCookStats::CallStats>& Stats)->bool
+		{
+			for (const auto& Iterator : Stats)
+			{
+				if (Iterator.Value.GetAccumulatedValueAnyThread(FCookStats::CallStats::EHitOrMiss::Hit, FCookStats::CallStats::EStatType::Counter) > 0)
+				{
+					return true;
+				}
+
+				if (Iterator.Value.GetAccumulatedValueAnyThread(FCookStats::CallStats::EHitOrMiss::Miss, FCookStats::CallStats::EStatType::Counter) > 0)
+				{
+					return true;
+				}
+			}
+
+			return false;
+		};
+
+		return HasAccumulatedData(PushStats) || HasAccumulatedData(PullStats);
+	}
+
 	void LogStats()
 	{
-		if (PushStats.IsEmpty() && PullStats.IsEmpty())
+		if (!HasProfilingData())
 		{
-			return; // Early out if we have no data to show at all
+			return; // Early out if we have no data
 		}
 
+		UE_LOG(LogVirtualization, Log, TEXT(""));
 		UE_LOG(LogVirtualization, Log, TEXT("Virtualization ProfileData"));
+		UE_LOG(LogVirtualization, Log, TEXT("======================================================================================="));
 
 		if (PushStats.Num() > 0)
 		{
 			UE_LOG(LogVirtualization, Log, TEXT("%-40s|%17s|%12s|%14s|"), TEXT("Pushing Data"), TEXT("TotalSize (MB)"), TEXT("TotalTime(s)"), TEXT("DataRate(MB/S)"));
+			UE_LOG(LogVirtualization, Log, TEXT("----------------------------------------|-----------------|------------|--------------|"));
 
 			for (const auto& Iterator : PushStats)
 			{
@@ -151,11 +178,14 @@ namespace Profiling
 					Time,
 					MBps);
 			}
+
+			UE_LOG(LogVirtualization, Log, TEXT("======================================================================================="));
 		}
 
 		if (PullStats.Num() > 0)
 		{
 			UE_LOG(LogVirtualization, Log, TEXT("%-40s|%17s|%12s|%14s|"), TEXT("Pulling Data"), TEXT("TotalSize (MB)"), TEXT("TotalTime(s)"), TEXT("DataRate(MB/S)"));
+			UE_LOG(LogVirtualization, Log, TEXT("----------------------------------------|-----------------|------------|--------------|"));
 
 			for (const auto& Iterator : PullStats)
 			{
@@ -169,6 +199,8 @@ namespace Profiling
 					Time,
 					MBps);
 			}
+
+			UE_LOG(LogVirtualization, Log, TEXT("======================================================================================="));
 		}
 	}
 #endif // ENABLE_COOK_STATS
