IoStore:
Add more flexibility to the redirects/localization functionality by moving redirect information away from the package summaries and import tables
Changed import references from global import hashes to package local export hashes

#rnx

diff --git a/Engine/Source/Developer/IoStoreUtilities/Private/IoStoreUtilities.cpp b/Engine/Source/Developer/IoStoreUtilities/Private/IoStoreUtilities.cpp
--- a/Engine/Source/Developer/IoStoreUtilities/Private/IoStoreUtilities.cpp
+++ b/Engine/Source/Developer/IoStoreUtilities/Private/IoStoreUtilities.cpp
@@ -857,7 +857,6 @@ static void AssignPackagesDiskOrder(
 					Cluster->Packages.Add(PackageToProcess);
 					TArray<FPackageId> AllReferencedPackageIds;
 					AllReferencedPackageIds.Append(PackageToProcess->OptimizedPackage->GetImportedPackageIds());
-					AllReferencedPackageIds.Append(PackageToProcess->OptimizedPackage->GetImportedRedirectedPackageIds().Array());
 					for (const FPackageId& ReferencedPackageId : AllReferencedPackageIds)
 					{
 						FLegacyCookedPackage* FindReferencedPackage = PackageIdMap.FindRef(ReferencedPackageId);
@@ -2303,13 +2302,6 @@ int32 CreateTarget(const FIoStoreArguments& Arguments, const FIoStoreWriterSetti
 			{
 				Package->OptimizedPackage->RedirectFrom(RedirectedPackageName);
 			}
-			else
-			{
-				// TODO: Should NAME_None be the default source package name instead?
-				// RemapLocalizationPathIfNeeded sets it to the original
-				// package name when not remapping plugin content
-				Package->OptimizedPackage->RedirectFrom(NAME_None);
-			}
 		}
 	}
 
@@ -2346,14 +2338,8 @@ int32 CreateTarget(const FIoStoreArguments& Arguments, const FIoStoreWriterSetti
 		OptimizedPackagesMap.Add(Package->OptimizedPackage->GetId(), Package->OptimizedPackage);
 	}
 
-	if (Arguments.PackageStore.IsValid())
-	{
-		UE_LOG(LogIoStore, Display, TEXT("Processing imports.."));
-		PackageStoreOptimizer.FindImports(OptimizedPackagesMap);
-	}
-
 	UE_LOG(LogIoStore, Display, TEXT("Processing redirects..."));
-	PackageStoreOptimizer.ProcessRedirects(OptimizedPackagesMap);
+	PackageStoreOptimizer.ProcessRedirects(OptimizedPackagesMap, Arguments.IsDLC());
 
 	UE_LOG(LogIoStore, Display, TEXT("Optimizing packages..."));
 	PackageStoreOptimizer.OptimizeExportBundles(OptimizedPackagesMap);
@@ -2820,11 +2806,11 @@ int32 Describe(
 		FPackageDesc* Package = nullptr;
 		FName Name;
 		FName FullName;
+		uint32 ExportHash;
 		FPackageObjectIndex OuterIndex;
 		FPackageObjectIndex ClassIndex;
 		FPackageObjectIndex SuperIndex;
 		FPackageObjectIndex TemplateIndex;
-		FPackageObjectIndex GlobalImportIndex;
 		uint64 SerialOffset = 0;
 		uint64 SerialSize = 0;
 		FSHAHash Hash;
@@ -2862,6 +2848,7 @@ int32 Describe(
 		int32 NameCount = -1;
 		int32 ExportBundleCount = -1;
 		TArray<FPackageLocation, TInlineAllocator<1>> Locations;
+		TArray<FPackageId> ImportedPackageIds;
 		TArray<FImportDesc> Imports;
 		TArray<FExportDesc> Exports;
 		TArray<TArray<FExportBundleEntryDesc>, TInlineAllocator<1>> ExportBundles;
@@ -2954,7 +2941,7 @@ int32 Describe(
 		TArray<FPackageDesc*> Packages;
 		FIoStoreReader* Reader = nullptr;
 		FCulturePackageMap RawCulturePackageMap;
-		TArray<TPair<FPackageId, FPackageId>> RawPackageRedirects;
+		TArray<FContainerHeaderPackageRedirect> RawPackageRedirects;
 	};
 
 	TArray<FLoadContainerHeaderJob> LoadContainerHeaderJobs;
@@ -3016,6 +3003,7 @@ int32 Describe(
 				PackageDesc->Exports.SetNum(ContainerEntry.ExportCount);
 				PackageDesc->ExportBundleCount = ContainerEntry.ExportBundleCount;
 				PackageDesc->LoadOrder = ContainerEntry.LoadOrder;
+				PackageDesc->ImportedPackageIds = TArrayView<FPackageId>(ContainerEntry.ImportedPackages.Data(), ContainerEntry.ImportedPackages.Num());
 				Job.Packages.Add(PackageDesc);
 				++TotalPackageCount;
 			}
@@ -3063,8 +3051,8 @@ int32 Describe(
 		for (const auto& RedirectPair : LoadContainerHeaderJob.RawPackageRedirects)
 		{
 			FPackageRedirect& PackageRedirect = LoadContainerHeaderJob.ContainerDesc->PackageRedirects.AddDefaulted_GetRef();
-			PackageRedirect.Source = PackageByIdMap.FindRef(RedirectPair.Get<0>());
-			PackageRedirect.Target = PackageByIdMap.FindRef(RedirectPair.Get<1>());
+			PackageRedirect.Source = PackageByIdMap.FindRef(RedirectPair.SourcePackageId);
+			PackageRedirect.Target = PackageByIdMap.FindRef(RedirectPair.TargetPackageId);
 		}
 		for (const auto& CultureRedirectsPair : LoadContainerHeaderJob.RawCulturePackageMap)
 		{
@@ -3072,8 +3060,8 @@ int32 Describe(
 			for (const auto& RedirectPair : CultureRedirectsPair.Get<1>())
 			{
 				FPackageRedirect& PackageRedirect = LoadContainerHeaderJob.ContainerDesc->PackageRedirects.AddDefaulted_GetRef();
-				PackageRedirect.Source = PackageByIdMap.FindRef(RedirectPair.Get<0>());
-				PackageRedirect.Target = PackageByIdMap.FindRef(RedirectPair.Get<1>());
+				PackageRedirect.Source = PackageByIdMap.FindRef(RedirectPair.SourcePackageId);
+				PackageRedirect.Target = PackageByIdMap.FindRef(RedirectPair.TargetPackageId);
 				PackageRedirect.Culture = Culture;
 			}
 		}
@@ -3142,7 +3130,7 @@ int32 Describe(
 			ExportDesc.ClassIndex = ExportMapEntry.ClassIndex;
 			ExportDesc.SuperIndex = ExportMapEntry.SuperIndex;
 			ExportDesc.TemplateIndex = ExportMapEntry.TemplateIndex;
-			ExportDesc.GlobalImportIndex = ExportMapEntry.GlobalImportIndex;
+			ExportDesc.ExportHash = ExportMapEntry.ExportHash;
 			ExportDesc.SerialSize = ExportMapEntry.CookedSerialSize;
 		}
 
@@ -3179,7 +3167,7 @@ int32 Describe(
 	}, EParallelForFlags::Unbalanced);
 
 	UE_LOG(LogIoStore, Display, TEXT("Connecting imports and exports..."));
-	TMap<FPackageObjectIndex, FExportDesc*> ExportByGlobalIdMap;
+	TMap<FPublicExportKey, FExportDesc*> ExportByKeyMap;
 	{
 		TRACE_CPUPROFILER_EVENT_SCOPE(ConnectImportsAndExports);
 
@@ -3187,9 +3175,10 @@ int32 Describe(
 		{
 			for (FExportDesc& ExportDesc : PackageDesc->Exports)
 			{
-				if (!ExportDesc.GlobalImportIndex.IsNull())
+				if (ExportDesc.ExportHash)
 				{
-					ExportByGlobalIdMap.Add(ExportDesc.GlobalImportIndex, &ExportDesc);
+					FPublicExportKey Key = FPublicExportKey::MakeKey(PackageDesc->PackageId, ExportDesc.ExportHash);
+					ExportByKeyMap.Add(Key, &ExportDesc);
 				}
 			}
 		}
@@ -3245,7 +3234,8 @@ int32 Describe(
 				{
 					if (Import.GlobalImportIndex.IsPackageImport())
 					{
-						Import.Export = ExportByGlobalIdMap.FindRef(Import.GlobalImportIndex);
+						FPublicExportKey Key = FPublicExportKey::FromPackageImport(Import.GlobalImportIndex, PackageDesc->ImportedPackageIds);
+						Import.Export = ExportByKeyMap.FindRef(Key);
 						if (!Import.Export)
 						{
 							UE_LOG(LogIoStore, Warning, TEXT("Missing import: 0x%llX in package 0x%llX '%s'"), Import.GlobalImportIndex.Value(), PackageDesc->PackageId.ValueForDebugging(), *PackageDesc->PackageName.ToString());
@@ -3363,7 +3353,7 @@ int32 Describe(
 		TGuardValue<bool> GuardPrintLogCategory(GPrintLogCategory, false);
 		TGuardValue<bool> GuardPrintLogVerbosity(GPrintLogVerbosity, false);
 
-		auto PackageObjectIndexToString = [&ScriptObjectByGlobalIdMap, &ExportByGlobalIdMap](const FPackageObjectIndex& PackageObjectIndex, bool bIncludeName) -> FString
+		auto PackageObjectIndexToString = [&ScriptObjectByGlobalIdMap, &ExportByKeyMap](const FPackageDesc* Package, const FPackageObjectIndex& PackageObjectIndex, bool bIncludeName) -> FString
 		{
 			if (PackageObjectIndex.IsNull())
 			{
@@ -3371,7 +3361,8 @@ int32 Describe(
 			}
 			else if (PackageObjectIndex.IsPackageImport())
 			{
-				FExportDesc* ExportDesc = ExportByGlobalIdMap.FindRef(PackageObjectIndex);
+				FPublicExportKey Key = FPublicExportKey::FromPackageImport(PackageObjectIndex, Package->ImportedPackageIds);
+				FExportDesc* ExportDesc = ExportByKeyMap.FindRef(Key);
 				if (ExportDesc && bIncludeName)
 				{
 					return FString::Printf(TEXT("0x%llX '%s'"), PackageObjectIndex.Value(), *ExportDesc->FullName.ToString());
@@ -3470,7 +3461,7 @@ int32 Describe(
 			{
 				OutputOverride->Logf(ELogVerbosity::Display, TEXT("\t*************************"));
 				OutputOverride->Logf(ELogVerbosity::Display, TEXT("\tImport %d: '%s'"), Index++, *Import.Name.ToString());
-				OutputOverride->Logf(ELogVerbosity::Display, TEXT("\t\tGlobalImportIndex: %s"), *PackageObjectIndexToString(Import.GlobalImportIndex, false));
+				OutputOverride->Logf(ELogVerbosity::Display, TEXT("\t\tGlobalImportIndex: %s"), *PackageObjectIndexToString(PackageDesc, Import.GlobalImportIndex, false));
 			}
 
 			OutputOverride->Logf(ELogVerbosity::Display, TEXT("--------------------------------------------"));
@@ -3481,11 +3472,11 @@ int32 Describe(
 			{
 				OutputOverride->Logf(ELogVerbosity::Display, TEXT("\t*************************"));
 				OutputOverride->Logf(ELogVerbosity::Display, TEXT("\tExport %d: '%s'"), Index++, *Export.Name.ToString());
-				OutputOverride->Logf(ELogVerbosity::Display, TEXT("\t\t       OuterIndex: %s"), *PackageObjectIndexToString(Export.OuterIndex, true));
-				OutputOverride->Logf(ELogVerbosity::Display, TEXT("\t\t       ClassIndex: %s"), *PackageObjectIndexToString(Export.ClassIndex, true));
-				OutputOverride->Logf(ELogVerbosity::Display, TEXT("\t\t       SuperIndex: %s"), *PackageObjectIndexToString(Export.SuperIndex, true));
-				OutputOverride->Logf(ELogVerbosity::Display, TEXT("\t\t    TemplateIndex: %s"), *PackageObjectIndexToString(Export.TemplateIndex, true));
-				OutputOverride->Logf(ELogVerbosity::Display, TEXT("\t\tGlobalImportIndex: %s"), *PackageObjectIndexToString(Export.GlobalImportIndex, false));
+				OutputOverride->Logf(ELogVerbosity::Display, TEXT("\t\t       OuterIndex: %s"), *PackageObjectIndexToString(PackageDesc, Export.OuterIndex, true));
+				OutputOverride->Logf(ELogVerbosity::Display, TEXT("\t\t       ClassIndex: %s"), *PackageObjectIndexToString(PackageDesc, Export.ClassIndex, true));
+				OutputOverride->Logf(ELogVerbosity::Display, TEXT("\t\t       SuperIndex: %s"), *PackageObjectIndexToString(PackageDesc, Export.SuperIndex, true));
+				OutputOverride->Logf(ELogVerbosity::Display, TEXT("\t\t    TemplateIndex: %s"), *PackageObjectIndexToString(PackageDesc, Export.TemplateIndex, true));
+				OutputOverride->Logf(ELogVerbosity::Display, TEXT("\t\t       ExportHash: %d"), Export.ExportHash);
 				OutputOverride->Logf(ELogVerbosity::Display, TEXT("\t\t           Offset: %lld"), Export.SerialOffset);
 				OutputOverride->Logf(ELogVerbosity::Display, TEXT("\t\t             Size: %lld"), Export.SerialSize);
 				if (bIncludeExportHashes)
diff --git a/Engine/Source/Developer/IoStoreUtilities/Private/PackageStoreOptimizer.cpp b/Engine/Source/Developer/IoStoreUtilities/Private/PackageStoreOptimizer.cpp
--- a/Engine/Source/Developer/IoStoreUtilities/Private/PackageStoreOptimizer.cpp
+++ b/Engine/Source/Developer/IoStoreUtilities/Private/PackageStoreOptimizer.cpp
@@ -336,9 +336,9 @@ FPackageStoreOptimizer::FPackageStoreHeaderData FPackageStoreOptimizer::LoadPack
 	return PackageStoreHeaderData;
 }
 
-void FPackageStoreOptimizer::ResolveImport(FPackageStorePackage::FImport* Imports, const FObjectImport* ObjectImports, int32 LocalImportIndex) const
+void FPackageStoreOptimizer::ResolveImport(FPackageStorePackage::FUnresolvedImport* Imports, const FObjectImport* ObjectImports, int32 LocalImportIndex) const
 {
-	FPackageStorePackage::FImport* Import = Imports + LocalImportIndex;
+	FPackageStorePackage::FUnresolvedImport* Import = Imports + LocalImportIndex;
 	if (Import->FullName.Len() == 0)
 	{
 		Import->FullName.Reserve(256);
@@ -349,22 +349,26 @@ void FPackageStoreOptimizer::ResolveImport(FPackageStorePackage::FImport* Import
 			FName PackageName = ObjectImport->ObjectName;
 			PackageName.AppendString(Import->FullName);
 			Import->FullName.ToLowerInline();
-			Import->PackageId = FPackageId::FromName(PackageName);
+			Import->FromPackageId = FPackageId::FromName(PackageName);
+			Import->FromPackageName = PackageName;
+			Import->FromPackageNameLen = Import->FullName.Len();
 			Import->bIsScriptImport = Import->FullName.StartsWith(TEXT("/Script/"));
-			Import->bIsPackageImport = true;
+			Import->bIsImportOfPackage = true;
 		}
 		else
 		{
 			const int32 OuterIndex = ObjectImport->OuterIndex.ToImport();
 			ResolveImport(Imports, ObjectImports, OuterIndex);
-			FPackageStorePackage::FImport* OuterImport = Imports + OuterIndex;
+			FPackageStorePackage::FUnresolvedImport* OuterImport = Imports + OuterIndex;
 			check(OuterImport->FullName.Len() > 0);
 			Import->bIsScriptImport = OuterImport->bIsScriptImport;
 			Import->FullName.Append(OuterImport->FullName);
 			Import->FullName.AppendChar(TEXT('/'));
 			ObjectImport->ObjectName.AppendString(Import->FullName);
 			Import->FullName.ToLowerInline();
-			Import->PackageId = OuterImport->PackageId;
+			Import->FromPackageId = OuterImport->FromPackageId;
+			Import->FromPackageName = OuterImport->FromPackageName;
+			Import->FromPackageNameLen = OuterImport->FromPackageNameLen;
 		}
 	}
 }
@@ -372,42 +376,55 @@ void FPackageStoreOptimizer::ResolveImport(FPackageStorePackage::FImport* Import
 void FPackageStoreOptimizer::ProcessImports(const FCookedHeaderData& CookedHeaderData, FPackageStorePackage* Package) const
 {
 	int32 ImportCount = CookedHeaderData.ObjectImports.Num();
+	TArray<FPackageStorePackage::FUnresolvedImport> UnresolvedImports;
+	UnresolvedImports.SetNum(ImportCount);
 	Package->Imports.SetNum(ImportCount);
 
 	TSet<FPackageId> ImportedPackageIds;
 	for (int32 ImportIndex = 0; ImportIndex < ImportCount; ++ImportIndex)
 	{
-		ResolveImport(Package->Imports.GetData(), CookedHeaderData.ObjectImports.GetData(), ImportIndex);
-		FPackageStorePackage::FImport& Import = Package->Imports[ImportIndex];
-		if (Import.bIsScriptImport)
+		ResolveImport(UnresolvedImports.GetData(), CookedHeaderData.ObjectImports.GetData(), ImportIndex);
+		FPackageStorePackage::FUnresolvedImport& UnresolvedImport = UnresolvedImports[ImportIndex];
+		if (!UnresolvedImport.bIsScriptImport && UnresolvedImport.bIsImportOfPackage)
 		{
-			Import.GlobalImportIndex = FPackageObjectIndex::FromScriptPath(Import.FullName);
+			ImportedPackageIds.Add(UnresolvedImport.FromPackageId);
 		}
-		else if (!Import.bIsPackageImport)
+	}
+	Package->ImportedPackageIds = ImportedPackageIds.Array();
+	Algo::Sort(Package->ImportedPackageIds);
+
+	for (int32 ImportIndex = 0; ImportIndex < ImportCount; ++ImportIndex)
+	{
+		FPackageStorePackage::FUnresolvedImport& UnresolvedImport = UnresolvedImports[ImportIndex];
+		if (UnresolvedImport.bIsScriptImport)
 		{
-			Import.GlobalImportIndex = FPackageObjectIndex::FromPackagePath(Import.FullName);
+			Package->Imports[ImportIndex] = FPackageObjectIndex::FromScriptPath(UnresolvedImport.FullName);
 		}
-		else
+		else if (!UnresolvedImport.bIsImportOfPackage)
 		{
-			ImportedPackageIds.Add(Import.PackageId);
+			bool bFoundPackageIndex = false;
+			for (uint32 PackageIndex = 0, PackageCount = static_cast<uint32>(Package->ImportedPackageIds.Num()); PackageIndex < PackageCount; ++PackageIndex)
+			{
+				if (UnresolvedImport.FromPackageId == Package->ImportedPackageIds[PackageIndex])
+				{
+					FStringView PackageRelativeName = FStringView(UnresolvedImport.FullName).RightChop(UnresolvedImport.FromPackageNameLen);
+					check(PackageRelativeName.Len());
+					uint32 ExportHash = GetTypeHash(PackageRelativeName);
+					FPackageImportReference PackageImportRef(PackageIndex, ExportHash);
+					Package->Imports[ImportIndex] = FPackageObjectIndex::FromPackageImportRef(PackageImportRef);
+					bFoundPackageIndex = true;
+					break;
+				}
+			}
+			check(bFoundPackageIndex);
 		}
 	}
-	Package->ImportedPackageIds = ImportedPackageIds.Array();
-	Algo::Sort(Package->ImportedPackageIds);
 }
 
 void FPackageStoreOptimizer::ProcessImports(const FPackageStoreHeaderData& PackageStoreHeaderData, FPackageStorePackage* Package) const
 {
 	Package->ImportedPackageIds = PackageStoreHeaderData.ImportedPackageIds;
-
-	for (FPackageObjectIndex ImportObjectIndex : PackageStoreHeaderData.Imports)
-	{
-		FPackageStorePackage::FImport& Import = Package->Imports.AddDefaulted_GetRef();
-		Import.GlobalImportIndex = ImportObjectIndex;
-		Import.bIsScriptImport = ImportObjectIndex.IsScriptImport();
-		Import.bIsPackageImport = ImportObjectIndex.IsPackageImport();
-		// Import.FullName and Import.PackageId will be resolve in FindImports after all package headers have been loaded
-	}
+	Package->Imports = PackageStoreHeaderData.Imports;
 }
 
 void FPackageStoreOptimizer::ResolveExport(
@@ -481,11 +498,11 @@ void FPackageStoreOptimizer::ProcessExports(const FCookedHeaderData& CookedHeade
 	Package->ExportGraphNodes.Reserve(ExportCount * 2);
 
 	auto PackageObjectIdFromPackageIndex =
-		[](const TArray<FPackageStorePackage::FImport>& Imports, const FPackageIndex& PackageIndex) -> FPackageObjectIndex
+		[](const TArray<FPackageObjectIndex>& Imports, const FPackageIndex& PackageIndex) -> FPackageObjectIndex
 	{
 		if (PackageIndex.IsImport())
 		{
-			return Imports[PackageIndex.ToImport()].GlobalImportIndex;
+			return Imports[PackageIndex.ToImport()];
 		}
 		if (PackageIndex.IsExport())
 		{
@@ -494,6 +511,7 @@ void FPackageStoreOptimizer::ProcessExports(const FCookedHeaderData& CookedHeade
 		return FPackageObjectIndex();
 	};
 
+	FString PackageNameStr = Package->Name.ToString();
 	for (int32 ExportIndex = 0; ExportIndex < ExportCount; ++ExportIndex)
 	{
 		const FObjectExport& ObjectExport = CookedHeaderData.ObjectExports[ExportIndex];
@@ -512,7 +530,9 @@ void FPackageStoreOptimizer::ProcessExports(const FCookedHeaderData& CookedHeade
 		if (Export.bIsPublic)
 		{
 			check(Export.FullName.Len() > 0);
-			Export.GlobalImportIndex = FPackageObjectIndex::FromPackagePath(Export.FullName);
+			FStringView PackageRelativeName = FStringView(Export.FullName).RightChop(PackageNameStr.Len());
+			check(PackageRelativeName.Len());
+			Export.ExportHash = GetTypeHash(PackageRelativeName);
 		}
 
 		Export.OuterIndex = PackageObjectIdFromPackageIndex(Package->Imports, ObjectExport.OuterIndex);
@@ -546,11 +566,11 @@ void FPackageStoreOptimizer::ProcessExports(const FPackageStoreHeaderData& Packa
 
 		FPackageStorePackage::FExport& Export = Package->Exports[ExportIndex];
 		Export.ObjectName = FName::CreateFromDisplayId(NameMap[ExportEntry.ObjectName.GetIndex()], ExportEntry.ObjectName.GetNumber());
+		Export.ExportHash = ExportEntry.ExportHash;
 		Export.OuterIndex = ExportEntry.OuterIndex;
 		Export.ClassIndex = ExportEntry.ClassIndex;
 		Export.SuperIndex = ExportEntry.SuperIndex;
 		Export.TemplateIndex = ExportEntry.TemplateIndex;
-		Export.GlobalImportIndex = ExportEntry.GlobalImportIndex;
 		Export.ObjectFlags = ExportEntry.ObjectFlags;
 		Export.CookedSerialOffset = ExportEntry.CookedSerialOffset;
 		Export.SerialSize = ExportEntry.CookedSerialSize;
@@ -615,15 +635,6 @@ TArray<FPackageStorePackage*> FPackageStoreOptimizer::SortPackagesInLoadOrder(co
 				SourceArray.Add(Package);
 			}
 		}
-		for (FPackageId ImportedRedirectedPackageId : Package->ImportedRedirectedPackageIds)
-		{
-			FPackageStorePackage* FindImportedPackage = PackagesMap.FindRef(ImportedRedirectedPackageId);
-			if (FindImportedPackage)
-			{
-				TArray<FPackageStorePackage*>& SourceArray = SortedEdges.FindOrAdd(FindImportedPackage);
-				SourceArray.Add(Package);
-			}
-		}
 	}
 	for (auto& KV : SortedEdges)
 	{
@@ -769,42 +780,27 @@ void FPackageStoreOptimizer::OptimizeExportBundles(const TMap<FPackageId, FPacka
 				}
 				for (FPackageStorePackage::FExternalDependency& ExternalDependency : ExportGraphNode->ExternalDependencies)
 				{
-					const FPackageStorePackage::FImport& FromImport = Package->Imports[ExternalDependency.ImportIndex];
-					FPackageStorePackage* FindFromPackage = PackagesMap.FindRef(FromImport.PackageId);
+					FPackageObjectIndex FromImport = Package->Imports[ExternalDependency.ImportIndex];
+					check(FromImport.IsPackageImport());
+					FPackageImportReference FromPackageImportRef = FromImport.ToPackageImportRef();
+					FPackageId FromPackageId = Package->ImportedPackageIds[FromPackageImportRef.GetImportedPackageIndex()];
+					FPackageStorePackage* FindFromPackage = PackagesMap.FindRef(FromPackageId);
 					if (FindFromPackage)
 					{
 						bool bFoundExport = false;
 						for (int32 ExportIndex = 0; ExportIndex < FindFromPackage->Exports.Num(); ++ExportIndex)
 						{
 							FPackageStorePackage::FExport& FromExport = FindFromPackage->Exports[ExportIndex];
-							if (FromExport.GlobalImportIndex == FromImport.GlobalImportIndex)
+							if (FromExport.ExportHash == FromPackageImportRef.GetExportHash())
 							{
 								FPackageStorePackage::FExportGraphNode* FromExportGraphNode = FromExport.Nodes[ExternalDependency.ExportBundleCommandType];
 								check(FromExportGraphNode->ExportBundleIndex >= 0);
 								FPackageStorePackage::FExportBundleGraphNode* FromNode = &FindFromPackage->ExportBundleGraphNodes[FromExportGraphNode->ExportBundleIndex];
 								Dependencies.Add(FromNode);
-
-								for (const FPackageId& FromRedirectedPackageId : FindFromPackage->RedirectedToPackageIds)
-								{
-									FPackageStorePackage* FindRedirectedFromPackage = PackagesMap.FindRef(FromRedirectedPackageId);
-									if (FindRedirectedFromPackage)
-									{
-										FPackageStorePackage::FExport& FromRedirectedExport = FindRedirectedFromPackage->Exports[ExportIndex];
-										FPackageStorePackage::FExportGraphNode* FromRedirectedExportGraphNode = FromRedirectedExport.Nodes[ExternalDependency.ExportBundleCommandType];
-										check(FromRedirectedExportGraphNode->ExportBundleIndex >= 0);
-										FPackageStorePackage::FExportBundleGraphNode* FromRedirectedNode = &FindFromPackage->ExportBundleGraphNodes[FromRedirectedExportGraphNode->ExportBundleIndex];
-										Dependencies.Add(FromRedirectedNode);
-									}
-								}
-
 								bFoundExport = true;
 								break;
 							}
 						}
-						if (!bFoundExport)
-						{
-							ExternalDependency.bIsConfirmedMissing = true;
-						}
 					}
 				}
 			}
@@ -840,51 +836,6 @@ void FPackageStoreOptimizer::OptimizeExportBundles(const TMap<FPackageId, FPacka
 	}
 }
 
-void FPackageStoreOptimizer::FindImports(const TMap<FPackageId, FPackageStorePackage*>& PackagesMap) const
-{
-	using FPackageExportIndexPair = TPair<FPackageStorePackage*, int32>;
-	using FGlobalImportIndexToExport = TMap<FPackageObjectIndex, FPackageExportIndexPair>;
-
-	FGlobalImportIndexToExport GlobalImportIndexToExport;
-
-	for (const auto& PackageIdPackage : PackagesMap)
-	{
-		const FPackageStorePackage& Package = *PackageIdPackage.Value;
-		for (int32 ExportIndex = 0; ExportIndex < Package.Exports.Num(); ++ExportIndex)
-		{
-			const FPackageStorePackage::FExport& Export = Package.Exports[ExportIndex];
-			if (Export.bIsPublic)
-			{
-				GlobalImportIndexToExport.Add(Export.GlobalImportIndex, FPackageExportIndexPair(PackageIdPackage.Value, ExportIndex));
-			}
-		}
-	}
-
-	for (const auto& PackageIdPackage : PackagesMap)
-	{
-		FPackageStorePackage& Package = *PackageIdPackage.Value;
-		for (FPackageStorePackage::FImport& Import : Package.Imports)
-		{
-			if (Import.GlobalImportIndex.IsPackageImport())
-			{
-				if (const FPackageExportIndexPair* Entry = GlobalImportIndexToExport.Find(Import.GlobalImportIndex))
-				{
-					const FPackageStorePackage& ImportedPackage = *Entry->Key;
-					const int32 ExportIndex = Entry->Value;
-					Import.PackageId = ImportedPackage.GetId();
-					Import.FullName = ImportedPackage.Exports[ExportIndex].FullName;
-				}
-				else
-				{
-					UE_LOG(LogPackageStoreOptimizer, Warning,
-						TEXT("Failed to find import '0x%llX' in package '%s'"),
-						Import.GlobalImportIndex.Value(), *Package.Name.ToString());
-				}
-			}
-		}
-	}
-}
-
 void FPackageStoreOptimizer::ProcessPreloadDependencies(const FCookedHeaderData& CookedHeaderData, FPackageStorePackage* Package) const
 {
 	TRACE_CPUPROFILER_EVENT_SCOPE(ProcessPreloadDependencies);
@@ -898,8 +849,8 @@ void FPackageStoreOptimizer::ProcessPreloadDependencies(const FCookedHeaderData&
 
 	auto AddExternalDependency = [](FPackageStorePackage* Package, int32 FromImportIndex, FExportBundleEntry::EExportCommandType FromExportBundleCommandType, FPackageStorePackage::FExportGraphNode* ToNode)
 	{
-		const FPackageStorePackage::FImport& FromImport = Package->Imports[FromImportIndex];
-		if (FromImport.GlobalImportIndex.IsScriptImport())
+		FPackageObjectIndex FromImport = Package->Imports[FromImportIndex];
+		if (FromImport.IsScriptImport())
 		{
 			return;
 		}
@@ -1135,40 +1086,65 @@ static const TCHAR* GetExportNameSafe(const FString& ExportFullName, const FName
 	return nullptr;
 };
 
-bool FPackageStoreOptimizer::RedirectPackage(
-	const FPackageStorePackage& SourcePackage,
+bool FPackageStoreOptimizer::VerifyRedirect(
+	const FPackageStorePackage* SourcePackage,
 	FPackageStorePackage& TargetPackage,
-	TMap<FPackageObjectIndex, FPackageObjectIndex>& RedirectedToSourceImportIndexMap) const
+	bool bIsBuildingDLC) const
 {
+	if (!SourcePackage)
+	{
+		if (bIsBuildingDLC)
+		{
+			// We can't verify against the source package but check for presence of UStructs
+			for (const FPackageStorePackage::FExport& Export : TargetPackage.Exports)
+			{
+				if (!Export.SuperIndex.IsNull() && Export.OuterIndex.IsNull())
+				{
+					UE_LOG(LogPackageStoreOptimizer, Warning, TEXT("Skipping redirect to package '%s' due to presence of UStruct '%s'"), *TargetPackage.Name.ToString(), *Export.ObjectName.ToString());
+					return false;
+				}
+			}
+			return true;
+		}
+
+		if (TargetPackage.Region.Len() > 0)
+		{
+			// no update or verification required
+			UE_LOG(LogPackageStoreOptimizer, Verbose,
+				TEXT("For culture '%s': Localized package '%s' (0x%llX) is unique and does not override a source package."),
+				*TargetPackage.Region,
+				*TargetPackage.Name.ToString(),
+				TargetPackage.Id.ValueForDebugging());
+		}
+		return false;
+	}
+
 	const int32 ExportCount =
-		SourcePackage.Exports.Num() < TargetPackage.Exports.Num() ?
-		SourcePackage.Exports.Num() :
+		SourcePackage->Exports.Num() < TargetPackage.Exports.Num() ?
+		SourcePackage->Exports.Num() :
 		TargetPackage.Exports.Num();
 
-	UE_CLOG(SourcePackage.Exports.Num() != TargetPackage.Exports.Num(), LogPackageStoreOptimizer, Verbose,
+	UE_CLOG(SourcePackage->Exports.Num() != TargetPackage.Exports.Num(), LogPackageStoreOptimizer, Verbose,
 		TEXT("Redirection target package '%s' (0x%llX) for source package '%s' (0x%llX)  - Has ExportCount %d vs. %d"),
 		*TargetPackage.Name.ToString(),
 		TargetPackage.Id.ValueForDebugging(),
-		*SourcePackage.Name.ToString(),
-		SourcePackage.Id.ValueForDebugging(),
+		*SourcePackage->Name.ToString(),
+		SourcePackage->Id.ValueForDebugging(),
 		TargetPackage.Exports.Num(),
-		SourcePackage.Exports.Num());
+		SourcePackage->Exports.Num());
 
-	auto AppendMismatchMessage = [&TargetPackage, &SourcePackage](
+	auto AppendMismatchMessage = [&TargetPackage, SourcePackage](
 		const TCHAR* Text, FName ExportName, FPackageObjectIndex TargetIndex, FPackageObjectIndex SourceIndex, FString& FailReason)
 	{
 		FailReason.Appendf(TEXT("Public export '%s' has %s %s vs. %s"),
 			*ExportName.ToString(),
 			Text,
 			*TargetPackage.Exports[TargetIndex.ToExport()].FullName,
-			*SourcePackage.Exports[SourceIndex.ToExport()].FullName);
+			*SourcePackage->Exports[SourceIndex.ToExport()].FullName);
 	};
 
 	const int32 TargetPackageNameLen = TargetPackage.Name.GetStringLength();
-	const int32 SourcePackageNameLen = SourcePackage.Name.GetStringLength();
-
-	TArray <TPair<int32, int32>, TInlineAllocator<64>> NewPublicExports;
-	NewPublicExports.Reserve(ExportCount);
+	const int32 SourcePackageNameLen = SourcePackage->Name.GetStringLength();
 
 	bool bSuccess = true;
 	int32 TargetIndex = 0;
@@ -1177,12 +1153,12 @@ bool FPackageStoreOptimizer::RedirectPackage(
 	{
 		FString FailReason;
 		const FPackageStorePackage::FExport& TargetExport = TargetPackage.Exports[TargetIndex];
-		const FPackageStorePackage::FExport& SourceExport = SourcePackage.Exports[SourceIndex];
+		const FPackageStorePackage::FExport& SourceExport = SourcePackage->Exports[SourceIndex];
 
 		const TCHAR* TargetExportStr = GetExportNameSafe(
 			TargetExport.FullName, TargetPackage.Name, TargetPackageNameLen);
 		const TCHAR* SourceExportStr = GetExportNameSafe(
-			SourceExport.FullName, SourcePackage.Name, SourcePackageNameLen);
+			SourceExport.FullName, SourcePackage->Name, SourcePackageNameLen);
 
 		if (!TargetExportStr || !SourceExportStr)
 		{
@@ -1190,8 +1166,8 @@ bool FPackageStoreOptimizer::RedirectPackage(
 				TEXT("Redirection target package '%s' (0x%llX) for source package '%s' (0x%llX) - Has some bad data from an earlier phase."),
 				*TargetPackage.Name.ToString(),
 				TargetPackage.Id.ValueForDebugging(),
-				*SourcePackage.Name.ToString(),
-				SourcePackage.Id.ValueForDebugging())
+				*SourcePackage->Name.ToString(),
+				SourcePackage->Id.ValueForDebugging())
 				return false;
 		}
 
@@ -1199,12 +1175,6 @@ bool FPackageStoreOptimizer::RedirectPackage(
 		if (CompareResult < 0)
 		{
 			++TargetIndex;
-
-			if (TargetExport.bIsPublic)
-			{
-				// public localized export is missing in the source package, so just keep it as it is
-				NewPublicExports.Emplace(TargetIndex - 1, 1);
-			}
 		}
 		else if (CompareResult > 0)
 		{
@@ -1244,10 +1214,6 @@ bool FPackageStoreOptimizer::RedirectPackage(
 					AppendMismatchMessage(TEXT("super"), TargetExport.ObjectName,
 						TargetExport.SuperIndex, SourceExport.SuperIndex, FailReason);
 				}
-				else
-				{
-					NewPublicExports.Emplace(TargetIndex - 1, SourceIndex - 1);
-				}
 			}
 			else if (TargetExport.bIsPublic)
 			{
@@ -1263,56 +1229,19 @@ bool FPackageStoreOptimizer::RedirectPackage(
 				TEXT("Redirection target package '%s' (0x%llX) for '%s' (0x%llX) - %s"),
 				*TargetPackage.Name.ToString(),
 				TargetPackage.Id.ValueForDebugging(),
-				*SourcePackage.Name.ToString(),
-				SourcePackage.Id.ValueForDebugging(),
+				*SourcePackage->Name.ToString(),
+				SourcePackage->Id.ValueForDebugging(),
 				*FailReason);
 			bSuccess = false;
 		}
 	}
 
-	if (bSuccess)
-	{
-		for (TPair<int32, int32>& Pair : NewPublicExports)
-		{
-			FPackageStorePackage::FExport& TargetExport = TargetPackage.Exports[Pair.Key];
-			if (Pair.Value != -1)
-			{
-				const FPackageStorePackage::FExport& SourceExport = SourcePackage.Exports[Pair.Value];
-
-				RedirectedToSourceImportIndexMap.Add(TargetExport.GlobalImportIndex, SourceExport.GlobalImportIndex);
-				TargetExport.GlobalImportIndex = SourceExport.GlobalImportIndex;
-			}
-		}
-	}
-
 	return bSuccess;
 }
 
-void FPackageStoreOptimizer::RedirectPackageUnverified(FPackageStorePackage& TargetPackage, TMap<FPackageObjectIndex, FPackageObjectIndex>& RedirectedToSourceImportIndexMap) const
-{
-	const int32 TargetPackageNameLen = TargetPackage.Name.GetStringLength();
-	TStringBuilder<256> NameStringBuilder;
-	for (FPackageStorePackage::FExport& Export : TargetPackage.Exports)
-	{
-		if (Export.bIsPublic)
-		{
-			const TCHAR* ExportNameStr = GetExportNameSafe(Export.FullName, TargetPackage.Name, TargetPackageNameLen);
-			NameStringBuilder.Reset();
-			TargetPackage.SourceName.AppendString(NameStringBuilder);
-			NameStringBuilder.Append(TEXT("/"));
-			NameStringBuilder.Append(ExportNameStr);
-			FPackageObjectIndex SourceGlobalImportIndex = FPackageObjectIndex::FromPackagePath(NameStringBuilder);
-			RedirectedToSourceImportIndexMap.Add(Export.GlobalImportIndex, SourceGlobalImportIndex);
-			Export.GlobalImportIndex = SourceGlobalImportIndex;
-		}
-	}
-}
-
-void FPackageStoreOptimizer::ProcessRedirects(const TMap<FPackageId, FPackageStorePackage*>& PackagesMap) const
+void FPackageStoreOptimizer::ProcessRedirects(const TMap<FPackageId, FPackageStorePackage*>& PackagesMap, bool bIsBuildingDLC) const
 {
 	TRACE_CPUPROFILER_EVENT_SCOPE(ProcessRedirects);
-	TMap<FPackageObjectIndex, FPackageObjectIndex> RedirectedToSourceImportIndexMap;
-	TMap<FPackageId, FPackageId> SourceToRedirectedPackageMap;
 
 	for (const auto& KV : PackagesMap)
 	{
@@ -1324,39 +1253,7 @@ void FPackageStoreOptimizer::ProcessRedirects(const TMap<FPackageId, FPackageSto
 		
 		FPackageId SourcePackageId = FPackageId::FromName(Package->SourceName);
 		FPackageStorePackage* SourcePackage = PackagesMap.FindRef(SourcePackageId);
-		
-		if (SourcePackage)
-		{
-			Package->bIsRedirected = RedirectPackage(*SourcePackage, *Package, RedirectedToSourceImportIndexMap);
-		}
-		else
-		{
-			if (Package->Region.Len() > 0)
-			{
-				// no update or verification required
-				UE_LOG(LogPackageStoreOptimizer, Verbose,
-					TEXT("For culture '%s': Localized package '%s' (0x%llX) is unique and does not override a source package."),
-					*Package->Region,
-					*Package->Name.ToString(),
-					Package->Id.ValueForDebugging());
-				continue;
-			}
-
-			// Assume DLC base game redirection
-			// TODO: Unify these two redirection paths
-			RedirectPackageUnverified(*Package, RedirectedToSourceImportIndexMap);
-			SourceToRedirectedPackageMap.Add(SourcePackageId, Package->Id);
-			Package->bIsRedirected = true;
-			for (const FPackageStorePackage::FExport& Export : Package->Exports)
-			{
-				if (!Export.SuperIndex.IsNull() && Export.OuterIndex.IsNull())
-				{
-					UE_LOG(LogPackageStoreOptimizer, Warning, TEXT("Skipping redirect to package '%s' due to presence of UStruct '%s'"), *Package->Name.ToString(), *Export.ObjectName.ToString());
-					Package->bIsRedirected = false;
-					break;
-				}
-			}
-		}
+		Package->bIsRedirected = VerifyRedirect(SourcePackage, *Package, bIsBuildingDLC);
 		
 		if (Package->bIsRedirected)
 		{
@@ -1365,10 +1262,6 @@ void FPackageStoreOptimizer::ProcessRedirects(const TMap<FPackageId, FPackageSto
 				SourcePackageId.ValueForDebugging(),
 				*Package->Name.ToString(),
 				Package->Id.ValueForDebugging());
-			if (SourcePackage)
-			{
-				SourcePackage->RedirectedToPackageIds.Add(Package->Id);
-			}
 		}
 		else
 		{
@@ -1380,37 +1273,6 @@ void FPackageStoreOptimizer::ProcessRedirects(const TMap<FPackageId, FPackageSto
 				Package->Id.ValueForDebugging());
 		}
 	}
-
-	for (const auto& KV : PackagesMap)
-	{
-		FPackageStorePackage* Package = KV.Value;
-		for (FPackageId ImportedPackageId : Package->ImportedPackageIds)
-		{
-			FPackageStorePackage* FindImportedPackage = PackagesMap.FindRef(ImportedPackageId);
-			if (FindImportedPackage)
-			{
-				Package->ImportedRedirectedPackageIds.Append(FindImportedPackage->RedirectedToPackageIds);
-			}
-		}
-
-		for (FPackageStorePackage::FImport& Import : Package->Imports)
-		{
-			if (Import.GlobalImportIndex.IsPackageImport())
-			{
-				const FPackageObjectIndex* FindSourceGlobalImportIndex = RedirectedToSourceImportIndexMap.Find(Import.GlobalImportIndex);
-				if (FindSourceGlobalImportIndex)
-				{
-					Import.GlobalImportIndex = *FindSourceGlobalImportIndex;
-				}
-
-				const FPackageId* FindRedirectedPackageId = SourceToRedirectedPackageMap.Find(Import.PackageId);
-				if (FindRedirectedPackageId)
-				{
-					Import.PackageId = *FindRedirectedPackageId;
-				}
-			}
-		}
-	}
 }
 
 void FPackageStoreOptimizer::SerializeGraphData(const TArray<FPackageId>& ImportedPackageIds, FPackageStorePackage::FGraphData& GraphData, FBufferWriter& GraphArchive) const
@@ -1477,10 +1339,9 @@ void FPackageStoreOptimizer::SerializeGraphData(const TArray<FPackageId>& Import
 void FPackageStoreOptimizer::FinalizePackageHeader(FPackageStorePackage* Package) const
 {
 	FBufferWriter ImportMapArchive(nullptr, 0, EBufferWriterFlags::AllowResize | EBufferWriterFlags::TakeOwnership);
-	for (const FPackageStorePackage::FImport& Import : Package->Imports)
+	for (FPackageObjectIndex Import : Package->Imports)
 	{
-		FPackageObjectIndex GlobalImportIndex = Import.GlobalImportIndex;
-		ImportMapArchive << GlobalImportIndex;
+		ImportMapArchive << Import;
 	}
 	Package->ImportMapSize = ImportMapArchive.Tell();
 
@@ -1491,11 +1352,11 @@ void FPackageStoreOptimizer::FinalizePackageHeader(FPackageStorePackage* Package
 		ExportMapEntry.CookedSerialOffset = Export.CookedSerialOffset;
 		ExportMapEntry.CookedSerialSize = Export.SerialSize;
 		ExportMapEntry.ObjectName = Package->NameMapBuilder.MapName(Export.ObjectName);
+		ExportMapEntry.ExportHash = Export.ExportHash;
 		ExportMapEntry.OuterIndex = Export.OuterIndex;
 		ExportMapEntry.ClassIndex = Export.ClassIndex;
 		ExportMapEntry.SuperIndex = Export.SuperIndex;
 		ExportMapEntry.TemplateIndex = Export.TemplateIndex;
-		ExportMapEntry.GlobalImportIndex = Export.GlobalImportIndex;
 		ExportMapEntry.ObjectFlags = Export.ObjectFlags;
 		ExportMapEntry.FilterFlags = EExportFilterFlags::None;
 		if (Export.bNotForClient)
@@ -1527,8 +1388,6 @@ void FPackageStoreOptimizer::FinalizePackageHeader(FPackageStorePackage* Package
 
 	Package->NameMapBuilder.MarkNameAsReferenced(Package->Name);
 	FMappedName MappedPackageName = Package->NameMapBuilder.MapName(Package->Name);
-	Package->NameMapBuilder.MarkNameAsReferenced(Package->SourceName);
-	FMappedName MappedPackageSourceName = Package->NameMapBuilder.MapName(Package->SourceName);
 
 	FBufferWriter NameMapArchive(nullptr, 0, EBufferWriterFlags::AllowResize | EBufferWriterFlags::TakeOwnership);
 	SaveNameBatch(Package->NameMapBuilder.GetNameMap(), NameMapArchive);
@@ -1548,7 +1407,6 @@ void FPackageStoreOptimizer::FinalizePackageHeader(FPackageStorePackage* Package
 	FPackageSummary* PackageSummary = reinterpret_cast<FPackageSummary*>(HeaderData);
 	PackageSummary->HeaderSize = Package->HeaderSize;
 	PackageSummary->Name = MappedPackageName;
-	PackageSummary->SourceName = MappedPackageSourceName;
 	PackageSummary->PackageFlags = Package->PackageFlags;
 	PackageSummary->CookedHeaderSize = Package->CookedHeaderSize;
 	FBufferWriter HeaderArchive(HeaderData, Package->HeaderSize);
@@ -1597,14 +1455,10 @@ void FPackageStoreOptimizer::FinalizePackage(FPackageStorePackage* Package)
 		for (FPackageStorePackage::FExternalDependency& ExternalDependency : Node.ExternalDependencies)
 		{
 			check(ExternalDependency.ImportIndex >= 0);
-			const FPackageStorePackage::FImport& Import = Package->Imports[ExternalDependency.ImportIndex];
-			if (ExternalDependency.bIsConfirmedMissing)
-			{
-				UE_LOG(LogPackageStoreOptimizer, Warning, TEXT("Package '%s' missing import '%s'"), *Package->Name.ToString(), *Import.FullName);
-				continue;
-			}
-
-			TArray<FPackageStorePackage::FExternalArc>& ArcsFromImportedPackage = Package->GraphData.ExternalArcs.FindOrAdd(Import.PackageId);
+			const FPackageObjectIndex& Import = Package->Imports[ExternalDependency.ImportIndex];
+			check(Import.IsPackageImport());
+			FPackageImportReference PackageImportRef = Import.ToPackageImportRef();
+			TArray<FPackageStorePackage::FExternalArc>& ArcsFromImportedPackage = Package->GraphData.ExternalArcs.FindOrAdd(Package->ImportedPackageIds[PackageImportRef.GetImportedPackageIndex()]);
 			FPackageStorePackage::FExternalArc Arc;
 			Arc.FromImportIndex = ExternalDependency.ImportIndex;
 			Arc.FromCommandType = ExternalDependency.ExportBundleCommandType;
@@ -1874,18 +1728,22 @@ FContainerHeader FPackageStoreOptimizer::CreateContainerHeader(const FIoContaine
 	});
 
 	Header.PackageIds.Reserve(SortedPackageStoreEntries.Num());
+	FPackageStoreNameMapBuilder RedirectsNameMapBuilder;
+	RedirectsNameMapBuilder.SetNameMapType(FMappedName::EType::Container);
 	for (const FPackageStoreEntryResource* Entry : SortedPackageStoreEntries)
 	{
 		Header.PackageIds.Add(Entry->GetPackageId());
 		if (Entry->IsRedirected())
 		{
+			RedirectsNameMapBuilder.MarkNameAsReferenced(Entry->GetSourcePackageName());
+			FMappedName MappedSourcePackageName = RedirectsNameMapBuilder.MapName(Entry->GetSourcePackageName());
 			if (Entry->Region.IsValid())
 			{
-				Header.CulturePackageMap.FindOrAdd(Entry->Region.ToString()).Emplace(Entry->GetSourcePackageId(), Entry->GetPackageId());
+				Header.CulturePackageMap.FindOrAdd(Entry->Region.ToString()).Add({ Entry->GetSourcePackageId(), Entry->GetPackageId(), MappedSourcePackageName });
 			}
 			else
 			{
-				Header.PackageRedirects.Add(MakeTuple(Entry->GetSourcePackageId(), Entry->GetPackageId()));
+				Header.PackageRedirects.Add({ Entry->GetSourcePackageId(), Entry->GetPackageId(), MappedSourcePackageName });
 			}
 		}
 		
@@ -1912,6 +1770,7 @@ FContainerHeader FPackageStoreOptimizer::CreateContainerHeader(const FIoContaine
 			StoreDataArchive << const_cast<FSHAHash&>(ShaderMapHash);
 		}
 	}
+	Header.RedirectsNameMap = RedirectsNameMapBuilder.GetNameMap();
 
 	check(StoreTocArchive.TotalSize() == StoreTocSize);
 
diff --git a/Engine/Source/Developer/IoStoreUtilities/Private/PackageStoreOptimizer.h b/Engine/Source/Developer/IoStoreUtilities/Private/PackageStoreOptimizer.h
--- a/Engine/Source/Developer/IoStoreUtilities/Private/PackageStoreOptimizer.h
+++ b/Engine/Source/Developer/IoStoreUtilities/Private/PackageStoreOptimizer.h
@@ -139,11 +139,6 @@ public:
 		return ImportedPackageIds;
 	}
 
-	const TSet<FPackageId>& GetImportedRedirectedPackageIds() const
-	{
-		return ImportedRedirectedPackageIds;
-	}
-
 	void RedirectFrom(FName SourcePackageName)
 	{
 		SourceName = SourcePackageName;
@@ -166,7 +161,6 @@ private:
 	{
 		int32 ImportIndex = -1;
 		FExportBundleEntry::EExportCommandType ExportBundleCommandType;
-		bool bIsConfirmedMissing = false;
 	};
 
 	struct FExportGraphNode
@@ -183,7 +177,7 @@ private:
 	{
 		FString FullName;
 		FName ObjectName;
-		FPackageObjectIndex GlobalImportIndex;
+		uint32 ExportHash = 0;
 		FPackageObjectIndex OuterIndex;
 		FPackageObjectIndex ClassIndex;
 		FPackageObjectIndex SuperIndex;
@@ -204,13 +198,14 @@ private:
 		TArray<FExportBundleEntry> Entries;
 	};
 
-	struct FImport
+	struct FUnresolvedImport
 	{
 		FString FullName;
-		FPackageId PackageId;
-		FPackageObjectIndex GlobalImportIndex;
+		FName FromPackageName;
+		int32 FromPackageNameLen = 0;
+		FPackageId FromPackageId;
 		bool bIsScriptImport = false;
-		bool bIsPackageImport = false;
+		bool bIsImportOfPackage = false;
 	};
 
 	struct FInternalArc
@@ -260,14 +255,12 @@ private:
 	FString Region;
 
 	FPackageStoreNameMapBuilder NameMapBuilder;
-	TArray<FImport> Imports;
+	TArray<FPackageObjectIndex> Imports;
 	TArray<FExport> Exports;
 	TArray<FExportGraphNode> ExportGraphNodes;
 	FGraphData GraphData;
 
 	TArray<FPackageId> ImportedPackageIds;
-	TSet<FPackageId> RedirectedToPackageIds;
-	TSet<FPackageId> ImportedRedirectedPackageIds;
 	TSet<FSHAHash> ShaderMapHashes;
 
 	FIoBuffer HeaderBuffer;
@@ -337,9 +330,8 @@ public:
 	FPackageStoreEntryResource CreatePackageStoreEntry(const FPackageStorePackage* Package) const;
 	FContainerHeader CreateContainerHeader(const FIoContainerId& ContainerId, TArrayView<const FPackageStoreEntryResource> PackageStoreEntries) const;
 	FIoBuffer CreateScriptObjectsBuffer() const;
-	void ProcessRedirects(const TMap<FPackageId, FPackageStorePackage*>& PackagesMap) const;
+	void ProcessRedirects(const TMap<FPackageId, FPackageStorePackage*>& PackagesMap, bool bIsBuildingDLC) const;
 	void OptimizeExportBundles(const TMap<FPackageId, FPackageStorePackage*>& PackagesMap);
-	void FindImports(const TMap<FPackageId, FPackageStorePackage*>& PackagesMap) const;
 
 private:
 	struct FScriptObjectData
@@ -378,7 +370,7 @@ private:
 
 	FCookedHeaderData LoadCookedHeader(const FIoBuffer& CookedHeaderBuffer) const;
 	FPackageStoreHeaderData LoadPackageStoreHeader(const FIoBuffer& PackageStoreHeaderBuffer, const FPackageStoreEntryResource& PackageStoreEntry) const;
-	void ResolveImport(FPackageStorePackage::FImport* Imports, const FObjectImport* ObjectImports, int32 LocalImportIndex) const;
+	void ResolveImport(FPackageStorePackage::FUnresolvedImport* Imports, const FObjectImport* ObjectImports, int32 LocalImportIndex) const;
 	void ProcessImports(const FCookedHeaderData& CookedHeaderData, FPackageStorePackage* Package) const;
 	void ProcessImports(const FPackageStoreHeaderData& PackageStoreHeaderData, FPackageStorePackage* Package) const;
 	void ResolveExport(FPackageStorePackage::FExport* Exports, const FObjectExport* ObjectExports, const int32 LocalExportIndex, const FName& PackageName) const;
@@ -392,8 +384,7 @@ private:
 	TArray<FPackageStorePackage::FExportGraphNode*> SortExportGraphNodesInLoadOrder(FPackageStorePackage* Package, FExportGraphEdges& Edges) const;
 	TArray<FPackageStorePackage::FExportBundleGraphNode*> SortExportBundleGraphNodesInLoadOrder(const TArray<FPackageStorePackage*>& Packages, FExportBundleGraphEdges& Edges) const;
 	void CreateExportBundles(FPackageStorePackage* Package) const;
-	bool RedirectPackage(const FPackageStorePackage& SourcePackage, FPackageStorePackage& TargetPackage, TMap<FPackageObjectIndex, FPackageObjectIndex>& RedirectedToSourceImportIndexMap) const;
-	void RedirectPackageUnverified(FPackageStorePackage& TargetPackage, TMap<FPackageObjectIndex, FPackageObjectIndex>& RedirectedToSourceImportIndexMap) const;
+	bool VerifyRedirect(const FPackageStorePackage* SourcePackage, FPackageStorePackage& TargetPackage, bool bIsBuildingDLC) const;
 	void FinalizePackageHeader(FPackageStorePackage* Package) const;
 	void FindScriptObjectsRecursive(FPackageObjectIndex OuterIndex, UObject* Object, EObjectMark ExcludedObjectMarks, const ITargetPlatform* TargetPlatform);
 	void FindScriptObjects(const ITargetPlatform* TargetPlatform);
diff --git a/Engine/Source/Runtime/Core/Public/IO/PackageStore.h b/Engine/Source/Runtime/Core/Public/IO/PackageStore.h
--- a/Engine/Source/Runtime/Core/Public/IO/PackageStore.h
+++ b/Engine/Source/Runtime/Core/Public/IO/PackageStore.h
@@ -84,6 +84,11 @@ struct FPackageStoreEntryResource
 		return SourcePackageName.IsNone() ? FPackageId() : FPackageId::FromName(SourcePackageName);
 	}
 
+	FName GetSourcePackageName() const
+	{
+		return SourcePackageName;
+	}
+
 	/** Returns whether this package is redirected. */
 	bool IsRedirected() const
 	{
@@ -195,9 +200,6 @@ public:
 	/* Returns the package store entry data with export info and imported packages for the specified package ID. */
 	virtual FPackageStoreEntry GetPackageEntry(FPackageStoreEntryHandle Handle) = 0;
 
-	/* Returns the redirected package ID for the specified package ID. */
-	virtual FPackageId GetRedirectedPackageId(FPackageId PackageId) = 0;
-
-	/* Returns whether the package ID is a redirect. */
-	virtual bool IsRedirect(FPackageId PackageId) = 0;
+	/* Returns the redirected package ID and source package name for the specified package ID if it's being redirected. */
+	virtual bool GetPackageRedirectInfo(FPackageId PackageId, FName& OutSourcePackageName, FPackageId& OutRedirectedToPackageId) = 0;
 };
diff --git a/Engine/Source/Runtime/CoreUObject/Private/Serialization/AsyncLoading2.cpp b/Engine/Source/Runtime/CoreUObject/Private/Serialization/AsyncLoading2.cpp
--- a/Engine/Source/Runtime/CoreUObject/Private/Serialization/AsyncLoading2.cpp
+++ b/Engine/Source/Runtime/CoreUObject/Private/Serialization/AsyncLoading2.cpp
@@ -92,12 +92,33 @@ static FAutoConsoleVariableRef CVar_AsyncLoadingMaxPendingRequestsSizeMB(
 );
 
 
+FArchive& operator<<(FArchive& Ar, FContainerHeaderPackageRedirect& Redirect)
+{
+	Ar << Redirect.SourcePackageId;
+	Ar << Redirect.TargetPackageId;
+	Ar << Redirect.SourcePackageName;
+
+	return Ar;
+}
+
 FArchive& operator<<(FArchive& Ar, FContainerHeader& ContainerHeader)
 {
 	Ar << ContainerHeader.ContainerId;
 	Ar << ContainerHeader.PackageCount;
 	Ar << ContainerHeader.PackageIds;
 	Ar << ContainerHeader.StoreEntries;
+	if (Ar.IsLoading())
+	{
+		ContainerHeader.RedirectsNameMap = LoadNameBatch(Ar);
+	}
+	else
+	{
+#if ALLOW_NAME_BATCH_SAVING
+		SaveNameBatch(ContainerHeader.RedirectsNameMap, Ar);
+#else
+		check(false);
+#endif
+	}
 	Ar << ContainerHeader.CulturePackageMap;
 	Ar << ContainerHeader.PackageRedirects;
 
@@ -139,7 +160,7 @@ FArchive& operator<<(FArchive& Ar, FExportMapEntry& ExportMapEntry)
 	Ar << ExportMapEntry.ClassIndex;
 	Ar << ExportMapEntry.SuperIndex;
 	Ar << ExportMapEntry.TemplateIndex;
-	Ar << ExportMapEntry.GlobalImportIndex;
+	Ar << ExportMapEntry.ExportHash;
 
 	uint32 ObjectFlags = uint32(ExportMapEntry.ObjectFlags);
 	Ar << ObjectFlags;
@@ -269,7 +290,7 @@ static FAutoConsoleVariableRef CVar_VerbosePackageNames(
 #endif
 
 #define UE_ASYNC_PACKAGE_DEBUG(PackageDesc) \
-if (GAsyncLoading2_DebugPackageIds.Contains((PackageDesc).DiskPackageId)) \
+if (GAsyncLoading2_DebugPackageIds.Contains((PackageDesc).UPackageId)) \
 { \
 	UE_DEBUG_BREAK(); \
 }
@@ -286,24 +307,14 @@ if (GAsyncLoading2_DebugPackageIds.Contains((UPackage)->GetPackageId())) \
 #define UE_ASYNC_PACKAGE_LOG(Verbosity, PackageDesc, LogDesc, Format, ...) \
 if ((ELogVerbosity::Type(ELogVerbosity::Verbosity & ELogVerbosity::VerbosityMask) < ELogVerbosity::Verbose) || \
 	(GAsyncLoading2_VerboseLogFilter == 2) || \
-	(GAsyncLoading2_VerboseLogFilter == 1 && GAsyncLoading2_VerbosePackageIds.Contains((PackageDesc).DiskPackageId))) \
+	(GAsyncLoading2_VerboseLogFilter == 1 && GAsyncLoading2_VerbosePackageIds.Contains((PackageDesc).UPackageId))) \
 { \
-	if (!(PackageDesc).CustomPackageName.IsNone()) \
-	{ \
-		UE_LOG(LogStreaming, Verbosity, LogDesc TEXT(": %s (0x%llX) %s (0x%llX) - ") Format, \
-			*(PackageDesc).CustomPackageName.ToString(), \
-			(PackageDesc).CustomPackageId.ValueForDebugging(), \
-			*(PackageDesc).DiskPackageName.ToString(), \
-			(PackageDesc).DiskPackageId.ValueForDebugging(), \
-			##__VA_ARGS__); \
-	} \
-	else \
-	{ \
-		UE_LOG(LogStreaming, Verbosity, LogDesc TEXT(": %s (0x%llX) - ") Format, \
-			*(PackageDesc).DiskPackageName.ToString(), \
-			(PackageDesc).DiskPackageId.ValueForDebugging(), \
-			##__VA_ARGS__); \
-	} \
+	UE_LOG(LogStreaming, Verbosity, LogDesc TEXT(": %s (0x%llX) %s (0x%llX) - ") Format, \
+		*(PackageDesc).UPackageName.ToString(), \
+		(PackageDesc).UPackageId.ValueForDebugging(), \
+		*(PackageDesc).PackageNameToLoad.ToString(), \
+		(PackageDesc).PackageIdToLoad.ValueForDebugging(), \
+		##__VA_ARGS__); \
 }
 
 #define UE_ASYNC_PACKAGE_CLOG(Condition, Verbosity, PackageDesc, LogDesc, Format, ...) \
@@ -391,32 +402,24 @@ struct FPackageRequest
 {
 	int32 RequestId = -1;
 	int32 Priority = -1;
-	FPackageId PackageId;
-	FPackagePath PackagePath;
 	FName CustomName;
+	FName PackageNameToLoad;
+	FName UPackageName;
+	FPackageId PackageIdToLoad;
 	TUniquePtr<FLoadPackageAsyncDelegate> PackageLoadedDelegate;
 	FPackageStoreEntryHandle StoreEntryHandle;
 	FPackageRequest* Next = nullptr;
 
-	FName GetUPackageName() const
-	{
-		return CustomName.IsNone() ? PackagePath.GetPackageFName() : CustomName;
-	}
-
-	FName GetDiskPackageName() const
-	{
-		return PackagePath.GetPackageFName();
-	}
-
-	static FPackageRequest Create(const int32 RequestId, const int32 Priority, FPackagePath PackagePath, FName CustomPackageName, TUniquePtr<FLoadPackageAsyncDelegate> PackageLoadedDelegate)
+	static FPackageRequest Create(const int32 RequestId, const int32 Priority, FName PackageNameToLoad, FPackageId PackageIdToLoad, FName CustomName, TUniquePtr<FLoadPackageAsyncDelegate> PackageLoadedDelegate)
 	{
 		return FPackageRequest
 		{
 			RequestId,
 			Priority,
-			FPackageId::FromName(PackagePath.GetPackageFName()),
-			MoveTemp(PackagePath),
-			CustomPackageName,
+			CustomName,
+			PackageNameToLoad,
+			FName(),
+			PackageIdToLoad,
 			MoveTemp(PackageLoadedDelegate),
 			FPackageStoreEntryHandle(),
 			nullptr
@@ -470,16 +473,56 @@ public:
 		: PackageStore(InPackageStore)
 	{ }
 
-	void Enqueue(FPackageRequest LoadPackageRequest)
+	void Enqueue(FPackageRequest&& LoadPackageRequest)
 	{
+		TCHAR NameBuffer[FName::StringBufferSize];
+		uint32 NameLen = LoadPackageRequest.PackageNameToLoad.ToString(NameBuffer);
+		FStringView PackageNameStr = FStringView(NameBuffer, NameLen);
+		if (!FPackageName::IsValidLongPackageName(PackageNameStr))
+		{
+			FString NewPackageNameStr;
+			if (FPackageName::TryConvertFilenameToLongPackageName(FString(PackageNameStr), NewPackageNameStr))
+			{
+				LoadPackageRequest.PackageNameToLoad = *NewPackageNameStr;
+				LoadPackageRequest.PackageIdToLoad = FPackageId::FromName(LoadPackageRequest.PackageNameToLoad);
+			}
+		}
+
+		LoadPackageRequest.UPackageName = LoadPackageRequest.PackageNameToLoad;
+		{
+			FName SourcePackageName;
+			FPackageId RedirectedToPackageId;
+			if (PackageStore.GetPackageRedirectInfo(LoadPackageRequest.PackageIdToLoad, SourcePackageName, RedirectedToPackageId))
+			{
+				LoadPackageRequest.PackageIdToLoad = RedirectedToPackageId;
+				LoadPackageRequest.UPackageName = SourcePackageName;
+			}
+		}
+
+		LoadPackageRequest.StoreEntryHandle = PackageStore.GetPackageEntryHandle(LoadPackageRequest.PackageIdToLoad, LoadPackageRequest.PackageNameToLoad);
+		
+		// Fixup CustomName to handle any input string that can be converted to a long package name.
+		if (!LoadPackageRequest.CustomName.IsNone())
+		{
+			NameLen = LoadPackageRequest.CustomName.ToString(NameBuffer);
+			PackageNameStr = FStringView(NameBuffer, NameLen);
+			if (!FPackageName::IsValidLongPackageName(PackageNameStr))
+			{
+				FString NewPackageNameStr;
+				if (FPackageName::TryConvertFilenameToLongPackageName(FString(PackageNameStr), NewPackageNameStr))
+				{
+					LoadPackageRequest.CustomName = *NewPackageNameStr;
+				}
+			}
+			LoadPackageRequest.UPackageName = LoadPackageRequest.CustomName;
+		}
+
 		FScopeLock _(&CriticalSection);
 
 		FPackageRequest* Request = Alloc();
 		check(Request->Next == nullptr)
 		*Request = MoveTemp(LoadPackageRequest);
 
-		Request->StoreEntryHandle = PackageStore.GetPackageEntryHandle(Request->PackageId, Request->GetDiskPackageName());
-
 #if WITH_COTF
 		EPackageStoreEntryStatus Status = Request->StoreEntryHandle.Status();
 
@@ -502,7 +545,7 @@ public:
 		for (auto It = PendingRequests.CreateIterator(); It; ++It)
 		{
 			FPackageRequest* PendingRequest = *It;
-			PendingRequest->StoreEntryHandle = PackageStore.GetPackageEntryHandle(PendingRequest->PackageId, PendingRequest->GetDiskPackageName());
+			PendingRequest->StoreEntryHandle = PackageStore.GetPackageEntryHandle(PendingRequest->PackageIdToLoad, PendingRequest->PackageNameToLoad);
 			if (PendingRequest->StoreEntryHandle.Status() != EPackageStoreEntryStatus::Pending)
 			{
 				ReadyRequests.Add(PendingRequest);
@@ -571,88 +614,47 @@ struct FAsyncPackageDesc2
 	int32 Priority;
 	// The package store entry with meta data about the actual disk package
 	const FPackageStoreEntryHandle StoreEntry;
-	// The disk package id corresponding to the StoreEntry
-	// It is used by the loader for io chunks and to handle ref tracking of loaded packages and import objects.
-	FPackageId DiskPackageId; 
-	// The custom package id is only set for temp packages with a valid but "fake" CustomPackageName,
-	// if set, it will be used as key when tracking active async packages in AsyncPackageLookup
-	FPackageId CustomPackageId;
-	// The disk package name from the LoadPackage call, or none for imported packages
-	// up until the package summary has been serialized
-	FName DiskPackageName;
-	// The custom package name from the LoadPackage call is only used for temp packages,
-	// if set, it will be used as the runtime UPackage name
-	FName CustomPackageName;
-	// Set from the package summary,
-	FName SourcePackageName;
-
-	void SetDiskPackageName(FName SerializedDiskPackageName, FName SerializedSourcePackageName = FName())
-	{
-		check(DiskPackageName.IsNone() || DiskPackageName == SerializedDiskPackageName);
-		check(SourcePackageName.IsNone() || SourcePackageName == SerializedSourcePackageName);
-		DiskPackageName = SerializedDiskPackageName;
-		SourcePackageName = SerializedSourcePackageName;
-	}
-
-	bool CanBeImported() const
-	{
-		return CustomPackageName.IsNone();
-	}
-
-	/**
-	 * The UPackage name is used by the engine and game code for in-memory and network communication.
-	 */
-	FName GetUPackageName() const
-	{
-		if (!CustomPackageName.IsNone())
-		{
-			// temp packages
-			return CustomPackageName;
-		}
-		else if(!SourcePackageName.IsNone())
-		{
-			// localized packages
-			return SourcePackageName;
-		}
-		// normal packages
-		return DiskPackageName;
-	}
-
-	/**
-	 * The AsyncPackage id is used by the loader as a key in AsyncPackageLookup to track active load requests,
-	 * which in turn is used for looking up packages for setting up serialized arcs (mostly post load dependencies).
-	 */
-	FORCEINLINE FPackageId GetAsyncPackageId() const
-	{
-		return CustomPackageId.IsValid() ? CustomPackageId : DiskPackageId;
-	}
+	// The package id of the UPackage being loaded
+	// It will be used as key when tracking active async packages
+	FPackageId UPackageId;
+	// The id of the package being loaded from disk
+	FPackageId PackageIdToLoad; 
+	// The name of the UPackage being loaded
+	// Set to none for imported packages up until the package summary has been serialized
+	FName UPackageName;
+	// The name of the package being loaded from disk
+	// Set to none up until the package summary has been serialized
+	FName PackageNameToLoad;
+	// Packages with a a custom name can't be imported
+	bool bCanBeImported;
 
 	static FAsyncPackageDesc2 FromPackageRequest(
 		int32 RequestID,
 		int32 Priority,
+		FName UPackageName,
 		FPackageId PackageIdToLoad,
-		const FPackageStoreEntryHandle StoreEntryHandle,
-		FName DiskPackageName,
-		FPackageId PackageId,
-		FName CustomName)
+		bool bHasCustomName,
+		const FPackageStoreEntryHandle StoreEntryHandle)
 	{
 		return FAsyncPackageDesc2
 		{
 			RequestID,
 			Priority,
 			StoreEntryHandle,
+			FPackageId::FromName(UPackageName),
 			PackageIdToLoad,
-			PackageId,
-			DiskPackageName,
-			CustomName,
-			FName()
+			UPackageName,
+			FName(),
+			!bHasCustomName
 		};
 	}
 
 	static FAsyncPackageDesc2 FromPackageImport(
 		int32 RequestID,
 		int32 Priority,
+		FPackageId ImportedPackageId,
 		FPackageId PackageIdToLoad,
+		FName UPackageName,
 		const FPackageStoreEntryHandle StoreEntryHandle)
 	{
 		return FAsyncPackageDesc2
@@ -660,33 +662,26 @@ struct FAsyncPackageDesc2
 			RequestID,
 			Priority,
 			StoreEntryHandle,
+			ImportedPackageId,
 			PackageIdToLoad,
-			FPackageId(),
-			FName(),
+			UPackageName,
 			FName(),
-			FName()
+			true
 		};
 	}
 };
 
-struct FPublicExport
-{
-	UObject* Object = nullptr;
-	FPackageId PackageId; // for fast clear of package load status during GC
-};
-
 // Note: RemoveUnreachableObjects could move from GT to ALT by removing the debug raw pointers here,
 // the tradeoff would be increased complexity and more restricted debug and log possibilities.
-using FUnreachablePackage = TPair<FName, UPackage*>;
 using FUnreachablePublicExport = TPair<int32, UObject*>;
-using FUnreachablePackages = TArray<FUnreachablePackage>;
+using FUnreachablePackages = TArray<UPackage*>;
 using FUnreachablePublicExports = TArray<FUnreachablePublicExport>;
 
 struct FGlobalImportStore
 {
 	TMap<FPackageObjectIndex, UObject*> ScriptObjects;
-	TMap<FPackageObjectIndex, FPublicExport> PublicExportObjects;
-	TMap<int32, FPackageObjectIndex> ObjectIndexToPublicExport;
+	TMap<FPublicExportKey, UObject*> PublicExportObjects;
+	TMap<int32, FPublicExportKey> ObjectIndexToPublicExport;
 	// Temporary initial load data
 	TArray<FScriptObjectEntry> ScriptObjectEntries;
 	TMap<FPackageObjectIndex, FScriptObjectEntry*> ScriptObjectEntriesMap;
@@ -735,83 +730,78 @@ struct FGlobalImportStore
 	TArray<FPackageId> RemovePublicExports(const FUnreachablePublicExports& PublicExports)
 	{
 		TArray<FPackageId> PackageIds;
-		TArray<FPackageObjectIndex> GlobalIndices;
-		GlobalIndices.Reserve(PublicExports.Num());
+		TArray<FPublicExportKey> PublicExportKeys;
+		PublicExportKeys.Reserve(PublicExports.Num());
 		PackageIds.Reserve(PublicExports.Num());
 
 		for (const FUnreachablePublicExport& Item : PublicExports)
 		{
 			int32 ObjectIndex = Item.Key;
-			FPackageObjectIndex GlobalIndex;
-			if (ObjectIndexToPublicExport.RemoveAndCopyValue(ObjectIndex, GlobalIndex))
+			FPublicExportKey PublicExportKey;
+			if (ObjectIndexToPublicExport.RemoveAndCopyValue(ObjectIndex, PublicExportKey))
 			{
-				GlobalIndices.Emplace(GlobalIndex);
+				PublicExportKeys.Emplace(PublicExportKey);
 #if DO_CHECK
-				FPublicExport* PublicExport = PublicExportObjects.Find(GlobalIndex);
-				checkf(PublicExport, TEXT("Missing entry in ImportStore for object %s with id 0x%llX"), *Item.Value->GetPathName(), GlobalIndex.Value());
-				int32 ObjectIndex2 = GUObjectArray.ObjectToIndex(PublicExport->Object);
-				checkf(ObjectIndex2 == ObjectIndex, TEXT("Mismatch in ImportStore for %s with id 0x%llX"), *Item.Value->GetPathName(), GlobalIndex.Value());
+				UObject* PublicExportObject = PublicExportObjects.FindRef(PublicExportKey);
+				checkf(PublicExportObject, TEXT("Missing entry in ImportStore for object %s with id 0x%llX:%d"),
+					*Item.Value->GetPathName(),
+					PublicExportKey.GetPackageId().Value(),
+					PublicExportKey.GetExportHash());
+				int32 ObjectIndex2 = GUObjectArray.ObjectToIndex(PublicExportObject);
+				checkf(ObjectIndex2 == ObjectIndex, TEXT("Mismatch in ImportStore for %s with id 0x%llX:%d"),
+					*Item.Value->GetPathName(),
+					PublicExportKey.GetPackageId().Value(),
+					PublicExportKey.GetExportHash());
 #endif
 			}
 		}
 
 		FPackageId LastPackageId;
-		for (const FPackageObjectIndex& GlobalIndex : GlobalIndices)
+		for (const FPublicExportKey& PublicExportKey : PublicExportKeys)
 		{
-			FPublicExport PublicExport;
-			PublicExportObjects.RemoveAndCopyValue(GlobalIndex, PublicExport);
-			if (!(PublicExport.PackageId == LastPackageId)) // fast approximation of Contains()
+			PublicExportObjects.Remove(PublicExportKey);
+			FPackageId PackageId = PublicExportKey.GetPackageId();
+			if (!(PackageId == LastPackageId)) // fast approximation of Contains()
 			{
-				LastPackageId = PublicExport.PackageId;
+				LastPackageId = PackageId;
 				PackageIds.Emplace(LastPackageId);
 			}
 		}
 		return PackageIds;
 	}
 
-	inline UObject* GetPublicExportObject(FPackageObjectIndex GlobalIndex)
+	inline UObject* FindPublicExportObject(const FPublicExportKey& Key)
 	{
-		check(GlobalIndex.IsPackageImport());
-		UObject* Object = nullptr;
-		FPublicExport* PublicExport = PublicExportObjects.Find(GlobalIndex);
-		if (PublicExport)
-		{
-			Object = PublicExport->Object;
-			checkf(Object && !Object->IsUnreachable(), TEXT("%s"), Object ? *Object->GetFullName() : TEXT("null"));
-		}
+		UObject* Object = PublicExportObjects.FindRef(Key);
+		checkf(!Object || !Object->IsUnreachable(), TEXT("%s"), Object ? *Object->GetFullName() : TEXT("null"));
 		return Object;
 	}
 
 	UObject* FindScriptImportObjectFromIndex(FPackageObjectIndex ScriptImportIndex);
 
-	inline UObject* FindOrGetImportObject(FPackageObjectIndex GlobalIndex)
+	inline UObject* FindScriptImportObject(FPackageObjectIndex GlobalIndex)
 	{
-		check(GlobalIndex.IsImport());
+		check(GlobalIndex.IsScriptImport());
 		UObject* Object = nullptr;
-		if (GlobalIndex.IsScriptImport())
+		if (!bHasInitializedScriptObjects)
 		{
-			if (!bHasInitializedScriptObjects)
-			{
-				Object = FindScriptImportObjectFromIndex(GlobalIndex);
-			}
-			else
-			{
-				Object = ScriptObjects.FindRef(GlobalIndex);
-			}
+			Object = FindScriptImportObjectFromIndex(GlobalIndex);
 		}
 		else
 		{
-			Object = GetPublicExportObject(GlobalIndex);
+			Object = ScriptObjects.FindRef(GlobalIndex);
 		}
 		return Object;
 	}
 
-	void StoreGlobalObject(FPackageId PackageId, FPackageObjectIndex GlobalIndex, UObject* Object)
+	void StoreGlobalObject(FPackageId PackageId, uint32 ExportHash, UObject* Object)
 	{
-		check(GlobalIndex.IsPackageImport());
+		check(PackageId.IsValid());
+		check(ExportHash != 0);
 		int32 ObjectIndex = GUObjectArray.ObjectToIndex(Object);
-		PublicExportObjects.Add(GlobalIndex, {Object, PackageId});
-		ObjectIndexToPublicExport.Add(ObjectIndex, GlobalIndex);
+		FPublicExportKey Key = FPublicExportKey::MakeKey(PackageId, ExportHash);
+		PublicExportObjects.Add(Key, Object);
+		ObjectIndexToPublicExport.Add(ObjectIndex, Key);
 	}
 
 	void FindAllScriptObjects();
@@ -995,7 +985,7 @@ public:
 	}
 #endif
 
-	void RemovePackage(FPackageId PackageId, UPackage* Package)
+	void RemovePackage(UPackage* Package)
 	{
 		UE_ASYNC_UPACKAGE_DEBUG(Package);
 		check(IsGarbageCollecting());
@@ -1005,21 +995,13 @@ public:
 			return;
 		}
 
-		int32 RefCount = RemovePackage(PackageId);
-		if (RefCount < 0) // not found
-		{
-			FPackageId RedirectedId = PackageStore->GetRedirectedPackageId(PackageId);
-			if (RedirectedId.IsValid())
-			{
-				RefCount = RemovePackage(RedirectedId);
-			}
-		}
+		int32 RefCount = RemovePackage(Package->GetPackageId());
 		if (RefCount > 0)
 		{
 			UE_LOG(LogStreaming, Error,
-				TEXT("RemovePackage: %s (0x%llX) %s (0x%llX) - with (ObjectFlags=%x, InternalObjectFlags=%x) - ")
+				TEXT("RemovePackage: %s %s (0x%llX) - with (ObjectFlags=%x, InternalObjectFlags=%x) - ")
 				TEXT("Package destroyed while still being referenced, RefCount %d > 0."),
-				*Package->GetName(), PackageId.Value(),
+				*Package->GetName(),
 				*Package->GetLoadedPath().GetDebugName(), Package->GetPackageId().Value(),
 				Package->GetFlags(), Package->GetInternalFlags(), RefCount);
 			checkf(false, TEXT("Package %s destroyed with RefCount"), *Package->GetName());
@@ -1027,9 +1009,9 @@ public:
 		else if (RefCount < 0)
 		{
 			UE_LOG(LogStreaming, Error,
-				TEXT("RemovePackage: %s (0x%llX) %s (0x%llX) - with (ObjectFlags=%x, InternalObjectFlags=%x) - ")
+				TEXT("RemovePackage: %s %s (0x%llX) - with (ObjectFlags=%x, InternalObjectFlags=%x) - ")
 				TEXT("Package not found!"),
-				*Package->GetName(), PackageId.Value(),
+				*Package->GetName(),
 				*Package->GetLoadedPath().GetDebugName(), Package->GetPackageId().Value(),
 				Package->GetFlags(), Package->GetInternalFlags());
 			checkf(false, TEXT("Package %s not found"), *Package->GetName());
@@ -1039,19 +1021,9 @@ public:
 	void RemovePackages(const FUnreachablePackages& UnreachablePackages)
 	{
 		const int32 PackageCount = UnreachablePackages.Num();
-		TArray<FPackageId> PackageIds;
-		PackageIds.AddUninitialized(PackageCount);
-		bool bForceSingleThreaded = PackageCount < 64;
-		{
-			// TRACE_CPUPROFILER_EVENT_SCOPE(FPackageId::FromName);
-			ParallelFor(PackageCount, [&UnreachablePackages, &PackageIds](int32 Index)
-			{
-				PackageIds[Index] = FPackageId::FromName(UnreachablePackages[Index].Key);
-			}, bForceSingleThreaded);
-		}
 		for (int32 Index = 0; Index < PackageCount; ++Index)
 		{
-			RemovePackage(PackageIds[Index], UnreachablePackages[Index].Value);
+			RemovePackage(UnreachablePackages[Index]);
 		}
 	}
 
@@ -1075,13 +1047,15 @@ struct FPackageImportStore
 	FGlobalImportStore& GlobalImportStore;
 	FLoadedPackageStore& LoadedPackageStore;
 	const FAsyncPackageDesc2& Desc;
+	const TArrayView<const FPackageId>& ImportedPackageIds;
 	TArrayView<const FPackageObjectIndex> ImportMap;
 
-	FPackageImportStore(IPackageStore& InPackageStore, FGlobalImportStore& InGlobalImportStore, FLoadedPackageStore& InLoadedPackageStore, const FAsyncPackageDesc2& InDesc)
+	FPackageImportStore(IPackageStore& InPackageStore, FGlobalImportStore& InGlobalImportStore, FLoadedPackageStore& InLoadedPackageStore, const FAsyncPackageDesc2& InDesc, const TArrayView<const FPackageId>& InImportedPackageIds)
 		: PackageStore(InPackageStore)
 		, GlobalImportStore(InGlobalImportStore)
 		, LoadedPackageStore(InLoadedPackageStore)
 		, Desc(InDesc)
+		, ImportedPackageIds(InImportedPackageIds)
 	{
 	}
 
@@ -1103,10 +1077,20 @@ struct FPackageImportStore
 		const int32 LocalImportIndex = LocalIndex.ToImport();
 		check(LocalImportIndex < ImportMap.Num());
 		const FPackageObjectIndex GlobalIndex = ImportMap[LocalIndex.ToImport()];
+		return FindOrGetImportObject(GlobalIndex);
+	}
+
+	inline UObject* FindOrGetImportObject(FPackageObjectIndex GlobalIndex)
+	{
+		check(GlobalIndex.IsImport());
 		UObject* Object = nullptr;
-		if (GlobalIndex.IsImport())
+		if (GlobalIndex.IsScriptImport())
 		{
-			Object = GlobalImportStore.FindOrGetImportObject(GlobalIndex);
+			Object = GlobalImportStore.FindScriptImportObject(GlobalIndex);
+		}
+		else if (GlobalIndex.IsPackageImport())
+		{
+			Object = GlobalImportStore.FindPublicExportObject(FPublicExportKey::FromPackageImport(GlobalIndex, ImportedPackageIds));
 		}
 		else
 		{
@@ -1115,12 +1099,6 @@ struct FPackageImportStore
 		return Object;
 	}
 
-	inline UObject* FindOrGetImportObject(FPackageObjectIndex GlobalIndex)
-	{
-		check(GlobalIndex.IsImport());
-		return GlobalImportStore.FindOrGetImportObject(GlobalIndex);
-	}
-
 	bool GetUnresolvedCDOs(TArray<UClass*, TInlineAllocator<8>>& Classes)
 	{
 		for (const FPackageObjectIndex& Index : ImportMap)
@@ -1152,9 +1130,9 @@ struct FPackageImportStore
 		return Classes.Num() > 0;
 	}
 
-	inline void StoreGlobalObject(FPackageId PackageId, FPackageObjectIndex GlobalIndex, UObject* Object)
+	inline void StoreGlobalObject(FPackageId PackageId, uint32 ExportHash, UObject* Object)
 	{
-		GlobalImportStore.StoreGlobalObject(PackageId, GlobalIndex, Object);
+		GlobalImportStore.StoreGlobalObject(PackageId, ExportHash, Object);
 	}
 
 private:
@@ -1199,7 +1177,7 @@ private:
 	}
 
 public:
-	void AddPackageReferences(const TArrayView<const FPackageId>& ImportedPackageIds)
+	void AddPackageReferences()
 	{
 		for (const FPackageId& ImportedPackageId : ImportedPackageIds)
 		{
@@ -1209,9 +1187,9 @@ public:
 				AddAsyncFlags(PackageRef.GetPackage());
 			}
 		}
-		if (Desc.CanBeImported())
+		if (Desc.bCanBeImported)
 		{
-			FLoadedPackageRef& PackageRef = LoadedPackageStore.GetPackageRef(Desc.DiskPackageId);
+			FLoadedPackageRef& PackageRef = LoadedPackageStore.GetPackageRef(Desc.UPackageId);
 			PackageRef.ClearErrorFlags();
 			if (PackageRef.AddRef())
 			{
@@ -1220,21 +1198,21 @@ public:
 		}
 	}
 
-	void ReleasePackageReferences(const TArrayView<const FPackageId>& ImportedPackageIds)
+	void ReleasePackageReferences()
 	{
 		for (const FPackageId& ImportedPackageId : ImportedPackageIds)
 		{
 			FLoadedPackageRef& PackageRef = LoadedPackageStore.GetPackageRef(ImportedPackageId);
-			if (PackageRef.ReleaseRef(Desc.DiskPackageId, ImportedPackageId))
+			if (PackageRef.ReleaseRef(Desc.UPackageId, ImportedPackageId))
 			{
 				ClearAsyncFlags(PackageRef.GetPackage());
 			}
 		}
-		if (Desc.CanBeImported())
+		if (Desc.bCanBeImported)
 		{
 			// clear own reference, and possible all async flags if no remaining ref count
-			FLoadedPackageRef& PackageRef =	LoadedPackageStore.GetPackageRef(Desc.DiskPackageId);
-			if (PackageRef.ReleaseRef(Desc.DiskPackageId, Desc.DiskPackageId))
+			FLoadedPackageRef& PackageRef = LoadedPackageStore.GetPackageRef(Desc.UPackageId);
+			if (PackageRef.ReleaseRef(Desc.UPackageId, Desc.UPackageId))
 			{
 				ClearAsyncFlags(PackageRef.GetPackage());
 			}
@@ -1854,7 +1832,7 @@ private:
 
 	struct FExportToBundleMapping
 	{
-		FPackageObjectIndex GlobalImportIndex;
+		uint32 ExportHash;
 		int32 BundleIndex[FExportBundleEntry::ExportCommandType_Count];
 	};
 
@@ -2484,7 +2462,7 @@ private:
 	{
 		TRACE_CPUPROFILER_EVENT_SCOPE(CreateAsyncPackage);
 		UE_ASYNC_PACKAGE_DEBUG(Desc);
-		checkf(Desc.StoreEntry, TEXT("No package store entry for package %s"), *Desc.DiskPackageName.ToString());
+		checkf(Desc.StoreEntry, TEXT("No package store entry for package %s"), *Desc.UPackageName.ToString());
 
 		FPackageStoreEntry PackageStoreEntry = PackageStore->GetPackageEntry(Desc.StoreEntry);
 
@@ -2679,14 +2657,14 @@ FAsyncPackage2* FAsyncLoadingThread2::FindOrInsertPackage(FAsyncPackageDesc2& De
 	bInserted = false;
 	{
 		FScopeLock LockAsyncPackages(&AsyncPackagesCritical);
-		Package = AsyncPackageLookup.FindRef(Desc.GetAsyncPackageId());
+		Package = AsyncPackageLookup.FindRef(Desc.UPackageId);
 		if (!Package)
 		{
 			Package = CreateAsyncPackage(Desc);
-			checkf(Package, TEXT("Failed to create async package %s"), *Desc.DiskPackageName.ToString());
+			checkf(Package, TEXT("Failed to create async package %s"), *Desc.UPackageName.ToString());
 			Package->AddRef();
 			ActiveAsyncPackagesCounter.Increment();
-			AsyncPackageLookup.Add(Desc.GetAsyncPackageId(), Package);
+			AsyncPackageLookup.Add(Desc.UPackageId, Package);
 			bInserted = true;
 		}
 		else
@@ -2731,7 +2709,7 @@ bool FAsyncLoadingThread2::CreateAsyncPackagesFromQueue(FAsyncLoadingThreadState
 			{
 				bool bInserted;
 				FAsyncPackage2* Package = FindOrInsertPackage(PackageDesc, bInserted, MoveTemp(Request.PackageLoadedDelegate));
-				checkf(Package, TEXT("Failed to find or insert imported package %s"), *PackageDesc.DiskPackageName.ToString());
+				checkf(Package, TEXT("Failed to find or insert package %s"), *PackageDesc.UPackageName.ToString());
 
 				if (bInserted)
 				{
@@ -2811,7 +2789,7 @@ void FAsyncLoadingThread2::StartBundleIoRequests()
 		WaitingIoRequests.HeapPop(BundleIoRequest, false);
 
 		FIoReadOptions ReadOptions;
-		Package->IoRequest = IoBatch.ReadWithCallback(CreateIoChunkId(Package->Desc.DiskPackageId.Value(), 0, EIoChunkType::ExportBundleData),
+		Package->IoRequest = IoBatch.ReadWithCallback(CreateIoChunkId(Package->Desc.PackageIdToLoad.Value(), 0, EIoChunkType::ExportBundleData),
 			ReadOptions,
 			Package->Desc.Priority,
 			[Package](TIoStatusOr<FIoBuffer> Result)
@@ -3253,9 +3231,10 @@ void FAsyncPackage2::SetupSerializedArcs()
 				check(ToExportBundleIndex < Data.ExportInfo.ExportBundleCount);
 				
 				FPackageObjectIndex GlobalImportIndex = ImportStore.ImportMap[FromImportIndex];
+				FPackageImportReference PackageImportRef = GlobalImportIndex.ToPackageImportRef();
 				for (const FExportToBundleMapping& ExportToBundleMapping : ImportedPackage->ExportToBundleMappings)
 				{
-					if (ExportToBundleMapping.GlobalImportIndex == GlobalImportIndex)
+					if (ExportToBundleMapping.ExportHash == PackageImportRef.GetExportHash())
 					{
 						int32 FromExportBundleIndex = ExportToBundleMapping.BundleIndex[FromCommandType];
 						if (PreviousFromExportBundleIndex != FromExportBundleIndex || PreviousToExportBundleIndex != ToExportBundleIndex)
@@ -3384,6 +3363,18 @@ void FAsyncPackage2::ImportPackagesRecursive()
 	Data.ImportedAsyncPackages = MakeArrayView(Data.ImportedAsyncPackages.GetData(), Data.ImportedPackageIds.Num());
 	for (const FPackageId& ImportedPackageId : Data.ImportedPackageIds)
 	{
+		FName ImportedPackageUPackageName = FName();
+		FPackageId ImportedPackageIdToLoad = ImportedPackageId;
+		{
+			FName SourcePackageName;
+			FPackageId RedirectedToPackageId;
+			if (PackageStore.GetPackageRedirectInfo(ImportedPackageId, SourcePackageName, RedirectedToPackageId))
+			{
+				ImportedPackageUPackageName = SourcePackageName;
+				ImportedPackageIdToLoad = RedirectedToPackageId;
+			}
+		}
+		
 		FLoadedPackageRef& PackageRef = ImportStore.LoadedPackageStore.GetPackageRef(ImportedPackageId);
 		if (PackageRef.AreAllPublicExportsLoaded())
 		{
@@ -3391,7 +3382,7 @@ void FAsyncPackage2::ImportPackagesRecursive()
 			continue;
 		}
 
-		const FPackageStoreEntryHandle ImportedPackageEntry = PackageStore.GetPackageEntryHandle(ImportedPackageId);
+		const FPackageStoreEntryHandle ImportedPackageEntry = PackageStore.GetPackageEntryHandle(ImportedPackageIdToLoad);
 
 		if (!ImportedPackageEntry)
 		{
@@ -3402,7 +3393,7 @@ void FAsyncPackage2::ImportPackagesRecursive()
 			continue;
 		}
 
-		FAsyncPackageDesc2 PackageDesc = FAsyncPackageDesc2::FromPackageImport(INDEX_NONE, Desc.Priority, ImportedPackageId, ImportedPackageEntry);
+		FAsyncPackageDesc2 PackageDesc = FAsyncPackageDesc2::FromPackageImport(INDEX_NONE, Desc.Priority, ImportedPackageId, ImportedPackageIdToLoad, ImportedPackageUPackageName, ImportedPackageEntry);
 		bool bInserted;
 		FAsyncPackage2* ImportedPackage = AsyncLoadingThread.FindOrInsertPackage(PackageDesc, bInserted);
 
@@ -3461,9 +3452,9 @@ EAsyncPackageState::Type FAsyncPackage2::Event_ProcessPackageSummary(FAsyncLoadi
 
 	if (Package->bLoadHasFailed)
 	{
-		if (Package->Desc.CanBeImported())
+		if (Package->Desc.bCanBeImported)
 		{
-			FLoadedPackageRef* PackageRef = Package->ImportStore.LoadedPackageStore.FindPackageRef(Package->Desc.DiskPackageId);
+			FLoadedPackageRef* PackageRef = Package->ImportStore.LoadedPackageStore.FindPackageRef(Package->Desc.UPackageId);
 			check(PackageRef);
 			PackageRef->SetHasFailed();
 		}
@@ -3483,18 +3474,15 @@ EAsyncPackageState::Type FAsyncPackage2::Event_ProcessPackageSummary(FAsyncLoadi
 
 		{
 			FName PackageName = Package->NameMap.GetName(PackageSummary->Name);
-			// Don't apply any redirects in editor builds
-#if !WITH_IOSTORE_IN_EDITOR
-			if (PackageSummary->SourceName != PackageSummary->Name)
+			check(Package->Desc.PackageNameToLoad.IsNone());
+			Package->Desc.PackageNameToLoad = PackageName;
+			check(Package->Desc.PackageIdToLoad == FPackageId::FromName(Package->Desc.PackageNameToLoad));
+			// Imported packages won't have a UPackage name set unless they were redirected, in which case they will have the source package name
+			if (Package->Desc.UPackageName.IsNone())
 			{
-				FName SourcePackageName = Package->NameMap.GetName(PackageSummary->SourceName);
-				Package->Desc.SetDiskPackageName(PackageName, SourcePackageName);
-			}
-			else
-#endif
-			{
-				Package->Desc.SetDiskPackageName(PackageName);
+				Package->Desc.UPackageName = PackageName;
 			}
+			check(Package->Desc.UPackageId == FPackageId::FromName(Package->Desc.UPackageName));
 		}
 
 		Package->CookedHeaderSize = PackageSummary->CookedHeaderSize;
@@ -3525,13 +3513,13 @@ EAsyncPackageState::Type FAsyncPackage2::Event_ProcessPackageSummary(FAsyncLoadi
 			{
 				const FExportMapEntry& ExportMapEntry = Package->ExportMap[BundleEntry->LocalExportIndex];
 				FExportToBundleMapping& ExportToBundleMapping = Package->ExportToBundleMappings[BundleEntry->LocalExportIndex];
-				ExportToBundleMapping.GlobalImportIndex = ExportMapEntry.GlobalImportIndex;
+				ExportToBundleMapping.ExportHash = ExportMapEntry.ExportHash;
 				ExportToBundleMapping.BundleIndex[BundleEntry->CommandType] = ExportBundleIndex;
 				++BundleEntry;
 			}
 		}
 
-		TRACE_LOADTIME_PACKAGE_SUMMARY(Package, Package->Desc.DiskPackageName, PackageSummary->HeaderSize, Package->ImportStore.ImportMap.Num(), Package->Data.ExportInfo.ExportCount);
+		TRACE_LOADTIME_PACKAGE_SUMMARY(Package, Package->Desc.PackageNameToLoad, PackageSummary->HeaderSize, Package->ImportStore.ImportMap.Num(), Package->Data.ExportInfo.ExportCount);
 	}
 
 	Package->AsyncPackageLoadingState = EAsyncPackageLoadingState2::SetupDependencies;
@@ -3727,7 +3715,7 @@ UObject* FAsyncPackage2::EventDrivenIndexToObject(FPackageObjectIndex Index, boo
 		UE_CLOG(!Result, LogStreaming, Warning, TEXT("Missing %s import 0x%llX for package %s"),
 			Index.IsScriptImport() ? TEXT("script") : TEXT("package"),
 			Index.Value(),
-			*Desc.DiskPackageName.ToString());
+			*Desc.PackageNameToLoad.ToString());
 	}
 #if DO_CHECK
 	if (bCheckSerialized && !IsFullyLoadedObj(Result))
@@ -3938,19 +3926,13 @@ void FAsyncPackage2::EventDrivenCreateExport(int32 LocalExportIndex)
 	check(Object);
 	PinObjectForGC(Object, bIsNewObject);
 
-	if (Desc.CanBeImported() && !Export.GlobalImportIndex.IsNull())
+	if (Desc.bCanBeImported && Export.ExportHash)
 	{
 		check(Object->HasAnyFlags(RF_Public));
-		FPackageObjectIndex GlobalImportIndex = Export.GlobalImportIndex;
-#if WITH_IOSTORE_IN_EDITOR
-		// Always compute the global import index when loading cooked packages in editor builds
-		// to prevent localized packages to overwrite the redirected package name
-		GlobalImportIndex = FPackageObjectIndex::FromPackagePath(Object->GetPathName());
-#endif
-		ImportStore.StoreGlobalObject(Desc.DiskPackageId, GlobalImportIndex, Object);
+		ImportStore.StoreGlobalObject(Desc.UPackageId, Export.ExportHash, Object);
 
 		UE_ASYNC_PACKAGE_LOG_VERBOSE(VeryVerbose, Desc, TEXT("CreateExport"),
-			TEXT("Created public export %s. Tracked as 0x%llX"), *Object->GetPathName(), Export.GlobalImportIndex.Value());
+			TEXT("Created public export %s. Tracked as 0x%llX:%d"), *Object->GetPathName(), Desc.UPackageId.Value(), Export.ExportHash);
 	}
 	else
 	{
@@ -4058,10 +4040,10 @@ EAsyncPackageState::Type FAsyncPackage2::Event_ExportsDone(FAsyncLoadingThreadSt
 	UE_ASYNC_PACKAGE_DEBUG(Package->Desc);
 	check(Package->AsyncPackageLoadingState == EAsyncPackageLoadingState2::ExportsDone);
 
-	if (!Package->bLoadHasFailed && Package->Desc.CanBeImported())
+	if (!Package->bLoadHasFailed && Package->Desc.bCanBeImported)
 	{
 		FLoadedPackageRef& PackageRef =
-			Package->AsyncLoadingThread.LoadedPackageStore.GetPackageRef(Package->Desc.DiskPackageId);
+			Package->AsyncLoadingThread.LoadedPackageStore.GetPackageRef(Package->Desc.UPackageId);
 		PackageRef.SetAllPublicExportsLoaded();
 	}
 
@@ -4477,7 +4459,7 @@ bool FAsyncPackage2::AreAllDependenciesFullyLoaded(TSet<FPackageId>& VisitedPack
 	{
 		FAsyncPackage2* AsyncRoot = AsyncLoadingThread.GetAsyncPackage(PackageId);
 		UE_LOG(LogStreaming, Verbose, TEXT("AreAllDependenciesFullyLoaded: '%s' doesn't have all exports processed by DeferredPostLoad"),
-			*AsyncRoot->Desc.DiskPackageName.ToString());
+			*AsyncRoot->Desc.PackageNameToLoad.ToString());
 	}
 	return bLoaded;
 }
@@ -4596,7 +4578,7 @@ EAsyncPackageState::Type FAsyncLoadingThread2::ProcessLoadedPackagesFromGameThre
 
 			{
 				FScopeLock LockAsyncPackages(&AsyncPackagesCritical);
-				AsyncPackageLookup.Remove(Package->Desc.GetAsyncPackageId());
+				AsyncPackageLookup.Remove(Package->Desc.UPackageId);
 				if (!Package->bLoadHasFailed)
 				{
 					Package->ClearConstructedObjects();
@@ -5251,7 +5233,7 @@ FORCENOINLINE static void FilterUnreachableObjects(
 				if (Package->HasAnyPackageFlags(PKG_Cooked))
 #endif
 				{
-					Packages.Emplace(Package->GetLoadedPath().GetPackageFName(), Package);
+					Packages.Emplace(Package);
 				}
 			}
 		}
@@ -5382,10 +5364,10 @@ FAsyncPackage2::FAsyncPackage2(
 , Data(InData)
 , AsyncLoadingThread(InAsyncLoadingThread)
 , GraphAllocator(InGraphAllocator)
-, ImportStore(AsyncLoadingThread.GetPackageStore(), AsyncLoadingThread.GlobalImportStore, AsyncLoadingThread.LoadedPackageStore, Desc)
+, ImportStore(AsyncLoadingThread.GetPackageStore(), AsyncLoadingThread.GlobalImportStore, AsyncLoadingThread.LoadedPackageStore, Desc, Data.ImportedPackageIds)
 {
 	TRACE_LOADTIME_NEW_ASYNC_PACKAGE(this);
-	ImportStore.AddPackageReferences(Data.ImportedPackageIds);
+	ImportStore.AddPackageReferences();
 	AddRequestID(Desc.RequestID);
 
 	ConstructedObjects.Reserve(Data.ExportInfo.ExportCount + 1); // +1 for UPackage, may grow dynamically beoynd that
@@ -5421,16 +5403,16 @@ FAsyncPackage2::~FAsyncPackage2()
 	TRACE_LOADTIME_DESTROY_ASYNC_PACKAGE(this);
 	UE_ASYNC_PACKAGE_LOG(Verbose, Desc, TEXT("AsyncThread: Deleted"), TEXT("Package deleted."));
 
-	ImportStore.ReleasePackageReferences(Data.ImportedPackageIds);
+	ImportStore.ReleasePackageReferences();
 
 	checkf(RefCount == 0, TEXT("RefCount is not 0 when deleting package %s"),
-		*Desc.DiskPackageName.ToString());
+		*Desc.PackageNameToLoad.ToString());
 
 	checkf(RequestIDs.Num() == 0, TEXT("MarkRequestIDsAsComplete() has not been called for package %s"),
-		*Desc.DiskPackageName.ToString());
+		*Desc.PackageNameToLoad.ToString());
 	
 	checkf(ConstructedObjects.Num() == 0, TEXT("ClearConstructedObjects() has not been called for package %s"),
-		*Desc.DiskPackageName.ToString());
+		*Desc.PackageNameToLoad.ToString());
 }
 
 void FAsyncPackage2::ReleaseRef()
@@ -5473,7 +5455,7 @@ void FAsyncPackage2::ClearConstructedObjects()
 	ConstructedObjects.Empty();
 
 	// the async flag of all GC'able public export objects in non-temp packages are handled by FGlobalImportStore::ClearAsyncFlags
-	const bool bShouldClearAsyncFlagForPublicExports = GUObjectArray.IsDisregardForGC(LinkerRoot) || !Desc.CanBeImported();
+	const bool bShouldClearAsyncFlagForPublicExports = GUObjectArray.IsDisregardForGC(LinkerRoot) || !Desc.bCanBeImported;
 
 	for (FExportObject& Export : Data.Exports)
 	{
@@ -5589,15 +5571,15 @@ void FAsyncPackage2::CreateUPackage(const FPackageSummary* PackageSummary)
 	UPackage* ExistingPackage = nullptr;
 	{
 		TRACE_CPUPROFILER_EVENT_SCOPE(UPackageFind);
-		if (Desc.CanBeImported())
+		if (Desc.bCanBeImported)
 		{
-			PackageRef = ImportStore.LoadedPackageStore.FindPackageRef(Desc.DiskPackageId);
+			PackageRef = ImportStore.LoadedPackageStore.FindPackageRef(Desc.UPackageId);
 			UE_ASYNC_PACKAGE_CLOG(!PackageRef, Fatal, Desc, TEXT("CreateUPackage"), TEXT("Package has been destroyed by GC."));
 			LinkerRoot = PackageRef->GetPackage();
 #if DO_CHECK
 			if (LinkerRoot)
 			{
-				UPackage* FoundPackage = FindObjectFast<UPackage>(nullptr, Desc.GetUPackageName());
+				UPackage* FoundPackage = FindObjectFast<UPackage>(nullptr, Desc.UPackageName);
 				checkf(LinkerRoot == FoundPackage,
 					TEXT("LinkerRoot '%s' (%p) is different from FoundPackage '%s' (%p)"),
 					*LinkerRoot->GetName(), LinkerRoot, *FoundPackage->GetName(), FoundPackage);
@@ -5607,7 +5589,7 @@ void FAsyncPackage2::CreateUPackage(const FPackageSummary* PackageSummary)
 		if (!LinkerRoot)
 		{
 			// Packages can be created outside the loader, i.e from ResolveName via StaticLoadObject
-			ExistingPackage = FindObjectFast<UPackage>(nullptr, Desc.GetUPackageName());
+			ExistingPackage = FindObjectFast<UPackage>(nullptr, Desc.UPackageName);
 		}
 	}
 	if (!LinkerRoot)
@@ -5619,13 +5601,14 @@ void FAsyncPackage2::CreateUPackage(const FPackageSummary* PackageSummary)
 		}
 		else 
 		{
-			LinkerRoot = NewObject<UPackage>(/*Outer*/nullptr, Desc.GetUPackageName());
+			LinkerRoot = NewObject<UPackage>(/*Outer*/nullptr, Desc.UPackageName);
 			bCreatedLinkerRoot = true;
 		}
 		LinkerRoot->SetFlags(RF_Public | RF_WasLoaded);
-		LinkerRoot->SetLoadedPath(FPackagePath::FromPackageNameChecked(Desc.DiskPackageName));
-		LinkerRoot->SetCanBeImportedFlag(Desc.CanBeImported());
-		LinkerRoot->SetPackageId(Desc.DiskPackageId);
+		LinkerRoot->SetLoadedPath(FPackagePath::FromPackageNameChecked(Desc.PackageNameToLoad));
+		LinkerRoot->SetCanBeImportedFlag(Desc.bCanBeImported);
+		LinkerRoot->SetPackageId(Desc.UPackageId);
+		LinkerRoot->SetPackageIdToLoad(Desc.PackageIdToLoad);
 		LinkerRoot->SetPackageFlagsTo(PackageSummary->PackageFlags | PKG_Cooked);
 		LinkerRoot->LinkerPackageVersion = GPackageFileUEVersion;
 		LinkerRoot->LinkerLicenseeVersion = GPackageFileLicenseeUEVersion;
@@ -5640,8 +5623,8 @@ void FAsyncPackage2::CreateUPackage(const FPackageSummary* PackageSummary)
 	}
 	else
 	{
-		check(LinkerRoot->CanBeImported() == Desc.CanBeImported());
-		check(LinkerRoot->GetPackageId() == Desc.DiskPackageId);
+		check(LinkerRoot->CanBeImported() == Desc.bCanBeImported);
+		check(LinkerRoot->GetPackageId() == Desc.UPackageId);
 		check(LinkerRoot->GetPackageFlags() == (PackageSummary->PackageFlags | PKG_Cooked));
 		check(LinkerRoot->LinkerPackageVersion == GPackageFileUEVersion);
 		check(LinkerRoot->LinkerLicenseeVersion == GPackageFileLicenseeUEVersion);
@@ -5745,7 +5728,7 @@ void FAsyncPackage2::CallCompletionCallbacks(EAsyncLoadingResult::Type LoadingRe
 #endif
 	for (FCompletionCallback& CompletionCallback : CompletionCallbacks)
 	{
-		CompletionCallback->ExecuteIfBound(Desc.GetUPackageName(), LoadedPackage, LoadingResult);
+		CompletionCallback->ExecuteIfBound(Desc.UPackageName, LoadedPackage, LoadingResult);
 	}
 	CompletionCallbacks.Empty();
 }
@@ -5787,9 +5770,15 @@ int32 FAsyncLoadingThread2::LoadPackage(const FPackagePath& InPackagePath, FName
 
 	LazyInitializeFromLoadPackage();
 	
-	const FName PackageName = InCustomName.IsNone() ? InPackagePath.GetPackageFName() : InCustomName;
+	const FName PackageNameToLoad = InPackagePath.GetPackageFName();
+	const FPackageId PackageIdToLoad = FPackageId::FromName(PackageNameToLoad);
+	if (InCustomName == PackageNameToLoad)
+	{
+		InCustomName = NAME_None;
+	}
 	if (FCoreDelegates::OnAsyncLoadPackage.IsBound())
 	{
+		const FName PackageName = InCustomName.IsNone() ? PackageNameToLoad : InCustomName;
 		FCoreDelegates::OnAsyncLoadPackage.Broadcast(PackageName.ToString());
 	}
 
@@ -5804,7 +5793,7 @@ int32 FAsyncLoadingThread2::LoadPackage(const FPackagePath& InPackagePath, FName
 		? MakeUnique<FLoadPackageAsyncDelegate>(MoveTemp(InCompletionDelegate))
 		: TUniquePtr<FLoadPackageAsyncDelegate>();
 
-	PackageRequestQueue->Enqueue(FPackageRequest::Create(RequestId, InPackagePriority, InPackagePath, InCustomName, MoveTemp(CompletionDelegate)));
+	PackageRequestQueue->Enqueue(FPackageRequest::Create(RequestId, InPackagePriority, PackageNameToLoad, PackageIdToLoad, InCustomName, MoveTemp(CompletionDelegate)));
 	++QueuedPackagesCounter;
 
 	AltZenaphore.NotifyOne();
@@ -5814,116 +5803,30 @@ int32 FAsyncLoadingThread2::LoadPackage(const FPackagePath& InPackagePath, FName
 
 FAsyncPackageDesc2 FAsyncLoadingThread2::CreatePackageDesc(FPackageRequest& Request, bool& bIsMissingPackage)
 {
-	const FName Name = Request.GetUPackageName();
-	FName DiskPackageName = Request.GetDiskPackageName();
-	bool bHasCustomPackageName = Name != DiskPackageName;
-
-	FPackageId DiskPackageId = Request.PackageId;
+	FPackageId PackageIdToLoad = Request.PackageIdToLoad;
 	FPackageStoreEntryHandle StoreEntry = Request.StoreEntryHandle;
-	if (!StoreEntry.IsValid())
-	{
-		FString PackageNameStr = DiskPackageName.ToString();
-		if (!FPackageName::IsValidLongPackageName(PackageNameStr))
-		{
-			FString NewPackageNameStr;
-			if (FPackageName::TryConvertFilenameToLongPackageName(PackageNameStr, NewPackageNameStr))
-			{
-				DiskPackageName = *NewPackageNameStr;
-				DiskPackageId = FPackageId::FromName(DiskPackageName);
-				StoreEntry = PackageStore->GetPackageEntryHandle(DiskPackageId);
-				bHasCustomPackageName &= Name != DiskPackageName;
-			}
-		}
-	}
-
-	// Verify CustomPackageName, or fixup to handle any input string that can be converted to a long package name.
-	// CustomPackageName must not be an existing disk package name,
-	// that could cause missing or incorrect import objects for other packages.
-	FName CustomPackageName;
-	FPackageId CustomPackageId;
-	if (bHasCustomPackageName)
-	{
-		FPackageId PackageId = FPackageId::FromName(Name);
-		if (!PackageStore->GetPackageEntryHandle(PackageId))
-		{
-			FString PackageNameStr = Name.ToString();
-			if (FPackageName::IsValidLongPackageName(PackageNameStr))
-			{
-				CustomPackageName = Name;
-				CustomPackageId = PackageId;
-			}
-			else
-			{
-				FString NewPackageNameStr;
-				if (FPackageName::TryConvertFilenameToLongPackageName(PackageNameStr, NewPackageNameStr))
-				{
-					PackageId = FPackageId::FromName(FName(*NewPackageNameStr));
-					if (!PackageStore->GetPackageEntryHandle(PackageId))
-					{
-						CustomPackageName = *NewPackageNameStr;
-						CustomPackageId = PackageId;
-					}
-				}
-			}
-		}
-	}
-	// When explicitly requesting a redirected package then set CustomName to
-	// the redirected name, otherwise the UPackage name will be set to the base game name.
-	else if (PackageStore->IsRedirect(DiskPackageId))
-	{
-		bHasCustomPackageName = true;
-		CustomPackageName = DiskPackageName;
-		CustomPackageId = DiskPackageId;
-	}
-
-	check(CustomPackageId.IsValid() == !CustomPackageName.IsNone());
-
-	bool bCustomNameIsValid = (!bHasCustomPackageName && CustomPackageName.IsNone()) || (bHasCustomPackageName && !CustomPackageName.IsNone());
-	bool bDiskPackageIdIsValid = StoreEntry.IsValid();
-	if (!bDiskPackageIdIsValid)
+	
+	if (StoreEntry.Status() == EPackageStoreEntryStatus::Missing)
 	{
 		// While there is an active load request for (InName=/Temp/PackageABC_abc, InPackageToLoadFrom=/Game/PackageABC), then allow these requests too:
 		// (InName=/Temp/PackageA_abc, InPackageToLoadFrom=/Temp/PackageABC_abc) and (InName=/Temp/PackageABC_xyz, InPackageToLoadFrom=/Temp/PackageABC_abc)
-		FAsyncPackage2* Package = GetAsyncPackage(DiskPackageId);
+		FAsyncPackage2* Package = GetAsyncPackage(PackageIdToLoad);
 		if (Package)
 		{
-			if (CustomPackageName.IsNone())
-			{
-				CustomPackageName = Package->Desc.CustomPackageName;
-				CustomPackageId = Package->Desc.CustomPackageId;
-				bHasCustomPackageName = bCustomNameIsValid = true;
-			}
-			DiskPackageName = Package->Desc.DiskPackageName;
-			DiskPackageId = Package->Desc.DiskPackageId;
+			PackageIdToLoad = Package->Desc.PackageIdToLoad;
 			StoreEntry = Package->Desc.StoreEntry;
-			bDiskPackageIdIsValid = true;
 		}
 	}
 
-	FAsyncPackageDesc2 PackageDesc = FAsyncPackageDesc2::FromPackageRequest(Request.RequestId, Request.Priority, DiskPackageId, StoreEntry, DiskPackageName, CustomPackageId, CustomPackageName);
-
-	bool bIsValid = bDiskPackageIdIsValid && bCustomNameIsValid;
-	if (bIsValid)
-	{
-		// Fixup for redirected packages since the slim StoreEntry itself has been stripped from both package names and package ids
-		FPackageId RedirectedDiskPackageId = PackageStore->GetRedirectedPackageId(DiskPackageId);
-		if (RedirectedDiskPackageId.IsValid())
-		{
-			PackageDesc.DiskPackageId = RedirectedDiskPackageId;
-			PackageDesc.SourcePackageName = PackageDesc.DiskPackageName;
-			PackageDesc.DiskPackageName = FName();
-		}
-	}
-	
-	bIsMissingPackage = !bIsValid || StoreEntry.Status() == EPackageStoreEntryStatus::Missing;
+	FAsyncPackageDesc2 PackageDesc = FAsyncPackageDesc2::FromPackageRequest(Request.RequestId, Request.Priority, Request.UPackageName, PackageIdToLoad, !Request.CustomName.IsNone(), Request.StoreEntryHandle);
+	bIsMissingPackage = StoreEntry.Status() == EPackageStoreEntryStatus::Missing;
 
 	return MoveTemp(PackageDesc);
 }
 
 void FAsyncLoadingThread2::QueueMissingPackage(FAsyncPackageDesc2& PackageDesc, TUniquePtr<FLoadPackageAsyncDelegate>&& PackageLoadedDelegate)
 {
-	const bool bInvalidPackageName = PackageDesc.StoreEntry.IsValid();
-	const FName FailedPackageName = bInvalidPackageName ? PackageDesc.CustomPackageName : PackageDesc.DiskPackageName;
+	const FName FailedPackageName = PackageDesc.UPackageName;
 
 	static TSet<FName> SkippedPackages;
 	bool bIsAlreadySkipped = false;
@@ -5932,18 +5835,9 @@ void FAsyncLoadingThread2::QueueMissingPackage(FAsyncPackageDesc2& PackageDesc,
 
 	if (!bIsAlreadySkipped)
 	{
-		if (bInvalidPackageName)
-		{
-			UE_LOG(LogStreaming, Warning,
-				TEXT("LoadPackage: SkipPackage: %s (0x%llX) - The package name is invalid"),
-				*FailedPackageName.ToString(), PackageDesc.CustomPackageId.ValueForDebugging());
-		}
-		else
-		{
-			UE_LOG(LogStreaming, Warning,
-				TEXT("LoadPackage: SkipPackage: %s (0x%llX) - The package to load does not exist on disk or in the loader"),
-				*FailedPackageName.ToString(), PackageDesc.DiskPackageId.ValueForDebugging());
-		}
+		UE_LOG(LogStreaming, Warning,
+			TEXT("LoadPackage: SkipPackage: %s (0x%llX) - The package to load does not exist on disk or in the loader"),
+			*FailedPackageName.ToString(), PackageDesc.PackageIdToLoad.ValueForDebugging());
 	}
 
 	if (PackageLoadedDelegate.IsValid())
diff --git a/Engine/Source/Runtime/CoreUObject/Private/Serialization/BulkData2.cpp b/Engine/Source/Runtime/CoreUObject/Private/Serialization/BulkData2.cpp
--- a/Engine/Source/Runtime/CoreUObject/Private/Serialization/BulkData2.cpp
+++ b/Engine/Source/Runtime/CoreUObject/Private/Serialization/BulkData2.cpp
@@ -863,7 +863,7 @@ void FBulkDataBase::Serialize(FArchive& Ar, UObject* Owner, int32 /*Index*/, boo
 			{
 				if (bSeparateFilesUseIoStore)
 				{
-					Data.PackageID = Package->GetPackageId().Value();
+					Data.PackageID = Package->GetPackageIdToLoad().Value();
 					SetRuntimeBulkDataFlags(BULKDATA_UsesIoDispatcher); // Indicates that this BulkData should use the FIoChunkId rather than a PackagePath
 				}
 				else
diff --git a/Engine/Source/Runtime/CoreUObject/Private/Serialization/CookOnTheFlyPackageStore.cpp b/Engine/Source/Runtime/CoreUObject/Private/Serialization/CookOnTheFlyPackageStore.cpp
--- a/Engine/Source/Runtime/CoreUObject/Private/Serialization/CookOnTheFlyPackageStore.cpp
+++ b/Engine/Source/Runtime/CoreUObject/Private/Serialization/CookOnTheFlyPackageStore.cpp
@@ -183,12 +183,7 @@ private:
 		};
 	}
 
-	virtual FPackageId GetRedirectedPackageId(FPackageId PackageId)  override
-	{
-		return FPackageId();
-	}
-
-	virtual bool IsRedirect(FPackageId PackageId) override
+	virtual bool GetPackageRedirectInfo(FPackageId PackageId, FName& OutSourcePackageName, FPackageId& OutRedirectedToPackageId)  override
 	{
 		return false;
 	}
diff --git a/Engine/Source/Runtime/CoreUObject/Private/Serialization/FilePackageStore.cpp b/Engine/Source/Runtime/CoreUObject/Private/Serialization/FilePackageStore.cpp
--- a/Engine/Source/Runtime/CoreUObject/Private/Serialization/FilePackageStore.cpp
+++ b/Engine/Source/Runtime/CoreUObject/Private/Serialization/FilePackageStore.cpp
@@ -70,16 +70,20 @@ public:
 		};
 	}
 
-	virtual FPackageId GetRedirectedPackageId(FPackageId PackageId) override
+	virtual bool GetPackageRedirectInfo(FPackageId PackageId, FName& OutSourcePackageName, FPackageId& OutRedirectedToPackageId) override
 	{
 		FScopeLock Lock(&PackageNameMapsCritical);
-		FPackageId RedirectedId = RedirectsPackageMap.FindRef(PackageId);
-		return RedirectedId;
-	}
-
-	virtual bool IsRedirect(FPackageId PackageId) override
-	{
-		return TargetRedirectIds.Contains(PackageId);
+		TTuple<FName, FPackageId>* FindRedirect = RedirectsPackageMap.Find(PackageId);
+		if (FindRedirect)
+		{
+			OutSourcePackageName = FindRedirect->Get<0>();
+			OutRedirectedToPackageId = FindRedirect->Get<1>();
+			return true;
+		}
+		else
+		{
+			return false;
+		}
 	}
 
 private:
@@ -193,24 +197,22 @@ private:
 
 							if (LocalizedPackages)
 							{
-								for (auto& Pair : *LocalizedPackages)
+								for (const FContainerHeaderPackageRedirect& Redirect : *LocalizedPackages)
 								{
-									const FPackageId& SourceId = Pair.Key;
-									const FPackageId& LocalizedId = Pair.Value;
-									RedirectsPackageMap.Emplace(SourceId, LocalizedId);
-									TargetRedirectIds.Add(LocalizedId);
+									FNameEntryId NameEntry = ContainerHeader.RedirectsNameMap[Redirect.SourcePackageName.GetIndex()];
+									FName SourcePackageName = FName::CreateFromDisplayId(NameEntry, Redirect.SourcePackageName.GetNumber());
+									RedirectsPackageMap.Emplace(Redirect.SourcePackageId, MakeTuple(SourcePackageName, Redirect.TargetPackageId));
 								}
 							}
 						}
 
 						{
 							TRACE_CPUPROFILER_EVENT_SCOPE(LoadPackageStoreRedirects);
-							for (const TPair<FPackageId, FPackageId>& Redirect : ContainerHeader.PackageRedirects)
+							for (const FContainerHeaderPackageRedirect& Redirect : ContainerHeader.PackageRedirects)
 							{
-								const FPackageId& SourceId = Redirect.Key;
-								const FPackageId& RedirectedId = Redirect.Value;
-								RedirectsPackageMap.Emplace(SourceId, RedirectedId);
-								TargetRedirectIds.Add(RedirectedId);
+								FNameEntryId NameEntry = ContainerHeader.RedirectsNameMap[Redirect.SourcePackageName.GetIndex()];
+								FName SourcePackageName = FName::CreateFromDisplayId(NameEntry, Redirect.SourcePackageName.GetNumber());
+								RedirectsPackageMap.Emplace(Redirect.SourcePackageId, MakeTuple(SourcePackageName, Redirect.TargetPackageId));
 							}
 						}
 					}
@@ -226,8 +228,6 @@ private:
 		IoBatch.Issue();
 		Event->Wait();
 		FPlatformProcess::ReturnSynchEventToPool(Event);
-
-		ApplyRedirects(RedirectsPackageMap);
 	}
 
 	void OnContainerMounted(const FIoContainerId& ContainerId)
@@ -236,45 +236,6 @@ private:
 		LoadContainers(MakeArrayView(&ContainerId, 1));
 	}
 
-	void ApplyRedirects(const TMap<FPackageId, FPackageId>& Redirects)
-	{
-		TRACE_CPUPROFILER_EVENT_SCOPE(ApplyRedirects);
-
-		FScopeLock Lock(&PackageNameMapsCritical);
-
-		if (Redirects.Num() == 0)
-		{
-			return;
-		}
-
-		for (auto It = Redirects.CreateConstIterator(); It; ++It)
-		{
-			const FPackageId& SourceId = It.Key();
-			const FPackageId& RedirectId = It.Value();
-			check(RedirectId.IsValid());
-			FFilePackageStoreEntry* RedirectEntry = StoreEntriesMap.FindRef(RedirectId);
-			check(RedirectEntry);
-			FFilePackageStoreEntry*& PackageEntry = StoreEntriesMap.FindOrAdd(SourceId);
-			if (RedirectEntry)
-			{
-				PackageEntry = RedirectEntry;
-			}
-		}
-
-		for (auto It = StoreEntriesMap.CreateIterator(); It; ++It)
-		{
-			FFilePackageStoreEntry* StoreEntry = It.Value();
-
-			for (FPackageId& ImportedPackageId : StoreEntry->ImportedPackages)
-			{
-				if (const FPackageId* RedirectId = Redirects.Find(ImportedPackageId))
-				{
-					ImportedPackageId = *RedirectId;
-				}
-			}
-		}
-	}
-
 	struct FLoadedContainer
 	{
 		TUniquePtr<FNameMap> ContainerNameMap;
@@ -291,8 +252,7 @@ private:
 	FCriticalSection PackageNameMapsCritical;
 
 	TMap<FPackageId, FFilePackageStoreEntry*> StoreEntriesMap;
-	TMap<FPackageId, FPackageId> RedirectsPackageMap;
-	TSet<FPackageId> TargetRedirectIds;
+	TMap<FPackageId, TTuple<FName, FPackageId>> RedirectsPackageMap;
 	int32 NextCustomPackageIndex = 0;
 };
 
diff --git a/Engine/Source/Runtime/CoreUObject/Public/Serialization/AsyncLoading2.h b/Engine/Source/Runtime/CoreUObject/Public/Serialization/AsyncLoading2.h
--- a/Engine/Source/Runtime/CoreUObject/Public/Serialization/AsyncLoading2.h
+++ b/Engine/Source/Runtime/CoreUObject/Public/Serialization/AsyncLoading2.h
@@ -18,7 +18,16 @@ class IAsyncPackageLoader;
 class FIoDispatcher;
 class IEDLBootNotificationManager;
 
-using FSourceToLocalizedPackageIdMap = TArray<TPair<FPackageId, FPackageId>>;
+struct FContainerHeaderPackageRedirect
+{
+	FPackageId SourcePackageId;
+	FPackageId TargetPackageId;
+	FMappedName SourcePackageName;
+
+	COREUOBJECT_API friend FArchive& operator<<(FArchive& Ar, FContainerHeaderPackageRedirect& PackageRedirect);
+};
+
+using FSourceToLocalizedPackageIdMap = TArray<FContainerHeaderPackageRedirect>;
 using FCulturePackageMap = TMap<FString, FSourceToLocalizedPackageIdMap>;
 
 struct FContainerHeader
@@ -27,12 +36,37 @@ struct FContainerHeader
 	uint32 PackageCount = 0;
 	TArray<FPackageId> PackageIds;
 	TArray<uint8> StoreEntries; //FPackageStoreEntry[PackageCount]
+	TArray<FNameEntryId> RedirectsNameMap;
 	FCulturePackageMap CulturePackageMap;
-	TArray<TPair<FPackageId, FPackageId>> PackageRedirects;
+	TArray<FContainerHeaderPackageRedirect> PackageRedirects;
 
 	COREUOBJECT_API friend FArchive& operator<<(FArchive& Ar, FContainerHeader& ContainerHeader);
 };
 
+class FPackageImportReference
+{
+public:
+	FPackageImportReference(uint32 InImportedPackageIndex, uint32 InExportHash)
+		: ImportedPackageIndex(InImportedPackageIndex)
+		, ExportHash(InExportHash)
+	{
+	}
+
+	uint32 GetImportedPackageIndex() const
+	{
+		return ImportedPackageIndex;
+	}
+
+	uint32 GetExportHash() const
+	{
+		return ExportHash;
+	}
+
+private:
+	uint32 ImportedPackageIndex;
+	uint32 ExportHash;
+};
+
 class FPackageObjectIndex
 {
 	static constexpr uint64 IndexBits = 62ull;
@@ -70,9 +104,11 @@ public:
 		return FPackageObjectIndex(ScriptImport, GenerateImportHashFromObjectPath(ScriptObjectPath));
 	}
 
-	inline static FPackageObjectIndex FromPackagePath(const FStringView& PackageObjectPath)
+	inline static FPackageObjectIndex FromPackageImportRef(const FPackageImportReference& PackageImportRef)
 	{
-		return FPackageObjectIndex(PackageImport, GenerateImportHashFromObjectPath(PackageObjectPath));
+		uint64 Id = static_cast<uint64>(PackageImportRef.GetImportedPackageIndex()) << 32 | PackageImportRef.GetExportHash();
+		check(!(Id & TypeMask));
+		return FPackageObjectIndex(PackageImport, Id);
 	}
 
 	inline bool IsNull() const
@@ -106,6 +142,13 @@ public:
 		return uint32(TypeAndId);
 	}
 
+	inline FPackageImportReference ToPackageImportRef() const
+	{
+		uint32 ImportedPackageIndex = static_cast<uint32>((TypeAndId & IndexMask) >> 32);
+		uint32 ExportHash = static_cast<uint32>(TypeAndId);
+		return FPackageImportReference(ImportedPackageIndex, ExportHash);
+	}
+
 	inline uint64 Value() const
 	{
 		return TypeAndId & IndexMask;
@@ -133,6 +176,69 @@ public:
 	}
 };
 
+class FPublicExportKey
+{
+public:
+	FPublicExportKey()
+	{
+	}
+
+	bool IsNull() const
+	{
+		return ExportHash == 0;
+	}
+
+	FPackageId GetPackageId() const
+	{
+		return FPackageId::FromValue(uint64(PackageIdHigh) << 32 | PackageIdLow);
+	}
+
+	uint32 GetExportHash() const
+	{
+		return ExportHash;
+	}
+
+	inline bool operator==(const FPublicExportKey& Other) const
+	{
+		return ExportHash == Other.ExportHash &&
+			PackageIdHigh == Other.PackageIdHigh &&
+			PackageIdLow == Other.PackageIdLow;
+	}
+
+	inline friend uint32 GetTypeHash(const FPublicExportKey& In)
+	{
+		return HashCombine(HashCombine(In.PackageIdHigh, In.PackageIdLow), In.ExportHash);
+	}
+
+	static FPublicExportKey MakeKey(FPackageId PackageId, uint32 ExportHash)
+	{
+		check(PackageId.IsValid());
+		check(ExportHash);
+		uint64 PackageIdValue = PackageId.Value();
+		return FPublicExportKey(uint32(PackageIdValue >> 32), uint32(PackageIdValue), ExportHash);
+	}
+
+	static FPublicExportKey FromPackageImport(FPackageObjectIndex ObjectIndex, const TArrayView<const FPackageId>& ImportedPackageIds)
+	{
+		check(ObjectIndex.IsPackageImport());
+		FPackageImportReference PackageImportRef = ObjectIndex.ToPackageImportRef();
+		FPackageId PackageId = ImportedPackageIds[PackageImportRef.GetImportedPackageIndex()];
+		return MakeKey(PackageId, PackageImportRef.GetExportHash());
+	}
+
+private:
+	FPublicExportKey(uint32 InPackageIdHigh, uint32 InPackageIdLow, uint32 InExportHash)
+		: PackageIdHigh(InPackageIdHigh)
+		, PackageIdLow(InPackageIdLow)
+		, ExportHash(InExportHash)
+	{
+	}
+
+	uint32 PackageIdHigh = 0;
+	uint32 PackageIdLow = 0;
+	uint32 ExportHash = 0;
+};
+
 /**
  * Export filter flags.
  */
@@ -150,7 +256,6 @@ struct FPackageSummary
 {
 	uint32 HeaderSize;
 	FMappedName Name;
-	FMappedName SourceName;
 	uint32 PackageFlags;
 	uint32 CookedHeaderSize;
 	int32 ImportMapOffset;
@@ -211,10 +316,10 @@ struct FExportMapEntry
 	FPackageObjectIndex ClassIndex;
 	FPackageObjectIndex SuperIndex;
 	FPackageObjectIndex TemplateIndex;
-	FPackageObjectIndex GlobalImportIndex;
+	uint32 ExportHash;
 	EObjectFlags ObjectFlags = EObjectFlags::RF_NoFlags;
 	EExportFilterFlags FilterFlags = EExportFilterFlags::None;
-	uint8 Pad[3] = {};
+	uint8 Pad[7] = {};
 
 	COREUOBJECT_API friend FArchive& operator<<(FArchive& Ar, FExportMapEntry& ExportMapEntry);
 };
diff --git a/Engine/Source/Runtime/CoreUObject/Public/UObject/Package.h b/Engine/Source/Runtime/CoreUObject/Public/UObject/Package.h
--- a/Engine/Source/Runtime/CoreUObject/Public/UObject/Package.h
+++ b/Engine/Source/Runtime/CoreUObject/Public/UObject/Package.h
@@ -227,9 +227,12 @@ private:
 	/** Package Flags */
 	uint32	PackageFlagsPrivate;
 	
-	/** Globally unique id used to address I/O chunks within the package */
+	/** Globally unique id */
 	FPackageId PackageId;
 
+	/** Globally unique id used to address I/O chunks within the package */
+	FPackageId PackageIdToLoad;
+
 	/** The PackagePath this package was loaded from */
 	FPackagePath LoadedPath;
 public:
@@ -597,6 +600,18 @@ public:
 		PackageId = InPackageId;
 	}
 
+	/** returns the unique package id to load */
+	FORCEINLINE FPackageId GetPackageIdToLoad() const
+	{
+		return PackageIdToLoad;
+	}
+
+	/** sets the unique package id to load */
+	FORCEINLINE void SetPackageIdToLoad(FPackageId InPackageIdToLoad)
+	{
+		PackageIdToLoad = InPackageIdToLoad;
+	}
+
 	/**
 	 * Utility function to find Asset in this package, if any
 	 * @return the asset in the package, if any
