SavePackage PackageStoreWriter with EditorDomain:  fix offsets written into the BulkDatas to be relative to the end of the Linker's archive as expected, rather than relative to some number of bytes before the end of the archive because we were copying Linker->Tell at a point before AppendAdditionalData and the Package Post Tag were written to the archive.

Also add the writing of AppendAdditionalData to a callback on the PackageStoreWriter, so that FVirtualizedUntypedBulkDatas can be separated out into separate archives in the same way FUntypedBaseData are.

#rnx

diff --git a/Engine/Source/Developer/IoStoreUtilities/Private/ZenStoreWriter.cpp b/Engine/Source/Developer/IoStoreUtilities/Private/ZenStoreWriter.cpp
--- a/Engine/Source/Developer/IoStoreUtilities/Private/ZenStoreWriter.cpp
+++ b/Engine/Source/Developer/IoStoreUtilities/Private/ZenStoreWriter.cpp
@@ -478,6 +478,12 @@ bool FZenStoreWriter::WriteAdditionalFile(const FAdditionalFileInfo& Info, const
 	return true;
 }
 
+void FZenStoreWriter::WriteLinkerAdditionalData(const FLinkerAdditionalDataInfo& Info, const FIoBuffer& Data, const TArray<FFileRegion>& FileRegions)
+{
+	// Should not be called because IsLinkerAdditionalDataInSeparateArchive returned false
+	checkNoEntry();
+}
+
 void FZenStoreWriter::BeginCook(const FCookInfo& Info)
 {
 	CookMode = Info.CookMode;
diff --git a/Engine/Source/Developer/IoStoreUtilities/Public/FilePackageStoreWriter.h b/Engine/Source/Developer/IoStoreUtilities/Public/FilePackageStoreWriter.h
--- a/Engine/Source/Developer/IoStoreUtilities/Public/FilePackageStoreWriter.h
+++ b/Engine/Source/Developer/IoStoreUtilities/Public/FilePackageStoreWriter.h
@@ -69,6 +69,12 @@ public:
 	IOSTOREUTILITIES_API virtual void CommitPackage(const FCommitPackageInfo& Info) override;
 	IOSTOREUTILITIES_API virtual void WritePackageData(const FPackageInfo& Info, const FIoBuffer& PackageData, const TArray<FFileRegion>& FileRegions) override;
 	IOSTOREUTILITIES_API virtual void WriteBulkdata(const FBulkDataInfo& Info, const FIoBuffer& BulkData, const TArray<FFileRegion>& FileRegions) override;
+	IOSTOREUTILITIES_API virtual void WriteLinkerAdditionalData(const FLinkerAdditionalDataInfo& Info, const FIoBuffer& Data, const TArray<FFileRegion>& FileRegions) override
+	{
+		// Should not be called because IsLinkerAdditionalDataInSeparateArchive returned false
+		checkNoEntry();
+	}
+
 	IOSTOREUTILITIES_API virtual bool WriteAdditionalFile(const FAdditionalFileInfo& Info, const FIoBuffer& FileData) override { return false; }
 
 	IOSTOREUTILITIES_API virtual void BeginCook(const FCookInfo& Info);
diff --git a/Engine/Source/Developer/IoStoreUtilities/Public/ZenStoreWriter.h b/Engine/Source/Developer/IoStoreUtilities/Public/ZenStoreWriter.h
--- a/Engine/Source/Developer/IoStoreUtilities/Public/ZenStoreWriter.h
+++ b/Engine/Source/Developer/IoStoreUtilities/Public/ZenStoreWriter.h
@@ -39,6 +39,7 @@ public:
 
 	IOSTOREUTILITIES_API virtual void WritePackageData(const FPackageInfo& Info, const FIoBuffer& PackageData, const TArray<FFileRegion>& FileRegions) override;
 	IOSTOREUTILITIES_API virtual bool WriteAdditionalFile(const FAdditionalFileInfo& Info, const FIoBuffer& FileData) override;
+	IOSTOREUTILITIES_API virtual void WriteLinkerAdditionalData(const FLinkerAdditionalDataInfo& Info, const FIoBuffer& Data, const TArray<FFileRegion>& FileRegions) override;
 	IOSTOREUTILITIES_API virtual void WriteBulkdata(const FBulkDataInfo& Info, const FIoBuffer& BulkData, const TArray<FFileRegion>& FileRegions) override;
 	IOSTOREUTILITIES_API virtual void BeginCook(const FCookInfo& Info) override;
 	IOSTOREUTILITIES_API virtual void EndCook() override;
diff --git a/Engine/Source/Editor/UnrealEd/Private/EditorDomain/EditorDomainUtils.cpp b/Engine/Source/Editor/UnrealEd/Private/EditorDomain/EditorDomainUtils.cpp
--- a/Engine/Source/Editor/UnrealEd/Private/EditorDomain/EditorDomainUtils.cpp
+++ b/Engine/Source/Editor/UnrealEd/Private/EditorDomain/EditorDomainUtils.cpp
@@ -232,8 +232,12 @@ class FMemoryPackageStoreWriter : public IPackageStoreWriter
 {
 public:
 	// IPackageStoreWriter
+	virtual bool IsAdditionalFilesNeedLinkerSize() const override { return true; }
+	virtual bool IsLinkerAdditionalDataInSeparateArchive() const override { return true; }
+
 	virtual void WritePackageData(const FPackageInfo& Info, const FIoBuffer& PackageData, const TArray<FFileRegion>& FileRegions) override;
 	virtual void WriteBulkdata(const FBulkDataInfo& Info, const FIoBuffer& BulkData, const TArray<FFileRegion>& FileRegions) override;
+	virtual void WriteLinkerAdditionalData(const FLinkerAdditionalDataInfo& Info, const FIoBuffer& Data, const TArray<FFileRegion>& FileRegions) override;
 	virtual void Flush() override
 	{
 	}
@@ -243,7 +247,7 @@ public:
 	}
 
 	// IPackageStoreWriter cooking and accessor interface: not implemented in this writer
-	bool WriteAdditionalFile(const FAdditionalFileInfo& Info, const FIoBuffer& FileData) override
+	virtual bool WriteAdditionalFile(const FAdditionalFileInfo& Info, const FIoBuffer& FileData) override
 	{
 		checkNoEntry();
 		return false;
@@ -343,6 +347,22 @@ void FMemoryPackageStoreWriter::WriteBulkdata(const FBulkDataInfo& Info, const F
 	}
 }
 
+void FMemoryPackageStoreWriter::WriteLinkerAdditionalData(const FLinkerAdditionalDataInfo& Info, const FIoBuffer& Data, const TArray<FFileRegion>& FileRegions)
+{
+	SetPackageName(Info.PackageName);
+
+	FIoBuffer& DataOwner = BulkDataArchives.Add_GetRef(Data);
+	const uint8* DataStart = Data.Data();
+	uint64 DataLen = Data.DataSize();
+	for (const FFileRegion& FileRegion : FileRegions)
+	{
+		checkf(FileRegion.Type == EFileRegionType::None, TEXT("FMemoryPackageStoreWriter does not currently support FileRegion types other than None."));
+		checkf(FileRegion.Offset + FileRegion.Length <= DataLen, TEXT("FileRegions in WriteLinkerAdditionalData were outside of the range of the Data's size."));
+		BulkDataRegions.Add(FIoBuffer(DataStart + FileRegion.Offset, FileRegion.Length, DataOwner));
+	}
+
+}
+
 bool TrySavePackage(UPackage* Package)
 {
 	using namespace UE::DerivedData;
@@ -407,7 +427,7 @@ bool TrySavePackage(UPackage* Package)
 	{
 		TArray64<uint8> TempBytes;
 		const FIoBuffer& MainBuffer = PackageStoreWriter->GetHeaderAndExports();
-		TempBytes.Append(MainBuffer .Data(), MainBuffer.DataSize());
+		TempBytes.Append(MainBuffer.Data(), MainBuffer.DataSize());
 		for (const FIoBuffer& BulkBuffer : PackageStoreWriter->GetBulkDatas())
 		{
 			TempBytes.Append(BulkBuffer.Data(), BulkBuffer.DataSize());
diff --git a/Engine/Source/Runtime/Core/Public/IO/PackageStoreWriter.h b/Engine/Source/Runtime/Core/Public/IO/PackageStoreWriter.h
--- a/Engine/Source/Runtime/Core/Public/IO/PackageStoreWriter.h
+++ b/Engine/Source/Runtime/Core/Public/IO/PackageStoreWriter.h
@@ -12,6 +12,22 @@ class IPackageStoreWriter
 public:
 	virtual ~IPackageStoreWriter() = default;
 
+
+	// Properties of the PackageStoreWriter
+	/** Whether this writer needs BulkDatas to be written after the Linker's archive has finalized its size.
+
+		Some PackageStoreWriters need that behavior because they put the BulkDatas in a segment following
+		the exports in a Composite archive.
+	 */
+	virtual bool IsAdditionalFilesNeedLinkerSize() const { return false; }
+	/** Whether data stored in Linker.AdditionalDataToAppend should be serialized to a separate archive.
+	
+		If false, the data will be serialized to the end of the LinkerSave archive instead.
+	*/
+	virtual bool IsLinkerAdditionalDataInSeparateArchive() const { return false; }
+
+
+	// Events the PackageStoreWriter receives
 	struct FBeginPackageInfo
 	{
 		FName	PackageName;
@@ -75,8 +91,20 @@ public:
 		FIoChunkId ChunkId;
 	};
 
+	/** Write separate files written by UObjects during cooking via UObject::CookAdditionalFiles. */
 	virtual bool WriteAdditionalFile(const FAdditionalFileInfo& Info, const FIoBuffer& FileData) = 0;
 
+	struct FLinkerAdditionalDataInfo
+	{
+		FName PackageName;
+	};
+	/** Write separate data written by UObjects via FLinkerSave::AdditionalDataToAppend.
+
+		This function will not be called unless IsLinkerAdditionalDataInSeparateArchive returned true.
+		If that function is false, the data was inlined into the PackageData instead.
+	*/
+	virtual void WriteLinkerAdditionalData(const FLinkerAdditionalDataInfo& Info, const FIoBuffer& Data, const TArray<FFileRegion>& FileRegions) = 0;
+
 	struct FCookInfo
 	{
 		enum ECookMode
diff --git a/Engine/Source/Runtime/CoreUObject/Private/UObject/SavePackage.cpp b/Engine/Source/Runtime/CoreUObject/Private/UObject/SavePackage.cpp
--- a/Engine/Source/Runtime/CoreUObject/Private/UObject/SavePackage.cpp
+++ b/Engine/Source/Runtime/CoreUObject/Private/UObject/SavePackage.cpp
@@ -4162,57 +4162,72 @@ FSavePackageResultStruct UPackage::Save(UPackage* InOuter, UObject* Base, EObjec
 
 				SlowTask.EnterProgressFrame(1, NSLOCTEXT("Core", "SerializingBulkData", "Serializing bulk data"));
 
-				ESavePackageResult SaveResult = 
-					SavePackageUtilities::SaveBulkData(Linker.Get(), InOuter, Filename, TargetPlatform, SavePackageContext, SaveFlags, bTextFormat, bDiffing,
-							 bComputeHash, AsyncWriteAndHashSequence, TotalPackageSizeUncompressed );
-				if (SaveResult != ESavePackageResult::Success)
+				IPackageStoreWriter* PackageStoreWriter = SavePackageContext ? SavePackageContext->PackageStoreWriter : nullptr;
+				FSavePackageOutputFileArray AdditionalOutputFiles;
+				auto WriteAdditionalFiles = [&](int64 LinkerSize) -> ESavePackageResult
 				{
-					return SaveResult;
-				}
+					int64 DataStartOffset = LinkerSize >= 0 ? LinkerSize : Linker->Tell();
+					ESavePackageResult SaveResult =
+						SavePackageUtilities::SaveBulkData(Linker.Get(), DataStartOffset, InOuter, Filename, TargetPlatform, SavePackageContext, SaveFlags, bTextFormat, bDiffing,
+							bComputeHash, AsyncWriteAndHashSequence, TotalPackageSizeUncompressed);
+					if (SaveResult != ESavePackageResult::Success)
+					{
+						return SaveResult;
+					}
 
-				ESavePackageResult AppendAdditionalDataResult = SavePackageUtilities::AppendAdditionalData(*Linker.Get());
-				if (AppendAdditionalDataResult != ESavePackageResult::Success)
-				{
-					return AppendAdditionalDataResult;
-				}
+					ESavePackageResult AppendAdditionalDataResult = SavePackageUtilities::AppendAdditionalData(*Linker.Get(), DataStartOffset, PackageStoreWriter);
+					if (AppendAdditionalDataResult != ESavePackageResult::Success)
+					{
+						return AppendAdditionalDataResult;
+					}
 
-				FSavePackageOutputFileArray AdditionalOutputFiles;
-				ESavePackageResult CreateSidecarResult = SavePackageUtilities::CreatePayloadSidecarFile(*Linker, TargetPackagePath, bSaveAsync, !bDiffing, AdditionalOutputFiles);
-				if (CreateSidecarResult != ESavePackageResult::Success)
-				{
-					return CreateSidecarResult;
-				}
+					ESavePackageResult CreateSidecarResult = SavePackageUtilities::CreatePayloadSidecarFile(*Linker, TargetPackagePath, bSaveAsync, !bDiffing, AdditionalOutputFiles);
+					if (CreateSidecarResult != ESavePackageResult::Success)
+					{
+						return CreateSidecarResult;
+					}
 
 #if WITH_EDITOR
-				if (bIsCooking && AdditionalFilesFromExports.Num() > 0)
-				{
-					const bool bWriteFileToDisk = !bDiffing;
-					for (FLargeMemoryWriter& Writer : AdditionalFilesFromExports)
+					if (bIsCooking && AdditionalFilesFromExports.Num() > 0)
 					{
-						const int64 Size = Writer.TotalSize();
-						TotalPackageSizeUncompressed += Size;
-
-						if (bComputeHash || bWriteFileToDisk)
+						const bool bWriteFileToDisk = !bDiffing;
+						for (FLargeMemoryWriter& Writer : AdditionalFilesFromExports)
 						{
-							FLargeMemoryPtr DataPtr(Writer.ReleaseOwnership());
+							const int64 Size = Writer.TotalSize();
+							TotalPackageSizeUncompressed += Size;
 
-							EAsyncWriteOptions WriteOptions(EAsyncWriteOptions::None);
-							if (bComputeHash)
+							if (bComputeHash || bWriteFileToDisk)
 							{
-								WriteOptions |= EAsyncWriteOptions::ComputeHash;
-							}
-							if (bWriteFileToDisk)
-							{
-								WriteOptions |= EAsyncWriteOptions::WriteFileToDisk;
+								FLargeMemoryPtr DataPtr(Writer.ReleaseOwnership());
+
+								EAsyncWriteOptions WriteOptions(EAsyncWriteOptions::None);
+								if (bComputeHash)
+								{
+									WriteOptions |= EAsyncWriteOptions::ComputeHash;
+								}
+								if (bWriteFileToDisk)
+								{
+									WriteOptions |= EAsyncWriteOptions::WriteFileToDisk;
+								}
+								SavePackageUtilities::AsyncWriteFile(AsyncWriteAndHashSequence, MoveTemp(DataPtr), Size, *Writer.GetArchiveName(), WriteOptions, {});
 							}
-							SavePackageUtilities::AsyncWriteFile(AsyncWriteAndHashSequence, MoveTemp(DataPtr), Size, *Writer.GetArchiveName(), WriteOptions, {});
 						}
+						AdditionalFilesFromExports.Empty();
 					}
-					AdditionalFilesFromExports.Empty();
-				}
 #endif
-
+					return ESavePackageResult::Success;
+				};
 			
+				bool bAdditionalFilesNeedLinkerSize = PackageStoreWriter ? PackageStoreWriter->IsAdditionalFilesNeedLinkerSize() : false;
+				if (!bAdditionalFilesNeedLinkerSize)
+				{
+					ESavePackageResult AdditionalFilesResult = WriteAdditionalFiles(-1);
+					if (AdditionalFilesResult != ESavePackageResult::Success)
+					{
+						return AdditionalFilesResult;
+					}
+				}
+
 				// write the package post tag
 				if (!bTextFormat)
 				{
@@ -4224,6 +4239,16 @@ FSavePackageResultStruct UPackage::Save(UPackage* InOuter, UObject* Base, EObjec
 				// size correctly while the file is still being written to.
 				PackageSize = Linker->Tell();
 
+				if (bAdditionalFilesNeedLinkerSize)
+				{
+					ESavePackageResult AdditionalFilesResult = WriteAdditionalFiles(PackageSize);
+					if (AdditionalFilesResult != ESavePackageResult::Success)
+					{
+						return AdditionalFilesResult;
+					}
+					checkf(Linker->Tell() == PackageSize, TEXT("The writing of additional files is not allowed to append to the LinkerSave when IsAdditionalFilesNeedLinkerSize is true."));
+				}
+
 				// Save the import map.
 				{
 #if WITH_EDITOR
@@ -4442,7 +4467,6 @@ FSavePackageResultStruct UPackage::Save(UPackage* InOuter, UObject* Base, EObjec
 					}
 
 					// Compress the temporarily file to destination.
-					IPackageStoreWriter* PackageStoreWriter = SavePackageContext ? SavePackageContext->PackageStoreWriter : nullptr;
 					if (bSaveAsync)
 					{	
 						FString NewPathToSave = NewPath;
diff --git a/Engine/Source/Runtime/CoreUObject/Private/UObject/SavePackage/SaveContext.cpp b/Engine/Source/Runtime/CoreUObject/Private/UObject/SavePackage/SaveContext.cpp
--- a/Engine/Source/Runtime/CoreUObject/Private/UObject/SavePackage/SaveContext.cpp
+++ b/Engine/Source/Runtime/CoreUObject/Private/UObject/SavePackage/SaveContext.cpp
@@ -2,6 +2,8 @@
 
 #include "SaveContext.h"
 
+#include "IO/PackageStoreWriter.h"
+
 void FSaveContext::MarkUnsaveable(UObject* InObject)
 {
 	if (IsUnsaveable(InObject))
@@ -54,3 +56,9 @@ bool FSaveContext::IsUnsaveable(UObject* InObject) const
 	}
 	return false;	
 }
+
+bool FSaveContext::IsAdditionalFilesNeedLinkerSize() const
+{
+	IPackageStoreWriter* PackageStoreWriter = GetPackageStoreWriter();
+	return PackageStoreWriter ? PackageStoreWriter->IsAdditionalFilesNeedLinkerSize() : false;
+}
diff --git a/Engine/Source/Runtime/CoreUObject/Private/UObject/SavePackage/SaveContext.h b/Engine/Source/Runtime/CoreUObject/Private/UObject/SavePackage/SaveContext.h
--- a/Engine/Source/Runtime/CoreUObject/Private/UObject/SavePackage/SaveContext.h
+++ b/Engine/Source/Runtime/CoreUObject/Private/UObject/SavePackage/SaveContext.h
@@ -548,6 +548,13 @@ public:
 		return ObjectSaveContext;
 	}
 
+	IPackageStoreWriter* GetPackageStoreWriter() const
+	{
+		return SaveArgs.SavePackageContext ? SaveArgs.SavePackageContext->PackageStoreWriter : nullptr;
+	}
+
+	bool IsAdditionalFilesNeedLinkerSize() const;
+
 public:
 	ESavePackageResult Result;
 
diff --git a/Engine/Source/Runtime/CoreUObject/Private/UObject/SavePackage/SavePackageUtilities.cpp b/Engine/Source/Runtime/CoreUObject/Private/UObject/SavePackage/SavePackageUtilities.cpp
--- a/Engine/Source/Runtime/CoreUObject/Private/UObject/SavePackage/SavePackageUtilities.cpp
+++ b/Engine/Source/Runtime/CoreUObject/Private/UObject/SavePackage/SavePackageUtilities.cpp
@@ -2119,13 +2119,51 @@ void SaveThumbnails(UPackage* InOuter, FLinkerSave* Linker, FStructuredArchive::
 #endif
 }
 
-ESavePackageResult AppendAdditionalData(FLinkerSave& Linker)
+class FLargeMemoryWriterWithRegions : public FLargeMemoryWriter
 {
-	FArchive& Ar = Linker;
+public:
+	FLargeMemoryWriterWithRegions()
+		: FLargeMemoryWriter(0, /* IsPersistent */ true)
+	{}
+
+	TArray<FFileRegion> FileRegions;
+};
+
+ESavePackageResult AppendAdditionalData(FLinkerSave& Linker, int64& InOutDataStartOffset, IPackageStoreWriter* PackageStoreWriter)
+{
+	if (Linker.AdditionalDataToAppend.Num() == 0)
+	{
+		return ESavePackageResult::Success;
+	}
 
-	for (FLinkerSave::AdditionalDataCallback& Callback : Linker.AdditionalDataToAppend)
+	if (PackageStoreWriter && PackageStoreWriter->IsLinkerAdditionalDataInSeparateArchive())
 	{
-		Callback(Ar);
+		FLargeMemoryWriterWithRegions DataArchive;
+		for (FLinkerSave::AdditionalDataCallback& Callback : Linker.AdditionalDataToAppend)
+		{
+			int64 RegionStart = DataArchive.Tell();
+			Callback(Linker, DataArchive, InOutDataStartOffset + RegionStart);
+			int64 RegionEnd = DataArchive.Tell();
+			if (RegionEnd != RegionStart)
+			{
+				DataArchive.FileRegions.Add(FFileRegion(RegionStart, RegionEnd - RegionStart, EFileRegionType::None));
+			}
+		}
+		IPackageStoreWriter::FLinkerAdditionalDataInfo DataInfo{ Linker.LinkerRoot->GetFName() };
+		int64 DataSize = DataArchive.TotalSize();
+		FIoBuffer DataBuffer = FIoBuffer(FIoBuffer::AssumeOwnership, DataArchive.ReleaseOwnership(), DataSize);
+		PackageStoreWriter->WriteLinkerAdditionalData(DataInfo, DataBuffer, DataArchive.FileRegions);
+		InOutDataStartOffset += DataSize;
+	}
+	else
+	{
+		int64 LinkerStart = Linker.Tell();
+		FArchive& Ar = Linker;
+		for (FLinkerSave::AdditionalDataCallback& Callback : Linker.AdditionalDataToAppend)
+		{
+			Callback(Linker, Linker, Linker.Tell());
+		}
+		InOutDataStartOffset += Linker.Tell() - LinkerStart;
 	}
 
 	Linker.AdditionalDataToAppend.Empty();
@@ -2217,14 +2255,14 @@ ESavePackageResult CreatePayloadSidecarFile(FLinkerSave& Linker, const FPackageP
 	return ESavePackageResult::Success;
 }
 
-ESavePackageResult SaveBulkData(FLinkerSave* Linker, const UPackage* InOuter, const TCHAR* Filename, const ITargetPlatform* TargetPlatform,
+ESavePackageResult SaveBulkData(FLinkerSave* Linker, int64& InOutStartOffset, const UPackage* InOuter, const TCHAR* Filename, const ITargetPlatform* TargetPlatform,
 				  FSavePackageContext* SavePackageContext, uint32 SaveFlags, const bool bTextFormat, const bool bDiffing,
 				  const bool bComputeHash, TAsyncWorkSequence<FMD5>& AsyncWriteAndHashSequence, int64& TotalPackageSizeUncompressed)
 {
 	// Now we write all the bulkdata that is supposed to be at the end of the package
 	// and fix up the offset
 	IPackageStoreWriter* PackageStoreWriter = SavePackageContext != nullptr ? SavePackageContext->PackageStoreWriter : nullptr;
-	Linker->Summary.BulkDataStartOffset = Linker->Tell();
+	Linker->Summary.BulkDataStartOffset = InOutStartOffset;
 
 	if (Linker->BulkDataToAppend.Num() == 0)
 	{
@@ -2236,16 +2274,6 @@ ESavePackageResult SaveBulkData(FLinkerSave* Linker, const UPackage* InOuter, co
 
 	FScopedSlowTask BulkDataFeedback(Linker->BulkDataToAppend.Num());
 
-	class FLargeMemoryWriterWithRegions : public FLargeMemoryWriter
-	{
-	public:
-		FLargeMemoryWriterWithRegions()
-			: FLargeMemoryWriter(0, /* IsPersistent */ true)
-		{}
-
-		TArray<FFileRegion> FileRegions;
-	};
-
 	TUniquePtr<FLargeMemoryWriterWithRegions> BulkArchive;
 	TUniquePtr<FLargeMemoryWriterWithRegions> OptionalBulkArchive;
 	TUniquePtr<FLargeMemoryWriterWithRegions> MappedBulkArchive;
@@ -2267,14 +2295,19 @@ ESavePackageResult SaveBulkData(FLinkerSave* Linker, const UPackage* InOuter, co
 	} CookerUseSeparateSegments;
 	const bool bSeparateSegmentsEnabled = CookerUseSeparateSegments.bEnable && Linker->IsCooking();
 
+	int64 LinkerStart = 0;
 	if (PackageStoreWriter || bSeparateSegmentsEnabled)
 	{
 		BulkArchive.Reset(new FLargeMemoryWriterWithRegions);
+		if (bSeparateSegmentsEnabled)
+		{
+			OptionalBulkArchive.Reset(new FLargeMemoryWriterWithRegions);
+			MappedBulkArchive.Reset(new FLargeMemoryWriterWithRegions);
+		}
 	}
-	if (bSeparateSegmentsEnabled)
+	else
 	{
-		OptionalBulkArchive.Reset(new FLargeMemoryWriterWithRegions);
-		MappedBulkArchive.Reset(new FLargeMemoryWriterWithRegions);
+		LinkerStart = Linker->Tell();
 	}
 	bool bRequestSaveByReference = SaveFlags & SAVE_BulkDataByReference;
 
@@ -2475,7 +2508,7 @@ ESavePackageResult SaveBulkData(FLinkerSave* Linker, const UPackage* InOuter, co
 				// We record offsets in the bulkdata segment relative to the beginning of the CompositeArchive.
 				// We have to do this to distinguish them from inline bulkdatas in the exports segment
 				check((BulkDataFlags& BULKDATA_NoOffsetFixUp) != 0); // offset fixups are not supported with this SaveLocation
-				BulkDataOffsetInFile = BulkStartOffset + Linker->Summary.BulkDataStartOffset;
+				BulkDataOffsetInFile = BulkStartOffset + InOutStartOffset;
 			}
 			else if ((BulkDataFlags & BULKDATA_NoOffsetFixUp) == 0)
 			{
@@ -2587,6 +2620,20 @@ ESavePackageResult SaveBulkData(FLinkerSave* Linker, const UPackage* InOuter, co
 		}
 	}
 
+	if (!bSeparateSegmentsEnabled)
+	{
+		if (PackageStoreWriter)
+		{
+			check(BulkArchive.IsValid() && LinkerStart == 0);
+			InOutStartOffset += BulkArchive->TotalSize();
+		}
+		else
+		{
+			check(!BulkArchive.IsValid() && LinkerStart > 0);
+			InOutStartOffset += Linker->Tell() - LinkerStart;
+		}
+	}
+
 	Linker->BulkDataToAppend.Empty();
 	return ESavePackageResult::Success;
 }
diff --git a/Engine/Source/Runtime/CoreUObject/Private/UObject/SavePackage/SavePackageUtilities.h b/Engine/Source/Runtime/CoreUObject/Private/UObject/SavePackage/SavePackageUtilities.h
--- a/Engine/Source/Runtime/CoreUObject/Private/UObject/SavePackage/SavePackageUtilities.h
+++ b/Engine/Source/Runtime/CoreUObject/Private/UObject/SavePackage/SavePackageUtilities.h
@@ -21,6 +21,7 @@ class FMD5;
 class FPackagePath;
 class FSaveContext;
 class FSavePackageContext;
+class IPackageStoreWriter;
 template<typename StateType> class TAsyncWorkSequence;
 
 enum class ESavePackageResult;
@@ -386,12 +387,42 @@ namespace SavePackageUtilities
 	void DecrementOutstandingAsyncWrites();
 
 	void SaveThumbnails(UPackage* InOuter, FLinkerSave* Linker, FStructuredArchive::FSlot Slot);
-	ESavePackageResult SaveBulkData(FLinkerSave* Linker, const UPackage* InOuter, const TCHAR* Filename, const ITargetPlatform* TargetPlatform,
+
+	/**
+	 * Save all of the BulkDatas that have registered as not being inline.
+	 * They may be saved to the end of the file, or to an archive for a separate file.
+	 * 
+	 * @param Linker The linker containing the exports. Provides metadata for BulkData, and BulkData may write to it as their target archive.
+	 * @param InOutStartOffset In value is the offset in the Linker's archive where the BulkDatas will be put. If SavePackageContext settings direct
+	 *        the bulkdatas to write in a separate archive that will be combined after the linker, the value is the Linker archive's totalsize.
+	 *        Output value is incremented by the number of bytes written the Linker or the separate archive at the end of the linker.
+	 *        Bytes written to separate files do not contribute to this offset.
+	 * @param InOuter The package being saved
+	 * @param Filename The filename the package is being saved to. Used to decide the filename of separate archives, if any are used.
+	 * @param SavePackageContext Provides The PackageStoreWriter and other parameters of how the bulkdata should be saved.
+	 * @param SaveFlags The flags passed into SavePackage; affects how the bulkdata should be saved.
+	 * @param bTextFormat True if package is saving to text. Bulkdata has special handling for text output.
+	 * @param bDiffing True if the package is only diffing, so bulk data should be recorded but not saved to disk.
+	 * @param bComputeHash True if package hash needs to be computed; bulkdatas contribute to the hash.
+	 * @param AsyncWriteAndHashSequence Output: Collects the writes of the bulkdata to disk for async writing.
+	 * @param TotalPackageSizeUncompressed Output: Bulkdata sizes are added to it.
+	 */
+	ESavePackageResult SaveBulkData(FLinkerSave* Linker, int64& InOutStartOffset, const UPackage* InOuter, const TCHAR* Filename, const ITargetPlatform* TargetPlatform,
 		FSavePackageContext* SavePackageContext, uint32 SaveFlags, const bool bTextFormat, const bool bDiffing,
 		const bool bComputeHash, TAsyncWorkSequence<FMD5>& AsyncWriteAndHashSequence, int64& TotalPackageSizeUncompressed);
 	
-	/** Used to append additional data to the end of the package file by invoking callbacks stored in the linker */
-	ESavePackageResult AppendAdditionalData(FLinkerSave& Linker);
+	/**
+	 * Used to append additional data to the end of the package file by invoking callbacks stored in the linker.
+	 * They may be saved to the end of the file, or to a separate archive passed into the PackageStoreWriter.
+	 
+	 * @param Linker The linker containing the exports. Provides the list of AdditionalData, and the data may write to it as their target archive.
+	 * @param InOutStartOffset In value is the offset in the Linker's archive where the datas will be put. If SavePackageContext settings direct
+	 *        the datas to write in a separate archive that will be combined after the linker, the value is the offset after the Linker archive's 
+	 *        totalsize and after any previous post-Linker archive data such as BulkDatas.
+	 *        Output value is incremented by the number of bytes written the Linker or the separate archive at the end of the linker.
+	 * @param PackageStoreWriter If non-null and configured to require it, data is passed to this writer rather than appended to the Linker archive.
+	  */
+	ESavePackageResult AppendAdditionalData(FLinkerSave& Linker, int64& InOutDataStartOffset, IPackageStoreWriter* PackageStoreWriter);
 	
 	/** Used to create the sidecar file (.upayload) from payloads that have been added to the linker */
 	ESavePackageResult CreatePayloadSidecarFile(FLinkerSave& Linker, const FPackagePath& PackagePath, const bool bSaveAsync, const bool bWriteToDisk, FSavePackageOutputFileArray& AdditionalPackageFiles);
diff --git a/Engine/Source/Runtime/CoreUObject/Private/UObject/SavePackage2.cpp b/Engine/Source/Runtime/CoreUObject/Private/UObject/SavePackage2.cpp
--- a/Engine/Source/Runtime/CoreUObject/Private/UObject/SavePackage2.cpp
+++ b/Engine/Source/Runtime/CoreUObject/Private/UObject/SavePackage2.cpp
@@ -54,6 +54,8 @@ bool IsEditorOnlyObject(const UObject* InObject, bool bCheckRecursive, bool bChe
 namespace
 {
 
+ESavePackageResult WriteAdditionalFiles(FSaveContext& SaveContext, FScopedSlowTask& SlowTask, int64 LinkerSize);
+
 ESavePackageResult ReturnSuccessOrCancel()
 {
 	return !GWarn->ReceivedUserCancel() ? ESavePackageResult::Success : ESavePackageResult::Canceled;
@@ -2107,32 +2109,13 @@ ESavePackageResult InnerSave(FSaveContext& SaveContext)
 		}
 	}
 
-	// Save Bulk Data
-	SlowTask.EnterProgressFrame();
-	SaveContext.Result = SavePackageUtilities::SaveBulkData(SaveContext.Linker.Get(), SaveContext.GetPackage(), SaveContext.GetFilename(), SaveContext.GetTargetPlatform(),
-		SaveContext.GetSavePackageContext(), SaveContext.GetSaveArgs().SaveFlags, SaveContext.IsTextFormat(), SaveContext.IsDiffing(),
-		SaveContext.IsComputeHash(), SaveContext.AsyncWriteAndHashSequence, SaveContext.TotalPackageSizeUncompressed);
-
-	// Add any pending data blobs to the end of the file by invoking the callbacks
-	SaveContext.Result = SavePackageUtilities::AppendAdditionalData(*SaveContext.Linker.Get());
-	if (SaveContext.Result != ESavePackageResult::Success)
-	{
-		return SaveContext.Result;
-	}
-
-	// Create the payload side car file (if needed)
-	SaveContext.Result = SavePackageUtilities::CreatePayloadSidecarFile(*SaveContext.Linker.Get(), SaveContext.GetTargetPackagePath(), SaveContext.IsSaveAsync(), !SaveContext.IsDiffing(), SaveContext.AdditionalPackageFiles);
-	if (SaveContext.Result != ESavePackageResult::Success)
-	{
-		return SaveContext.Result;
-	}
-
-	// Write Additional files from export
-	SlowTask.EnterProgressFrame();
-	SaveContext.Result = WriteAdditionalExportFiles(SaveContext);
-	if (SaveContext.Result != ESavePackageResult::Success)
+	if (!SaveContext.IsAdditionalFilesNeedLinkerSize())
 	{
-		return SaveContext.Result;
+		SaveContext.Result = WriteAdditionalFiles(SaveContext, SlowTask, -1);
+		if (SaveContext.Result != ESavePackageResult::Success)
+		{
+			return SaveContext.Result;
+		}
 	}
 	
 	// Write Package Post Tag
@@ -2146,6 +2129,16 @@ ESavePackageResult InnerSave(FSaveContext& SaveContext)
 	int32 PackageSize = SaveContext.Linker->Tell();
 	SaveContext.TotalPackageSizeUncompressed += PackageSize;
 
+	if (SaveContext.IsAdditionalFilesNeedLinkerSize())
+	{
+		SaveContext.Result = WriteAdditionalFiles(SaveContext, SlowTask, PackageSize);
+		if (SaveContext.Result != ESavePackageResult::Success)
+		{
+			return SaveContext.Result;
+		}
+		checkf(SaveContext.Linker->Tell() == PackageSize, TEXT("The writing of additional files is not allowed to append to the LinkerSave when IsAdditionalFilesNeedLinkerSize is true."));
+	}
+
 	for (const FSavePackageOutputFile& File : SaveContext.AdditionalPackageFiles)
 	{
 		SaveContext.TotalPackageSizeUncompressed += File.DataSize;
@@ -2204,6 +2197,46 @@ ESavePackageResult InnerSave(FSaveContext& SaveContext)
 	return SaveContext.Result;
 }
 
+/** WriteAdditionalPayloads writes BulkData, PayloadSidecarFiles, and other data that is separate from exports.
+ * Depending on settings, these additional files may be appended to the LinkerSave after the exports, or they
+ * may be written into separate sidecar files.
+ * 
+ * @param SaveContext The context for the overall save, including data about the additional payloads
+ * @param LinkerSize If the Linker has finished writing, this is the size of the Linker's archive: Linker->Tell(). Otherwise it is -1.
+ */
+ESavePackageResult WriteAdditionalFiles(FSaveContext& SaveContext, FScopedSlowTask& SlowTask, int64 LinkerSize)
+{
+	// Save Bulk Data
+	SlowTask.EnterProgressFrame();
+	int64 DataStartOffset = LinkerSize >= 0 ? LinkerSize : SaveContext.Linker.Get()->Tell();
+	SaveContext.Result = SavePackageUtilities::SaveBulkData(SaveContext.Linker.Get(), DataStartOffset, SaveContext.GetPackage(), SaveContext.GetFilename(), SaveContext.GetTargetPlatform(),
+		SaveContext.GetSavePackageContext(), SaveContext.GetSaveArgs().SaveFlags, SaveContext.IsTextFormat(), SaveContext.IsDiffing(),
+		SaveContext.IsComputeHash(), SaveContext.AsyncWriteAndHashSequence, SaveContext.TotalPackageSizeUncompressed);
+
+	// Add any pending data blobs to the end of the file by invoking the callbacks
+	ESavePackageResult Result = SavePackageUtilities::AppendAdditionalData(*SaveContext.Linker.Get(), DataStartOffset, SaveContext.GetPackageStoreWriter());
+	if (Result != ESavePackageResult::Success)
+	{
+		return Result;
+	}
+
+	// Create the payload side car file (if needed)
+	Result = SavePackageUtilities::CreatePayloadSidecarFile(*SaveContext.Linker.Get(), SaveContext.GetTargetPackagePath(), SaveContext.IsSaveAsync(), !SaveContext.IsDiffing(), SaveContext.AdditionalPackageFiles);
+	if (Result != ESavePackageResult::Success)
+	{
+		return Result;
+	}
+
+	// Write Additional files from export
+	SlowTask.EnterProgressFrame();
+	Result = WriteAdditionalExportFiles(SaveContext);
+	if (Result != ESavePackageResult::Success)
+	{
+		return Result;
+	}
+	return ESavePackageResult::Success;
+}
+
 FText GetSlowTaskStatusMessage(const FSaveContext& SaveContext)
 {
 	FString CleanFilename = FPaths::GetCleanFilename(SaveContext.GetFilename());
diff --git a/Engine/Source/Runtime/CoreUObject/Private/Virtualization/VirtualizedBulkData.cpp b/Engine/Source/Runtime/CoreUObject/Private/Virtualization/VirtualizedBulkData.cpp
--- a/Engine/Source/Runtime/CoreUObject/Private/Virtualization/VirtualizedBulkData.cpp
+++ b/Engine/Source/Runtime/CoreUObject/Private/Virtualization/VirtualizedBulkData.cpp
@@ -432,27 +432,24 @@ void FVirtualizedUntypedBulkData::Serialize(FArchive& Ar, UObject* Owner, bool b
 
 				// The lambda is mutable so that PayloadToSerialize is not const (due to FArchive api not
 				// accepting const values)
-				auto SerializePayload = [this, OffsetPos, PayloadToSerialize, UpdatedFlags, Owner](FArchive& Ar) mutable
+				auto SerializePayload = [this, OffsetPos, PayloadToSerialize, UpdatedFlags, Owner](FLinkerSave* LinkerSave, FArchive& ExportsArchive, FArchive& DataArchive, int64 DataStartOffset) mutable
 				{
-					checkf(Ar.IsCooking() == false, TEXT("FVirtualizedUntypedBulkData::Serialize should not be called during a cook"));
+					checkf(ExportsArchive.IsCooking() == false, TEXT("FVirtualizedUntypedBulkData::Serialize should not be called during a cook"));
 
-					int64 PayloadOffset = Ar.Tell();
-
-					SerializeData(Ar, PayloadToSerialize, UpdatedFlags);
+					SerializeData(DataArchive, PayloadToSerialize, UpdatedFlags);
 
 					// Record the current archive offset (probably EOF but we cannot be sure)
-					const int64 ReturnPos = Ar.Tell();
+					const int64 ReturnPos = ExportsArchive.Tell();
 
 					// Update the offset/size entries that we set up during ::Serialize
-					Ar.Seek(OffsetPos);
-					Ar << PayloadOffset;
+					ExportsArchive.Seek(OffsetPos);
+					ExportsArchive << DataStartOffset;
 					// Restore the archive's offset
-					Ar.Seek(ReturnPos);
+					ExportsArchive.Seek(ReturnPos);
 
 					// If we are saving the package to disk (we have access to FLinkerSave and its filepath is valid) 
 					// then we should register a callback to be received once the package has actually been saved to 
 					// disk so that we can update the object's members to be redirected to the saved file.
-					FLinkerSave* LinkerSave = Cast<FLinkerSave>(Ar.GetLinker());
 					if (LinkerSave != nullptr && !LinkerSave->GetFilename().IsEmpty())
 					{	
 						// At some point saving to the sidecar file will be mutually exclusive with saving to the asset file, at that point
@@ -464,7 +461,7 @@ void FVirtualizedUntypedBulkData::Serialize(FArchive& Ar, UObject* Owner, bool b
 							SidecarData.Payload = PayloadToSerialize;
 						}
 
-						auto OnSavePackage = [this, PayloadOffset, UpdatedFlags, Owner](const FPackagePath& InPackagePath, FObjectPostSaveContext ObjectSaveContext)
+						auto OnSavePackage = [this, DataStartOffset, UpdatedFlags, Owner](const FPackagePath& InPackagePath, FObjectPostSaveContext ObjectSaveContext)
 						{
 							if (!ObjectSaveContext.IsUpdatingLoadedPath())
 							{
@@ -473,7 +470,7 @@ void FVirtualizedUntypedBulkData::Serialize(FArchive& Ar, UObject* Owner, bool b
 
 							this->PackagePath = InPackagePath;
 							check(!this->PackagePath.IsEmpty()); // LinkerSave guarantees a valid PackagePath if we're updating loaded path
-							this->OffsetInFile = PayloadOffset;
+							this->OffsetInFile = DataStartOffset;
 							this->Flags = UpdatedFlags;
 
 							if (CanUnloadData())
@@ -494,11 +491,15 @@ void FVirtualizedUntypedBulkData::Serialize(FArchive& Ar, UObject* Owner, bool b
 				// If we do not have a linker then we might as well just serialize right away.
 				if(LinkerSave != nullptr)
 				{
-					 LinkerSave->AdditionalDataToAppend.Add(MoveTemp(SerializePayload));
+					auto AdditionalDataCallback = [SerializePayload = MoveTemp(SerializePayload)](FLinkerSave& ExportsArchive, FArchive& DataArchive, int64 DataStartOffset) mutable
+					{
+						SerializePayload(&ExportsArchive, ExportsArchive, DataArchive, DataStartOffset);
+					};
+					LinkerSave->AdditionalDataToAppend.Add(MoveTemp(AdditionalDataCallback));
 				}
 				else
 				{
-					SerializePayload(Ar);
+					SerializePayload(nullptr, Ar, Ar, Ar.Tell());
 				}
 			}
 
diff --git a/Engine/Source/Runtime/CoreUObject/Public/UObject/LinkerSave.h b/Engine/Source/Runtime/CoreUObject/Public/UObject/LinkerSave.h
--- a/Engine/Source/Runtime/CoreUObject/Public/UObject/LinkerSave.h
+++ b/Engine/Source/Runtime/CoreUObject/Public/UObject/LinkerSave.h
@@ -73,7 +73,17 @@ public:
 
 	// TODO: Look into removing BulkDataToAppend and use AdditionalDataToAppend instead.
 
-	using AdditionalDataCallback = TUniqueFunction<void(FArchive& Ar)>;
+	/**
+	 * Callback for arbitrary serializers to append data to the end of the ExportsArchive.
+	 * Some PackageStoreWriters used by SavePackage will write this data to a separate archive.
+	 * 
+	 * @param ExportsArchive The archive containing the UObjects and structs, this is always this LinkerSave.
+	 * @param DataArchive The archive to which the data should be written. Might be this LinkerSave, or might be a separate archive.
+	 * @param DataStartOffset The offset to the beginning of the range in DataArchive that should be stored in the UObject or struct's
+	 *        export data. Reading at DataStartOffset from the FArchive passed into Serialize during a load will return the data that the
+	 *		  callback wrote to DataArchive.
+	 */
+	using AdditionalDataCallback = TUniqueFunction<void(FLinkerSave& ExportsArchive, FArchive& DataArchive, int64 DataStartOffset)>;
 	/** 
 	 * Array of callbacks that will be invoked when it is possible to serialize out data 
 	 * to the end of the output file.
