VirtualizedBulkData fix SetCompressionFormat(None) then SetCompressionFormat(other) resulting in no compression

diff --git a/Engine/Source/Runtime/CoreUObject/Private/Virtualization/VirtualizedBulkData.cpp b/Engine/Source/Runtime/CoreUObject/Private/Virtualization/VirtualizedBulkData.cpp
--- a/Engine/Source/Runtime/CoreUObject/Private/Virtualization/VirtualizedBulkData.cpp
+++ b/Engine/Source/Runtime/CoreUObject/Private/Virtualization/VirtualizedBulkData.cpp
@@ -1041,6 +1041,10 @@ void FVirtualizedUntypedBulkData::SetCompressionFormat(FName InCompressionFormat
 	{
 		EnumAddFlags(Flags, EFlags::DisablePayloadCompression);
 	}
+	else
+	{
+		EnumRemoveFlags(Flags, EFlags::DisablePayloadCompression);
+	}
 }
 
 FCustomVersionContainer FVirtualizedUntypedBulkData::GetCustomVersions(FArchive& InlineArchive)
