AsyncLoading with EditorDomain: batch up the callbacks in FAsyncLoadingThread::ProcessLoadedPackages so that the EditorDomain callback that saves the just-loaded package does not cause a stack overflow by adding 14 functions on the stack for each pending package discovered in the FlushAsyncLoading call.
#rnx

diff --git a/Engine/Source/Runtime/CoreUObject/Private/Serialization/AsyncLoading.cpp b/Engine/Source/Runtime/CoreUObject/Private/Serialization/AsyncLoading.cpp
--- a/Engine/Source/Runtime/CoreUObject/Private/Serialization/AsyncLoading.cpp
+++ b/Engine/Source/Runtime/CoreUObject/Private/Serialization/AsyncLoading.cpp
@@ -4355,6 +4355,7 @@ EAsyncPackageState::Type FAsyncLoadingThread::ProcessLoadedPackages(bool bUseTim
 		
 	bDidSomething = LoadedPackagesToProcess.Num() > 0;
 
+	TArray<FAsyncPackage*, TInlineAllocator<4>> CompletedPackages;
 	for (int32 PackageIndex = 0; PackageIndex < LoadedPackagesToProcess.Num() && !IsAsyncLoadingSuspendedInternal(); ++PackageIndex)
 	{
 		FPlatformMisc::PumpEssentialAppMessages();
@@ -4406,17 +4407,7 @@ EAsyncPackageState::Type FAsyncLoadingThread::ProcessLoadedPackages(bool bUseTim
 
 				UE_CLOG(NewExistingAsyncPackagesCounterValue < 0, LogStreaming, Fatal, TEXT("ExistingAsyncPackagesCounter is negative, this means we loaded more packages then requested so there must be a bug in async loading code."));
 
-				// Call external callbacks
-				const bool bInternalCallbacks = false;
-				const EAsyncLoadingResult::Type LoadingResult = Package->HasLoadFailed() ? EAsyncLoadingResult::Failed : EAsyncLoadingResult::Succeeded;
-				Package->CallCompletionCallbacks(bInternalCallbacks, LoadingResult);
-#if WITH_EDITOR
-				// In the editor we need to find any assets and add them to list for later callback
-				Package->GetLoadedAssets(LoadedAssets);
-#endif
-				// We don't need the package anymore
-				PackagesToDelete.AddUnique(Package);
-				Package->MarkRequestIDsAsComplete();
+				CompletedPackages.Add(Package);
 
 				TRACE_LOADTIME_END_LOAD_ASYNC_PACKAGE(Package);
 
@@ -4438,6 +4429,25 @@ EAsyncPackageState::Type FAsyncLoadingThread::ProcessLoadedPackages(bool bUseTim
 			break;
 		}
 	}
+
+	// Call callbacks in a batch in a stack-local array. This is to ensure that callbacks that trigger
+	// on each package load and call FlushAsyncLoading do not stack overflow by adding one FlushAsyncLoading
+	// call per LoadedPackageToProcess onto the stack
+	for (FAsyncPackage* Package : CompletedPackages)
+	{
+		// Call external callbacks
+		const bool bInternalCallbacks = false;
+		const EAsyncLoadingResult::Type LoadingResult = Package->HasLoadFailed() ? EAsyncLoadingResult::Failed : EAsyncLoadingResult::Succeeded;
+		Package->CallCompletionCallbacks(bInternalCallbacks, LoadingResult);
+#if WITH_EDITOR
+		// In the editor we need to find any assets and add them to list for later callback
+		Package->GetLoadedAssets(LoadedAssets);
+#endif
+		// We don't need the package anymore
+		PackagesToDelete.AddUnique(Package);
+		Package->MarkRequestIDsAsComplete();
+
+	}
 	bDidSomething = bDidSomething || PackagesToDelete.Num() > 0;
 
 	// Delete packages we're done processing and are no longer dependencies of anything else
diff --git a/Engine/Source/Runtime/CoreUObject/Private/Serialization/AsyncLoading2.cpp b/Engine/Source/Runtime/CoreUObject/Private/Serialization/AsyncLoading2.cpp
--- a/Engine/Source/Runtime/CoreUObject/Private/Serialization/AsyncLoading2.cpp
+++ b/Engine/Source/Runtime/CoreUObject/Private/Serialization/AsyncLoading2.cpp
@@ -4501,6 +4501,7 @@ EAsyncPackageState::Type FAsyncLoadingThread2::ProcessLoadedPackagesFromGameThre
 		bLocalDidSomething |= MainThreadEventQueue.PopAndExecute(ThreadState);
 
 		bLocalDidSomething |= LoadedPackagesToProcess.Num() > 0;
+		TArray<FAsyncPackage2*, TInlineAllocator<4>> PackagesReadyForCallback;
 		for (int32 PackageIndex = 0; PackageIndex < LoadedPackagesToProcess.Num(); ++PackageIndex)
 		{
 			SCOPED_LOADTIMER(ProcessLoadedPackagesTime);
@@ -4594,14 +4595,6 @@ EAsyncPackageState::Type FAsyncLoadingThread2::ProcessLoadedPackagesFromGameThre
 
 			TRACE_LOADTIME_END_LOAD_ASYNC_PACKAGE(Package);
 
-			// Call external callbacks
-			const EAsyncLoadingResult::Type LoadingResult = Package->HasLoadFailed() ? EAsyncLoadingResult::Failed : EAsyncLoadingResult::Succeeded;
-			{
-				TRACE_CPUPROFILER_EVENT_SCOPE(PackageCompletionCallbacks);
-				Package->CallCompletionCallbacks(LoadingResult);
-			}
-
-			// We don't need the package anymore
 			check(Package->AsyncPackageLoadingState == EAsyncPackageLoadingState2::Finalize);
 			if (bHasClusterObjects)
 			{
@@ -4611,6 +4604,21 @@ EAsyncPackageState::Type FAsyncLoadingThread2::ProcessLoadedPackagesFromGameThre
 			{
 				Package->AsyncPackageLoadingState = EAsyncPackageLoadingState2::Complete;
 			}
+			PackagesReadyForCallback.Add(Package);
+		}
+
+		// Call callbacks in a batch in a stack-local array. This is to ensure that callbacks that trigger
+		// on each package load and call FlushAsyncLoading do not stack overflow by adding one FlushAsyncLoading
+		// call per LoadedPackageToProcess onto the stack
+		for (FAsyncPackage2* Package : PackagesReadyForCallback)
+		{
+			// Call external callbacks
+			const EAsyncLoadingResult::Type LoadingResult = Package->HasLoadFailed() ? EAsyncLoadingResult::Failed : EAsyncLoadingResult::Succeeded;
+			{
+				TRACE_CPUPROFILER_EVENT_SCOPE(PackageCompletionCallbacks);
+				Package->CallCompletionCallbacks(LoadingResult);
+			}
+
 			check(!CompletedPackages.Contains(Package));
 			CompletedPackages.Add(Package);
 			Package->MarkRequestIDsAsComplete();
