AsyncLoading2:
Ensure that all reachable packages have all their exports serialized before starting post load instead of relying on the serialization order. This is still required in some cases.
Apply the same logic to ensure that all reachable packages are post loaded before finishing loading

#rnx

diff --git a/Engine/Source/Runtime/CoreUObject/Private/Serialization/AsyncLoading2.cpp b/Engine/Source/Runtime/CoreUObject/Private/Serialization/AsyncLoading2.cpp
--- a/Engine/Source/Runtime/CoreUObject/Private/Serialization/AsyncLoading2.cpp
+++ b/Engine/Source/Runtime/CoreUObject/Private/Serialization/AsyncLoading2.cpp
@@ -1702,9 +1702,6 @@ enum EEventLoadNode2 : uint8
 	Package_ProcessSummary,
 	Package_SetupDependencies,
 	Package_ExportsSerialized,
-	Package_BeginPostLoad,
-	Package_PostLoadDone,
-	Package_FinishLoading,
 	Package_NumPhases,
 
 	ExportBundle_Process = 0,
@@ -1815,9 +1812,6 @@ struct FAsyncPackage2
 	/** Returns the UPackage wrapped by this, if it is valid */
 	UPackage* GetLoadedPackage();
 
-	/** Checks if all dependencies (imported packages) of this package have been fully loaded */
-	bool AreAllDependenciesFullyLoaded(TSet<FPackageId>& VisitedPackages);
-
 	/** Creates GC clusters from loaded objects */
 	EAsyncPackageState::Type CreateClusters(FAsyncLoadingThreadState2& ThreadState);
 
@@ -1836,9 +1830,6 @@ private:
 		int32 BundleIndex[FExportBundleEntry::ExportCommandType_Count];
 	};
 
-	/** Checks if all dependencies (imported packages) of this package have been fully loaded */
-	bool AreAllDependenciesFullyLoadedInternal(FAsyncPackage2* Package, TSet<FPackageId>& VisitedPackages, FPackageId& OutPackageId);
-
 	/** Basic information associated with this package */
 	FAsyncPackageDesc2 Desc;
 	FAsyncPackageData Data;
@@ -1859,6 +1850,100 @@ private:
 	/** Current index into DeferredClusterObjects array used to spread routing CreateClusters over several frames			*/
 	int32						DeferredClusterIndex = 0;
 	EAsyncPackageLoadingState2	AsyncPackageLoadingState = EAsyncPackageLoadingState2::NewPackage;
+
+	struct FAllDependenciesState
+	{
+		FAsyncPackage2* WaitingForPackage = nullptr;
+		FAsyncPackage2* PackagesWaitingForThisHead = nullptr;
+		FAsyncPackage2* PackagesWaitingForThisTail = nullptr;
+		FAsyncPackage2* PrevLink = nullptr;
+		FAsyncPackage2* NextLink = nullptr;
+		uint32 LastTick = 0;
+		bool bAllDone = false;
+		bool bAnyNotDone = false;
+		bool bVisitedMark = false;
+
+		void UpdateTick(int32 CurrentTick)
+		{
+			if (LastTick != CurrentTick)
+			{
+				LastTick = CurrentTick;
+				bAnyNotDone = false;
+				bVisitedMark = false;
+			}
+		}
+
+		static void AddToWaitList(FAllDependenciesState FAsyncPackage2::* StateMemberPtr, FAsyncPackage2* WaitListPackage, FAsyncPackage2* PackageToAdd)
+		{
+			check(WaitListPackage);
+			check(PackageToAdd);
+			FAllDependenciesState& WaitListPackageState = WaitListPackage->*StateMemberPtr;
+			FAllDependenciesState& PackageToAddState = PackageToAdd->*StateMemberPtr;
+			
+			if (PackageToAddState.WaitingForPackage == WaitListPackage)
+			{
+				return;
+			}
+			if (PackageToAddState.WaitingForPackage)
+			{
+				PackageToAddState.RemoveFromWaitList(StateMemberPtr, PackageToAddState.WaitingForPackage, PackageToAdd);
+			}
+
+			check(!PackageToAddState.PrevLink);
+			check(!PackageToAddState.NextLink);
+			if (WaitListPackageState.PackagesWaitingForThisTail)
+			{
+				FAllDependenciesState& WaitListTailState = WaitListPackageState.PackagesWaitingForThisTail->*StateMemberPtr;
+				check(!WaitListTailState.NextLink);
+				WaitListTailState.NextLink = PackageToAdd;
+				PackageToAddState.PrevLink = WaitListPackageState.PackagesWaitingForThisTail;
+			}
+			else
+			{
+				check(!WaitListPackageState.PackagesWaitingForThisHead);
+				WaitListPackageState.PackagesWaitingForThisHead = PackageToAdd;
+			}
+			WaitListPackageState.PackagesWaitingForThisTail = PackageToAdd;
+			PackageToAddState.WaitingForPackage = WaitListPackage;
+		}
+
+		static void RemoveFromWaitList(FAllDependenciesState FAsyncPackage2::* StateMemberPtr, FAsyncPackage2* WaitListPackage, FAsyncPackage2* PackageToRemove)
+		{
+			check(WaitListPackage);
+			check(PackageToRemove);
+
+			FAllDependenciesState& WaitListPackageState = WaitListPackage->*StateMemberPtr;
+			FAllDependenciesState& PackageToRemoveState = PackageToRemove->*StateMemberPtr;
+
+			check(PackageToRemoveState.WaitingForPackage == WaitListPackage);
+			if (PackageToRemoveState.PrevLink)
+			{
+				FAllDependenciesState& PrevLinkState = PackageToRemoveState.PrevLink->*StateMemberPtr;
+				PrevLinkState.NextLink = PackageToRemoveState.NextLink;
+			}
+			else
+			{
+				check(WaitListPackageState.PackagesWaitingForThisHead == PackageToRemove);
+				WaitListPackageState.PackagesWaitingForThisHead = PackageToRemoveState.NextLink;
+			}
+			if (PackageToRemoveState.NextLink)
+			{
+				FAllDependenciesState& NextLinkState = PackageToRemoveState.NextLink->*StateMemberPtr;
+				NextLinkState.PrevLink = PackageToRemoveState.PrevLink;
+			}
+			else
+			{
+				check(WaitListPackageState.PackagesWaitingForThisTail == PackageToRemove);
+				WaitListPackageState.PackagesWaitingForThisTail = PackageToRemoveState.PrevLink;
+			}
+			PackageToRemoveState.PrevLink = nullptr;
+			PackageToRemoveState.NextLink = nullptr;
+			PackageToRemoveState.WaitingForPackage = nullptr;
+		}
+	};
+	FAllDependenciesState		AllDependenciesSerializedState;
+	FAllDependenciesState		AllDependenciesFullyLoadedState;
+
 	/** True if our load has failed */
 	bool						bLoadHasFailed = false;
 	/** True if this package was created by this async package */
@@ -1902,9 +1987,6 @@ public:
 	static EAsyncPackageState::Type Event_ProcessPackageSummary(FAsyncLoadingThreadState2& ThreadState, FAsyncPackage2* Package, int32);
 	static EAsyncPackageState::Type Event_SetupDependencies(FAsyncLoadingThreadState2& ThreadState, FAsyncPackage2* Package, int32);
 	static EAsyncPackageState::Type Event_ExportsDone(FAsyncLoadingThreadState2& ThreadState, FAsyncPackage2* Package, int32);
-	static EAsyncPackageState::Type Event_BeginPostLoad(FAsyncLoadingThreadState2& ThreadState, FAsyncPackage2* Package, int32);
-	static EAsyncPackageState::Type Event_PostLoadDone(FAsyncLoadingThreadState2& ThreadState, FAsyncPackage2* Package, int32);
-	static EAsyncPackageState::Type Event_FinishLoading(FAsyncLoadingThreadState2& ThreadState, FAsyncPackage2* Package, int32);
 	static EAsyncPackageState::Type Event_PostLoadExportBundle(FAsyncLoadingThreadState2& ThreadState, FAsyncPackage2* Package, int32 ExportBundleIndex);
 	static EAsyncPackageState::Type Event_DeferredPostLoadExportBundle(FAsyncLoadingThreadState2& ThreadState, FAsyncPackage2* Package, int32 ExportBundleIndex);
 	
@@ -1934,6 +2016,12 @@ private:
 	void CreateNodes(const FAsyncLoadEventSpec* EventSpecs);
 	void SetupSerializedArcs();
 	void SetupScriptDependencies();
+	bool HaveAllDependenciesReachedStateDebug(FAsyncPackage2* Package, TSet<FAsyncPackage2*>& VisitedPackages, EAsyncPackageLoadingState2 WaitForPackageState);
+	bool HaveAllDependenciesReachedState(FAllDependenciesState FAsyncPackage2::* StateMemberPtr, EAsyncPackageLoadingState2 WaitForPackageState, uint32 CurrentTick);
+	void UpdateDependenciesStateRecursive(FAllDependenciesState FAsyncPackage2::* StateMemberPtr, EAsyncPackageLoadingState2 WaitForPackageState, uint32 CurrentTick, FAsyncPackage2* Root);
+	void WaitForAllDependenciesToReachState(FAllDependenciesState FAsyncPackage2::* StateMemberPtr, EAsyncPackageLoadingState2 WaitForPackageState, uint32& CurrentTickVariable, TFunctionRef<void(FAsyncPackage2*)> OnStateReached);
+	void ConditionalBeginPostLoad();
+	void ConditionalFinishLoading();
 
 	/**
 	 * Begin async loading process. Simulates parts of BeginLoad.
@@ -2149,6 +2237,9 @@ private:
 	TArray<FBundleIoRequest> WaitingIoRequests;
 	uint64 PendingBundleIoRequestsTotalSize = 0;
 
+	uint32 ConditionalBeginPostLoadTick = 0;
+	uint32 ConditionalFinishLoadingTick = 0;
+
 public:
 
 	//~ Begin FRunnable Interface.
@@ -3413,8 +3504,6 @@ void FAsyncPackage2::ImportPackagesRecursive()
 		ImportedPackage->AddRef();
 		Data.ImportedAsyncPackages[ImportedPackageIndex++] = ImportedPackage;
 		GetPackageNode(Package_SetupDependencies).DependsOn(&ImportedPackage->GetPackageNode(Package_ProcessSummary));
-		GetPackageNode(Package_BeginPostLoad).DependsOn(&ImportedPackage->GetPackageNode(Package_ExportsSerialized));
-		GetPackageNode(Package_FinishLoading).DependsOn(&ImportedPackage->GetPackageNode(Package_PostLoadDone));
 
 		if (bInserted)
 		{
@@ -3684,7 +3773,6 @@ EAsyncPackageState::Type FAsyncPackage2::Event_ProcessExportBundle(FAsyncLoading
 		if (Package->ExternalReadDependencies.Num() == 0)
 		{
 			check(Package->AsyncPackageLoadingState == EAsyncPackageLoadingState2::ProcessExportBundles);
-			Package->AsyncPackageLoadingState = EAsyncPackageLoadingState2::ExportsDone;
 			Package->GetPackageNode(Package_ExportsSerialized).ReleaseBarrier(&ThreadState);
 		}
 		else
@@ -4038,7 +4126,8 @@ EAsyncPackageState::Type FAsyncPackage2::Event_ExportsDone(FAsyncLoadingThreadSt
 {
 	TRACE_CPUPROFILER_EVENT_SCOPE(Event_ExportsDone);
 	UE_ASYNC_PACKAGE_DEBUG(Package->Desc);
-	check(Package->AsyncPackageLoadingState == EAsyncPackageLoadingState2::ExportsDone);
+	check(Package->AsyncPackageLoadingState == EAsyncPackageLoadingState2::ProcessExportBundles || Package->AsyncPackageLoadingState == EAsyncPackageLoadingState2::WaitingForExternalReads);
+	Package->AsyncPackageLoadingState = EAsyncPackageLoadingState2::ExportsDone;
 
 	if (!Package->bLoadHasFailed && Package->Desc.bCanBeImported)
 	{
@@ -4053,33 +4142,205 @@ EAsyncPackageState::Type FAsyncPackage2::Event_ExportsDone(FAsyncLoadingThreadSt
 		FCoreDelegates::ReleasePreloadedPackageShaderMaps.ExecuteIfBound(Package->Data.ShaderMapHashes);
 	}
 
-	Package->AsyncPackageLoadingState = EAsyncPackageLoadingState2::PostLoad;
-	Package->GetPackageNode(Package_BeginPostLoad).ReleaseBarrier(&ThreadState);
+	Package->ConditionalBeginPostLoad();
 	return EAsyncPackageState::Complete;
 }
 
-EAsyncPackageState::Type FAsyncPackage2::Event_BeginPostLoad(FAsyncLoadingThreadState2& ThreadState, FAsyncPackage2* Package, int32)
+bool FAsyncPackage2::HaveAllDependenciesReachedStateDebug(FAsyncPackage2* Package, TSet<FAsyncPackage2*>& VisitedPackages, EAsyncPackageLoadingState2 WaitForPackageState)
 {
-	TRACE_CPUPROFILER_EVENT_SCOPE(Event_BeginPostLoad);
-	for (int32 ExportBundleIndex = 0, ExportBundleCount = Package->Data.ExportInfo.ExportBundleCount; ExportBundleIndex < ExportBundleCount; ++ExportBundleIndex)
+	for (FAsyncPackage2* ImportedPackage : Package->Data.ImportedAsyncPackages)
 	{
-		Package->GetExportBundleNode(EEventLoadNode2::ExportBundle_PostLoad, ExportBundleIndex).ReleaseBarrier();
+		if (ImportedPackage)
+		{
+
+			if (VisitedPackages.Contains(ImportedPackage))
+			{
+				continue;
+			}
+			VisitedPackages.Add(ImportedPackage);
+
+			if (ImportedPackage->AsyncPackageLoadingState < WaitForPackageState)
+			{
+				return false;
+			}
+
+			if (!HaveAllDependenciesReachedStateDebug(ImportedPackage, VisitedPackages, WaitForPackageState))
+			{
+				return false;
+			}
+		}
 	}
-	return EAsyncPackageState::Complete;
+	return true;
 }
 
-EAsyncPackageState::Type FAsyncPackage2::Event_PostLoadDone(FAsyncLoadingThreadState2& ThreadState, FAsyncPackage2* Package, int32)
+bool FAsyncPackage2::HaveAllDependenciesReachedState(FAllDependenciesState FAsyncPackage2::* StateMemberPtr, EAsyncPackageLoadingState2 WaitForPackageState, uint32 CurrentTick)
 {
-	TRACE_CPUPROFILER_EVENT_SCOPE(Event_PostLoadDone);
-	Package->GetPackageNode(Package_FinishLoading).ReleaseBarrier(&ThreadState);
-	return EAsyncPackageState::Complete;
+	FAllDependenciesState& ThisState = this->*StateMemberPtr;
+	if (ThisState.bAllDone)
+	{
+		return true;
+	}
+	if (AsyncPackageLoadingState < WaitForPackageState)
+	{
+		return false;
+	}
+	ThisState.UpdateTick(CurrentTick);
+	UpdateDependenciesStateRecursive(StateMemberPtr, WaitForPackageState, CurrentTick, this);
+	check(ThisState.bAllDone || (ThisState.WaitingForPackage && ThisState.WaitingForPackage->AsyncPackageLoadingState <= WaitForPackageState));
+	return ThisState.bAllDone;
 }
 
-EAsyncPackageState::Type FAsyncPackage2::Event_FinishLoading(FAsyncLoadingThreadState2& ThreadState, FAsyncPackage2* Package, int32)
+void FAsyncPackage2::UpdateDependenciesStateRecursive(FAllDependenciesState FAsyncPackage2::* StateMemberPtr, EAsyncPackageLoadingState2 WaitForPackageState, uint32 CurrentTick, FAsyncPackage2* Root)
 {
-	TRACE_CPUPROFILER_EVENT_SCOPE(Event_FinishLoading);
-	Package->AsyncLoadingThread.LoadedPackagesToProcess.Add(Package);
-	return EAsyncPackageState::Complete;
+	FAllDependenciesState& ThisState = this->*StateMemberPtr;
+
+	check(!ThisState.bVisitedMark);
+	check(!ThisState.bAllDone);
+	check(!ThisState.bAnyNotDone);
+
+	ThisState.bVisitedMark = true;
+
+	if (FAsyncPackage2* WaitingForPackage = ThisState.WaitingForPackage)
+	{
+		FAllDependenciesState& WaitingForPackageState = WaitingForPackage->*StateMemberPtr;
+		if (WaitingForPackage->AsyncPackageLoadingState < WaitForPackageState)
+		{
+			ThisState.bAnyNotDone = true;
+			return;
+		}
+		else if (!WaitingForPackageState.bAllDone)
+		{
+			WaitingForPackageState.UpdateTick(CurrentTick);
+			if (!WaitingForPackageState.bVisitedMark)
+			{
+				WaitingForPackage->UpdateDependenciesStateRecursive(StateMemberPtr, WaitForPackageState, CurrentTick, Root);
+			}
+			if (WaitingForPackageState.bAnyNotDone)
+			{
+				ThisState.bAnyNotDone = true;
+				return;
+			}
+		}
+	}
+
+	bool bAllDone = true;
+	FAsyncPackage2* WaitingForPackage = nullptr;
+	for (FAsyncPackage2* ImportedPackage : Data.ImportedAsyncPackages)
+	{
+		if (!ImportedPackage)
+		{
+			continue;
+		}
+
+		FAllDependenciesState& ImportedPackageState = ImportedPackage->*StateMemberPtr;
+
+		if (ImportedPackageState.bAllDone)
+		{
+			continue;
+		}
+
+		ImportedPackageState.UpdateTick(CurrentTick);
+
+		if (ImportedPackage->AsyncPackageLoadingState < WaitForPackageState)
+		{
+			ImportedPackageState.bAnyNotDone = true;
+		}
+		else if (!ImportedPackageState.bVisitedMark)
+		{
+			ImportedPackage->UpdateDependenciesStateRecursive(StateMemberPtr, WaitForPackageState, CurrentTick, Root);
+		}
+
+		if (ImportedPackageState.bAnyNotDone)
+		{
+			ThisState.bAnyNotDone = true;
+			WaitingForPackage = ImportedPackage;
+			break;
+		}
+		else if (!ImportedPackageState.bAllDone)
+		{
+			bAllDone = false;
+		}
+	}
+	if (WaitingForPackage)
+	{
+		check(WaitingForPackage != this);
+		FAllDependenciesState::AddToWaitList(StateMemberPtr, WaitingForPackage, this);
+	}
+	else if (bAllDone || this == Root)
+	{
+		// If we're the root an not waiting for any package we're done
+		ThisState.bAllDone = true;
+	}
+	else
+	{
+		// We didn't find any imported package that was not done but we could have a circular dependency back to the root which could either be done or end up waiting
+		// for another package. Make us wait for the root so that we are ticked when it completes.
+		FAllDependenciesState::AddToWaitList(StateMemberPtr, Root, this);
+	}
+}
+
+void FAsyncPackage2::WaitForAllDependenciesToReachState(FAllDependenciesState FAsyncPackage2::* StateMemberPtr, EAsyncPackageLoadingState2 WaitForPackageState, uint32& CurrentTickVariable, TFunctionRef<void(FAsyncPackage2*)> OnStateReached)
+{
+	if (HaveAllDependenciesReachedState(StateMemberPtr, WaitForPackageState, CurrentTickVariable++))
+	{
+		FAsyncPackage2* FirstPackageReadyToProceed = this;
+
+		while (FirstPackageReadyToProceed)
+		{
+			FAsyncPackage2* PackageReadyToProceed = FirstPackageReadyToProceed;
+			FAllDependenciesState& PackageReadyToProceedState = PackageReadyToProceed->*StateMemberPtr;
+			FirstPackageReadyToProceed = PackageReadyToProceedState.NextLink;
+
+			if (PackageReadyToProceed->AsyncPackageLoadingState > WaitForPackageState)
+			{
+				continue;
+			}
+
+#if DO_CHECK
+			TSet<FAsyncPackage2*> VisitedPackages;
+			check(HaveAllDependenciesReachedStateDebug(this, VisitedPackages, WaitForPackageState));
+#endif
+
+			while (FAsyncPackage2* WaitingPackage = PackageReadyToProceedState.PackagesWaitingForThisHead)
+			{
+				FAllDependenciesState& WaitingPackageState = WaitingPackage->*StateMemberPtr;
+				check(WaitingPackageState.WaitingForPackage == PackageReadyToProceed);
+				if (WaitingPackage->HaveAllDependenciesReachedState(StateMemberPtr, WaitForPackageState, CurrentTickVariable++))
+				{
+					FAllDependenciesState::RemoveFromWaitList(StateMemberPtr, PackageReadyToProceed, WaitingPackage);
+					WaitingPackageState.NextLink = FirstPackageReadyToProceed;
+					FirstPackageReadyToProceed = WaitingPackage;
+				}
+			}
+			check(!PackageReadyToProceedState.PackagesWaitingForThisTail);
+			check(PackageReadyToProceed->AsyncPackageLoadingState == WaitForPackageState);
+			PackageReadyToProceed->AsyncPackageLoadingState = static_cast<EAsyncPackageLoadingState2>(static_cast<uint32>(WaitForPackageState) + 1);
+			OnStateReached(PackageReadyToProceed);
+		}
+	}
+}
+
+void FAsyncPackage2::ConditionalBeginPostLoad()
+{
+	TRACE_CPUPROFILER_EVENT_SCOPE(ConditionalBeginPostLoad);
+	WaitForAllDependenciesToReachState(&FAsyncPackage2::AllDependenciesSerializedState, EAsyncPackageLoadingState2::ExportsDone, AsyncLoadingThread.ConditionalBeginPostLoadTick,
+		[](FAsyncPackage2* Package)
+		{
+			for (int32 ExportBundleIndex = 0, ExportBundleCount = Package->Data.ExportInfo.ExportBundleCount; ExportBundleIndex < ExportBundleCount; ++ExportBundleIndex)
+			{
+				Package->GetExportBundleNode(EEventLoadNode2::ExportBundle_PostLoad, ExportBundleIndex).ReleaseBarrier();
+			}
+		});
+}
+
+void FAsyncPackage2::ConditionalFinishLoading()
+{
+	TRACE_CPUPROFILER_EVENT_SCOPE(ConditionalFinishLoading);
+	WaitForAllDependenciesToReachState(&FAsyncPackage2::AllDependenciesFullyLoadedState, EAsyncPackageLoadingState2::DeferredPostLoadDone, AsyncLoadingThread.ConditionalFinishLoadingTick,
+		[](FAsyncPackage2* Package)
+		{
+			Package->AsyncLoadingThread.LoadedPackagesToProcess.Add(Package);
+		});
 }
 
 EAsyncPackageState::Type FAsyncPackage2::Event_PostLoadExportBundle(FAsyncLoadingThreadState2& ThreadState, FAsyncPackage2* Package, int32 InExportBundleIndex)
@@ -4287,7 +4548,7 @@ EAsyncPackageState::Type FAsyncPackage2::Event_DeferredPostLoadExportBundle(FAsy
 		Package->ProcessedExportBundlesCount = 0;
 		check(Package->AsyncPackageLoadingState == EAsyncPackageLoadingState2::DeferredPostLoad);
 		Package->AsyncPackageLoadingState = EAsyncPackageLoadingState2::DeferredPostLoadDone;
-		Package->GetPackageNode(Package_PostLoadDone).ReleaseBarrier(&ThreadState);
+		Package->ConditionalFinishLoading();
 	}
 
 	return EAsyncPackageState::Complete;
@@ -4421,49 +4682,6 @@ EAsyncPackageState::Type FAsyncLoadingThread2::ProcessAsyncLoadingFromGameThread
 	return EAsyncPackageState::Complete;
 }
 
-bool FAsyncPackage2::AreAllDependenciesFullyLoadedInternal(FAsyncPackage2* Package, TSet<FPackageId>& VisitedPackages, FPackageId& OutPackageId)
-{
-	for (const FPackageId& ImportedPackageId : Package->Data.ImportedPackageIds)
-	{
-		if (VisitedPackages.Contains(ImportedPackageId))
-		{
-			continue;
-		}
-		VisitedPackages.Add(ImportedPackageId);
-
-		FAsyncPackage2* AsyncRoot = AsyncLoadingThread.GetAsyncPackage(ImportedPackageId);
-		if (AsyncRoot)
-		{
-			if (AsyncRoot->AsyncPackageLoadingState < EAsyncPackageLoadingState2::DeferredPostLoadDone)
-			{
-				OutPackageId = ImportedPackageId;
-				return false;
-			}
-
-			if (!AreAllDependenciesFullyLoadedInternal(AsyncRoot, VisitedPackages, OutPackageId))
-			{
-				return false;
-			}
-		}
-	}
-	return true;
-}
-
-bool FAsyncPackage2::AreAllDependenciesFullyLoaded(TSet<FPackageId>& VisitedPackages)
-{
-	TRACE_CPUPROFILER_EVENT_SCOPE(AreAllDependenciesFullyLoaded);
-	VisitedPackages.Reset();
-	FPackageId PackageId;
-	const bool bLoaded = AreAllDependenciesFullyLoadedInternal(this, VisitedPackages, PackageId);
-	if (!bLoaded)
-	{
-		FAsyncPackage2* AsyncRoot = AsyncLoadingThread.GetAsyncPackage(PackageId);
-		UE_LOG(LogStreaming, Verbose, TEXT("AreAllDependenciesFullyLoaded: '%s' doesn't have all exports processed by DeferredPostLoad"),
-			*AsyncRoot->Desc.PackageNameToLoad.ToString());
-	}
-	return bLoaded;
-}
-
 EAsyncPackageState::Type FAsyncLoadingThread2::ProcessLoadedPackagesFromGameThread(FAsyncLoadingThreadState2& ThreadState, bool& bDidSomething, int32 FlushRequestID)
 {
 	EAsyncPackageState::Type Result = EAsyncPackageState::Complete;
@@ -4484,9 +4702,6 @@ EAsyncPackageState::Type FAsyncLoadingThread2::ProcessLoadedPackagesFromGameThre
 		}
 	}
 
-	// For performance reasons this set is created here and reset inside of AreAllDependenciesFullyLoaded
-	TSet<FPackageId> VisistedPackages;
-
 	for (;;)
 	{
 		FPlatformMisc::PumpEssentialAppMessages();
@@ -4507,8 +4722,7 @@ EAsyncPackageState::Type FAsyncLoadingThread2::ProcessLoadedPackagesFromGameThre
 			SCOPED_LOADTIMER(ProcessLoadedPackagesTime);
 			FAsyncPackage2* Package = LoadedPackagesToProcess[PackageIndex];
 			UE_ASYNC_PACKAGE_DEBUG(Package->Desc);
-			check(Package->AsyncPackageLoadingState == EAsyncPackageLoadingState2::DeferredPostLoadDone);
-			Package->AsyncPackageLoadingState = EAsyncPackageLoadingState2::Finalize;
+			check(Package->AsyncPackageLoadingState == EAsyncPackageLoadingState2::Finalize);
 
 			bool bHasClusterObjects = false;
 			TArray<UObject*> CDODefaultSubobjects;
@@ -4653,20 +4867,17 @@ EAsyncPackageState::Type FAsyncLoadingThread2::ProcessLoadedPackagesFromGameThre
 				{
 					SCOPE_CYCLE_COUNTER(STAT_FAsyncPackage_CreateClustersGameThread);
 					// This package will create GC clusters but first check if all dependencies of this package have been fully loaded
-					if (Package->AreAllDependenciesFullyLoaded(VisistedPackages))
+					if (Package->CreateClusters(ThreadState) == EAsyncPackageState::Complete)
 					{
-						if (Package->CreateClusters(ThreadState) == EAsyncPackageState::Complete)
-						{
-							// All clusters created, it's safe to delete the package
-							bSafeToDelete = true;
-							Package->AsyncPackageLoadingState = EAsyncPackageLoadingState2::Complete;
-						}
-						else
-						{
-							// Cluster creation timed out
-							Result = EAsyncPackageState::TimeOut;
-							break;
-						}
+						// All clusters created, it's safe to delete the package
+						bSafeToDelete = true;
+						Package->AsyncPackageLoadingState = EAsyncPackageLoadingState2::Complete;
+					}
+					else
+					{
+						// Cluster creation timed out
+						Result = EAsyncPackageState::TimeOut;
+						break;
 					}
 				}
 				else
@@ -4792,9 +5003,6 @@ FAsyncLoadingThread2::FAsyncLoadingThread2(FIoDispatcher& InIoDispatcher)
 	EventSpecs[EEventLoadNode2::Package_ProcessSummary] = { &FAsyncPackage2::Event_ProcessPackageSummary, &EventQueue, false };
 	EventSpecs[EEventLoadNode2::Package_SetupDependencies] = { &FAsyncPackage2::Event_SetupDependencies, &EventQueue, false };
 	EventSpecs[EEventLoadNode2::Package_ExportsSerialized] = { &FAsyncPackage2::Event_ExportsDone, &EventQueue, true };
-	EventSpecs[EEventLoadNode2::Package_BeginPostLoad] = { &FAsyncPackage2::Event_BeginPostLoad, &EventQueue, true };
-	EventSpecs[EEventLoadNode2::Package_PostLoadDone] = { &FAsyncPackage2::Event_PostLoadDone, &MainThreadEventQueue, true };
-	EventSpecs[EEventLoadNode2::Package_FinishLoading] = { &FAsyncPackage2::Event_FinishLoading, &MainThreadEventQueue, true };
 
 	EventSpecs[EEventLoadNode2::Package_NumPhases + EEventLoadNode2::ExportBundle_Process] = { &FAsyncPackage2::Event_ProcessExportBundle, &EventQueue, false };
 	EventSpecs[EEventLoadNode2::Package_NumPhases + EEventLoadNode2::ExportBundle_PostLoad] = { &FAsyncPackage2::Event_PostLoadExportBundle, &EventQueue, false };
@@ -5677,7 +5885,6 @@ EAsyncPackageState::Type FAsyncPackage2::ProcessExternalReads(EExternalReadActio
 	}
 
 	ExternalReadDependencies.Empty();
-	AsyncPackageLoadingState = EAsyncPackageLoadingState2::ExportsDone;
 	GetPackageNode(Package_ExportsSerialized).ReleaseBarrier();
 	return EAsyncPackageState::Complete;
 }
