Fix fatal log error during cook when referencing an unsaveable dependency

diff --git a/Engine/Source/Runtime/CoreUObject/Private/UObject/SavePackage2.cpp b/Engine/Source/Runtime/CoreUObject/Private/UObject/SavePackage2.cpp
--- a/Engine/Source/Runtime/CoreUObject/Private/UObject/SavePackage2.cpp
+++ b/Engine/Source/Runtime/CoreUObject/Private/UObject/SavePackage2.cpp
@@ -1122,7 +1122,7 @@ void SavePreloadDependencies(FStructuredArchive::FRecord& StructuredArchiveRoot,
 					AddTo.Add(Index);
 					return;
 				}
-				else if (!ToTest->HasAnyFlags(RF_Transient))
+				else if (!SaveContext.IsUnsaveable(ToTest))
 				{
 					UE_CLOG(Outermost->HasAnyPackageFlags(PKG_CompiledIn), LogSavePackage, Verbose, TEXT("A compiled in dependency '%s' of '%s' was not actually in the linker tables and so will be ignored (%d)."), *ToTest->GetFullName(), *ForObj->GetFullName(), CallSite);
 					UE_CLOG(!Outermost->HasAnyPackageFlags(PKG_CompiledIn), LogSavePackage, Fatal, TEXT("A dependency '%s' of '%s' was not actually in the linker tables and so will be ignored (%d)."), *ToTest->GetFullName(), *ForObj->GetFullName(), CallSite);
