The local caching of Mirage payloads after pulling them from persistent storage can now be enabled/disabled via the config file option '[Core.ContentVirtualization].EnableCacheAfterPull', defaults to true.

#rnx

diff --git a/Engine/Config/BaseEngine.ini b/Engine/Config/BaseEngine.ini
--- a/Engine/Config/BaseEngine.ini
+++ b/Engine/Config/BaseEngine.ini
@@ -1046,6 +1046,7 @@ EnablePackageSidecarSaving=False
 ; Options for the content virtualization system
 [Core.ContentVirtualization]
 EnablePushToBackend=False
+EnableCacheAfterPull=True
 MinPayloadLength=0
 BackendGraph=ContentVirtualizationBackendGraph_None
 
diff --git a/Engine/Source/Runtime/CoreUObject/Private/Virtualization/VirtualizationManager.cpp b/Engine/Source/Runtime/CoreUObject/Private/Virtualization/VirtualizationManager.cpp
--- a/Engine/Source/Runtime/CoreUObject/Private/Virtualization/VirtualizationManager.cpp
+++ b/Engine/Source/Runtime/CoreUObject/Private/Virtualization/VirtualizationManager.cpp
@@ -215,6 +215,7 @@ FVirtualizationManager& FVirtualizationManager::Get()
 
 FVirtualizationManager::FVirtualizationManager()
 	: bEnablePayloadPushing(true)
+	, bEnableCacheAfterPull(true)
 	, MinPayloadLength(0)
 	, BackendGraphName(TEXT("ContentVirtualizationBackendGraph_None"))
 	, bForceSingleThreaded(false)
@@ -365,7 +366,11 @@ FCompressedBuffer FVirtualizationManager::PullData(const FPayloadId& Id)
 		FCompressedBuffer Payload = PullDataFromBackend(*Backend, Id);
 		if (Payload)
 		{
-			CachePayload(Id, Payload, Backend);
+			if (bEnableCacheAfterPull)
+			{
+				CachePayload(Id, Payload, Backend);
+			}
+			
 			return Payload;
 		}
 	}
@@ -417,6 +422,17 @@ void FVirtualizationManager::ApplySettingsFromConfigFiles(const FConfigFile& Pla
 		UE_LOG(LogVirtualization, Error, TEXT("Failed to load [Core.ContentVirtualization].EnablePushToBackend from config file!"));
 	}
 
+	bool bEnableCacheAfterPullFromIni = false;
+	if (PlatformEngineIni.GetBool(TEXT("Core.ContentVirtualization"), TEXT("EnableCacheAfterPull"), bEnableCacheAfterPullFromIni))
+	{
+		bEnableCacheAfterPull = bEnableCacheAfterPullFromIni;
+		UE_LOG(LogVirtualization, Log, TEXT("\tCachePulledPayloads : %s"), bEnableCacheAfterPull ? TEXT("true") : TEXT("false"));
+	}
+	else
+	{
+		UE_LOG(LogVirtualization, Error, TEXT("Failed to load [Core.ContentVirtualization].EnableCacheAfterPull from config file!"));
+	}
+
 	int64 MinPayloadLengthFromIni = 0;
 	if (PlatformEngineIni.GetInt64(TEXT("Core.ContentVirtualization"), TEXT("MinPayloadLength"), MinPayloadLengthFromIni))
 	{
diff --git a/Engine/Source/Runtime/CoreUObject/Public/Virtualization/VirtualizationManager.h b/Engine/Source/Runtime/CoreUObject/Public/Virtualization/VirtualizationManager.h
--- a/Engine/Source/Runtime/CoreUObject/Public/Virtualization/VirtualizationManager.h
+++ b/Engine/Source/Runtime/CoreUObject/Public/Virtualization/VirtualizationManager.h
@@ -141,6 +141,9 @@ private:
 	/** Are payloads allowed to be virtualized. Defaults to true. */
 	bool bEnablePayloadPushing;
 
+	/** Should payloads be cached locally after being pulled from persistent storage? Defaults to true. */
+	bool bEnableCacheAfterPull;
+
 	/** The minimum length for a payload to be considered for virtualization. Defaults to 0 bytes. */
 	int64 MinPayloadLength;
 
