FAutomationTestAttemptToFindUninitializedScriptStructMembers test skipped when running automation tests under UHT, as it requires a NoExportType which doesn't exist in that config.

diff --git a/Engine/Source/Runtime/CoreUObject/Private/UObject/Class.cpp b/Engine/Source/Runtime/CoreUObject/Private/UObject/Class.cpp
--- a/Engine/Source/Runtime/CoreUObject/Private/UObject/Class.cpp
+++ b/Engine/Source/Runtime/CoreUObject/Private/UObject/Class.cpp
@@ -3610,7 +3610,12 @@ FAutoConsoleCommandWithWorldAndArgs GCmdListBadScriptStructs(
 IMPLEMENT_SIMPLE_AUTOMATION_TEST(FAutomationTestAttemptToFindUninitializedScriptStructMembers, "UObject.Class AttemptToFindUninitializedScriptStructMembers", EAutomationTestFlags::EditorContext | EAutomationTestFlags::ApplicationContextMask | EAutomationTestFlags::ServerContext | EAutomationTestFlags::SmokeFilter)
 bool FAutomationTestAttemptToFindUninitializedScriptStructMembers::RunTest(const FString& Parameters)
 {
+	// This test fails when running tests under UHT because there is no TestUninitializedScriptStructMembersTest, so just skip it in that config.
+#if !HACK_HEADER_GENERATOR
 	return FStructUtils::AttemptToFindUninitializedScriptStructMembers() == 0;
+#else
+	return true;
+#endif
 }
 
 #endif
