SetSparseClassDataStruct also edits the superclass pointers of subclasses' SparseClassDataStructs.
#rnx

diff --git a/Engine/Source/Runtime/CoreUObject/Private/UObject/Class.cpp b/Engine/Source/Runtime/CoreUObject/Private/UObject/Class.cpp
--- a/Engine/Source/Runtime/CoreUObject/Private/UObject/Class.cpp
+++ b/Engine/Source/Runtime/CoreUObject/Private/UObject/Class.cpp
@@ -5173,6 +5173,25 @@ void UClass::SetSparseClassDataStruct(UScriptStruct* InSparseClassDataStruct)
 { 
 	if (SparseClassDataStruct != InSparseClassDataStruct)
 	{
+		if (SparseClassDataStruct && InSparseClassDataStruct)
+		{
+			// Find all subclasses and point the SuperClass of their SparseClassDataStruct to point to this new SparseClassDataStruct.
+			// We have to do this when compilation creates a new SparseClassDataStruct that has already loaded subclasses.
+			// Most class-owned subobjects are regenerated at compile time (properties,functions,etc) or mirrored in the UBlueprint
+			// (see FKismetCompilerContext::SaveSubObjectsFromCleanAndSanitizeClass) and copied over to the new class inside of
+			// FKismetCompilerContext::FinishCompilingClass (SimpleConstructionScript,etc). But SparseClassDataStruct is not covered
+			// by either of those cases, so we instead fixup subclasses' SparseClassDataStructs' superclass pointers here.
+			TArray<UClass*> SubClasses;
+			GetDerivedClasses(this, SubClasses, true /* bRecursive */);
+			for (UClass* SubClass : SubClasses)
+			{
+				UScriptStruct* SubClassSparseClassDataStruct = SubClass->GetSparseClassDataStruct();
+				if (SubClassSparseClassDataStruct && SubClassSparseClassDataStruct->GetSuperStruct() == SparseClassDataStruct)
+				{
+					SubClassSparseClassDataStruct->SetSuperStruct(InSparseClassDataStruct);
+				}
+			}
+		}
 		// the old type and new type may not match when we do a reload so get rid of the old data
 		CleanupSparseClassData();
 
