AsyncLoading2:
Fixed edge case where imported packages failed to register as dependencies if they had serialized all their exports but not postloaded them yet
#rnx

diff --git a/Engine/Source/Runtime/CoreUObject/Private/Serialization/AsyncLoading2.cpp b/Engine/Source/Runtime/CoreUObject/Private/Serialization/AsyncLoading2.cpp
--- a/Engine/Source/Runtime/CoreUObject/Private/Serialization/AsyncLoading2.cpp
+++ b/Engine/Source/Runtime/CoreUObject/Private/Serialization/AsyncLoading2.cpp
@@ -2334,23 +2334,11 @@ public:
 
 	virtual void FireCompletedCompiledInImport(void* AsyncPackage, FPackageIndex Import) override {}
 
-	/**
-	* [ASYNC THREAD] Finds an existing async package in the AsyncPackages by its name.
-	*
-	* @param PackageName async package name.
-	* @return Pointer to the package or nullptr if not found
-	*/
-	FORCEINLINE FAsyncPackage2* FindAsyncPackage(const FName& PackageName)
+	FORCEINLINE FAsyncPackage2* FindAsyncPackage(FPackageId PackageId)
 	{
-		// TRACE_CPUPROFILER_EVENT_SCOPE(FindAsyncPackage);
-		FPackageId PackageId = FPackageId::FromName(PackageName);
-		if (PackageId.IsValid())
-		{
-			FScopeLock LockAsyncPackages(&AsyncPackagesCritical);
-			//checkSlow(IsInAsyncLoadThread());
-			return AsyncPackageLookup.FindRef(PackageId);
-		}
-		return nullptr;
+		FScopeLock LockAsyncPackages(&AsyncPackagesCritical);
+		//checkSlow(IsInAsyncLoadThread());
+		return AsyncPackageLookup.FindRef(PackageId);
 	}
 
 	FORCEINLINE FAsyncPackage2* GetAsyncPackage(const FPackageId& PackageId)
@@ -3467,12 +3455,6 @@ void FAsyncPackage2::ImportPackagesRecursive()
 		}
 		
 		FLoadedPackageRef& PackageRef = ImportStore.LoadedPackageStore.GetPackageRef(ImportedPackageId);
-		if (PackageRef.AreAllPublicExportsLoaded())
-		{
-			Data.ImportedAsyncPackages[ImportedPackageIndex++] = nullptr;
-			continue;
-		}
-
 		const FPackageStoreEntryHandle ImportedPackageEntry = PackageStore.GetPackageEntryHandle(ImportedPackageIdToLoad);
 
 		if (!ImportedPackageEntry)
@@ -3484,9 +3466,23 @@ void FAsyncPackage2::ImportPackagesRecursive()
 			continue;
 		}
 
+		FAsyncPackage2* ImportedPackage = nullptr;
+		bool bInserted = false;
 		FAsyncPackageDesc2 PackageDesc = FAsyncPackageDesc2::FromPackageImport(INDEX_NONE, Desc.Priority, ImportedPackageId, ImportedPackageIdToLoad, ImportedPackageUPackageName, ImportedPackageEntry);
-		bool bInserted;
-		FAsyncPackage2* ImportedPackage = AsyncLoadingThread.FindOrInsertPackage(PackageDesc, bInserted);
+		if (PackageRef.AreAllPublicExportsLoaded())
+		{
+			ImportedPackage = AsyncLoadingThread.FindAsyncPackage(ImportedPackageId);
+			if (!ImportedPackage)
+			{
+				Data.ImportedAsyncPackages[ImportedPackageIndex++] = nullptr;
+				continue;
+			}
+			bInserted = false;
+		}
+		else
+		{
+			ImportedPackage = AsyncLoadingThread.FindOrInsertPackage(PackageDesc, bInserted);
+		}
 
 		checkf(ImportedPackage, TEXT("Failed to find or insert imported package with id '0x%llX'"), ImportedPackageId.Value());
 		TRACE_LOADTIME_ASYNC_PACKAGE_IMPORT_DEPENDENCY(this, ImportedPackage);
