Change virtualized bulkdata to store payloads with oodle compression rather than Zlib/LZ4.

#rnx

* VirtualizedBulkData
- The compression option has been removed from ::UpdatePayload. If you wish to deviate from the default compression settings then call ::SetCompressionOptions
- We provide two overloads of ::SetCompressionOptions
-- The first version takes a high level option of type 'ECompressionOptions' that provides an easy way to disable compression or reset to default settings.
-- The second version allows the caller to chose the type of Oodle compression to be used and the level at which it should be applied.
- As compression settings are no longer reset as part of the call to ::UpdatePayload you must now consider if you need to make a manual call to ::SetCompressionOptions after updating a payload to avoid accidently using a previous compression setting.
- The compression settings are still reset when the package is serialized to disk as the compression settings are still considered an attribute of the payload itself, not the bulkdata object.

* Texture
- Updated to use the new API

* How the compression settings are chosen
- When a package is saved to disk, the compression applied to the payload will be determined:
-- If the bulkdata object has compression disabled then we will not compress the payload. If the payload was pulled from persistent storage and is already compressed, then compression will be removed.
-- If compression settings were provided when the payload was last updated, then they will be applied to the payload.
-- If no compression settings have been provided BUT the payload is already compressed (most likely from being pulled from persistent storage) then the pre-existing compression settings will be preserved.
-- If no compression settings have been provided and the payload is uncompressed then we will apply the default compression settings, which is currently Kraken-Fast.
- Note that currently just changing the compression of a payload will not actually cause the payload to be updated in persistent storage since we address payloads by the hash of the uncompressed data. It would be possible to add this functionality at a later date.

diff --git a/Engine/Source/Runtime/CoreUObject/Private/Virtualization/VirtualizedBulkData.cpp b/Engine/Source/Runtime/CoreUObject/Private/Virtualization/VirtualizedBulkData.cpp
--- a/Engine/Source/Runtime/CoreUObject/Private/Virtualization/VirtualizedBulkData.cpp
+++ b/Engine/Source/Runtime/CoreUObject/Private/Virtualization/VirtualizedBulkData.cpp
@@ -2,6 +2,7 @@
 
 #include "Virtualization/VirtualizedBulkData.h"
 
+#include "Compression/OodleDataCompression.h"
 #include "HAL/FileManager.h"
 #include "HAL/IConsoleManager.h"
 #include "Misc/ConfigCacheIni.h"
@@ -79,6 +80,91 @@ IVirtualizationSourceControlUtilities* GetSourceControlInterface()
 	return (IVirtualizationSourceControlUtilities*)IModularFeatures::Get().GetModularFeatureImplementation(FName("VirtualizationSourceControlUtilities"), 0);
 }
 
+namespace Private
+{
+
+FCompressionSettings::FCompressionSettings()
+	: Compressor(FOodleDataCompression::ECompressor::NotSet)
+	, CompressionLevel(FOodleDataCompression::ECompressionLevel::None)
+	, bIsSet(false)
+{
+
+}
+
+FCompressionSettings::FCompressionSettings(const FCompressedBuffer& Buffer)
+{
+	// Note that if the buffer is using a non-oodle format we consider it
+	// as not set.
+	if (!Buffer.TryGetCompressParameters(Compressor, CompressionLevel))
+	{
+		bIsSet = true;
+		Compressor = FOodleDataCompression::ECompressor::NotSet;
+		CompressionLevel = FOodleDataCompression::ECompressionLevel::None;
+	}
+}
+
+bool FCompressionSettings::operator==(const FCompressionSettings& Other) const
+{
+	return	Compressor == Other.Compressor &&
+			CompressionLevel == Other.CompressionLevel &&
+			bIsSet == Other.bIsSet;
+}
+
+bool FCompressionSettings::operator != (const FCompressionSettings& Other) const
+{
+	return !(*this == Other);
+}
+
+void FCompressionSettings::Reset()
+{
+	Compressor = FOodleDataCompression::ECompressor::NotSet;
+	CompressionLevel = FOodleDataCompression::ECompressionLevel::None;
+	bIsSet = false;
+}
+
+void FCompressionSettings::Set(FOodleDataCompression::ECompressor InCompressor, FOodleDataCompression::ECompressionLevel InCompressionLevel)
+{
+	Compressor = InCompressor;
+	CompressionLevel = InCompressionLevel;
+	bIsSet = true;
+}
+
+void FCompressionSettings::SetToDefault()
+{
+	Compressor = FOodleDataCompression::ECompressor::Kraken;
+	CompressionLevel = FOodleDataCompression::ECompressionLevel::Fast;
+	bIsSet = true;
+}
+
+void FCompressionSettings::SetToDisabled()
+{
+	Compressor = FOodleDataCompression::ECompressor::NotSet;
+	CompressionLevel = FOodleDataCompression::ECompressionLevel::None;
+	bIsSet = true;
+}
+
+bool FCompressionSettings::IsSet() const
+{
+	return bIsSet;
+}
+
+bool FCompressionSettings::IsCompressed() const
+{
+	return bIsSet == true && CompressionLevel != FOodleDataCompression::ECompressionLevel::None;
+}
+
+FOodleDataCompression::ECompressor FCompressionSettings::GetCompressor() const
+{
+	return Compressor;
+}
+
+FOodleDataCompression::ECompressionLevel FCompressionSettings::GetCompressionLevel()
+{
+	return CompressionLevel;
+}
+
+} // namespace Private
+
 FVirtualizedUntypedBulkData::FVirtualizedUntypedBulkData(FVirtualizedUntypedBulkData&& Other)
 {
 	*this = MoveTemp(Other);
@@ -94,11 +180,12 @@ FVirtualizedUntypedBulkData& FVirtualizedUntypedBulkData::operator=(FVirtualized
 	PayloadContentId = MoveTemp(Other.PayloadContentId);
 	Payload = MoveTemp(Other.Payload);
 	PayloadSize = MoveTemp(Other.PayloadSize);
-	CompressionFormatToUse = MoveTemp(Other.CompressionFormatToUse);
 	OffsetInFile = MoveTemp(Other.OffsetInFile);
 	PackagePath = MoveTemp(Other.PackagePath);
 	PackageSegment = MoveTemp(Other.PackageSegment);
 	Flags = MoveTemp(Other.Flags);
+	CompressionSettings = MoveTemp(Other.CompressionSettings);
+
 	Other.Reset();
 
 	Register(nullptr);
@@ -140,11 +227,12 @@ FVirtualizedUntypedBulkData& FVirtualizedUntypedBulkData::operator=(const FVirtu
 	PayloadContentId = Other.PayloadContentId;
 	Payload = Other.Payload;
 	PayloadSize = Other.PayloadSize;
-	CompressionFormatToUse = Other.CompressionFormatToUse;
 	OffsetInFile = Other.OffsetInFile;
 	PackagePath = Other.PackagePath;
 	PackageSegment = Other.PackageSegment;
 	Flags = Other.Flags;
+	CompressionSettings = Other.CompressionSettings;
+
 	EnumRemoveFlags(Flags, EFlags::TransientFlags);
 
 	if (bTornOff)
@@ -475,8 +563,10 @@ void FVirtualizedUntypedBulkData::Serialize(FArchive& Ar, UObject* Owner, bool b
 
 							if (CanUnloadData())
 							{
+								this->CompressionSettings.Reset();
 								this->Payload.Reset();
 							}
+
 							// Update our information in the registry
 							// TODO: Pass Owner into Register once the AssetRegistry has been fixed to use the updated PackageGuid from the save
 							Register(nullptr);
@@ -505,6 +595,7 @@ void FVirtualizedUntypedBulkData::Serialize(FArchive& Ar, UObject* Owner, bool b
 
 			if (CanUnloadData())
 			{
+				this->CompressionSettings.Reset();
 				Payload.Reset();
 			}
 		}
@@ -920,11 +1011,12 @@ void FVirtualizedUntypedBulkData::Reset()
 	PayloadContentId.Reset();
 	Payload.Reset();
 	PayloadSize = 0;
-	CompressionFormatToUse = NAME_Default;
 	OffsetInFile = INDEX_NONE;
 	PackagePath.Empty();
 	PackageSegment = EPackageSegment::Header;
 	Flags = EFlags::None;
+
+	CompressionSettings.Reset();
 }
 
 void FVirtualizedUntypedBulkData::UnloadData()
@@ -941,24 +1033,21 @@ FGuid FVirtualizedUntypedBulkData::GetIdentifier() const
 	return BulkDataId;
 }
 
-void FVirtualizedUntypedBulkData::UpdatePayloadImpl(FSharedBuffer&& InPayload, FPayloadId&& InPayloadID, FName InCompressionFormat)
+void FVirtualizedUntypedBulkData::UpdatePayloadImpl(FSharedBuffer&& InPayload, FPayloadId&& InPayloadID)
 {
 	TRACE_CPUPROFILER_EVENT_SCOPE(FVirtualizedUntypedBulkData::UpdatePayloadImpl);
 
 	UnloadData();
 
 	// Make sure that we own the memory in the shared buffer
-	Payload = InPayload.MakeOwned();
+	Payload = MoveTemp(InPayload).MakeOwned();
 	PayloadSize = (int64)Payload.GetSize();
 	PayloadContentId = MoveTemp(InPayloadID);
 
-	EnumRemoveFlags(Flags, EFlags::IsVirtualized |
-		EFlags::DisablePayloadCompression |
-		EFlags::ReferencesLegacyFile |
-		EFlags::LegacyFileIsCompressed |
-		EFlags::LegacyKeyWasGuidDerived);
-
-	SetCompressionFormat(InCompressionFormat);
+	EnumRemoveFlags(Flags,	EFlags::IsVirtualized |
+							EFlags::ReferencesLegacyFile |
+							EFlags::LegacyFileIsCompressed |
+							EFlags::LegacyKeyWasGuidDerived);
 
 	PackagePath.Empty();
 	PackageSegment = EPackageSegment::Header;
@@ -1040,11 +1129,11 @@ TFuture<FCompressedBuffer>FVirtualizedUntypedBulkData::GetCompressedPayload() co
 	return Promise.GetFuture();
 }
 
-void FVirtualizedUntypedBulkData::UpdatePayload(FSharedBuffer InPayload, FName InCompressionFormat)
+void FVirtualizedUntypedBulkData::UpdatePayload(FSharedBuffer InPayload)
 {
 	TRACE_CPUPROFILER_EVENT_SCOPE(FVirtualizedUntypedBulkData::UpdatePayload);
 	FPayloadId NewPayloadId(InPayload);
-	UpdatePayloadImpl(MoveTemp(InPayload), MoveTemp(NewPayloadId), InCompressionFormat);
+	UpdatePayloadImpl(MoveTemp(InPayload), MoveTemp(NewPayloadId));
 }
 
 FVirtualizedUntypedBulkData::FSharedBufferWithID::FSharedBufferWithID(FSharedBuffer InPayload)
@@ -1053,17 +1142,40 @@ FVirtualizedUntypedBulkData::FSharedBufferWithID::FSharedBufferWithID(FSharedBuf
 {
 }
 
-void FVirtualizedUntypedBulkData::UpdatePayload(FSharedBufferWithID InPayload, FName InCompressionFormat)
+void FVirtualizedUntypedBulkData::UpdatePayload(FSharedBufferWithID InPayload)
 {
-	UpdatePayloadImpl(MoveTemp(InPayload.Payload), MoveTemp(InPayload.PayloadId), InCompressionFormat);
+	UpdatePayloadImpl(MoveTemp(InPayload.Payload), MoveTemp(InPayload.PayloadId));
 }
 
-void FVirtualizedUntypedBulkData::SetCompressionFormat(FName InCompressionFormat)
+void FVirtualizedUntypedBulkData::SetCompressionOptions(ECompressionOptions Option)
 {
-	//TODO: Should we validate now or let FCompressedBuffer do that later?
-	CompressionFormatToUse = InCompressionFormat;
+	switch (Option)
+	{
+	case ECompressionOptions::Disabled:
+		CompressionSettings.SetToDisabled();
+		break;
+	case ECompressionOptions::Default:
+		CompressionSettings.Reset();
+		break;
+	default:
+		checkNoEntry();
+	}
 
-	if (InCompressionFormat == NAME_None)
+	if (CompressionSettings.GetCompressionLevel() == FOodleDataCompression::ECompressionLevel::None)
+	{
+		EnumAddFlags(Flags, EFlags::DisablePayloadCompression);
+	}
+	else
+	{
+		EnumRemoveFlags(Flags, EFlags::DisablePayloadCompression);
+	}
+}
+
+void FVirtualizedUntypedBulkData::SetCompressionOptions(FOodleDataCompression::ECompressor Compressor, FOodleDataCompression::ECompressionLevel CompressionLevel)
+{
+	CompressionSettings.Set(Compressor, CompressionLevel);
+
+	if (CompressionSettings.GetCompressionLevel() == FOodleDataCompression::ECompressionLevel::None)
 	{
 		EnumAddFlags(Flags, EFlags::DisablePayloadCompression);
 	}
@@ -1104,26 +1216,39 @@ void FVirtualizedUntypedBulkData::UpdateKeyIfNeeded()
 
 void FVirtualizedUntypedBulkData::RecompressForSerialization(FCompressedBuffer& InOutPayload, EFlags PayloadFlags) const
 {
-	const FName CurrentMethod = InOutPayload.GetFormatName();
+	if (!InOutPayload)
+	{
+		return; // Early out if we do not have a valid payload to operate on
+	}
+
+	Private::FCompressionSettings CurrentSettings(InOutPayload);
+	Private::FCompressionSettings TargetSettings;
 
-	FName TargetMethod;
-	
 	if (EnumHasAnyFlags(PayloadFlags, EFlags::DisablePayloadCompression))
 	{
-		TargetMethod = NAME_None;
+		// If the disable payload compression flag is set, then we should not compress the payload
+		TargetSettings.SetToDisabled(); 
+	}
+	else if (CompressionSettings.IsSet())
+	{
+		// If we have pending compression settings then we can apply them to the payload
+		TargetSettings = CompressionSettings;
 	}
-	else if (CompressionFormatToUse != NAME_Default)
+	else if(!CurrentSettings.IsCompressed()) 
 	{
-		check(!CompressionFormatToUse.IsNone()); // Should be caught by the DisablePayloadCompression flag
-		TargetMethod = CompressionFormatToUse;
+		// If we have no settings to apply to the payload and the payload is currently uncompressed then we
+		// should use the default compression settings.
+		TargetSettings.SetToDefault();
 	}
 	else
 	{
-		// TODO: Do we want to add more logic, min size etc?
-		TargetMethod = NAME_LZ4;
+		// If we have no settings to apply to the payload but the payload is already compressed then we can
+		// just keep the existing settings, what ever they are.
+		TargetSettings = CurrentSettings;
 	}
 	
-	if (TargetMethod != CurrentMethod) // No change in compression so we can just return it
+	// Now we will re-compress the input payload if the current compression settings differ from the desired settings
+	if (TargetSettings != CurrentSettings)
 	{
 		FCompositeBuffer DecompressedBuffer = InOutPayload.DecompressToComposite();
 
@@ -1131,7 +1256,8 @@ void FVirtualizedUntypedBulkData::RecompressForSerialization(FCompressedBuffer&
 		// payload in memory. Compressing it will create a third version so before doing that we should reset
 		// the original compressed buffer in case that we can release it to reduce higher water mark pressure.
 		InOutPayload.Reset();
-		InOutPayload = FCompressedBuffer::Compress(TargetMethod, DecompressedBuffer);
+
+		InOutPayload = FCompressedBuffer::Compress(DecompressedBuffer, TargetSettings.GetCompressor(), TargetSettings.GetCompressionLevel());
 	}
 }
 
diff --git a/Engine/Source/Runtime/CoreUObject/Public/Virtualization/VirtualizedBulkData.h b/Engine/Source/Runtime/CoreUObject/Public/Virtualization/VirtualizedBulkData.h
--- a/Engine/Source/Runtime/CoreUObject/Public/Virtualization/VirtualizedBulkData.h
+++ b/Engine/Source/Runtime/CoreUObject/Public/Virtualization/VirtualizedBulkData.h
@@ -19,7 +19,51 @@ class UObject;
 
 namespace UE::Virtualization
 {
-	
+
+namespace Private
+{
+/** A wrapper around the oodle compression settings used by FVirtualizedUntypedBulkData. */
+struct FCompressionSettings
+{
+	COREUOBJECT_API FCompressionSettings();
+	FCompressionSettings(const FCompressedBuffer& Buffer);
+
+	[[nodiscard]] bool operator ==(const FCompressionSettings& Other) const;
+	[[nodiscard]] bool operator != (const FCompressionSettings& Other) const;
+
+	void Reset();
+
+	void Set(FOodleDataCompression::ECompressor InCompressor, FOodleDataCompression::ECompressionLevel InCompressionLevel);
+	void SetToDefault();
+	void SetToDisabled();
+
+	[[nodiscard]] bool IsSet() const;
+	[[nodiscard]] bool IsCompressed() const;
+
+	[[nodiscard]] FOodleDataCompression::ECompressor GetCompressor() const;
+	[[nodiscard]] FOodleDataCompression::ECompressionLevel GetCompressionLevel();
+
+private:
+
+	FOodleDataCompression::ECompressor Compressor;
+	FOodleDataCompression::ECompressionLevel CompressionLevel;
+
+	bool bIsSet;
+};
+} // namespace Private
+
+/** 
+ * A set of higher level compression options that avoid the need to set the specific
+ * oodle options.
+ */
+enum class ECompressionOptions : uint8
+{
+	/** Use default compression settings. */
+	Default,
+	/** Disable compression for the bulkdata entirely. */
+	Disabled,
+};
+
 /**
  * The goal of this class is to provide an editor time version of BulkData that will work with the content
  * virtualization system.
@@ -121,8 +165,9 @@ public:
 	TFuture<FCompressedBuffer> GetCompressedPayload() const;
 
 	/**
-	 * Allows the existing payload to be replaced with a new one.
-	 *
+	 * Replaces the existing payload (if any) with a new one. 
+	 * It is important to consider the ownership model of the payload being passed in to the method.
+	 * 
 	 * To pass in a raw pointer use 'FSharedBuffer::...(Data, Size)' to create a valid FSharedBuffer.
 	 * Use 'FSharedBuffer::MakeView' if you want to retain ownership on the data being passed in, and use
 	 * 'FSharedBuffer::TakeOwnership' if you are okay with the bulkdata object taking over ownership of it.
@@ -130,15 +175,9 @@ public:
 	 * by using 'FSharedBuffer::MakeView') then a clone of the data will be created internally and assigned
 	 * to the bulkdata object.
 	 *
-	 * @param InPayload				The payload to update the bulkdata with
-	 * @param InCompressionFormat	The compression format to use, NAME_None indicates that the
-	 *								payload is already in a compressed format and will not gain from being 
-	 *								compressed again. These payloads will never be compressed. NAME_Default 
-	 *								will apply whichever compression format that the underlying code deems 
-	 *								appropriate. Other specific compression formats may be allowed, see the 
-	 *								documentation of FCompressedBuffer for details.
+	 * @param InPayload	The payload that this bulkdata object should reference. @see FSharedBuffer
 	 */
-	void UpdatePayload(FSharedBuffer InPayload, FName CompressionFormat = NAME_Default);
+	void UpdatePayload(FSharedBuffer InPayload);
 
 	/**
 	 * Utility struct used to compute the Payload ID before calling UpdatePayload
@@ -175,34 +214,29 @@ public:
 	 * Use this override if you want compute PayloadId before updating the bulkdata
 	 *
 	 * @param InPayload				The payload to update the bulkdata with
-	 * @param InCompressionFormat	The compression format to use, NAME_None indicates that the
-	 *								payload is already in a compressed format and will not gain from being 
-	 *								compressed again. These payloads will never be compressed. NAME_Default 
-	 *								will apply whichever compression format that the underlying code deems 
-	 *								appropriate. Other specific compression formats may be allowed, see the 
-	 *								documentation of FCompressedBuffer for details.
 	 */
-	void UpdatePayload(FSharedBufferWithID InPayload, FName CompressionFormat = NAME_Default);
+	void UpdatePayload(FSharedBufferWithID InPayload);
 
 	/** 
-	 * Allows the compression format to be specified that will be applied the next time that the bulkdata
-	 * is saved. For non-virtualized payloads this will occur when the package is next saved. For 
-	 * virtualized payloads this will only ever be applied when it is pushed to the local cache. If the
-	 * payload has already been pushed to long term storage backends then the compression is not likely to
-	 * be changed.
+	 * Sets the compression options to be applied to the payload during serialization.
 	 * 
-	 * This method is only exposed as part of the public api so that a long standing bug with FTextureSource
-	 * where older textures may have been saved with the wrong compression setting that needs to be fixed. 
-	 * There shouldn't generally be a need to call this.
+	 * These settings will continue to be used until the bulkdata object is reset, a subsequent
+	 * call to ::SetCompressionOptions is made or the owning package is serialized to disk.
 	 * 
-	 * @param InCompressionFormat	The compression format to use, NAME_None indicates that the
-	 *								payload is already in a compressed format and will not gain from being
-	 *								compressed again. These payloads will never be compressed. NAME_Default
-	 *								will apply whichever compression format that the underlying code deems
-	 *								appropriate. Other specific compression formats may be allowed, see the
-	 *								documentation of FCompressedBuffer for details.
+	 * @param Option	The high level option to use. @see UE::Virtualization::ECompressionOptions
+	 */ 
+	void SetCompressionOptions(ECompressionOptions Option);
+
+	/** 
+	 * Sets the compression options to be applied to the payload during serialization.
+	 * 
+	 * These settings will continue to be used until the bulkdata object is reset, a subsequent
+	 * call to ::SetCompressionOptions is made or the owning package is serialized to disk.
+	 * 
+	 * @param Compressor		The Oodle compressor to use. @see FOodleDataCompression::ECompressor
+	 * @param CompressionLevel	The Oodle compression level to use. @see FOodleDataCompression::ECompressionLevel
 	 */
-	void SetCompressionFormat(FName InCompressionFormat);
+	void SetCompressionOptions(FOodleDataCompression::ECompressor Compressor, FOodleDataCompression::ECompressionLevel CompressionLevel);
 
 	/**
 	* Get the CustomVersions used in the file containing the payload. Currently this is assumed
@@ -279,7 +313,7 @@ private:
 
 	FRIEND_ENUM_CLASS_FLAGS(EFlags);
 
-	void UpdatePayloadImpl(FSharedBuffer&& InPayload, FPayloadId&& InPayloadID, FName CompressionFormat = NAME_Default);
+	void UpdatePayloadImpl(FSharedBuffer&& InPayload, FPayloadId&& InPayloadID);
 
 	FCompressedBuffer GetDataInternal() const;
 
@@ -323,9 +357,6 @@ private:
 
 	//---- The remaining members are used when the payload is not virtualized.
 	
-	/** The compression algorithm to use when saving the member 'Payload' */
-	FName CompressionFormatToUse = NAME_Default;
-
 	/** Offset of the payload in the file that contains it (INDEX_NONE if the payload does not come from a file)*/
 	int64 OffsetInFile = INDEX_NONE;
 
@@ -337,6 +368,13 @@ private:
 
 	/** A 32bit bitfield of flags */
 	EFlags Flags = EFlags::None;
+
+	/** 
+	 * Compression settings to be applied to the payload when the package is next saved. The settings will be reset if
+	 * the payload is unloaded from memory during serialization (i.e. the payload was virtualized or the package was
+	 * saved to disk.&
+	 */
+	Private::FCompressionSettings CompressionSettings;
 };
 
 ENUM_CLASS_FLAGS(FVirtualizedUntypedBulkData::EFlags);
diff --git a/Engine/Source/Runtime/Engine/Private/Texture.cpp b/Engine/Source/Runtime/Engine/Private/Texture.cpp
--- a/Engine/Source/Runtime/Engine/Private/Texture.cpp
+++ b/Engine/Source/Runtime/Engine/Private/Texture.cpp
@@ -29,6 +29,8 @@
 #include "Interfaces/ITextureFormat.h"
 #include "Interfaces/ITextureFormatModule.h"
 #include "UObject/UE5MainStreamObjectVersion.h"
+#include "Compression/OodleDataCompression.h"
+#include "Engine/TextureCube.h"
 
 #if WITH_EDITOR
 #include "TextureCompiler.h"
@@ -39,8 +41,6 @@
 	#include "EditorFramework/AssetImportData.h"
 #endif
 
-#include "Engine/TextureCube.h"
-
 static TAutoConsoleVariable<int32> CVarVirtualTextures(
 	TEXT("r.VirtualTextures"),
 	0,
@@ -1093,6 +1093,7 @@ void FTextureSource::InitBlocked(const ETextureSourceFormat* InLayerFormats,
 	}
 
 	BulkData.UpdatePayload(Buffer.MoveToShared());
+	BulkData.SetCompressionOptions(UE::Virtualization::ECompressionOptions::Default);
 }
 
 void FTextureSource::InitBlocked(const ETextureSourceFormat* InLayerFormats,
@@ -1104,6 +1105,7 @@ void FTextureSource::InitBlocked(const ETextureSourceFormat* InLayerFormats,
 	InitBlockedImpl(InLayerFormats, InBlocks, InNumLayers, InNumBlocks);
 
 	BulkData.UpdatePayload(MoveTemp(NewData));
+	BulkData.SetCompressionOptions(UE::Virtualization::ECompressionOptions::Default);
 }
 
 void FTextureSource::InitLayered(
@@ -1141,6 +1143,8 @@ void FTextureSource::InitLayered(
 	{
 		BulkData.UpdatePayload(FUniqueBuffer::Alloc(TotalBytes).MoveToShared());
 	}
+
+	BulkData.SetCompressionOptions(UE::Virtualization::ECompressionOptions::Default);
 }
 
 void FTextureSource::InitLayered(
@@ -1162,6 +1166,7 @@ void FTextureSource::InitLayered(
 	);
 
 	BulkData.UpdatePayload(MoveTemp(NewData));
+	BulkData.SetCompressionOptions(UE::Virtualization::ECompressionOptions::Default);
 }
 
 void FTextureSource::Init(
@@ -1243,10 +1248,18 @@ void FTextureSource::InitWithCompressedSourceData(
 	CompressionFormat = NewSourceFormat;
 
 	checkf(LockState == ELockState::None, TEXT("InitWithCompressedSourceData shouldn't be called in-between LockMip/UnlockMip"));
-	
+
+	BulkData.UpdatePayload(FSharedBuffer::Clone(NewData.GetData(), NewData.Num()));
+
 	// Disable the internal bulkdata compression if the source data is already compressed
-	const FName CompressionName = CompressionFormat == TSCF_None ? NAME_Default : NAME_None;
-	BulkData.UpdatePayload(FSharedBuffer::Clone(NewData.GetData(), NewData.Num()), CompressionName);
+	if (CompressionFormat == TSCF_None)
+	{
+		BulkData.SetCompressionOptions(UE::Virtualization::ECompressionOptions::Default);
+	}
+	else
+	{
+		BulkData.SetCompressionOptions(UE::Virtualization::ECompressionOptions::Disabled);
+	}
 }
 
 FTextureSource FTextureSource::CopyTornOff() const
@@ -1292,8 +1305,9 @@ void FTextureSource::Compress()
 			TArray64<uint8> CompressedData = ImageWrapper->GetCompressed(PngQuality);
 			if ( CompressedData.Num() > 0 )
 			{
-				BulkData.UpdatePayload(MakeSharedBufferFromArray(MoveTemp(CompressedData)), NAME_None);
-				
+				BulkData.UpdatePayload(MakeSharedBufferFromArray(MoveTemp(CompressedData)));
+				BulkData.SetCompressionOptions(UE::Virtualization::ECompressionOptions::Disabled);
+
 				bPNGCompressed = true;
 				CompressionFormat = TSCF_PNG;
 			}
@@ -1306,22 +1320,18 @@ void FTextureSource::Compress()
 		CompressionFormat = TSCF_PNG;
 	}
 
-	// There was a brief period when textures could be saved with the wrong compression flag based on their storage format
-	// so we force set it here so that they will be fixed on resave.
-
-	FName Compressor = NAME_None;
-	if ( CompressionFormat == TSCF_None )
+	if (CompressionFormat == TSCF_PNG && bUseOodleOnPNGz0)
 	{
-		//was NAME_Default which is Zlib
-		Compressor = NAME_Default;
-		// @todo use Oodle for all non-PNG and JPEG
-		//Compressor = NAME_Oodle;
+		BulkData.SetCompressionOptions(UE::Virtualization::ECompressionOptions::Default);
 	}
-	else if ( CompressionFormat == TSCF_PNG && bUseOodleOnPNGz0 )
+	else if (CompressionFormat != TSCF_None)
 	{
-		Compressor = NAME_Oodle;
+		// There was a brief period when textures stored with a compressed image format could also have ZLib compression
+		// applied via bulkdata serialization, so we need to call ::SetCompressionFormat directly to make sure that the 
+		// disable compression flag is set internally if the image storage format is compressed.
+		// TODO: We should consider versioning this so that we can remove this code path at some point.
+		BulkData.SetCompressionOptions(FOodleDataCompression::ECompressor::NotSet, FOodleDataCompression::ECompressionLevel::None);
 	}
-	BulkData.SetCompressionFormat(Compressor);
 }
 
 FSharedBuffer FTextureSource::Decompress(IImageWrapperModule* ImageWrapperModule) const
@@ -1415,6 +1425,7 @@ void FTextureSource::UnlockMip(int32 BlockIndex, int32 LayerIndex, int32 MipInde
 			UE_CLOG(CompressionFormat == TSCF_JPEG, LogTexture, Warning, TEXT("Call to FTextureSource::UnlockMip will cause texture source to lose it's jpeg storage format"));
 
 			BulkData.UpdatePayload(LockedMipData.Release());
+			BulkData.SetCompressionOptions(UE::Virtualization::ECompressionOptions::Default);
 
 			bPNGCompressed = false;
 			CompressionFormat = TSCF_None;
