PackageAccessTracking: Add UE_TRACK_REFERENCING_PACKAGE_SCOPED_PLATFORM. Move PackageBuildDependencyTracker into UnrealEd; it's currently just a demo class, and we will start using it soon in UnrealEd.
Improve performance of FPackageBuildDependencyTracker::StaticOnObjectHandleRead by removing the iteration over PackageScopes and instead accumulating that information in the constructor of each PackageScope. The constructor accumulation is cheap so this should be cheaper overall.
#rnx

diff --git a/Engine/Source/Editor/UnrealEd/Private/Commandlets/CookCommandlet.cpp b/Engine/Source/Editor/UnrealEd/Private/Commandlets/CookCommandlet.cpp
--- a/Engine/Source/Editor/UnrealEd/Private/Commandlets/CookCommandlet.cpp
+++ b/Engine/Source/Editor/UnrealEd/Private/Commandlets/CookCommandlet.cpp
@@ -43,6 +43,7 @@
 #include "AssetRegistryModule.h"
 #include "StudioAnalytics.h"
 #include "Cooker/CookProfiling.h"
+#include "Cooker/PackageBuildDependencyTracker.h"
 #include "DerivedDataBuildRemoteExecutor.h"
 
 DEFINE_LOG_CATEGORY_STATIC(LogCookCommandlet, Log, All);
@@ -53,8 +54,6 @@ DEFINE_LOG_CATEGORY_STATIC(LogCookCommandlet, Log, All);
 #include "IAnalyticsProviderET.h"
 #include "AnalyticsET.h"
 
-ASSETREGISTRY_API void DumpBuildDependencyTrackerData();
-
 namespace DetailedCookStats
 {
 	FString CookProject;
@@ -380,7 +379,7 @@ namespace DetailedCookStats
 			}
 		}
 
-		DumpBuildDependencyTrackerData();
+		DumpBuildDependencyTrackerStats();
 	}
 }
 #endif
diff --git a/Engine/Source/Editor/UnrealEd/Private/Cooker/PackageBuildDependencyTracker.cpp b/Engine/Source/Editor/UnrealEd/Private/Cooker/PackageBuildDependencyTracker.cpp
new file mode 100644
--- /dev/null
+++ b/Engine/Source/Editor/UnrealEd/Private/Cooker/PackageBuildDependencyTracker.cpp
@@ -0,0 +1,111 @@
+// Copyright Epic Games, Inc. All Rights Reserved.
+
+#include "PackageBuildDependencyTracker.h"
+#include "HAL/Platform.h"
+#include "Logging/LogMacros.h"
+#include "UObject/Package.h"
+#include "Misc/PackageAccessTrackingOps.h"
+
+#if UE_WITH_PACKAGE_ACCESS_TRACKING
+
+DEFINE_LOG_CATEGORY_STATIC(LogPackageBuildDependencyTracker, Log, All);
+
+FPackageBuildDependencyTracker FPackageBuildDependencyTracker::Singleton;
+
+void FPackageBuildDependencyTracker::DumpStats() const
+{
+	FScopeLock RecordsScopeLock(&RecordsLock);
+	uint64 ReferencingPackageCount = 0;
+	uint64 ReferenceCount = 0;
+	for (const TPair<FName, TSet<FBuildDependencyAccessData>>& PackageAccessRecord : Records)
+	{
+		++ReferencingPackageCount;
+		for (const FBuildDependencyAccessData& AccessedData : PackageAccessRecord.Value)
+		{
+			++ReferenceCount;
+		}
+	}
+	UE_LOG(LogPackageBuildDependencyTracker, Display, TEXT("Package Accesses (%u referencing packages with a total of %u unique accesses)"), ReferencingPackageCount, ReferenceCount);
+
+	constexpr bool bDetailedDump = false;
+	if (bDetailedDump)
+	{
+		UE_LOG(LogPackageBuildDependencyTracker, Display, TEXT("========================================================================="));
+		for (const TPair<FName, TSet<FBuildDependencyAccessData>>& PackageAccessRecord : Records)
+		{
+			UE_LOG(LogPackageBuildDependencyTracker, Display, TEXT("%s:"), *PackageAccessRecord.Key.ToString());
+			for (const FBuildDependencyAccessData& AccessedData : PackageAccessRecord.Value)
+			{
+				UE_LOG(LogPackageBuildDependencyTracker, Display, TEXT("    %s"), *AccessedData.ReferencedPackage.ToString());
+			}
+		}
+	}
+}
+
+TArray<FBuildDependencyAccessData> FPackageBuildDependencyTracker::GetAccessDatas(FName ReferencerPackage) const
+{
+	FScopeLock RecordsScopeLock(&Singleton.RecordsLock);
+	const TSet<FBuildDependencyAccessData>* ReferencerSet = Records.Find(ReferencerPackage);
+	if (!ReferencerSet)
+	{
+		return TArray<FBuildDependencyAccessData>();
+	}
+	return ReferencerSet->Array();
+}
+
+FPackageBuildDependencyTracker::FPackageBuildDependencyTracker()
+{
+	PreviousObjectHandleReadFunction = SetObjectHandleReadCallback(StaticOnObjectHandleRead);
+}
+
+void FPackageBuildDependencyTracker::StaticOnObjectHandleRead(UObject* ReadObject)
+{
+	if (ReadObject)
+	{
+		if (PackageAccessTracking_Private::FTrackedData* AccumulatedScopeData = PackageAccessTracking_Private::FPackageAccessRefScope::GetCurrentThreadAccumulatedData())
+		{
+			if (!AccumulatedScopeData->BuildOpName.IsNone())
+			{
+				FName Referencer = AccumulatedScopeData->PackageName;
+				FName Referenced = ReadObject->GetOutermost()->GetFName();
+				if (Referencer != Referenced)
+				{
+					FBuildDependencyAccessData AccessData{ Referenced, AccumulatedScopeData->TargetPlatform };
+					FScopeLock RecordsScopeLock(&Singleton.RecordsLock);
+					if (Referencer == Singleton.LastReferencer)
+					{
+						if (AccessData != Singleton.LastAccessData)
+						{
+							Singleton.LastAccessData = AccessData;
+							Singleton.LastReferencerSet->Add(AccessData);
+						}
+					}
+					else
+					{
+						Singleton.LastAccessData = AccessData;
+						Singleton.LastReferencer = Referencer;
+						Singleton.LastReferencerSet = &Singleton.Records.FindOrAdd(Referencer);
+						Singleton.LastReferencerSet->Add(AccessData);
+					}
+				}
+			}
+		}
+	}
+
+	if (Singleton.PreviousObjectHandleReadFunction)
+	{
+		Singleton.PreviousObjectHandleReadFunction(ReadObject);
+	}
+}
+
+void DumpBuildDependencyTrackerStats()
+{
+	FPackageBuildDependencyTracker::Get().DumpStats();
+}
+
+#else
+void DumpBuildDependencyTrackerStats()
+{
+}
+
+#endif // UE_WITH_OBJECT_HANDLE_TRACKING
diff --git a/Engine/Source/Editor/UnrealEd/Private/Cooker/PackageBuildDependencyTracker.h b/Engine/Source/Editor/UnrealEd/Private/Cooker/PackageBuildDependencyTracker.h
new file mode 100644
--- /dev/null
+++ b/Engine/Source/Editor/UnrealEd/Private/Cooker/PackageBuildDependencyTracker.h
@@ -0,0 +1,55 @@
+// Copyright Epic Games, Inc. All Rights Reserved.
+
+#pragma once
+
+#include "HAL/CriticalSection.h"
+#include "Misc/PackageAccessTracking.h"
+#include "Templates/UnrealTemplate.h"
+#include "UObject/ObjectHandle.h"
+
+#if UE_WITH_PACKAGE_ACCESS_TRACKING
+
+struct FBuildDependencyAccessData
+{
+	FName ReferencedPackage;
+	const ITargetPlatform* TargetPlatform;
+	friend uint32 GetTypeHash(const FBuildDependencyAccessData& Data)
+	{
+		return HashCombine(GetTypeHash(Data.ReferencedPackage), GetTypeHash(Data.TargetPlatform));
+	}
+	bool operator==(const FBuildDependencyAccessData& Other) const
+	{
+		return ReferencedPackage == Other.ReferencedPackage && TargetPlatform == Other.TargetPlatform;
+	}
+	bool operator!=(const FBuildDependencyAccessData& Other) const
+	{
+		return !(*this == Other); 
+	}
+};
+
+class FPackageBuildDependencyTracker: public FNoncopyable
+{
+public:
+	static FPackageBuildDependencyTracker& Get() { return Singleton; }
+
+	void DumpStats() const;
+	TArray<FBuildDependencyAccessData> GetAccessDatas(FName ReferencerPackage) const;
+
+private:
+	FPackageBuildDependencyTracker();
+
+	/** Track object reference reads */
+	static void StaticOnObjectHandleRead(UObject* ReadObject);
+	ObjectHandleReadFunction* PreviousObjectHandleReadFunction = nullptr;
+
+	mutable FCriticalSection RecordsLock;
+	TMap<FName, TSet<FBuildDependencyAccessData>> Records;
+	FName LastReferencer = NAME_None;
+	FBuildDependencyAccessData LastAccessData{ NAME_None, nullptr };
+	TSet<FBuildDependencyAccessData>* LastReferencerSet = nullptr;
+	static FPackageBuildDependencyTracker Singleton;
+};
+
+#endif // UE_WITH_PACKAGE_ACCESS_TRACKING
+
+void DumpBuildDependencyTrackerStats();
diff --git a/Engine/Source/Runtime/AssetRegistry/Private/PackageBuildDependencyTracker.cpp b/Engine/Source/Runtime/AssetRegistry/Private/PackageBuildDependencyTracker.cpp
deleted file mode 100644
diff --git a/Engine/Source/Runtime/AssetRegistry/Private/PackageBuildDependencyTracker.h b/Engine/Source/Runtime/AssetRegistry/Private/PackageBuildDependencyTracker.h
deleted file mode 100644
diff --git a/Engine/Source/Runtime/CoreUObject/Private/Misc/PackageAccessTracking.cpp b/Engine/Source/Runtime/CoreUObject/Private/Misc/PackageAccessTracking.cpp
--- a/Engine/Source/Runtime/CoreUObject/Private/Misc/PackageAccessTracking.cpp
+++ b/Engine/Source/Runtime/CoreUObject/Private/Misc/PackageAccessTracking.cpp
@@ -1,6 +1,7 @@
 // Copyright Epic Games, Inc. All Rights Reserved.
 
 #include "Misc/PackageAccessTracking.h"
+#include "Misc/PackageAccessTrackingOps.h"
 
 #if UE_WITH_PACKAGE_ACCESS_TRACKING
 thread_local PackageAccessTracking_Private::FPackageAccessRefScope* PackageAccessTracking_Private::FPackageAccessRefScope::CurrentThreadScope = nullptr;
@@ -8,16 +9,72 @@ thread_local PackageAccessTracking_Private::FPackageAccessRefScope* PackageAcces
 namespace PackageAccessTracking_Private
 {
 
-FPackageAccessRefScope::FPackageAccessRefScope(FName InPackageName, FName InOpName)
+bool IsBuildOpName(FName OpName)
+{
+	//If this was a Package Access done in the context of a save or load, we record it as a build dependency.
+	//If we want to capture script package dependencies, we may also need to record accesses under the "PackageAccessTrackingOps::NAME_CreateDefaultObject" operation
+	//which occurs from "UObjectLoadAllCompiledInDefaultProperties" and possibly elsewhere.
+	return (OpName == PackageAccessTrackingOps::NAME_Load)
+		|| (OpName == PackageAccessTrackingOps::NAME_Save)
+		// || (OpName == PackageAccessTrackingOps::NAME_CreateDefaultObject)
+		;
+}
+
+FTrackedData::FTrackedData(FName InPackageName, FName InOpName, const ITargetPlatform* InTargetPlatform)
 	: PackageName(InPackageName)
 	, OpName(InOpName)
+	, BuildOpName(IsBuildOpName(InOpName) ? InOpName : NAME_None)
+	, TargetPlatform(InTargetPlatform)
+{
+}
+
+FTrackedData::FTrackedData(FTrackedData& DirectData, FTrackedData* OuterAccumulatedData)
+{
+	if (OuterAccumulatedData)
+	{
+		// NOTE: The accumulated PackageName is used as the referencer
+		// The referencer is always the PackageName from the latest scope with a PackageName, NOT the PackageName from the latest scope with a BuildOpName. 
+		// Consider a situation where the in the process of loading package A,
+		// we call PostLoad on an object on package B which references an object in package C.  The SearchThreadScope with OpName == Load is for package A,
+		// but we want to record that B depends on C, not that A depends on C.  The InnermostThreadScope will have a package name of B and an OpName of PostLoad
+		// and so we want the package name from the InnermostThreadScope while having searched upwards for a scope with an OpName == Load.
+
+		// Make sure this line matches the PackageName line in the FTrackedData accumulating constructor
+		PackageName = !DirectData.PackageName.IsNone() ? DirectData.PackageName : OuterAccumulatedData->PackageName;
+
+		// The accumulated OpName is the most recent non-empty OpName
+		OpName = !DirectData.OpName.IsNone() ? DirectData.OpName : OuterAccumulatedData->OpName;
+
+		// The accumulated BuildOpName is the most recent non-empty BuildOpName
+		BuildOpName = !DirectData.BuildOpName.IsNone() ? DirectData.BuildOpName : OuterAccumulatedData->BuildOpName;
+
+		// The accumulated TargetPlatform is the most recent non-empty TargetPlatform
+		TargetPlatform = DirectData.TargetPlatform ? DirectData.TargetPlatform : OuterAccumulatedData->TargetPlatform;
+	}
+	else
+	{
+		PackageName = DirectData.PackageName;
+		OpName = DirectData.OpName;
+		BuildOpName = DirectData.BuildOpName;
+		TargetPlatform = DirectData.TargetPlatform;
+	}
+}
+
+FPackageAccessRefScope::FPackageAccessRefScope(FName InPackageName, FName InOpName, const ITargetPlatform* InTargetPlatform)
+	: DirectData(InPackageName, InOpName, InTargetPlatform)
+	, AccumulatedData(DirectData, GetCurrentThreadAccumulatedData())
+	, Outer(CurrentThreadScope)
 {
-	Outer = CurrentThreadScope;
 	CurrentThreadScope = this;
 }
 
 FPackageAccessRefScope::FPackageAccessRefScope(const UPackage* InPackage, FName InOpName)
-	: FPackageAccessRefScope(InPackage->GetFName(), InOpName)
+	: FPackageAccessRefScope(InPackage->GetFName(), InOpName, nullptr)
+{
+}
+
+FPackageAccessRefScope::FPackageAccessRefScope(const ITargetPlatform* InTargetPlatform)
+	: FPackageAccessRefScope(NAME_None, NAME_None, InTargetPlatform)
 {
 }
 
@@ -27,10 +84,33 @@ FPackageAccessRefScope::~FPackageAccessRefScope()
 	CurrentThreadScope = Outer;
 }
 
+void FPackageAccessRefScope::SetPackageName(FName InPackageName)
+{
+	checkf(FPackageName::IsValidLongPackageName(InPackageName.ToString(), true), TEXT("Invalid package name: %s"), *InPackageName.ToString());
+	checkf(CurrentThreadScope == this, TEXT("Invalid SetPackageName while a child scope is already on the stack. We have not yet implemented propagating a change in PackageName to child scopes' AccumulatedData."));
+
+	DirectData.PackageName = InPackageName;
+	if (Outer)
+	{
+		// Make sure this line matches the PackageName line in the FTrackedData accumulating constructor
+		AccumulatedData.PackageName = !DirectData.PackageName.IsNone() ? DirectData.PackageName : Outer->AccumulatedData.PackageName;
+	}
+	else
+	{
+		AccumulatedData.PackageName = DirectData.PackageName;
+	}
+}
+
 FPackageAccessRefScope* FPackageAccessRefScope::GetCurrentThreadScope()
 {
 	return CurrentThreadScope;
 }
 
+FTrackedData* FPackageAccessRefScope::GetCurrentThreadAccumulatedData()
+{
+	return CurrentThreadScope ? &CurrentThreadScope->AccumulatedData : nullptr;
+}
+
+
 } // PackageAccessTracking_Private
 #endif // UE_WITH_OBJECT_HANDLE_TRACKING
diff --git a/Engine/Source/Runtime/CoreUObject/Public/Misc/PackageAccessTracking.h b/Engine/Source/Runtime/CoreUObject/Public/Misc/PackageAccessTracking.h
--- a/Engine/Source/Runtime/CoreUObject/Public/Misc/PackageAccessTracking.h
+++ b/Engine/Source/Runtime/CoreUObject/Public/Misc/PackageAccessTracking.h
@@ -9,13 +9,17 @@
 #include "UObject/ObjectHandle.h"
 #include "UObject/Package.h"
 
+class ITargetPlatform;
+
 #define UE_WITH_PACKAGE_ACCESS_TRACKING UE_WITH_OBJECT_HANDLE_TRACKING
 
 #if UE_WITH_PACKAGE_ACCESS_TRACKING
+#define UE_TRACK_REFERENCING_PACKAGE_SCOPED_PLATFORM(TargetPlatform) PackageAccessTracking_Private::FPackageAccessRefScope ANONYMOUS_VARIABLE(PackageAccessTracker_)(TargetPlatform);
 #define UE_TRACK_REFERENCING_PACKAGE_SCOPED(Package, OpName) PackageAccessTracking_Private::FPackageAccessRefScope ANONYMOUS_VARIABLE(PackageAccessTracker_)(Package, OpName);
 #define UE_TRACK_REFERENCING_PACKAGE_DELAYED_SCOPED(TrackerName, OpName) TOptional<PackageAccessTracking_Private::FPackageAccessRefScope> TrackerName; FName TrackerName##_OpName(OpName);
 #define UE_TRACK_REFERENCING_PACKAGE_DELAYED(TrackerName, Package) if (TrackerName) TrackerName->SetPackageName(Package->GetFName()); else TrackerName.Emplace(Package->GetFName(), TrackerName##_OpName);
 #else
+#define UE_TRACK_REFERENCING_PACKAGE_SCOPED_PLATFORM(TargetPlatform)
 #define UE_TRACK_REFERENCING_PACKAGE_SCOPED(Package, OpName)
 #define UE_TRACK_REFERENCING_PACKAGE_DELAYED_SCOPED(TrackerName, OpName)
 #define UE_TRACK_REFERENCING_PACKAGE_DELAYED(TrackerName, Package)
@@ -26,27 +30,40 @@
 #if UE_WITH_PACKAGE_ACCESS_TRACKING
 namespace PackageAccessTracking_Private
 {
+	struct FTrackedData
+	{
+		/** Standard constructor; sets variables from DirectData passed to a FPackageAccessRefScope. */
+		COREUOBJECT_API FTrackedData(FName PackgeName, FName OpName, const ITargetPlatform* InTargetPlatform);
+		/** Accumulating constructor; sets new AccumulatedData by combining DirectData with Outer's AccumulatedData. */
+		COREUOBJECT_API FTrackedData(FTrackedData& DirectData, FTrackedData* OuterAccumulatedData);
+
+		FName PackageName;
+		FName OpName;
+		FName BuildOpName;
+		const ITargetPlatform* TargetPlatform = nullptr;
+	};
+
 	class FPackageAccessRefScope
 	{
 	public:
-		COREUOBJECT_API FPackageAccessRefScope(FName InPackageName, FName InOpName);
+		COREUOBJECT_API FPackageAccessRefScope(FName InPackageName, FName InOpName, const ITargetPlatform* InTargetPlatform = nullptr);
 		COREUOBJECT_API FPackageAccessRefScope(const UPackage* InPackage, FName InOpName);
+		COREUOBJECT_API FPackageAccessRefScope(const ITargetPlatform* InTargetPlatform);
 
 		COREUOBJECT_API ~FPackageAccessRefScope();
 
-		FORCEINLINE FName GetPackageName() const { return PackageName; }
-		FORCEINLINE void SetPackageName(FName InPackageName)
-		{
-			checkf(FPackageName::IsValidLongPackageName(InPackageName.ToString(), true), TEXT("Invalid package name: %s"), *InPackageName.ToString());
-			PackageName = InPackageName;
-		}
-		FORCEINLINE FName GetOpName() const { return OpName; }
+		COREUOBJECT_API void SetPackageName(FName InPackageName);
+
+		FORCEINLINE FName GetPackageName() const { return DirectData.PackageName; }
+		FORCEINLINE FName GetOpName() const { return DirectData.OpName; }
 		FORCEINLINE FPackageAccessRefScope* GetOuter() const { return Outer; }
+		FORCEINLINE const ITargetPlatform* GetTargetPlatform() const { return DirectData.TargetPlatform; }
 
 		static COREUOBJECT_API FPackageAccessRefScope* GetCurrentThreadScope();
+		static COREUOBJECT_API FTrackedData* GetCurrentThreadAccumulatedData();
 	private:
-		FName PackageName;
-		FName OpName;
+		FTrackedData DirectData;
+		FTrackedData AccumulatedData;
 		FPackageAccessRefScope* Outer = nullptr;
 		static thread_local FPackageAccessRefScope* CurrentThreadScope;
 	};
