Further removal of UE4 references.

diff --git a/Engine/Source/Runtime/CoreUObject/Private/UObject/EnumProperty.cpp b/Engine/Source/Runtime/CoreUObject/Private/UObject/EnumProperty.cpp
--- a/Engine/Source/Runtime/CoreUObject/Private/UObject/EnumProperty.cpp
+++ b/Engine/Source/Runtime/CoreUObject/Private/UObject/EnumProperty.cpp
@@ -15,7 +15,7 @@
 // WARNING: This should always be the last include in any file that needs it (except .generated.h)
 #include "UObject/UndefineUPropertyMacros.h"
 
-namespace UE4EnumProperty_Private
+namespace UEEnumProperty_Private
 {
 	template <typename OldIntType>
 	void ConvertIntToEnumProperty(FStructuredArchive::FSlot Slot, FEnumProperty* EnumProp, FNumericProperty* UnderlyingProp, UEnum* Enum, void* Obj)
@@ -438,31 +438,31 @@ EConvertFromTypeResult FEnumProperty::ConvertFromType(const FPropertyTag& Tag, F
 	}
 	else if (Tag.Type == NAME_Int8Property)
 	{
-		UE4EnumProperty_Private::ConvertIntToEnumProperty<int8>(Slot, this, UnderlyingProp, Enum, ContainerPtrToValuePtr<void>(Data, Tag.ArrayIndex));
+		UEEnumProperty_Private::ConvertIntToEnumProperty<int8>(Slot, this, UnderlyingProp, Enum, ContainerPtrToValuePtr<void>(Data, Tag.ArrayIndex));
 	}
 	else if (Tag.Type == NAME_Int16Property)
 	{
-		UE4EnumProperty_Private::ConvertIntToEnumProperty<int16>(Slot, this, UnderlyingProp, Enum, ContainerPtrToValuePtr<void>(Data, Tag.ArrayIndex));
+		UEEnumProperty_Private::ConvertIntToEnumProperty<int16>(Slot, this, UnderlyingProp, Enum, ContainerPtrToValuePtr<void>(Data, Tag.ArrayIndex));
 	}
 	else if (Tag.Type == NAME_IntProperty)
 	{
-		UE4EnumProperty_Private::ConvertIntToEnumProperty<int32>(Slot, this, UnderlyingProp, Enum, ContainerPtrToValuePtr<void>(Data, Tag.ArrayIndex));
+		UEEnumProperty_Private::ConvertIntToEnumProperty<int32>(Slot, this, UnderlyingProp, Enum, ContainerPtrToValuePtr<void>(Data, Tag.ArrayIndex));
 	}
 	else if (Tag.Type == NAME_Int64Property)
 	{
-		UE4EnumProperty_Private::ConvertIntToEnumProperty<int64>(Slot, this, UnderlyingProp, Enum, ContainerPtrToValuePtr<void>(Data, Tag.ArrayIndex));
+		UEEnumProperty_Private::ConvertIntToEnumProperty<int64>(Slot, this, UnderlyingProp, Enum, ContainerPtrToValuePtr<void>(Data, Tag.ArrayIndex));
 	}
 	else if (Tag.Type == NAME_UInt16Property)
 	{
-		UE4EnumProperty_Private::ConvertIntToEnumProperty<uint16>(Slot, this, UnderlyingProp, Enum, ContainerPtrToValuePtr<void>(Data, Tag.ArrayIndex));
+		UEEnumProperty_Private::ConvertIntToEnumProperty<uint16>(Slot, this, UnderlyingProp, Enum, ContainerPtrToValuePtr<void>(Data, Tag.ArrayIndex));
 	}
 	else if (Tag.Type == NAME_UInt32Property)
 	{
-		UE4EnumProperty_Private::ConvertIntToEnumProperty<uint32>(Slot, this, UnderlyingProp, Enum, ContainerPtrToValuePtr<void>(Data, Tag.ArrayIndex));
+		UEEnumProperty_Private::ConvertIntToEnumProperty<uint32>(Slot, this, UnderlyingProp, Enum, ContainerPtrToValuePtr<void>(Data, Tag.ArrayIndex));
 	}
 	else if (Tag.Type == NAME_UInt64Property)
 	{
-		UE4EnumProperty_Private::ConvertIntToEnumProperty<uint64>(Slot, this, UnderlyingProp, Enum, ContainerPtrToValuePtr<void>(Data, Tag.ArrayIndex));
+		UEEnumProperty_Private::ConvertIntToEnumProperty<uint64>(Slot, this, UnderlyingProp, Enum, ContainerPtrToValuePtr<void>(Data, Tag.ArrayIndex));
 	}
 	else
 	{
diff --git a/Engine/Source/Runtime/CoreUObject/Private/UObject/LinkerDiff.cpp b/Engine/Source/Runtime/CoreUObject/Private/UObject/LinkerDiff.cpp
--- a/Engine/Source/Runtime/CoreUObject/Private/UObject/LinkerDiff.cpp
+++ b/Engine/Source/Runtime/CoreUObject/Private/UObject/LinkerDiff.cpp
@@ -90,8 +90,8 @@ void FLinkerDiff::GenerateSummaryDiff(FLinker* LHSLinker, FLinker* RHSLinker)
 
 	// Package Tag, file version
 	COMPARE_MEMBER(LHSSummary, RHSSummary, Tag);
-	//COMPARE_MEMBER(LHSSummary, RHSSummary, FileVersionUE4);
-	//COMPARE_MEMBER(LHSSummary, RHSSummary, FileVersionLicenseeUE4);
+	//COMPARE_MEMBER(LHSSummary, RHSSummary, FileVersionUE);
+	//COMPARE_MEMBER(LHSSummary, RHSSummary, FileVersionLicenseeUE);
 
 	//CustomVersionContainer
 
diff --git a/Engine/Source/Runtime/CoreUObject/Private/UObject/ObjectBaseUtility.cpp b/Engine/Source/Runtime/CoreUObject/Private/UObject/ObjectBaseUtility.cpp
--- a/Engine/Source/Runtime/CoreUObject/Private/UObject/ObjectBaseUtility.cpp
+++ b/Engine/Source/Runtime/CoreUObject/Private/UObject/ObjectBaseUtility.cpp
@@ -8,9 +8,9 @@
 #include "UObject/LinkerLoad.h"
 
 /**
- * Returns the UE4 version of the linker for this object.
+ * Returns the UE version of the linker for this object.
  *
- * @return	the UE4 version of the engine's package file when this object
+ * @return	the UEversion of the engine's package file when this object
  *			was last saved, or GPackageFileUEVersion (current version) if
  *			this object does not have a linker, which indicates that
  *			a) this object is a native only class, or
diff --git a/Engine/Source/Runtime/CoreUObject/Private/UObject/PackageFileSummary.cpp b/Engine/Source/Runtime/CoreUObject/Private/UObject/PackageFileSummary.cpp
--- a/Engine/Source/Runtime/CoreUObject/Private/UObject/PackageFileSummary.cpp
+++ b/Engine/Source/Runtime/CoreUObject/Private/UObject/PackageFileSummary.cpp
@@ -90,7 +90,7 @@ void operator<<(FStructuredArchive::FSlot Slot, FPackageFileSummary& Sum)
 		* The package file version number when this package was saved.
 		*
 		* Lower 16 bits stores the UE3 engine version
-		* Upper 16 bits stores the UE4/licensee version
+		* Upper 16 bits stores the UE licensee version
 		* For newer packages this is -7
 		*		-2 indicates presence of enum-based custom versions
 		*		-3 indicates guid-based custom versions
diff --git a/Engine/Source/Runtime/CoreUObject/Private/UObject/PropertyPrivate.cpp b/Engine/Source/Runtime/CoreUObject/Private/UObject/PropertyPrivate.cpp
--- a/Engine/Source/Runtime/CoreUObject/Private/UObject/PropertyPrivate.cpp
+++ b/Engine/Source/Runtime/CoreUObject/Private/UObject/PropertyPrivate.cpp
@@ -164,7 +164,7 @@ void UEnumProperty::AddReferencedObjects(UObject* InThis, FReferenceCollector& C
 	Super::AddReferencedObjects(InThis, Collector);
 }
 
-namespace UE4UEnumProperty_Private
+namespace UEEnumProperty_Private
 {
 	struct FEnumPropertyFriend
 	{
@@ -175,8 +175,8 @@ namespace UE4UEnumProperty_Private
 
 IMPLEMENT_CORE_INTRINSIC_CLASS(UEnumProperty, UProperty,
 	{
-		Class->EmitObjectReference(UE4UEnumProperty_Private::FEnumPropertyFriend::EnumOffset, TEXT("Enum"));
-		Class->EmitObjectReference(UE4UEnumProperty_Private::FEnumPropertyFriend::UnderlyingPropOffset, TEXT("UnderlyingProp"));
+		Class->EmitObjectReference(UEEnumProperty_Private::FEnumPropertyFriend::EnumOffset, TEXT("Enum"));
+		Class->EmitObjectReference(UEEnumProperty_Private::FEnumPropertyFriend::UnderlyingPropOffset, TEXT("UnderlyingProp"));
 	}
 );
 
diff --git a/Engine/Source/Runtime/CoreUObject/Public/AssetRegistry/AssetData.h b/Engine/Source/Runtime/CoreUObject/Public/AssetRegistry/AssetData.h
--- a/Engine/Source/Runtime/CoreUObject/Public/AssetRegistry/AssetData.h
+++ b/Engine/Source/Runtime/CoreUObject/Public/AssetRegistry/AssetData.h
@@ -638,7 +638,7 @@ public:
 	/** Total size of this asset on disk */
 	int64 DiskSize;
 
-	/** UE4 file version that the package was saved with */
+	/** UE file version that the package was saved with */
 	int32 FileVersionUE;
 
 	/** Licensee file version that the package was saved with */
diff --git a/Engine/Source/Runtime/CoreUObject/Public/Virtualization/VirtualizationUtilities.h b/Engine/Source/Runtime/CoreUObject/Public/Virtualization/VirtualizationUtilities.h
--- a/Engine/Source/Runtime/CoreUObject/Public/Virtualization/VirtualizationUtilities.h
+++ b/Engine/Source/Runtime/CoreUObject/Public/Virtualization/VirtualizationUtilities.h
@@ -129,7 +129,7 @@ protected:
 };
 
 
-namespace UE4VirtualizedBulkData_Private
+namespace UEVirtualizedBulkData_Private
 {
 	/** Wraps access to the FVirtualizedUntypedBulkData data for FVirtualizedBulkDataReader
 	so that it can be done before the FBufferReaderBase constructor is called */
@@ -168,7 +168,7 @@ namespace UE4VirtualizedBulkData_Private
 	};
 }
 
-class FVirtualizedBulkDataReader : protected UE4VirtualizedBulkData_Private::DataAccessWrapper, public FBufferReaderBase
+class FVirtualizedBulkDataReader : protected UEVirtualizedBulkData_Private::DataAccessWrapper, public FBufferReaderBase
 {
 public:
 	FVirtualizedBulkDataReader(FVirtualizedUntypedBulkData& InBulkData, bool bIsPersistent = false)
@@ -182,7 +182,7 @@ public:
 	/** Returns if the FVirtualizedBulkDataReader has a valid bulkdata payload or not */
 	bool IsValid() const 
 	{
-		return UE4VirtualizedBulkData_Private::DataAccessWrapper::IsValid();
+		return UEVirtualizedBulkData_Private::DataAccessWrapper::IsValid();
 	}
 
 	using FArchive::operator<<; // For visibility of the overloads we don't override
