Fixed a compile error if CPU tracing is compiled out.

diff --git a/Engine/Source/Runtime/CoreUObject/Public/UObject/UObjectBaseUtility.h b/Engine/Source/Runtime/CoreUObject/Public/UObject/UObjectBaseUtility.h
--- a/Engine/Source/Runtime/CoreUObject/Public/UObject/UObjectBaseUtility.h
+++ b/Engine/Source/Runtime/CoreUObject/Public/UObject/UObjectBaseUtility.h
@@ -828,10 +828,12 @@ public:
 				bStarted = true;
 			}
 
+#if CPUPROFILERTRACE_ENABLED
 			if (!bStarted && UE_TRACE_CHANNELEXPR_IS_ENABLED(CpuChannel))
 			{
 				StartTrace(Object->GetFName());
 			}
+#endif
 		}
 
 #if USE_MALLOC_PROFILER
@@ -859,10 +861,12 @@ public:
 				}
 			}
 
+#if CPUPROFILERTRACE_ENABLED
 			if (!bStarted && UE_TRACE_CHANNELEXPR_IS_ENABLED(CpuChannel))
 			{
 				StartTrace(Object->GetFName());
 			}
+#endif
 		}
 
 #if USE_MALLOC_PROFILER
