Fixed a compile error if load-time tracing is compiled out.

diff --git a/Engine/Source/Runtime/CoreUObject/Private/UObject/UObjectGlobals.cpp b/Engine/Source/Runtime/CoreUObject/Private/UObject/UObjectGlobals.cpp
--- a/Engine/Source/Runtime/CoreUObject/Private/UObject/UObjectGlobals.cpp
+++ b/Engine/Source/Runtime/CoreUObject/Private/UObject/UObjectGlobals.cpp
@@ -1108,10 +1108,12 @@ UPackage* LoadPackageInternal(UPackage* InOuter, const FPackagePath& PackagePath
 	DECLARE_SCOPE_CYCLE_COUNTER(TEXT("LoadPackageInternal"), STAT_LoadPackageInternal, STATGROUP_ObjectVerbose);
 
 	FString TracePackageName;
+#if LOADTIMEPROFILERTRACE_ENABLED
 	if (UE_TRACE_CHANNELEXPR_IS_ENABLED(LoadTimeChannel))
 	{
 		TracePackageName = PackagePath.GetPackageNameOrFallback();
 	}
+#endif
 	SCOPED_CUSTOM_LOADTIMER(LoadPackageInternal)
 		ADD_CUSTOM_LOADTIMER_META(LoadPackageInternal, PackageName, *TracePackageName);
 
