AsyncLoading2: Fixed memory leak in package request queue
#rnx

diff --git a/Engine/Source/Runtime/CoreUObject/Private/Serialization/AsyncLoading2.cpp b/Engine/Source/Runtime/CoreUObject/Private/Serialization/AsyncLoading2.cpp
--- a/Engine/Source/Runtime/CoreUObject/Private/Serialization/AsyncLoading2.cpp
+++ b/Engine/Source/Runtime/CoreUObject/Private/Serialization/AsyncLoading2.cpp
@@ -517,12 +517,12 @@ public:
 			LoadPackageRequest.UPackageName = LoadPackageRequest.CustomName;
 		}
 
-		FScopeLock _(&CriticalSection);
-
 		FPackageRequest* Request = Alloc();
 		check(Request->Next == nullptr)
 		*Request = MoveTemp(LoadPackageRequest);
 
+		FScopeLock _(&CriticalSection);
+
 #if WITH_COTF
 		EPackageStoreEntryStatus Status = Request->StoreEntryHandle.Status();
 
@@ -556,7 +556,9 @@ public:
 		FPackageRequest* ReadyRequest = ReadyRequests.Pop();
 		while (ReadyRequest)
 		{
-			if (!Fn(*ReadyRequest))
+			bool bContinue = Fn(*ReadyRequest);
+			Free(ReadyRequest);
+			if (!bContinue)
 			{
 				break;
 			}
@@ -565,42 +567,20 @@ public:
 		}
 	}
 
-	void Trim()
-	{
-		FScopeLock _(&CriticalSection);
-
-		if (!ReadyRequests.Head)
-		{
-			ReadyRequests = FPackageRequestList();
-			FreeRequests = FPackageRequestList();
-			Requests.Empty();
-		}
-	}
-
 private:
 	FPackageRequest* Alloc()
 	{
-		if (FPackageRequest* FreeRequest = FreeRequests.Pop())
-		{
-			return FreeRequest;
-		}
-		else
-		{
-			const int32 Index = Requests.Add();
-			return &Requests[Index];
-		}
+		return new FPackageRequest;
 	}
 
 	void Free(FPackageRequest* Request)
 	{
-		FreeRequests.Add(Request);
+		delete Request;
 	}
 
 	IPackageStore& PackageStore;
 	FCriticalSection CriticalSection;
-	TChunkedArray<FPackageRequest, 1024> Requests;
 	FPackageRequestList ReadyRequests;
-	FPackageRequestList FreeRequests;
 #if WITH_COTF
 	TArray<FPackageRequest*> PendingRequests;
 #endif
@@ -2350,7 +2330,7 @@ public:
 
 	void UpdatePackagePriority(FAsyncPackage2* Package, int32 NewPriority);
 
-	FAsyncPackageDesc2 CreatePackageDesc(FPackageRequest& Request, bool& bIsMissingPackage);
+	FAsyncPackageDesc2 CreatePackageDesc(const FPackageRequest& Request, bool& bIsMissingPackage);
 	FAsyncPackage2* FindOrInsertPackage(FAsyncPackageDesc2& InDesc, bool& bInserted, TUniquePtr<FLoadPackageAsyncDelegate>&& PackageLoadedDelegate = TUniquePtr<FLoadPackageAsyncDelegate>());
 	void QueueMissingPackage(FAsyncPackageDesc2& PackageDesc, TUniquePtr<FLoadPackageAsyncDelegate>&& LoadPackageAsyncDelegate);
 
@@ -6012,7 +5992,7 @@ int32 FAsyncLoadingThread2::LoadPackage(const FPackagePath& InPackagePath, FName
 	return RequestId;
 }
 
-FAsyncPackageDesc2 FAsyncLoadingThread2::CreatePackageDesc(FPackageRequest& Request, bool& bIsMissingPackage)
+FAsyncPackageDesc2 FAsyncLoadingThread2::CreatePackageDesc(const FPackageRequest& Request, bool& bIsMissingPackage)
 {
 	FPackageId PackageIdToLoad = Request.PackageIdToLoad;
 	FPackageStoreEntryHandle StoreEntry = Request.StoreEntryHandle;
