Add an overload of FPathName::IsPackageExtension that takes an EPackageExtension as the parameter.

#rnx

- Matches the overload of ::IsTextPackageExtension that also takes a EPackageExtension parameter.

diff --git a/Engine/Source/Runtime/CoreUObject/Private/Misc/PackageName.cpp b/Engine/Source/Runtime/CoreUObject/Private/Misc/PackageName.cpp
--- a/Engine/Source/Runtime/CoreUObject/Private/Misc/PackageName.cpp
+++ b/Engine/Source/Runtime/CoreUObject/Private/Misc/PackageName.cpp
@@ -1897,6 +1897,11 @@ bool FPackageName::IsPackageExtension( const TCHAR* Ext )
 	return IsAssetPackageExtension(Ext) || IsMapPackageExtension(Ext);
 }
 
+bool FPackageName::IsPackageExtension(EPackageExtension Extension)
+{
+	return Extension == EPackageExtension::Asset || Extension == EPackageExtension::Map;
+}
+
 bool FPackageName::IsAssetPackageExtension(const TCHAR* Ext)
 {
 	FStringView AssetPackageExtension(LexToString(EPackageExtension::Asset));
diff --git a/Engine/Source/Runtime/CoreUObject/Public/Misc/PackageName.h b/Engine/Source/Runtime/CoreUObject/Public/Misc/PackageName.h
--- a/Engine/Source/Runtime/CoreUObject/Public/Misc/PackageName.h
+++ b/Engine/Source/Runtime/CoreUObject/Public/Misc/PackageName.h
@@ -495,6 +495,14 @@ public:
 	 */
 	static bool IsPackageExtension(const TCHAR* Ext);
 
+	/**
+	 * Returns whether the passed in extension is a valid binary package extension.
+	 *
+	 * @param	Extension to test.
+	 * @return	True if Ext is either a binary asset or a binary map extension, otherwise false
+	 */
+	static bool IsPackageExtension(EPackageExtension Extension);
+
 	/**
 	 * Returns whether the passed in extension is a valid binary asset package
 	 * extension. Extensions with and without trailing dots are supported.
