Fix a crash when running PIE or the Editor.Import.Fbx.* tests.

#rnx

* VirtualizedBulkData
- Changes in meant that if a payload was empty we would write an invalid FCompressedBuffer to disk, where as before we would generate an empty FCompressedBuffer instead which would still write out it's header.
- When loading we were trying to load the header, even if it didn't exist causing all subsequent loading to fail.
- Reverted back to the old behaviour, so empty payloads will at least write out a FCompressedBuffer header again.
-- Note that this was only a problem when saving directly to a memory buffer, when we would serialize the payload inline. When saving the package, payloads are saved to the end of the file and we find them by recording the offset into the file, so this mismatch would be taken into account.
-- This means we don't need to do any versioning here as the packages saved to disk while this bug has been active are all still usable. This bug mostly showed up when duplicating objects for PIE.
- The optimized path in ::GetPayload, which just directly returns the payload if it is in memory will no longer run if the payload size is 0. This will stop us from returning a valid 0 sized FSharedBuffer (need to check the behaviour on this)

* VirtualizedBulkDataTests
- Moved the creation of random data to a shared function.
- Updated System.Core.Virtualization.BulkData.Empty to check GetCompressedPayload as well as GetPayload.

diff --git a/Engine/Source/Runtime/CoreUObject/Private/Tests/Virtualization/VirtualizedBulkDataTests.cpp b/Engine/Source/Runtime/CoreUObject/Private/Tests/Virtualization/VirtualizedBulkDataTests.cpp
--- a/Engine/Source/Runtime/CoreUObject/Private/Tests/Virtualization/VirtualizedBulkDataTests.cpp
+++ b/Engine/Source/Runtime/CoreUObject/Private/Tests/Virtualization/VirtualizedBulkDataTests.cpp
@@ -17,6 +17,19 @@ namespace UE::Virtualization
 	
 constexpr const uint32 TestFlags = EAutomationTestFlags::ApplicationContextMask | EAutomationTestFlags::EngineFilter;
 
+/** Creates a buffer full of random data to make it easy to have something to test against. */
+TUniquePtr<uint8[]> CreateRandomData(int64 BufferSize)
+{
+	TUniquePtr<uint8[]> Buffer = MakeUnique<uint8[]>(BufferSize);
+
+	for (int64 Index = 0; Index < BufferSize; ++Index)
+	{
+		Buffer[Index] = (uint8)(FMath::Rand() % 255);
+	}
+
+	return Buffer;
+}
+
 /**
  * This test creates a very basic VirtualizedBulkData object with in memory payload and validates that we are able to retrieve the 
  * payload via both TFuture and callback methods. It then creates copies of the object and makes sure that we can get the payload from
@@ -26,13 +39,7 @@ IMPLEMENT_SIMPLE_AUTOMATION_TEST(FVirtualizationWrapperTestBasic, TEXT("System.C
 bool FVirtualizationWrapperTestBasic::RunTest(const FString& Parameters)
 {
 	const int64 BufferSize = 1024;
-	TUniquePtr<uint8[]> SourceBuffer = MakeUnique<uint8[]>(BufferSize);
-
-	// Write some values to SourceBuffer
-	for (int64 Index = 0; Index < BufferSize; ++Index)
-	{
-		SourceBuffer[Index] = (uint8)(FMath::Rand() % 255);
-	}
+	TUniquePtr<uint8[]> SourceBuffer = CreateRandomData(BufferSize);
 
 	auto ValidateBulkData = [this, &SourceBuffer, BufferSize](const FVirtualizedUntypedBulkData& BulkDataToValidate, const TCHAR* Label)
 		{
@@ -88,6 +95,9 @@ bool FVirtualizationWrapperTestEmpty::RunTest(const FString& Parameters)
 	FSharedBuffer Payload = BulkData.GetPayload().Get();
 	TestTrue(TEXT("The payload from the GetPayload TFuture is null"), Payload.IsNull());
 
+	FCompressedBuffer CompressedPayload = BulkData.GetCompressedPayload().Get();
+	TestTrue(TEXT("The payload from the GetCompressedPayload TFuture is null"), Payload.IsNull());
+
 	return true;
 }
 
@@ -175,15 +185,9 @@ bool FVirtualizationWrapperTestReaderWriter::RunTest(const FString& Parameters)
 {
 	const int64 BufferSize = 1024;
 
-	TUniquePtr<uint8[]> SourceBuffer = MakeUnique<uint8[]>(BufferSize);
+	TUniquePtr<uint8[]> SourceBuffer = CreateRandomData(BufferSize);
 	TUniquePtr<uint8[]> ReplicatedBuffer = MakeUnique<uint8[]>(BufferSize);
 
-	// Write some values to SourceBuffer
-	for (int64 Index = 0; Index < BufferSize; ++Index)
-	{
-		SourceBuffer[Index] = (uint8)(FMath::Rand() % 255);
-	}
-
 	FByteVirtualizedBulkData BulkData;
 
 	// Serialize the SourceBuffer to BulkData 
@@ -212,6 +216,60 @@ bool FVirtualizationWrapperTestReaderWriter::RunTest(const FString& Parameters)
 	return true;
 }
 
+/** 
+ * This test will serialize a number of empty and valid bulkdata objects to a memory buffer and then serialize them back 
+ * as bulkdata objects again. This should help find problems where the data saved and loaded are mismatched.
+ * Note that this does not check the package saving code paths, only direct serialization to buffers.
+ */
+IMPLEMENT_SIMPLE_AUTOMATION_TEST(FVirtualizationWrapperTestSerializationToMemory, TEXT("System.Core.Virtualization.BulkData.SerializationToMemory"), TestFlags)
+bool FVirtualizationWrapperTestSerializationToMemory::RunTest(const FString& Parameters)
+{
+	const int64 BufferSize = 1024;
+
+	TUniquePtr<uint8[]> SourceBuffer = CreateRandomData(BufferSize);
+	
+	FByteVirtualizedBulkData EmptyBulkData;
+	FByteVirtualizedBulkData ValidBulkData;
+	ValidBulkData.UpdatePayload(FSharedBuffer::Clone(SourceBuffer.Get(), BufferSize));
+	
+	TArray<uint8> MemoryBuffer;
+	const bool bIsArPersistent = true;
+
+	// Write out
+	{
+		FMemoryWriter WriterAr(MemoryBuffer, bIsArPersistent);
+
+		ValidBulkData.Serialize(WriterAr, nullptr);
+		EmptyBulkData.Serialize(WriterAr, nullptr);
+		ValidBulkData.Serialize(WriterAr, nullptr);
+		EmptyBulkData.Serialize(WriterAr, nullptr);
+	}
+
+	// Read in
+	{
+		FByteVirtualizedBulkData SerializedBulkData;
+
+		FMemoryReader ReaderAr(MemoryBuffer, bIsArPersistent);
+
+		SerializedBulkData.Serialize(ReaderAr, nullptr);
+		TestEqual(TEXT("Bulkdata identifier should remain the same"), SerializedBulkData.GetIdentifier(), ValidBulkData.GetIdentifier());
+		TestTrue(TEXT("Bulkdata payload should remain the same"), FMemory::Memcmp(SourceBuffer.Get(), SerializedBulkData.GetPayload().Get().GetData(), BufferSize) == 0);
+
+		SerializedBulkData.Serialize(ReaderAr, nullptr);
+		TestEqual(TEXT("Bulkdata identifier should remain the same"), SerializedBulkData.GetIdentifier(), EmptyBulkData.GetIdentifier());
+		TestTrue(TEXT("Bulkdata should not have a payload"), SerializedBulkData.GetPayload().Get().IsNull());
+
+		SerializedBulkData.Serialize(ReaderAr, nullptr);
+		TestEqual(TEXT("Bulkdata identifier should remain the same"), SerializedBulkData.GetIdentifier(), ValidBulkData.GetIdentifier());
+		TestTrue(TEXT("Bulkdata payload should remain the same"), FMemory::Memcmp(SourceBuffer.Get(), SerializedBulkData.GetPayload().Get().GetData(), BufferSize) == 0);
+
+		SerializedBulkData.Serialize(ReaderAr, nullptr);
+		TestEqual(TEXT("Bulkdata identifier should remain the same"), SerializedBulkData.GetIdentifier(), EmptyBulkData.GetIdentifier());
+		TestTrue(TEXT("Bulkdata should not have a payload"), SerializedBulkData.GetPayload().Get().IsNull());
+	}
+	return true;
+}
+
 /**
  * This set of tests validate that the BulkData's identifier works how we expect it too. It should remain unique in all cases except
  * move semantics.
diff --git a/Engine/Source/Runtime/CoreUObject/Private/Virtualization/VirtualizedBulkData.cpp b/Engine/Source/Runtime/CoreUObject/Private/Virtualization/VirtualizedBulkData.cpp
--- a/Engine/Source/Runtime/CoreUObject/Private/Virtualization/VirtualizedBulkData.cpp
+++ b/Engine/Source/Runtime/CoreUObject/Private/Virtualization/VirtualizedBulkData.cpp
@@ -1100,10 +1100,10 @@ FCompressedBuffer FVirtualizedUntypedBulkData::GetDataInternal() const
 TFuture<FSharedBuffer> FVirtualizedUntypedBulkData::GetPayload() const
 {
 	TPromise<FSharedBuffer> Promise;
-
-	if (Payload)
+	
+	// Avoid a unnecessary compression and decompression if we already have the uncompressed payload
+	if (Payload && GetPayloadSize() > 0)
 	{
-		// Avoid a unnecessary compression and decompression if we already have the uncompressed payload
 		Promise.SetValue(Payload);
 	}
 	else
@@ -1216,11 +1216,6 @@ void FVirtualizedUntypedBulkData::UpdateKeyIfNeeded()
 
 void FVirtualizedUntypedBulkData::RecompressForSerialization(FCompressedBuffer& InOutPayload, EFlags PayloadFlags) const
 {
-	if (!InOutPayload)
-	{
-		return; // Early out if we do not have a valid payload to operate on
-	}
-
 	Private::FCompressionSettings CurrentSettings(InOutPayload);
 	Private::FCompressionSettings TargetSettings;
 
