Removed another UE4 reference in codebase comment

diff --git a/Engine/Source/Runtime/CoreUObject/Public/UObject/PackageFileSummary.h b/Engine/Source/Runtime/CoreUObject/Public/UObject/PackageFileSummary.h
--- a/Engine/Source/Runtime/CoreUObject/Public/UObject/PackageFileSummary.h
+++ b/Engine/Source/Runtime/CoreUObject/Public/UObject/PackageFileSummary.h
@@ -51,7 +51,7 @@ struct FPackageFileSummary
 	int32		Tag;
 
 private:
-	/* UE4 file version */
+	/* UE file version */
 	int32		FileVersionUE;
 	/* Licensee file version */
 	int32		FileVersionLicenseeUE;
