Add type comment to declaration of FObjectPtr/TObjectPtr.
Clean-up pass on TODOs that aren't relevant anymore due to changes in priority or are superfluous for general release.

diff --git a/Engine/Source/Runtime/CoreUObject/Private/UObject/ObjectHandle.cpp b/Engine/Source/Runtime/CoreUObject/Private/UObject/ObjectHandle.cpp
--- a/Engine/Source/Runtime/CoreUObject/Private/UObject/ObjectHandle.cpp
+++ b/Engine/Source/Runtime/CoreUObject/Private/UObject/ObjectHandle.cpp
@@ -395,11 +395,6 @@ UObject* ResolveObjectRef(const FObjectRef& ObjectRef, uint32 LoadFlags /*= LOAD
 		return nullptr;
 	}
 
-#if USE_CIRCULAR_DEPENDENCY_LOAD_DEFERRING
-	// @TODO: OBJPTR: Do we need something equivalent for handling circular dependencies during lazy load?
-	//DeferPotentialCircularImport(Index);
-#endif // USE_CIRCULAR_DEPENDENCY_LOAD_DEFERRING
-
 	ResolveObjectRefClass(ObjectRef, LoadFlags);
 
 	UPackage* TargetPackage = FindOrLoadPackage(ObjectRef.PackageName, LoadFlags);
@@ -413,10 +408,6 @@ UObject* ResolveObjectRef(const FObjectRef& ObjectRef, uint32 LoadFlags /*= LOAD
 	FObjectPathId::ResolvedNameContainerType ResolvedNames;
 	ObjectRef.ObjectPath.Resolve(ResolvedNames);
 
-	// @TODO: OBJPTR: This seems like a silly way to find objects in a hierarchy, investigate better options.
-	//		 Normal import resolving walks up the outer hierarchy and each outer can have a resolved
-	//		 result cached avoiding any further walking up the hierarchy. Is that something that
-	//		 could/should be emulated here?
 	UObject* CurrentObject = TargetPackage;
 	for (int32 ObjectPathIndex = 0; ObjectPathIndex < ResolvedNames.Num(); ++ObjectPathIndex)
 	{
diff --git a/Engine/Source/Runtime/CoreUObject/Public/UObject/ObjectHandle.h b/Engine/Source/Runtime/CoreUObject/Public/UObject/ObjectHandle.h
--- a/Engine/Source/Runtime/CoreUObject/Public/UObject/ObjectHandle.h
+++ b/Engine/Source/Runtime/CoreUObject/Public/UObject/ObjectHandle.h
@@ -10,7 +10,7 @@
 #include "UObject/Object.h"
 #include "UObject/ObjectPathId.h"
 
-#define UE_WITH_OBJECT_HANDLE_LATE_RESOLVE WITH_EDITORONLY_DATA // @TODO: OBJPTR: Should be WITH_EDITORONLY_DATA when ready to tackle lazy load/resolve issues
+#define UE_WITH_OBJECT_HANDLE_LATE_RESOLVE WITH_EDITORONLY_DATA
 #define UE_WITH_OBJECT_HANDLE_TRACKING WITH_EDITORONLY_DATA
 
 DECLARE_LOG_CATEGORY_EXTERN(LogObjectHandle, Log, All);
@@ -209,7 +209,6 @@ namespace ObjectHandle_Private
 	extern COREUOBJECT_API ObjectHandleClassResolvedFunction* ObjectHandleClassResolvedCallback;
 	extern COREUOBJECT_API ObjectHandleReferenceResolvedFunction* ObjectHandleReferenceResolvedCallback;
 
-	// @TODO: OBJPTR: Re-assess callsites for these On* functions to see if we have the coverage we want for each event type.
 	inline void OnHandleRead(UObject* Object)
 	{
 		if (ObjectHandleReadFunction* Callback = ObjectHandleReadCallback)
@@ -277,41 +276,8 @@ inline bool operator!=(FObjectHandle LHS, FObjectHandle RHS)
 
 inline uint32 GetTypeHash(FObjectHandle Handle)
 {
-	// @TODO: OBJPTR: Forcing resolve for ObjectPtr hashing.
-	// This is inefficient, and is most impactful to TSet<handle or ObjectPtr> or TMap<handle or ObjectPtr,*>
-	// These cases are often using local/private object references, not cross-package/public object references
-	// so this limits the negative impact of this hashing choice.
-	//
-	// An alternative implementation could use the hashes of the name path of the referenced object, however it 
-	// would have to address the following:
-	// 1) hash of unresolved handle would have to match hash of resolved handle
-	//   1a) even if the resolved handle points to a different object due to UObjectRedirector
-	//   1b) even if the resolved handle poitns to a different object due to an entry in FCoreRedirects
-	// 2) renaming an object with a resolved handle to it should not change the hash of the handle
-	// 3) reparenting an object with a resolved handle to it should not change the hash of the handle
-	// 4) two unresolved handles that resolve to the same destination object (through redirectors) should have the same hash
-
 	checkf(IsObjectHandleResolved(Handle), TEXT("Cannot hash an unresolved handle."));
 	return GetTypeHash(ReadObjectHandlePointerNoCheck(Handle));
-
-	// // Hash happens on the packed object ref to avoid having to resolve for the sake of hashing
-	// if (IsObjectHandleResolved(Handle))
-	// {
-	// 	UObject* Obj = ReadObjectHandlePointerNoCheck(Handle);
-	// 	if (!Obj || !Obj->HasAnyFlags(RF_Public))
-	// 	{
-	// 		// Null or non-public objects can be hashed on their address in memory as they can't
-	// 		// ever be referenced by an unresolved handle and:
-	// 		// 1) We can't rely on the object redirect mechanism to provide us with consistency if they are renamed
-	// 		// 2) It is faster to hash them this way
-	// 		return GetTypeHash(Obj);
-	// 	}
-	// 	return GetTypeHash(MakePackedObjectRef(Obj).EncodedRef);
-	// }
-	// else
-	// {
-	// 	return GetTypeHash(Handle.PointerOrRef);
-	// }
 }
 
 #else
diff --git a/Engine/Source/Runtime/CoreUObject/Public/UObject/ObjectPtr.h b/Engine/Source/Runtime/CoreUObject/Public/UObject/ObjectPtr.h
--- a/Engine/Source/Runtime/CoreUObject/Public/UObject/ObjectPtr.h
+++ b/Engine/Source/Runtime/CoreUObject/Public/UObject/ObjectPtr.h
@@ -39,6 +39,9 @@
 template <typename T>
 struct TObjectPtr;
 
+/**
+ * FObjectPtr is the basic, minimally typed version of TObjectPtr
+ */
 struct FObjectPtr
 {
 public:
@@ -80,9 +83,6 @@ public:
 		return ResolveObjectHandleClass(Handle);
 	}
 
-	// @TODO: OBJPTR: Garbage collector interaction (for Rob's requested feature to make incremental
-	//		 reachability analysis possible) could be hard coded here and not conditionally
-	//		 compiled out.
 	FObjectPtr(FObjectPtr&&) = default;
 	FObjectPtr(const FObjectPtr&) = default;
 	FObjectPtr& operator=(FObjectPtr&&) = default;
@@ -90,9 +90,6 @@ public:
 
 	FObjectPtr& operator=(UObject* Other)
 	{
-		// @TODO: OBJPTR: Garbage collector interaction (for Rob's requested feature to make incremental
-		//		 reachability analysis possible) could be hard coded here and not conditionally
-		//		 compiled out.
 		Handle = MakeObjectHandle(Other);
 		return *this;
 	}
@@ -100,18 +97,12 @@ public:
 	UE_OBJPTR_DEPRECATED(5.0, "Assignment with incomplete type pointer is deprecated.  Please update this code to use MakeObjectPtrUnsafe.")
 	FObjectPtr& operator=(void* IncompleteOther)
 	{
-		// @TODO: OBJPTR: Garbage collector interaction (for Rob's requested feature to make incremental
-		//		 reachability analysis possible) could be hard coded here and not conditionally
-		//		 compiled out.
 		Handle = MakeObjectHandle(reinterpret_cast<UObject*>(IncompleteOther));
 		return *this;
 	}
 
 	FObjectPtr& operator=(TYPE_OF_NULLPTR)
 	{
-		// @TODO: OBJPTR: Garbage collector interaction (for Rob's requested feature to make incremental
-		//		 reachability analysis possible) could be hard coded here and not conditionally
-		//		 compiled out.
 		Handle = MakeObjectHandle(nullptr);
 		return *this;
 	}
@@ -119,17 +110,11 @@ public:
 	FORCEINLINE bool operator==(FObjectPtr Other) const { return (Handle == Other.Handle); }
 	FORCEINLINE bool operator!=(FObjectPtr Other) const { return (Handle != Other.Handle); }
 
-	// @TODO: OBJPTR: ToTObjectPtr will be removed in the future when a proper casting layer is added
 	UE_OBJPTR_DEPRECATED(5.0, "Use of ToTObjectPtr is unsafe and is deprecated.")
 	FORCEINLINE TObjectPtr<UObject>& ToTObjectPtr();
 	UE_OBJPTR_DEPRECATED(5.0, "Use of ToTObjectPtr is unsafe and is deprecated.")
 	FORCEINLINE const TObjectPtr<UObject>& ToTObjectPtr() const;
 
-	// @TODO: OBJPTR: The presence of this conversion operator triggers a bug in MSVC where implicit conversion of the
-	//			derived TObjectPtr type to UObject* results in a compile error saying that the conversion  operator is
-	//			in a private base class.  That shouldn't matter as the derived TObjectPtr type has its own templated
-	//			conversion operator.  This issue does not come up in clang compiled builds.
-	//FORCEINLINE operator UObject*() const { return Get();  }
 	FORCEINLINE UObject* operator->() const { return Get(); }
 	FORCEINLINE UObject& operator*() const { return *Get(); }
 
@@ -172,7 +157,6 @@ protected:
 private:
 	friend FORCEINLINE uint32 GetTypeHash(const FObjectPtr& Object)
 	{
-		// @TODO: OBJPTR: Forcing resolve for ObjectPtr hashing.  See GetTypeHash in ObjectHandle.h.
 		Object.Get();
 		return GetTypeHash(Object.Handle);
 	}
@@ -195,6 +179,18 @@ namespace ObjectPtr_Private
 	FORCEINLINE bool IsObjectPtrEqualToRawPtrOfRelatedType(const TObjectPtr<T>& Ptr, const U* Other);
 };
 
+/**
+ * TObjectPtr is a type of pointer to a UObject that is meant to function as a drop-in replacement for raw pointer
+ * member properties. It is size equivalent to a 64-bit pointer and supports access tracking and optional lazy load
+ * behavior in editor builds. It stores either the address to the referenced object or (in editor builds) an index in
+ * the object handle table that describes a referenced object that hasn't been loaded yet. It is serialized
+ * identically to a raw pointer to a UObject. When resolved, its participation in garbage collection is identical to a
+ * raw pointer to a UObject.
+ *
+ * This is useful for automatic replacement of raw pointers to support advanced cook-time dependency tracking and
+ * editor-time lazy load use cases. See UnrealObjectPtrTool for tooling to automatically replace raw pointer members
+ * with FObjectPtr/TObjectPtr members instead.
+ */
 template <typename T>
 struct TObjectPtr : private FObjectPtr
 {
