Add PackageAccessTracking ops for reset and noaccessaccepted, to debug and fix the too-broad save op.
#rnx

diff --git a/Engine/Source/Editor/UnrealEd/Private/CookOnTheFlyServer.cpp b/Engine/Source/Editor/UnrealEd/Private/CookOnTheFlyServer.cpp
--- a/Engine/Source/Editor/UnrealEd/Private/CookOnTheFlyServer.cpp
+++ b/Engine/Source/Editor/UnrealEd/Private/CookOnTheFlyServer.cpp
@@ -4451,7 +4451,7 @@ void UCookOnTheFlyServer::SaveCookedPackage(UE::Cook::FSaveCookedPackageContext&
 		if (Context.bPlatformSetupSuccessful)
 		{
 			UE_SCOPED_HIERARCHICAL_COOKTIMER(GEditorSavePackage);
-			UE_TRACK_REFERENCING_PACKAGE_SCOPED_PLATFORM(TargetPlatform);
+			UE_TRACK_REFERENCING_PLATFORM_SCOPED(TargetPlatform);
 
 			try
 			{
diff --git a/Engine/Source/Editor/UnrealEd/Private/Cooker/PackageBuildDependencyTracker.cpp b/Engine/Source/Editor/UnrealEd/Private/Cooker/PackageBuildDependencyTracker.cpp
--- a/Engine/Source/Editor/UnrealEd/Private/Cooker/PackageBuildDependencyTracker.cpp
+++ b/Engine/Source/Editor/UnrealEd/Private/Cooker/PackageBuildDependencyTracker.cpp
@@ -70,6 +70,12 @@ void FPackageBuildDependencyTracker::StaticOnObjectHandleRead(UObject* ReadObjec
 				FName Referenced = ReadObject->GetOutermost()->GetFName();
 				if (Referencer != Referenced)
 				{
+					if (AccumulatedScopeData->OpName == PackageAccessTrackingOps::NAME_NoAccessExpected)
+					{
+						UE_LOG(LogPackageBuildDependencyTracker, Warning, TEXT("Object %s is referencing object %s inside of a NAME_NoAccessExpected scope. Programmer should narrow the scope or debug the reference."),
+							*Referencer.ToString(), *Referenced.ToString());
+					}
+
 					FBuildDependencyAccessData AccessData{ Referenced, AccumulatedScopeData->TargetPlatform };
 					FScopeLock RecordsScopeLock(&Singleton.RecordsLock);
 					if (Referencer == Singleton.LastReferencer)
diff --git a/Engine/Source/Runtime/AssetRegistry/Private/AssetRegistry.cpp b/Engine/Source/Runtime/AssetRegistry/Private/AssetRegistry.cpp
--- a/Engine/Source/Runtime/AssetRegistry/Private/AssetRegistry.cpp
+++ b/Engine/Source/Runtime/AssetRegistry/Private/AssetRegistry.cpp
@@ -2530,6 +2530,7 @@ void UAssetRegistryImpl::ScanPathsSynchronousInternal(const TArray<FString>& InD
 	bool bInForceRescan, bool bInIgnoreBlackListScanFilters)
 {
 	TRACE_CPUPROFILER_EVENT_SCOPE(UAssetRegistryImpl::ScanPathsSynchronousInternal);
+	UE_TRACK_REFERENCING_OPNAME_SCOPED(PackageAccessTrackingOps::NAME_ResetContext);
 	const double SearchStartTime = FPlatformTime::Seconds();
 
 	UE::AssetRegistry::Impl::FEventContext EventContext;
diff --git a/Engine/Source/Runtime/CoreUObject/Private/Misc/PackageAccessTracking.cpp b/Engine/Source/Runtime/CoreUObject/Private/Misc/PackageAccessTracking.cpp
--- a/Engine/Source/Runtime/CoreUObject/Private/Misc/PackageAccessTracking.cpp
+++ b/Engine/Source/Runtime/CoreUObject/Private/Misc/PackageAccessTracking.cpp
@@ -45,8 +45,16 @@ FTrackedData::FTrackedData(FTrackedData& DirectData, FTrackedData* OuterAccumula
 		// The accumulated OpName is the most recent non-empty OpName
 		OpName = !DirectData.OpName.IsNone() ? DirectData.OpName : OuterAccumulatedData->OpName;
 
-		// The accumulated BuildOpName is the most recent non-empty BuildOpName
-		BuildOpName = !DirectData.BuildOpName.IsNone() ? DirectData.BuildOpName : OuterAccumulatedData->BuildOpName;
+		// The accumulated BuildOpName is by default the most recent non-empty BuildOpName
+		// ResetContext ops set it back to none
+		if (OpName == PackageAccessTrackingOps::NAME_ResetContext)
+		{
+			BuildOpName = NAME_None;
+		}
+		else
+		{
+			BuildOpName = !DirectData.BuildOpName.IsNone() ? DirectData.BuildOpName : OuterAccumulatedData->BuildOpName;
+		}
 
 		// The accumulated TargetPlatform is the most recent non-empty TargetPlatform
 		TargetPlatform = DirectData.TargetPlatform ? DirectData.TargetPlatform : OuterAccumulatedData->TargetPlatform;
@@ -73,6 +81,11 @@ FPackageAccessRefScope::FPackageAccessRefScope(const UPackage* InPackage, FName
 {
 }
 
+FPackageAccessRefScope::FPackageAccessRefScope(FName InOpName)
+	: FPackageAccessRefScope(NAME_None, InOpName, nullptr)
+{
+}
+
 FPackageAccessRefScope::FPackageAccessRefScope(const ITargetPlatform* InTargetPlatform)
 	: FPackageAccessRefScope(NAME_None, NAME_None, InTargetPlatform)
 {
diff --git a/Engine/Source/Runtime/CoreUObject/Private/Misc/PackageAccessTrackingOps.cpp b/Engine/Source/Runtime/CoreUObject/Private/Misc/PackageAccessTrackingOps.cpp
--- a/Engine/Source/Runtime/CoreUObject/Private/Misc/PackageAccessTrackingOps.cpp
+++ b/Engine/Source/Runtime/CoreUObject/Private/Misc/PackageAccessTrackingOps.cpp
@@ -8,4 +8,7 @@ FName PackageAccessTrackingOps::NAME_PreLoad(TEXT("PreLoad"));
 FName PackageAccessTrackingOps::NAME_PostLoad(TEXT("PostLoad"));
 FName PackageAccessTrackingOps::NAME_Save(TEXT("Save"));
 FName PackageAccessTrackingOps::NAME_CreateDefaultObject(TEXT("CreateDefaultObject"));
+FName PackageAccessTrackingOps::NAME_ResetContext(TEXT("ResetContext"));
+FName PackageAccessTrackingOps::NAME_NoAccessExpected(TEXT("NoAccessExpected"));
+
 #endif // UE_WITH_OBJECT_HANDLE_TRACKING
diff --git a/Engine/Source/Runtime/CoreUObject/Public/Misc/PackageAccessTracking.h b/Engine/Source/Runtime/CoreUObject/Public/Misc/PackageAccessTracking.h
--- a/Engine/Source/Runtime/CoreUObject/Public/Misc/PackageAccessTracking.h
+++ b/Engine/Source/Runtime/CoreUObject/Public/Misc/PackageAccessTracking.h
@@ -14,15 +14,17 @@ class ITargetPlatform;
 #define UE_WITH_PACKAGE_ACCESS_TRACKING UE_WITH_OBJECT_HANDLE_TRACKING
 
 #if UE_WITH_PACKAGE_ACCESS_TRACKING
-#define UE_TRACK_REFERENCING_PACKAGE_SCOPED_PLATFORM(TargetPlatform) PackageAccessTracking_Private::FPackageAccessRefScope ANONYMOUS_VARIABLE(PackageAccessTracker_)(TargetPlatform);
 #define UE_TRACK_REFERENCING_PACKAGE_SCOPED(Package, OpName) PackageAccessTracking_Private::FPackageAccessRefScope ANONYMOUS_VARIABLE(PackageAccessTracker_)(Package, OpName);
 #define UE_TRACK_REFERENCING_PACKAGE_DELAYED_SCOPED(TrackerName, OpName) TOptional<PackageAccessTracking_Private::FPackageAccessRefScope> TrackerName; FName TrackerName##_OpName(OpName);
 #define UE_TRACK_REFERENCING_PACKAGE_DELAYED(TrackerName, Package) if (TrackerName) TrackerName->SetPackageName(Package->GetFName()); else TrackerName.Emplace(Package->GetFName(), TrackerName##_OpName);
+#define UE_TRACK_REFERENCING_PLATFORM_SCOPED(TargetPlatform) PackageAccessTracking_Private::FPackageAccessRefScope ANONYMOUS_VARIABLE(PackageAccessTracker_)(TargetPlatform);
+#define UE_TRACK_REFERENCING_OPNAME_SCOPED(OpName) PackageAccessTracking_Private::FPackageAccessRefScope ANONYMOUS_VARIABLE(PackageAccessTracker_)(OpName);
 #else
-#define UE_TRACK_REFERENCING_PACKAGE_SCOPED_PLATFORM(TargetPlatform)
 #define UE_TRACK_REFERENCING_PACKAGE_SCOPED(Package, OpName)
 #define UE_TRACK_REFERENCING_PACKAGE_DELAYED_SCOPED(TrackerName, OpName)
 #define UE_TRACK_REFERENCING_PACKAGE_DELAYED(TrackerName, Package)
+#define UE_TRACK_REFERENCING_PLATFORM_SCOPED(TargetPlatform)
+#define UE_TRACK_REFERENCING_OPNAME_SCOPED(OpName)
 #endif //UE_WITH_PACKAGE_ACCESS_TRACKING
 
 ///////////////////////////////////////////////////////////////////////////////////////////////////////////////////////
@@ -48,6 +50,7 @@ namespace PackageAccessTracking_Private
 	public:
 		COREUOBJECT_API FPackageAccessRefScope(FName InPackageName, FName InOpName, const ITargetPlatform* InTargetPlatform = nullptr);
 		COREUOBJECT_API FPackageAccessRefScope(const UPackage* InPackage, FName InOpName);
+		COREUOBJECT_API FPackageAccessRefScope(FName InOpName);
 		COREUOBJECT_API FPackageAccessRefScope(const ITargetPlatform* InTargetPlatform);
 
 		COREUOBJECT_API ~FPackageAccessRefScope();
diff --git a/Engine/Source/Runtime/CoreUObject/Public/Misc/PackageAccessTrackingOps.h b/Engine/Source/Runtime/CoreUObject/Public/Misc/PackageAccessTrackingOps.h
--- a/Engine/Source/Runtime/CoreUObject/Public/Misc/PackageAccessTrackingOps.h
+++ b/Engine/Source/Runtime/CoreUObject/Public/Misc/PackageAccessTrackingOps.h
@@ -14,6 +14,18 @@ namespace PackageAccessTrackingOps
 	extern COREUOBJECT_API FName NAME_PostLoad;
 	extern COREUOBJECT_API FName NAME_Save;
 	extern COREUOBJECT_API FName NAME_CreateDefaultObject;
+	/**
+	 * A global-scope operation suck as AssetRegistry tick is occurring on the current thread inside of another scope,
+	 * and the references made during this operation should not be associated with the previous scope.
+	 */
+	extern COREUOBJECT_API FName NAME_ResetContext;
+	/**
+	 * A debug-only scope that is used to find global-scope operations. Any accesses made when a NoAccessExpected
+	 * scope is on top of the context stack will log a warning. To resolve the warning either the NoAccessExpected
+	 * scope should be narrowed, or another operation should be added (ResetContext or a repeat of the higher
+	 * level scope).
+	 */
+	extern COREUOBJECT_API FName NAME_NoAccessExpected;
 }
 
 #endif //UE_WITH_PACKAGE_ACCESS_TRACKING
