AsyncLoading2: Fixed issue with IoBuffers being kept alive longer than intended

#rnx

diff --git a/Engine/Source/Runtime/Core/Private/IO/IoDispatcher.cpp b/Engine/Source/Runtime/Core/Private/IO/IoDispatcher.cpp
--- a/Engine/Source/Runtime/Core/Private/IO/IoDispatcher.cpp
+++ b/Engine/Source/Runtime/Core/Private/IO/IoDispatcher.cpp
@@ -1027,6 +1027,16 @@ FIoRequest::UpdatePriority(uint32 NewPriority)
 	Impl->Dispatcher.Reprioritize(Impl);
 }
 
+void
+FIoRequest::Release()
+{
+	if (Impl)
+	{
+		Impl->ReleaseRef();
+		Impl = nullptr;
+	}
+}
+
 FIoChunkId CreateExternalFileChunkId(const FStringView Filename)
 {
 	check(Filename.Len() > 0);
diff --git a/Engine/Source/Runtime/Core/Public/IO/IoDispatcher.h b/Engine/Source/Runtime/Core/Public/IO/IoDispatcher.h
--- a/Engine/Source/Runtime/Core/Public/IO/IoDispatcher.h
+++ b/Engine/Source/Runtime/Core/Public/IO/IoDispatcher.h
@@ -824,6 +824,7 @@ public:
 	CORE_API TIoStatusOr<FIoBuffer>			GetResult();
 	CORE_API void							Cancel();
 	CORE_API void							UpdatePriority(uint32 NewPriority);
+	CORE_API void							Release();
 
 private:
 	FIoRequestImpl* Impl = nullptr;
diff --git a/Engine/Source/Runtime/CoreUObject/Private/Serialization/AsyncLoading2.cpp b/Engine/Source/Runtime/CoreUObject/Private/Serialization/AsyncLoading2.cpp
--- a/Engine/Source/Runtime/CoreUObject/Private/Serialization/AsyncLoading2.cpp
+++ b/Engine/Source/Runtime/CoreUObject/Private/Serialization/AsyncLoading2.cpp
@@ -1939,7 +1939,6 @@ private:
 	TArray<FCompletionCallback, TInlineAllocator<2>> CompletionCallbacks;
 
 	FIoRequest IoRequest;
-	FIoBuffer IoBuffer;
 	const uint8* AllExportDataPtr = nullptr;
 	const uint8* CurrentExportDataPtr = nullptr;
 	uint32 CookedHeaderSize = 0;
@@ -2855,7 +2854,7 @@ void FAsyncLoadingThread2::StartBundleIoRequests()
 		{
 			if (Result.IsOk())
 			{
-				Package->IoBuffer = Result.ConsumeValueOrDie();
+				TRACE_COUNTER_ADD(AsyncLoadingTotalLoaded, Result.ValueOrDie().DataSize());
 			}
 			else
 			{
@@ -2865,7 +2864,6 @@ void FAsyncLoadingThread2::StartBundleIoRequests()
 			}
 			int32 LocalPendingIoRequestsCounter = Package->AsyncLoadingThread.PendingIoRequestsCounter.DecrementExchange() - 1;
 			TRACE_COUNTER_SET(AsyncLoadingPendingIoRequests, LocalPendingIoRequestsCounter);
-			TRACE_COUNTER_ADD(AsyncLoadingTotalLoaded, Package->IoBuffer.DataSize());
 			Package->GetPackageNode(EEventLoadNode2::Package_ProcessSummary).ReleaseBarrier();
 		});
 
@@ -3528,7 +3526,7 @@ EAsyncPackageState::Type FAsyncPackage2::Event_ProcessPackageSummary(FAsyncLoadi
 	{
 		check(Package->ExportBundleEntryIndex == 0);
 
-		const uint8* PackageSummaryData = Package->IoBuffer.Data();
+		const uint8* PackageSummaryData = Package->IoRequest.GetResult().ValueOrDie().Data();
 		const FPackageSummary* PackageSummary = reinterpret_cast<const FPackageSummary*>(PackageSummaryData);
 		{
 			TRACE_CPUPROFILER_EVENT_SCOPE(LoadPackageNameMap);
@@ -3647,7 +3645,8 @@ EAsyncPackageState::Type FAsyncPackage2::Event_ProcessExportBundle(FAsyncLoading
 	if (!Package->bLoadHasFailed)
 	{
 		const FExportBundleHeader* ExportBundle = Package->Data.ExportBundleHeaders + InExportBundleIndex;
-		const uint64 AllExportDataSize = Package->IoBuffer.DataSize() - (Package->AllExportDataPtr - Package->IoBuffer.Data());
+		const FIoBuffer& IoBuffer = Package->IoRequest.GetResult().ValueOrDie();
+		const uint64 AllExportDataSize = IoBuffer.DataSize() - (Package->AllExportDataPtr - IoBuffer.Data());
 		if (Package->ExportBundleEntryIndex == 0)
 		{
 			Package->CurrentExportDataPtr = Package->AllExportDataPtr + ExportBundle->SerialOffset;
@@ -3701,7 +3700,7 @@ EAsyncPackageState::Type FAsyncPackage2::Event_ProcessExportBundle(FAsyncLoading
 				const uint64 CookedSerialSize = ExportMapEntry.CookedSerialSize;
 				UObject* Object = Export.Object;
 
-				check(Package->CurrentExportDataPtr + CookedSerialSize <= Package->IoBuffer.Data() + Package->IoBuffer.DataSize());
+				check(Package->CurrentExportDataPtr + CookedSerialSize <= IoBuffer.Data() + IoBuffer.DataSize());
 				check(Object || Export.bFiltered || Export.bExportLoadFailed);
 
 				Ar.ExportBufferBegin(Object, ExportMapEntry.CookedSerialOffset, ExportMapEntry.CookedSerialSize);
@@ -3744,7 +3743,7 @@ EAsyncPackageState::Type FAsyncPackage2::Event_ProcessExportBundle(FAsyncLoading
 	{
 		Package->ProcessedExportBundlesCount = 0;
 		Package->ImportStore.ImportMap = TArrayView<const FPackageObjectIndex>();
-		Package->IoBuffer = FIoBuffer();
+		Package->IoRequest.Release();
 
 		if (Package->ExternalReadDependencies.Num() == 0)
 		{
