Fixed the function TryConvertFilenameToLongPackageName
converting a file to a package path that can't exist for that for that file.

This was causing an issue were the asset registry would think that a folder was deleted when in fact it was a file that was deleted.

Ex:
"../../../EngineTest/Content/NewFolder/.FileWithExtensionButNoName" would convert to  ""../../../EngineTest/Content/NewFolder/" when it should fail.

diff --git a/Engine/Source/Runtime/CoreUObject/Private/Misc/PackageName.cpp b/Engine/Source/Runtime/CoreUObject/Private/Misc/PackageName.cpp
--- a/Engine/Source/Runtime/CoreUObject/Private/Misc/PackageName.cpp
+++ b/Engine/Source/Runtime/CoreUObject/Private/Misc/PackageName.cpp
@@ -432,18 +432,25 @@ void FPackageName::InternalFilenameToLongPackageName(FStringView InFilename, FSt
 				Filename += TEXT("/");
 			}
 		}
-		Result = FPathViews::GetBaseFilenameWithPath(Filename);
 	}
-	else
+	
+	FStringView SplitOutPath;
+	FStringView SplitOutName;
+	FStringView SplitOutExtension;
+	FPathViews::Split(Filename, SplitOutPath, SplitOutName, SplitOutExtension);
+
+	// Do not strip the extension from paths with an empty pathname; if we do the caller can't tell the difference
+	// between /Path/.ext and /Path/
+	if (!SplitOutName.IsEmpty() || SplitOutExtension.IsEmpty())
 	{
-		Result = FPathViews::GetBaseFilenameWithPath(Filename);
-		if (Result.Len() != Filename.Len())
-		{
-			UE_LOG(LogPackageName, Warning, TEXT("TryConvertFilenameToLongPackageName was passed an ObjectPath (%.*s) rather than a PackageName or FilePath; it will be converted to the PackageName. ")
-				TEXT("Accepting ObjectPaths is deprecated behavior and will be removed in a future release; TryConvertFilenameToLongPackageName will fail on ObjectPaths."), InFilename.Len(), InFilename.GetData());
-		}
+		Result = FStringView(SplitOutPath.GetData(), SplitOutName.GetData() + SplitOutName.Len() - SplitOutPath.GetData());
 	}
 
+	if (bIsValidLongPackageName && Result.Len() != Filename.Len())
+	{
+		UE_LOG(LogPackageName, Warning, TEXT("TryConvertFilenameToLongPackageName was passed an ObjectPath (%.*s) rather than a PackageName or FilePath; it will be converted to the PackageName. ")
+			TEXT("Accepting ObjectPaths is deprecated behavior and will be removed in a future release; TryConvertFilenameToLongPackageName will fail on ObjectPaths."), InFilename.Len(), InFilename.GetData());
+	}
 
 	{
 		FReadScopeLock ScopeLock(ContentMountPointCriticalSection);
@@ -491,6 +498,16 @@ bool FPackageName::TryConvertFilenameToLongPackageName(const FString& InFilename
 	InternalFilenameToLongPackageName(InFilename, LongPackageNameBuilder);
 	FStringView LongPackageName = LongPackageNameBuilder.ToString();
 
+	if (LongPackageName.IsEmpty())
+	{
+		if (OutFailureReason)
+		{
+			FStringView FilenameWithoutExtension = FPathViews::GetBaseFilenameWithPath(InFilename);
+			*OutFailureReason = FString::Printf(TEXT("FilenameToLongPackageName failed to convert '%s'. The Result would be indistinguishable from using '%.*s' as the InFilename."), *InFilename, FilenameWithoutExtension.Len(), FilenameWithoutExtension.GetData());
+		}
+		return false;
+	}
+
 	// we don't support loading packages from outside of well defined places
 	int32 CharacterIndex;
 	const bool bContainsDot = LongPackageName.FindChar(TEXT('.'), CharacterIndex);
