AsyncLoading2:
Remove limit on pending package size since request buffers are no longer allocated up front.
Rely on the iodispatcher to sort requests instead of using the global load order for each package.
No longer store export bundle size and load order in the package store.
#rnx

diff --git a/Engine/Source/Developer/IoStoreUtilities/Private/IoStoreUtilities.cpp b/Engine/Source/Developer/IoStoreUtilities/Private/IoStoreUtilities.cpp
--- a/Engine/Source/Developer/IoStoreUtilities/Private/IoStoreUtilities.cpp
+++ b/Engine/Source/Developer/IoStoreUtilities/Private/IoStoreUtilities.cpp
@@ -407,10 +407,8 @@ public:
 			
 			Entry.ExportInfo = FPackageStoreExportInfo
 			{
-				FilePackageEntry.ExportBundlesSize,
 				FilePackageEntry.ExportCount,
-				FilePackageEntry.ExportBundleCount,
-				FilePackageEntry.LoadOrder
+				FilePackageEntry.ExportBundleCount
 			};
 			Entry.ImportedPackageIds = MakeArrayView<const FPackageId>(FilePackageEntry.ImportedPackages.Data(), FilePackageEntry.ImportedPackages.Num());
 
@@ -2855,8 +2853,6 @@ int32 Describe(
 	{
 		FPackageId PackageId;
 		FName PackageName;
-		uint64 Size = 0;
-		uint32 LoadOrder = uint32(-1);
 		uint32 PackageFlags = 0;
 		int32 NameCount = -1;
 		int32 ExportBundleCount = -1;
@@ -3012,10 +3008,8 @@ int32 Describe(
 				const FPackageId& PackageId = ContainerHeader.PackageIds[PackageIndex++];
 				FPackageDesc* PackageDesc = new FPackageDesc();
 				PackageDesc->PackageId = PackageId;
-				PackageDesc->Size = ContainerEntry.ExportBundlesSize;
 				PackageDesc->Exports.SetNum(ContainerEntry.ExportCount);
 				PackageDesc->ExportBundleCount = ContainerEntry.ExportBundleCount;
-				PackageDesc->LoadOrder = ContainerEntry.LoadOrder;
 				PackageDesc->ImportedPackageIds = TArrayView<FPackageId>(ContainerEntry.ImportedPackages.Data(), ContainerEntry.ImportedPackages.Num());
 				Job.Packages.Add(PackageDesc);
 				++TotalPackageCount;
@@ -3343,10 +3337,6 @@ int32 Describe(
 				}
 			}
 		}
-		Algo::Sort(OutputPackages, [](const FPackageDesc* A, const FPackageDesc* B)
-			{
-				return A->LoadOrder < B->LoadOrder;
-			});
 	}
 
 	UE_LOG(LogIoStore, Display, TEXT("Generating report..."));
@@ -3447,8 +3437,6 @@ int32 Describe(
 			OutputOverride->Logf(ELogVerbosity::Display, TEXT("--------------------------------------------"));
 
 			OutputOverride->Logf(ELogVerbosity::Display, TEXT("\t\t        PackageId: 0x%llX"), PackageDesc->PackageId.ValueForDebugging());
-			OutputOverride->Logf(ELogVerbosity::Display, TEXT("\t\t             Size: %lld"), PackageDesc->Size);
-			OutputOverride->Logf(ELogVerbosity::Display, TEXT("\t\t        LoadOrder: %d"), PackageDesc->LoadOrder);
 			OutputOverride->Logf(ELogVerbosity::Display, TEXT("\t\t     PackageFlags: %X"), PackageDesc->PackageFlags);
 			OutputOverride->Logf(ELogVerbosity::Display, TEXT("\t\t        NameCount: %d"), PackageDesc->NameCount);
 			OutputOverride->Logf(ELogVerbosity::Display, TEXT("\t\t      ImportCount: %d"), PackageDesc->Imports.Num());
@@ -4046,9 +4034,7 @@ int32 Staged2Zen(const FString& BuildPath, const FKeyChain& KeyChain, const ITar
 				FPackageStoreEntryResource& PackageStoreEntryResource = FindPackageInfo->PackageStoreEntry;
 				PackageStoreEntryResource.PackageName = *FindPackageName;
 				PackageStoreEntryResource.ExportInfo.ExportBundleCount = StoreEntry->ExportBundleCount;
-				PackageStoreEntryResource.ExportInfo.ExportBundlesSize = StoreEntry->ExportBundlesSize;
 				PackageStoreEntryResource.ExportInfo.ExportCount = StoreEntry->ExportCount;
-				PackageStoreEntryResource.ExportInfo.LoadOrder = StoreEntry->LoadOrder;
 				PackageStoreEntryResource.ImportedPackageIds.SetNum(StoreEntry->ImportedPackages.Num());
 				FMemory::Memcpy(PackageStoreEntryResource.ImportedPackageIds.GetData(), StoreEntry->ImportedPackages.Data(), sizeof(FPackageId) * StoreEntry->ImportedPackages.Num());
 			}
diff --git a/Engine/Source/Developer/IoStoreUtilities/Private/PackageStoreOptimizer.cpp b/Engine/Source/Developer/IoStoreUtilities/Private/PackageStoreOptimizer.cpp
--- a/Engine/Source/Developer/IoStoreUtilities/Private/PackageStoreOptimizer.cpp
+++ b/Engine/Source/Developer/IoStoreUtilities/Private/PackageStoreOptimizer.cpp
@@ -140,7 +140,6 @@ FPackageStorePackage* FPackageStoreOptimizer::CreatePackageFromPackageStoreHeade
 	FPackageStorePackage* Package = new FPackageStorePackage();
 	Package->Name = Name;
 	Package->Id = FPackageId::FromName(Name);
-	Package->LoadOrder = PackageStoreEntry.ExportInfo.LoadOrder;
 	
 	Package->SourceName = *RemapLocalizationPathIfNeeded(Name.ToString(), &Package->Region);
 
@@ -1688,10 +1687,8 @@ FPackageStoreEntryResource FPackageStoreOptimizer::CreatePackageStoreEntry(const
 	Result.PackageName = Package->Name;
 	Result.SourcePackageName = Package->SourceName;
 	Result.Region = FName(*Package->Region);
-	Result.ExportInfo.ExportBundlesSize = Package->HeaderSize + Package->ExportsSerialSize;
 	Result.ExportInfo.ExportCount = Package->Exports.Num();
 	Result.ExportInfo.ExportBundleCount = Package->GraphData.ExportBundles.Num();
-	Result.ExportInfo.LoadOrder = Package->GetLoadOrder();
 	Result.ImportedPackageIds = Package->ImportedPackageIds;
 	Result.ShaderMapHashes = Package->ShaderMapHashes.Array();
 	return Result;
@@ -1750,10 +1747,8 @@ FContainerHeader FPackageStoreOptimizer::CreateContainerHeader(const FIoContaine
 		
 		// StoreEntries
 		FPackageStoreExportInfo ExportInfo = Entry->ExportInfo;
-		int32 Pad = 0;
 		StoreTocArchive << ExportInfo; 
-		StoreTocArchive << Pad;
-
+		
 		// ImportedPackages
 		const TArray<FPackageId>& ImportedPackageIds = Entry->ImportedPackageIds;
 		SerializePackageEntryCArrayHeader(ImportedPackageIds.Num());
diff --git a/Engine/Source/Runtime/Core/Private/IO/PackageStore.cpp b/Engine/Source/Runtime/Core/Private/IO/PackageStore.cpp
--- a/Engine/Source/Runtime/Core/Private/IO/PackageStore.cpp
+++ b/Engine/Source/Runtime/Core/Private/IO/PackageStore.cpp
@@ -7,10 +7,8 @@
 
 FArchive& operator<<(FArchive& Ar, FPackageStoreExportInfo& ExportInfo)
 {
-	Ar << ExportInfo.ExportBundlesSize;
 	Ar << ExportInfo.ExportCount;
 	Ar << ExportInfo.ExportBundleCount;
-	Ar << ExportInfo.LoadOrder;
 
 	return Ar;
 }
@@ -18,10 +16,8 @@ FArchive& operator<<(FArchive& Ar, FPackageStoreExportInfo& ExportInfo)
 FCbWriter& operator<<(FCbWriter& Writer, const FPackageStoreExportInfo& ExportInfo)
 {
 	Writer.BeginObject();
-	Writer << "exportbundlessize" << ExportInfo.ExportBundlesSize;
 	Writer << "exportcount" << ExportInfo.ExportCount;
 	Writer << "exportbundlecount" << ExportInfo.ExportBundleCount;
-	Writer << "loadorder" << ExportInfo.LoadOrder;
 	Writer.EndObject();
 	
 	return Writer;
@@ -31,10 +27,8 @@ FPackageStoreExportInfo FPackageStoreExportInfo::FromCbObject(const FCbObject& O
 {
 	FPackageStoreExportInfo ExportInfo;
 	
-	ExportInfo.ExportBundlesSize	= Obj["exportbundlessize"].AsUInt64();
 	ExportInfo.ExportCount			= Obj["exportcount"].AsInt32();
 	ExportInfo.ExportBundleCount	= Obj["exportbundlecount"].AsInt32();
-	ExportInfo.LoadOrder			= Obj["loadorder"].AsUInt32();
 
 	return ExportInfo;
 }
diff --git a/Engine/Source/Runtime/Core/Public/IO/PackageStore.h b/Engine/Source/Runtime/Core/Public/IO/PackageStore.h
--- a/Engine/Source/Runtime/Core/Public/IO/PackageStore.h
+++ b/Engine/Source/Runtime/Core/Public/IO/PackageStore.h
@@ -16,11 +16,9 @@ class FCbWriter;
  */
 struct FPackageStoreExportInfo
 {
-	uint64 ExportBundlesSize = 0;
 	int32 ExportCount = 0;
 	int32 ExportBundleCount = 0;
-	uint32 LoadOrder = 0;
-
+	
 	CORE_API friend FArchive& operator<<(FArchive& Ar, FPackageStoreExportInfo& ExportInfo);
 	
 	CORE_API friend FCbWriter& operator<<(FCbWriter& Writer, const FPackageStoreExportInfo& ExportInfo);
diff --git a/Engine/Source/Runtime/CoreUObject/Private/Serialization/AsyncLoading2.cpp b/Engine/Source/Runtime/CoreUObject/Private/Serialization/AsyncLoading2.cpp
--- a/Engine/Source/Runtime/CoreUObject/Private/Serialization/AsyncLoading2.cpp
+++ b/Engine/Source/Runtime/CoreUObject/Private/Serialization/AsyncLoading2.cpp
@@ -78,20 +78,9 @@
 PRAGMA_DISABLE_OPTIMIZATION
 #endif
 
-TRACE_DECLARE_MEMORY_COUNTER(AsyncLoadingPendingIoRequestsSize, TEXT("AsyncLoading/PendingIoRequestsSize"));
 CSV_DECLARE_CATEGORY_MODULE_EXTERN(CORE_API, FileIO);
-CSV_DEFINE_STAT(FileIO, PendingAsyncLoadingIoRequestsSizeMB);
 CSV_DEFINE_STAT(FileIO, FrameCompletedExportBundleLoadsKB);
 
-int32 GAsyncLoadingMaxPendingRequestsSizeMB = 256;
-static FAutoConsoleVariableRef CVar_AsyncLoadingMaxPendingRequestsSizeMB(
-	TEXT("s.AsyncLoadingMaxPendingRequestsSizeMB"),
-	GAsyncLoadingMaxPendingRequestsSizeMB,
-	TEXT("Max amount of data in MB to request from IO"),
-	ECVF_Default
-);
-
-
 FArchive& operator<<(FArchive& Ar, FContainerHeaderPackageRedirect& Redirect)
 {
 	Ar << Redirect.SourcePackageId;
@@ -1795,8 +1784,8 @@ struct FAsyncPackage2
 	/** Creates GC clusters from loaded objects */
 	EAsyncPackageState::Type CreateClusters(FAsyncLoadingThreadState2& ThreadState);
 
-	void ImportPackagesRecursive();
-	void StartLoading();
+	void ImportPackagesRecursive(FIoBatch& IoBatch);
+	void StartLoading(FIoBatch& IoBatch);
 
 #if WITH_IOSTORE_IN_EDITOR
 	void GetLoadedAssets(TArray<FWeakObjectPtr>& AssetList);
@@ -2204,18 +2193,6 @@ private:
 	/** Initial load pending CDOs */
 	TMap<UClass*, TArray<FEventLoadNode2*>> PendingCDOs;
 
-	struct FBundleIoRequest
-	{
-		bool operator<(const FBundleIoRequest& Other) const
-		{
-			return Package->Data.ExportInfo.LoadOrder < Other.Package->Data.ExportInfo.LoadOrder;
-		}
-
-		FAsyncPackage2* Package;
-	};
-	TArray<FBundleIoRequest> WaitingIoRequests;
-	uint64 PendingBundleIoRequestsTotalSize = 0;
-
 	uint32 ConditionalBeginPostLoadTick = 0;
 	uint32 ConditionalFinishLoadingTick = 0;
 
@@ -2512,9 +2489,6 @@ private:
 	EAsyncPackageState::Type ProcessLoadedPackagesFromGameThread(FAsyncLoadingThreadState2& ThreadState, bool& bDidSomething, int32 FlushRequestID = INDEX_NONE);
 
 	bool CreateAsyncPackagesFromQueue(FAsyncLoadingThreadState2& ThreadState);
-	void AddBundleIoRequest(FAsyncPackage2* Package);
-	void BundleIoRequestCompleted(FAsyncPackage2* Package);
-	void StartBundleIoRequests();
 
 	FAsyncPackage2* CreateAsyncPackage(const FAsyncPackageDesc2& Desc)
 	{
@@ -2751,10 +2725,12 @@ bool FAsyncLoadingThread2::CreateAsyncPackagesFromQueue(FAsyncLoadingThreadState
 	bool bPackagesCreated = false;
 	const int32 TimeSliceGranularity = ThreadState.UseTimeLimit() ? 4 : MAX_int32;
 
+	FIoBatch IoBatch = IoDispatcher.NewBatch();
+
 	for (;;)
 	{
 		int32 NumDequeued = 0;
-		PackageRequestQueue->Dequeue([this, &NumDequeued, TimeSliceGranularity](FPackageRequest& Request)
+		PackageRequestQueue->Dequeue([this, &NumDequeued, TimeSliceGranularity, &IoBatch](FPackageRequest& Request)
 		{
 			bool bIsMissingPackage;
 			FAsyncPackageDesc2 PackageDesc = CreatePackageDesc(Request, bIsMissingPackage);
@@ -2785,15 +2761,13 @@ bool FAsyncLoadingThread2::CreateAsyncPackagesFromQueue(FAsyncLoadingThreadState
 				{
 					{
 						TRACE_CPUPROFILER_EVENT_SCOPE(ImportPackages);
-						Package->ImportPackagesRecursive();
+						Package->ImportPackagesRecursive(IoBatch);
 					}
 
 					if (bInserted)
 					{
-						Package->StartLoading();
+						Package->StartLoading(IoBatch);
 					}
-
-					StartBundleIoRequests();
 				}
 			}
 
@@ -2807,89 +2781,10 @@ bool FAsyncLoadingThread2::CreateAsyncPackagesFromQueue(FAsyncLoadingThreadState
 			break;
 		}
 	}
-	
-	return bPackagesCreated;
-}
-
-void FAsyncLoadingThread2::AddBundleIoRequest(FAsyncPackage2* Package)
-{
-	int32 LocalPendingIoRequestsCounter = PendingIoRequestsCounter.IncrementExchange() + 1;
-	TRACE_COUNTER_SET(AsyncLoadingPendingIoRequests, LocalPendingIoRequestsCounter);
-	WaitingIoRequests.HeapPush({ Package });
-}
-
-void FAsyncLoadingThread2::BundleIoRequestCompleted(FAsyncPackage2* Package)
-{
-	check(PendingBundleIoRequestsTotalSize >= Package->Data.ExportInfo.ExportBundlesSize)
-	PendingBundleIoRequestsTotalSize -= Package->Data.ExportInfo.ExportBundlesSize;
-	TRACE_COUNTER_SET(AsyncLoadingPendingIoRequestsSize, PendingBundleIoRequestsTotalSize);
-	CSV_CUSTOM_STAT_DEFINED(FrameCompletedExportBundleLoadsKB, float((double)Package->Data.ExportInfo.ExportBundlesSize / 1024.0), ECsvCustomStatOp::Accumulate);
-	if (WaitingIoRequests.Num())
-	{
-		StartBundleIoRequests();
-	}
-}
-
-void FAsyncLoadingThread2::StartBundleIoRequests()
-{
-	TRACE_CPUPROFILER_EVENT_SCOPE(StartBundleIoRequests);
-	static const uint64 MaxPendingRequestsSize = IsRunningCookOnTheFly() ? MAX_uint64 : uint64(GAsyncLoadingMaxPendingRequestsSizeMB) << 20;
-	FIoBatch IoBatch = IoDispatcher.NewBatch();
-	while (WaitingIoRequests.Num())
-	{
-		FBundleIoRequest& BundleIoRequest = WaitingIoRequests.HeapTop();
-		FAsyncPackage2* Package = BundleIoRequest.Package;
-		if (PendingBundleIoRequestsTotalSize > 0 && PendingBundleIoRequestsTotalSize + Package->Data.ExportInfo.ExportBundlesSize > MaxPendingRequestsSize)
-		{
-			break;
-		}
-		PendingBundleIoRequestsTotalSize += Package->Data.ExportInfo.ExportBundlesSize;
-		WaitingIoRequests.HeapPop(BundleIoRequest, false);
-
-		FIoReadOptions ReadOptions;
-		Package->IoRequest = IoBatch.ReadWithCallback(CreateIoChunkId(Package->Desc.PackageIdToLoad.Value(), 0, EIoChunkType::ExportBundleData),
-			ReadOptions,
-			Package->Desc.Priority,
-			[Package](TIoStatusOr<FIoBuffer> Result)
-		{
-			if (Result.IsOk())
-			{
-				TRACE_COUNTER_ADD(AsyncLoadingTotalLoaded, Result.ValueOrDie().DataSize());
-			}
-			else
-			{
-				UE_ASYNC_PACKAGE_LOG(Warning, Package->Desc, TEXT("StartBundleIoRequests: FailedRead"),
-					TEXT("Failed reading chunk for package: %s"), *Result.Status().ToString());
-				Package->bLoadHasFailed = true;
-			}
-			int32 LocalPendingIoRequestsCounter = Package->AsyncLoadingThread.PendingIoRequestsCounter.DecrementExchange() - 1;
-			TRACE_COUNTER_SET(AsyncLoadingPendingIoRequests, LocalPendingIoRequestsCounter);
-			Package->GetPackageNode(EEventLoadNode2::Package_ProcessSummary).ReleaseBarrier();
-		});
 
-		if (!Package->Data.ShaderMapHashes.IsEmpty())
-		{
-			TRACE_CPUPROFILER_EVENT_SCOPE(StartShaderMapRequests);
-			auto ReadShaderMapFunc = [Package, &IoBatch](const FIoChunkId& ChunkId, FGraphEventRef GraphEvent)
-			{
-				Package->GetPackageNode(Package_ExportsSerialized).AddBarrier();
-				int32 LocalPendingIoRequestsCounter = Package->AsyncLoadingThread.PendingIoRequestsCounter.IncrementExchange() + 1;
-				TRACE_COUNTER_SET(AsyncLoadingPendingIoRequests, LocalPendingIoRequestsCounter);
-				return IoBatch.ReadWithCallback(ChunkId, FIoReadOptions(), Package->Desc.Priority,
-					[Package, GraphEvent](TIoStatusOr<FIoBuffer> Result)
-					{
-						GraphEvent->DispatchSubsequents();
-						int32 LocalPendingIoRequestsCounter = Package->AsyncLoadingThread.PendingIoRequestsCounter.DecrementExchange() - 1;
-						TRACE_COUNTER_SET(AsyncLoadingPendingIoRequests, LocalPendingIoRequestsCounter);
-						Package->GetPackageNode(Package_ExportsSerialized).ReleaseBarrier();
-					});
-			};
-			FCoreDelegates::PreloadPackageShaderMaps.ExecuteIfBound(Package->Data.ShaderMapHashes, ReadShaderMapFunc);
-		}
-	}
 	IoBatch.Issue();
-
-	TRACE_COUNTER_SET(AsyncLoadingPendingIoRequestsSize, PendingBundleIoRequestsTotalSize);
+	
+	return bPackagesCreated;
 }
 
 FEventLoadNode2::FEventLoadNode2(const FAsyncLoadEventSpec* InSpec, FAsyncPackage2* InPackage, int32 InImportOrExportIndex, int32 InBarrierCount)
@@ -3395,7 +3290,7 @@ void FGlobalImportStore::FindAllScriptObjects()
 		ScriptObjects.Num(), (float)ScriptObjects.GetAllocatedSize() / 1024.f);
 }
 
-void FAsyncPackage2::ImportPackagesRecursive()
+void FAsyncPackage2::ImportPackagesRecursive(FIoBatch& IoBatch)
 {
 	if (AsyncPackageLoadingState >= EAsyncPackageLoadingState2::ImportPackages)
 	{
@@ -3481,8 +3376,8 @@ void FAsyncPackage2::ImportPackagesRecursive()
 
 		if (bInserted)
 		{
-			ImportedPackage->ImportPackagesRecursive();
-			ImportedPackage->StartLoading();
+			ImportedPackage->ImportPackagesRecursive(IoBatch);
+			ImportedPackage->StartLoading(IoBatch);
 		}
 	}
 	UE_ASYNC_PACKAGE_LOG_VERBOSE(VeryVerbose, Desc, TEXT("ImportPackages: ImportsDone"),
@@ -3492,7 +3387,7 @@ void FAsyncPackage2::ImportPackagesRecursive()
 	AsyncPackageLoadingState = EAsyncPackageLoadingState2::ImportPackagesDone;
 }
 
-void FAsyncPackage2::StartLoading()
+void FAsyncPackage2::StartLoading(FIoBatch& IoBatch)
 {
 	TRACE_CPUPROFILER_EVENT_SCOPE(StartLoading);
 	TRACE_LOADTIME_BEGIN_LOAD_ASYNC_PACKAGE(this);
@@ -3500,8 +3395,52 @@ void FAsyncPackage2::StartLoading()
 
 	LoadStartTime = FPlatformTime::Seconds();
 
-	AsyncLoadingThread.AddBundleIoRequest(this);
 	AsyncPackageLoadingState = EAsyncPackageLoadingState2::WaitingForIo;
+
+	int32 LocalPendingIoRequestsCounter = AsyncLoadingThread.PendingIoRequestsCounter.IncrementExchange() + 1;
+	TRACE_COUNTER_SET(AsyncLoadingPendingIoRequests, LocalPendingIoRequestsCounter);
+
+	FIoReadOptions ReadOptions;
+	IoRequest = IoBatch.ReadWithCallback(CreateIoChunkId(Desc.PackageIdToLoad.Value(), 0, EIoChunkType::ExportBundleData),
+		ReadOptions,
+		Desc.Priority,
+		[this](TIoStatusOr<FIoBuffer> Result)
+		{
+			if (Result.IsOk())
+			{
+				TRACE_COUNTER_ADD(AsyncLoadingTotalLoaded, Result.ValueOrDie().DataSize());
+				CSV_CUSTOM_STAT_DEFINED(FrameCompletedExportBundleLoadsKB, float((double)Result.ValueOrDie().DataSize() / 1024.0), ECsvCustomStatOp::Accumulate);
+			}
+			else
+			{
+				UE_ASYNC_PACKAGE_LOG(Warning, Desc, TEXT("StartBundleIoRequests: FailedRead"),
+					TEXT("Failed reading chunk for package: %s"), *Result.Status().ToString());
+				bLoadHasFailed = true;
+			}
+			int32 LocalPendingIoRequestsCounter = AsyncLoadingThread.PendingIoRequestsCounter.DecrementExchange() - 1;
+			TRACE_COUNTER_SET(AsyncLoadingPendingIoRequests, LocalPendingIoRequestsCounter);
+			GetPackageNode(EEventLoadNode2::Package_ProcessSummary).ReleaseBarrier();
+		});
+
+	if (!Data.ShaderMapHashes.IsEmpty())
+	{
+		TRACE_CPUPROFILER_EVENT_SCOPE(StartShaderMapRequests);
+		auto ReadShaderMapFunc = [this, &IoBatch](const FIoChunkId& ChunkId, FGraphEventRef GraphEvent)
+		{
+			GetPackageNode(Package_ExportsSerialized).AddBarrier();
+			int32 LocalPendingIoRequestsCounter = AsyncLoadingThread.PendingIoRequestsCounter.IncrementExchange() + 1;
+			TRACE_COUNTER_SET(AsyncLoadingPendingIoRequests, LocalPendingIoRequestsCounter);
+			return IoBatch.ReadWithCallback(ChunkId, FIoReadOptions(), Desc.Priority,
+				[this, GraphEvent](TIoStatusOr<FIoBuffer> Result)
+				{
+					GraphEvent->DispatchSubsequents();
+					int32 LocalPendingIoRequestsCounter = AsyncLoadingThread.PendingIoRequestsCounter.DecrementExchange() - 1;
+					TRACE_COUNTER_SET(AsyncLoadingPendingIoRequests, LocalPendingIoRequestsCounter);
+					GetPackageNode(Package_ExportsSerialized).ReleaseBarrier();
+				});
+		};
+		FCoreDelegates::PreloadPackageShaderMaps.ExecuteIfBound(Data.ShaderMapHashes, ReadShaderMapFunc);
+	}
 }
 
 EAsyncPackageState::Type FAsyncPackage2::Event_ProcessPackageSummary(FAsyncLoadingThreadState2& ThreadState, FAsyncPackage2* Package, int32)
@@ -3734,11 +3673,6 @@ EAsyncPackageState::Type FAsyncPackage2::Event_ProcessExportBundle(FAsyncLoading
 	
 	Package->ExportBundleEntryIndex = 0;
 
-	if (Package->ProcessedExportBundlesCount == 0)
-	{
-		Package->AsyncLoadingThread.BundleIoRequestCompleted(Package);
-	}
-
 	if (++Package->ProcessedExportBundlesCount == Package->Data.ExportInfo.ExportBundleCount)
 	{
 		Package->ProcessedExportBundlesCount = 0;
@@ -6055,8 +5989,7 @@ EAsyncPackageState::Type FAsyncLoadingThread2::ProcessLoadingFromGameThread(FAsy
 	// CSV_CUSTOM_STAT(FileIO, EDLEventQueueDepth, (int32)GraphAllocator.TotalNodeCount, ECsvCustomStatOp::Set);
 	CSV_CUSTOM_STAT(FileIO, QueuedPackagesQueueDepth, GetNumQueuedPackages(), ECsvCustomStatOp::Set);
 	CSV_CUSTOM_STAT(FileIO, ExistingQueuedPackagesQueueDepth, GetNumAsyncPackages(), ECsvCustomStatOp::Set);
-	CSV_CUSTOM_STAT_DEFINED(PendingAsyncLoadingIoRequestsSizeMB, float((double)PendingBundleIoRequestsTotalSize / 1024.0 / 1024.0), ECsvCustomStatOp::Set);
-
+	
 	TickAsyncLoadingFromGameThread(ThreadState, bUseTimeLimit, bUseFullTimeLimit, TimeLimit);
 	return IsAsyncLoading() ? EAsyncPackageState::TimeOut : EAsyncPackageState::Complete;
 }
diff --git a/Engine/Source/Runtime/CoreUObject/Private/Serialization/FilePackageStore.cpp b/Engine/Source/Runtime/CoreUObject/Private/Serialization/FilePackageStore.cpp
--- a/Engine/Source/Runtime/CoreUObject/Private/Serialization/FilePackageStore.cpp
+++ b/Engine/Source/Runtime/CoreUObject/Private/Serialization/FilePackageStore.cpp
@@ -60,10 +60,8 @@ public:
 		{
 			FPackageStoreExportInfo
 			{
-				Entry->ExportBundlesSize,
 				Entry->ExportCount,
-				Entry->ExportBundleCount,
-				Entry->LoadOrder
+				Entry->ExportBundleCount
 			},
 			MakeArrayView(Entry->ImportedPackages.Data(), Entry->ImportedPackages.Num()),
 			MakeArrayView(Entry->ShaderMapHashes.Data(), Entry->ShaderMapHashes.Num())
diff --git a/Engine/Source/Runtime/CoreUObject/Public/Serialization/FilePackageStore.h b/Engine/Source/Runtime/CoreUObject/Public/Serialization/FilePackageStore.h
--- a/Engine/Source/Runtime/CoreUObject/Public/Serialization/FilePackageStore.h
+++ b/Engine/Source/Runtime/CoreUObject/Public/Serialization/FilePackageStore.h
@@ -38,11 +38,8 @@ public:
  */
 struct FFilePackageStoreEntry
 {
-	uint64 ExportBundlesSize;
 	int32 ExportCount;
 	int32 ExportBundleCount;
-	uint32 LoadOrder;
-	uint32 Pad;
 	TFilePackageStoreEntryCArrayView<FPackageId> ImportedPackages;
 	TFilePackageStoreEntryCArrayView<FSHAHash> ShaderMapHashes;
 };
