Fix StartAsyncLoading() with IODispatcher.

This fix was originally made by [REDACTED_PERSON].

diff --git a/Engine/Source/Runtime/CoreUObject/Private/Serialization/BulkData2.cpp b/Engine/Source/Runtime/CoreUObject/Private/Serialization/BulkData2.cpp
--- a/Engine/Source/Runtime/CoreUObject/Private/Serialization/BulkData2.cpp
+++ b/Engine/Source/Runtime/CoreUObject/Private/Serialization/BulkData2.cpp
@@ -1715,6 +1715,7 @@ void FBulkDataBase::InternalLoadFromIoStoreAsync(void** DstBuffer, AsyncCallback
 
 	// Set up our options (we only need to set the target)
 	FIoReadOptions Options;
+	Options.SetRange(GetBulkDataOffsetInFile(), GetBulkDataSize());
 	Options.SetTargetVa(*DstBuffer);
 
 	FIoBatch Batch = GetIoDispatcher()->NewBatch();
