BulkDataRegistry: Fix spurious ensure failure when loading BulkData with no PackagePath; this can occur with BulkDatas stored in virtualized storage.
#rnx

diff --git a/Engine/Source/Runtime/CoreUObject/Private/Virtualization/VirtualizedBulkData.cpp b/Engine/Source/Runtime/CoreUObject/Private/Virtualization/VirtualizedBulkData.cpp
--- a/Engine/Source/Runtime/CoreUObject/Private/Virtualization/VirtualizedBulkData.cpp
+++ b/Engine/Source/Runtime/CoreUObject/Private/Virtualization/VirtualizedBulkData.cpp
@@ -680,7 +680,14 @@ void FVirtualizedUntypedBulkData::SerializeForRegistry(FArchive& Ar)
 	{
 		FString PackageName;
 		Ar << PackageName;
-		ensure(FPackagePath::TryFromPackageName(PackageName, PackagePath));
+		if (PackageName.IsEmpty())
+		{
+			PackagePath.Empty();
+		}
+		else
+		{
+			ensure(FPackagePath::TryFromPackageName(PackageName, PackagePath));
+		}
 		PackageSegment = EPackageSegment::Header;
 	}
 	Ar << OffsetInFile;
