I also fixed it to not crash when just testing, but we could lower the fatal error to a normal/ensure if we want more safety

diff --git a/Engine/Source/Runtime/CoreUObject/Private/UObject/Obj.cpp b/Engine/Source/Runtime/CoreUObject/Private/UObject/Obj.cpp
--- a/Engine/Source/Runtime/CoreUObject/Private/UObject/Obj.cpp
+++ b/Engine/Source/Runtime/CoreUObject/Private/UObject/Obj.cpp
@@ -173,8 +173,14 @@ bool UObject::Rename( const TCHAR* InName, UObject* NewOuter, ERenameFlags Flags
 #endif //WITH_EDITOR
 
 	// Check that we are not renaming a within object into an Outer of the wrong type, unless we're renaming the CDO of a Blueprint.
-	if( NewOuter && !NewOuter->IsA(GetClass()->ClassWithin) && !HasAnyFlags(RF_ClassDefaultObject))	
+	// Moving objects to the transient package is commonly used halfway through destroying them so that is also fine
+	if( NewOuter && !NewOuter->IsA(GetClass()->ClassWithin) && !HasAnyFlags(RF_ClassDefaultObject) && NewOuter != GetTransientPackage())	
 	{
+		if (Flags & REN_Test)
+		{
+			return false;
+		}
+
 		UE_LOG(LogObj, Fatal, TEXT("Cannot rename %s into Outer %s as it is not of type %s"), 
 			*GetFullName(), 
 			*NewOuter->GetFullName(), 
