Implemented a stripped down version of FPackagePath for use outside of the editor in order to save memory
#rnx

diff --git a/Engine/Source/Runtime/CoreUObject/Private/Misc/PackagePath.cpp b/Engine/Source/Runtime/CoreUObject/Private/Misc/PackagePath.cpp
--- a/Engine/Source/Runtime/CoreUObject/Private/Misc/PackagePath.cpp
+++ b/Engine/Source/Runtime/CoreUObject/Private/Misc/PackagePath.cpp
@@ -191,6 +191,170 @@
 	*this = Other;
 }
 
+void FPackagePath::Empty()
+{
+	*this = FPackagePath();
+}
+
+FPackagePath FPackagePath::FromPackageNameChecked(FStringView InPackageName)
+{
+	FPackagePath PackagePath;
+	if (!TryFromPackageName(InPackageName, PackagePath))
+	{
+		UE_LOG(LogPackageName, Error, TEXT("FromPackageNameChecked: Invalid LongPackageName \"%.*s\""),
+			InPackageName.Len(), InPackageName.GetData())
+	}
+	return PackagePath;
+}
+
+FPackagePath FPackagePath::FromPackageNameChecked(FName InPackageName)
+{
+	TStringBuilder<256> PackageNameString;
+	InPackageName.AppendString(PackageNameString);
+	return FromPackageNameChecked(PackageNameString);
+}
+
+FPackagePath FPackagePath::FromPackageNameChecked(const TCHAR* InPackageName)
+{
+	return FromPackageNameChecked(FStringView(InPackageName));
+}
+
+bool FPackagePath::TryFromPackageName(const TCHAR* InPackageName, FPackagePath& OutPackagePath)
+{
+	return TryFromPackageName(FStringView(InPackageName), OutPackagePath);
+}
+
+FPackagePath FPackagePath::FromLocalPath(FStringView InFilename)
+{
+	EPackageSegment PackageSegment;
+	return FromLocalPath(InFilename, PackageSegment);
+}
+
+bool FPackagePath::operator!=(const FPackagePath& Other) const
+{
+	return !(*this == Other);
+}
+
+FString FPackagePath::GetPackageName() const
+{
+	TStringBuilder<256> Result;
+	AppendPackageName(Result);
+	return FString(Result);
+}
+
+FString FPackagePath::GetLocalFullPath() const
+{
+	return GetLocalFullPath(EPackageSegment::Header);
+}
+
+EPackageExtension FPackagePath::GetHeaderExtension() const
+{
+	return HeaderExtension;
+}
+
+EPackageExtension FPackagePath::GetExtension(EPackageSegment PackageSegment, FStringView& OutCustomExtension) const
+{
+	switch (PackageSegment)
+	{
+	case EPackageSegment::Header:
+		OutCustomExtension = GetCustomExtension();
+		return HeaderExtension;
+	default:
+		OutCustomExtension = FStringView();
+		return SegmentToExtension(PackageSegment);
+	}
+}
+
+FStringView FPackagePath::GetExtensionString(EPackageSegment PackageSegment) const
+{
+	switch (PackageSegment)
+	{
+	case EPackageSegment::Header:
+		if (HeaderExtension == EPackageExtension::Custom)
+		{
+			return GetCustomExtension();
+		}
+		else
+		{
+			return LexToString(HeaderExtension);
+		}
+	default:
+		return LexToString(SegmentToExtension(PackageSegment));
+	}
+}
+
+namespace UE
+{
+namespace PackagePathPrivate
+{
+	EPackageExtension AllExtensions[] =
+	{
+		EPackageExtension::Unspecified,
+		EPackageExtension::Asset,
+		EPackageExtension::Map,
+		EPackageExtension::TextAsset,
+		EPackageExtension::TextMap,
+		EPackageExtension::Custom,
+		EPackageExtension::EmptyString,
+		EPackageExtension::Exports,
+		EPackageExtension::BulkDataDefault,
+		EPackageExtension::BulkDataOptional,
+		EPackageExtension::BulkDataMemoryMapped,
+		EPackageExtension::PayloadSidecar,
+	};
+	static_assert(static_cast<int>(EPackageExtension::Unspecified) == 0 && EPackageExtensionCount == UE_ARRAY_COUNT(AllExtensions), "Need to add new extensions to AllExtensions array");
+}
+}
+
+TConstArrayView<EPackageExtension> FPackagePath::GetPossibleExtensions(EPackageSegment PackageSegment) const
+{
+	FStringView OutCustomExtension;
+	EPackageExtension Extension = GetExtension(PackageSegment, OutCustomExtension);
+	if (Extension != EPackageExtension::Unspecified)
+	{
+		return TConstArrayView<EPackageExtension>(&UE::PackagePathPrivate::AllExtensions[static_cast<int>(Extension)], 1);
+	}
+	else
+	{
+#if WITH_TEXT_ARCHIVE_SUPPORT
+		constexpr int NumHeaderSearchExtensions = 4;
+#else
+		constexpr int NumHeaderSearchExtensions = 2;
+#endif
+		static_assert(static_cast<int>(EPackageExtension::TextMap) == static_cast<int>(EPackageExtension::Asset) + 3, "Need to update the list of header extensions");
+		return TConstArrayView<EPackageExtension>(&UE::PackagePathPrivate::AllExtensions[static_cast<int>(EPackageExtension::Asset)], NumHeaderSearchExtensions);
+	}
+}
+
+FString FPackagePath::GetDebugName() const
+{
+	return GetDebugName(EPackageSegment::Header);
+}
+
+FText FPackagePath::GetDebugNameText() const
+{
+	return GetDebugNameText(EPackageSegment::Header);
+}
+
+FText FPackagePath::GetDebugNameText(EPackageSegment PackageSegment) const
+{
+	return FText::FromString(GetDebugName(PackageSegment));
+}
+
+FString FPackagePath::GetDebugNameWithExtension() const
+{
+	return GetDebugNameWithExtension(EPackageSegment::Header);
+}
+
+FString FPackagePath::GetDebugNameWithExtension(EPackageSegment PackageSegment) const
+{
+	FString Result = GetDebugName(EPackageSegment::Header);
+	Result += GetExtensionString(PackageSegment);
+	return Result;
+}
+
+#if WITH_EDITOR
+
 FPackagePath& FPackagePath::operator=(const FPackagePath& Other)
 {
 	if (this == &Other)
@@ -213,11 +377,6 @@
 	return *this;
 }
 
-void FPackagePath::Empty()
-{
-	*this = FPackagePath();
-}
-
 bool FPackagePath::TryFromMountedName(FStringView InPackageNameOrFilePath, FPackagePath& OutPackagePath)
 {
 	TStringBuilder<64> PackageNameRoot;
@@ -263,38 +422,14 @@
 	return TryFromPackageName(PackageName, OutPackagePath);
 }
 
-bool FPackagePath::TryFromPackageName(const TCHAR* InPackageName, FPackagePath& OutPackagePath)
+FPackagePath FPackagePath::FromPackageNameUnchecked(FName InPackageName)
 {
-	return TryFromPackageName(FStringView(InPackageName), OutPackagePath);
-}
-
-FPackagePath FPackagePath::FromPackageNameChecked(FStringView InPackageName)
-{
 	FPackagePath PackagePath;
-	if (!TryFromPackageName(InPackageName, PackagePath))
-	{
-		UE_LOG(LogPackageName, Error, TEXT("FromPackageNameChecked: Invalid LongPackageName \"%.*s\""),
-			InPackageName.Len(), InPackageName.GetData())
-	}
-	return PackagePath;
-}
-
-FPackagePath FPackagePath::FromPackageNameChecked(FName InPackageName)
-{
 	TStringBuilder<256> PackageNameString;
 	InPackageName.AppendString(PackageNameString);
-	return FromPackageNameChecked(PackageNameString);
+	PackagePath.IdType = EPackageIdType::PackageOnlyPath;
+	PackagePath.SetStringData(PackageNameString, FStringView(), FStringView(), FStringView());
+	return PackagePath;
 }
 
-FPackagePath FPackagePath::FromPackageNameChecked(const TCHAR* InPackageName)
-{
-	return FromPackageNameChecked(FStringView(InPackageName));
-}
-
-FPackagePath FPackagePath::FromLocalPath(FStringView InFilename)
-{
-	EPackageSegment PackageSegment;
-	return FromLocalPath(InFilename, PackageSegment);
-}
-
 FPackagePath FPackagePath::FromLocalPath(FStringView InFilename, EPackageSegment& PackageSegment)
@@ -476,11 +611,6 @@
 	}
 }
 
-bool FPackagePath::operator!=(const FPackagePath& Other) const
-{
-	return !(*this == Other);
-}
-
 FArchive& operator<<(FArchive& Ar, FPackagePath& PackagePath)
 {
 	checkf(!Ar.IsPersistent(), TEXT("PackagePath is transient and does not support serialization to Persistent storage."));
@@ -582,13 +712,6 @@
 	}
 }
 
-FString FPackagePath::GetPackageName() const
-{
-	TStringBuilder<256> Result;
-	AppendPackageName(Result);
-	return FString(Result);
-}
-
 void FPackagePath::AppendPackageName(FStringBuilderBase& Builder) const
 {
 	switch (IdType)
@@ -625,11 +748,6 @@
 	return FName(Builder);
 }
 
-FString FPackagePath::GetLocalFullPath() const
-{
-	return GetLocalFullPath(EPackageSegment::Header);
-}
-
 FString FPackagePath::GetLocalFullPath(EPackageSegment PackageSegment) const
 {
 	TStringBuilder<256> Result;
@@ -717,90 +835,6 @@
 	}
 }
 
-EPackageExtension FPackagePath::GetHeaderExtension() const
-{
-	return HeaderExtension;
-}
-
-EPackageExtension FPackagePath::GetExtension(EPackageSegment PackageSegment, FStringView& OutCustomExtension) const
-{
-	switch (PackageSegment)
-	{
-	case EPackageSegment::Header:
-		OutCustomExtension = GetCustomExtension();
-		return HeaderExtension;
-	default:
-		OutCustomExtension = FStringView();
-		return SegmentToExtension(PackageSegment);
-	}
-}
-
-FStringView FPackagePath::GetExtensionString(EPackageSegment PackageSegment) const
-{
-	switch (PackageSegment)
-	{
-	case EPackageSegment::Header:
-		if (HeaderExtension == EPackageExtension::Custom)
-		{
-			return GetCustomExtension();
-		}
-		else
-		{
-			return LexToString(HeaderExtension);
-		}
-	default:
-		return LexToString(SegmentToExtension(PackageSegment));
-	}
-}
-
-namespace UE
-{
-namespace PackagePathPrivate
-{
-EPackageExtension AllExtensions[] =
-{
-	EPackageExtension::Unspecified,
-	EPackageExtension::Asset,
-	EPackageExtension::Map,
-	EPackageExtension::TextAsset,
-	EPackageExtension::TextMap,
-	EPackageExtension::Custom,
-	EPackageExtension::EmptyString,
-	EPackageExtension::Exports,
-	EPackageExtension::BulkDataDefault,
-	EPackageExtension::BulkDataOptional,
-	EPackageExtension::BulkDataMemoryMapped,
-	EPackageExtension::PayloadSidecar,
-};
-static_assert(static_cast<int>(EPackageExtension::Unspecified) == 0 && EPackageExtensionCount == UE_ARRAY_COUNT(AllExtensions), "Need to add new extensions to AllExtensions array");
-}
-}
-
-TConstArrayView<EPackageExtension> FPackagePath::GetPossibleExtensions(EPackageSegment PackageSegment) const
-{
-	FStringView OutCustomExtension;
-	EPackageExtension Extension = GetExtension(PackageSegment, OutCustomExtension);
-	if (Extension != EPackageExtension::Unspecified)
-	{
-		return TConstArrayView<EPackageExtension>(&UE::PackagePathPrivate::AllExtensions[static_cast<int>(Extension)], 1);
-	}
-	else
-	{
-#if WITH_TEXT_ARCHIVE_SUPPORT
-		constexpr int NumHeaderSearchExtensions = 4;
-#else
-		constexpr int NumHeaderSearchExtensions = 2;
-#endif
-		static_assert(static_cast<int>(EPackageExtension::TextMap) == static_cast<int>(EPackageExtension::Asset) + 3, "Need to update the list of header extensions");
-		return TConstArrayView<EPackageExtension>(&UE::PackagePathPrivate::AllExtensions[static_cast<int>(EPackageExtension::Asset)], NumHeaderSearchExtensions);
-	}
-}
-
-FString FPackagePath::GetDebugName() const
-{
-	return GetDebugName(EPackageSegment::Header);
-}
-
 FString FPackagePath::GetDebugName(EPackageSegment PackageSegment) const
 {
 	FString Result;
@@ -830,28 +864,6 @@
 	return Result;
 }
 
-FText FPackagePath::GetDebugNameText() const
-{
-	return GetDebugNameText(EPackageSegment::Header);
-}
-
-FText FPackagePath::GetDebugNameText(EPackageSegment PackageSegment) const
-{
-	return FText::FromString(GetDebugName(PackageSegment));
-}
-
-FString FPackagePath::GetDebugNameWithExtension() const
-{
-	return GetDebugNameWithExtension(EPackageSegment::Header);
-}
-
-FString FPackagePath::GetDebugNameWithExtension(EPackageSegment PackageSegment) const
-{
-	FString Result = GetDebugName(EPackageSegment::Header);
-	Result += GetExtensionString(PackageSegment);
-	return Result;
-}
-
 FString FPackagePath::GetPackageNameOrFallback() const
 {
 	switch (IdType)
@@ -983,3 +995,164 @@
 	SetStringData(NewPath.GetPathData(), NewPath.GetPackageNameRoot(), NewPath.GetFilePathRoot(), GetCustomExtension());
 	return true;
 }
+
+#else
+
+FPackagePath& FPackagePath::operator=(const FPackagePath& Other)
+{
+	PackageName = Other.PackageName;
+	HeaderExtension = Other.HeaderExtension;
+	return *this;
+}
+
+bool FPackagePath::TryFromMountedName(FStringView InPackageNameOrHeaderFilePath, FPackagePath& OutPackagePath)
+{
+	if (FPackageName::IsValidLongPackageName(InPackageNameOrHeaderFilePath))
+	{
+		OutPackagePath.PackageName = FName(InPackageNameOrHeaderFilePath);
+	}
+	else
+	{
+		OutPackagePath.PackageName = FName(FPackageName::FilenameToLongPackageName(FString(InPackageNameOrHeaderFilePath)));
+	}
+	return true;
+}
+
+bool FPackagePath::TryFromPackageName(FStringView InPackageName, FPackagePath& OutPackagePath)
+{
+	if (!FPackageName::IsValidTextForLongPackageName(InPackageName))
+	{
+		return false;
+	}
+	OutPackagePath.PackageName = FName(InPackageName);
+	return true;
+}
+
+bool FPackagePath::TryFromPackageName(FName InPackageName, FPackagePath& OutPackagePath)
+{
+	TStringBuilder<256> PackageNameString;
+	InPackageName.AppendString(PackageNameString);
+	if (!FPackageName::IsValidTextForLongPackageName(PackageNameString))
+	{
+		return false;
+	}
+	OutPackagePath.PackageName = InPackageName;
+	return true;
+}
+
+FPackagePath FPackagePath::FromPackageNameUnchecked(FName InPackageName)
+{
+	FPackagePath PackagePath;
+	PackagePath.PackageName = InPackageName;
+	return PackagePath;
+}
+
+FPackagePath FPackagePath::FromLocalPath(FStringView InFilename, EPackageSegment& OutPackageSegment)
+{
+	unimplemented();
+	return FPackagePath();
+}
+
+FPackagePath FPackagePath::FromMountedComponents(FStringView PackageNameRoot, FStringView FilePathRoot, FStringView RelPath,
+	EPackageExtension InExtension, FStringView InCustomExtension)
+{
+	if (ExtensionToSegment(InExtension) != EPackageSegment::Header)
+	{
+		InExtension = EPackageExtension::Unspecified;
+		InCustomExtension = FStringView();
+	}
+	if (!InCustomExtension.IsEmpty())
+	{
+		checkNoEntry();
+	}
+	TStringBuilder<256> FileNameString;
+	FileNameString.Append(FilePathRoot);
+	FileNameString.Append(RelPath);
+
+	FString PackageNameString;
+	if (!FPackageName::TryConvertFilenameToLongPackageName(FString(FileNameString), PackageNameString))
+	{
+		UE_LOG(LogPackageName, Error, TEXT("FromMountedComponents: Invalid FileName \"%s\""), *FileNameString);
+	}
+	FPackagePath PackagePath;
+	PackagePath.PackageName = FName(PackageNameString);
+	PackagePath.HeaderExtension = InExtension;
+	return PackagePath;
+}
+
+bool FPackagePath::TryMatchCase(const FPackagePath& SourcePackagePath, FStringView FilePathToMatch, FPackagePath& OutPackagePath)
+{
+	return true;
+}
+
+bool FPackagePath::operator==(const FPackagePath& Other) const
+{
+	return Other.PackageName == PackageName && Other.HeaderExtension == HeaderExtension;
+}
+
+FArchive& operator<<(FArchive& Ar, FPackagePath& PackagePath)
+{
+	unimplemented();
+	return Ar;
+}
+
+bool FPackagePath::IsEmpty() const
+{
+	return PackageName.IsNone();
+}
+
+bool FPackagePath::IsMountedPath() const
+{
+	return true;
+}
+
+void FPackagePath::AppendPackageName(FStringBuilderBase& Builder) const
+{
+	PackageName.AppendString(Builder);
+}
+
+FName FPackagePath::GetPackageFName() const
+{
+	return PackageName;
+}
+
+FString FPackagePath::GetLocalFullPath(EPackageSegment PackageSegment) const
+{
+	FString PackageNameString = PackageName.ToString();
+	EPackageExtension Extension = PackageSegment == EPackageSegment::Header ? HeaderExtension : SegmentToExtension(PackageSegment);
+	FString LocalPath = FPackageName::LongPackageNameToFilename(PackageNameString, LexToString(Extension));
+	return LocalPath;
+}
+
+FString FPackagePath::GetLocalBaseFilenameWithPath() const
+{
+	FString PackageNameString = PackageName.ToString();
+	FString LocalPath = FPackageName::LongPackageNameToFilename(PackageNameString);
+	return LocalPath;
+}
+
+FString FPackagePath::GetDebugName(EPackageSegment PackageSegment) const
+{
+	return PackageName.ToString();
+}
+
+FString FPackagePath::GetPackageNameOrFallback() const
+{
+	return GetPackageName();
+}
+
+void FPackagePath::SetHeaderExtension(EPackageExtension Extension, FStringView CustomExtension) const
+{
+	if (!CustomExtension.IsEmpty())
+	{
+		checkNoEntry();
+	}
+	HeaderExtension = Extension;
+}
+
+FStringView FPackagePath::GetCustomExtension() const
+{
+	return FStringView();
+}
+
+#endif
diff --git a/Engine/Source/Runtime/CoreUObject/Private/Serialization/AsyncLoading2.cpp b/Engine/Source/Runtime/CoreUObject/Private/Serialization/AsyncLoading2.cpp
--- a/Engine/Source/Runtime/CoreUObject/Private/Serialization/AsyncLoading2.cpp
+++ b/Engine/Source/Runtime/CoreUObject/Private/Serialization/AsyncLoading2.cpp
@@ -5730,7 +5730,7 @@ void FAsyncPackage2::CreateUPackage(const FPackageSummary* PackageSummary)
 			bCreatedLinkerRoot = true;
 		}
 		LinkerRoot->SetFlags(RF_Public | RF_WasLoaded);
-		LinkerRoot->SetLoadedPath(FPackagePath::FromPackageNameChecked(Desc.PackageNameToLoad));
+		LinkerRoot->SetLoadedPath(FPackagePath::FromPackageNameUnchecked(Desc.PackageNameToLoad));
 		LinkerRoot->SetCanBeImportedFlag(Desc.bCanBeImported);
 		LinkerRoot->SetPackageId(Desc.UPackageId);
 		LinkerRoot->SetPackageIdToLoad(Desc.PackageIdToLoad);
diff --git a/Engine/Source/Runtime/CoreUObject/Public/Misc/PackagePath.h b/Engine/Source/Runtime/CoreUObject/Public/Misc/PackagePath.h
--- a/Engine/Source/Runtime/CoreUObject/Public/Misc/PackagePath.h
+++ b/Engine/Source/Runtime/CoreUObject/Public/Misc/PackagePath.h
@@ -141,6 +141,17 @@ public:
 	COREUOBJECT_API static FPackagePath FromPackageNameChecked(FName InPackageName);
 	COREUOBJECT_API static FPackagePath FromPackageNameChecked(const TCHAR* InPackageName);
 
+	/**
+	 * Construct a PackagePath from a known valid LongPackageName FName
+	 *
+	 * Will be converted to a MountedPath when the LocalPath is required; if the package is not mounted at that point
+	 * the LocalPath will be empty
+	 *
+	 * @param InPackageName A valid LongPackageName, does not have to be mounted or existing
+	 * @return The constructed PackagePath
+	 */
+	COREUOBJECT_API static FPackagePath FromPackageNameUnchecked(FName InPackageName);
+
 	/**
 	 * Construct a PackagePath from a LocalPath string
 	 *
@@ -446,6 +457,7 @@ private:
 		LocalOnlyPath,
 	};
 
+#if WITH_EDITOR
 	/**
 	 * Combined storage for (EPackageIDType-specific path+filename storage) followed by (possibly empty)
 	 * PackageNameRoot followed by (possibly empty) FilePathRoot followed by (possibly empty) Extension
@@ -458,4 +470,8 @@ private:
 	mutable uint16 ExtensionLen = 0;
 	mutable EPackageIdType IdType = EPackageIdType::Empty;
 	mutable EPackageExtension HeaderExtension = EPackageExtension::Unspecified;
-};
\ No newline at end of file
+#else
+	FName PackageName;
+	mutable EPackageExtension HeaderExtension = EPackageExtension::Unspecified;
+#endif
+};
