Fix for a possible crash when incremental GC is running and other threads create new UObjects

diff --git a/Engine/Source/Runtime/CoreUObject/Private/UObject/GarbageCollection.cpp b/Engine/Source/Runtime/CoreUObject/Private/UObject/GarbageCollection.cpp
--- a/Engine/Source/Runtime/CoreUObject/Private/UObject/GarbageCollection.cpp
+++ b/Engine/Source/Runtime/CoreUObject/Private/UObject/GarbageCollection.cpp
@@ -1447,6 +1447,7 @@ void IncrementalPurgeGarbage(bool bUseTimeLimit, float TimeLimit)
 
 	if (!bTimeLimitReached)
 	{
+		FConditionalGCLock ScopedGCLock;
 		bCompleted = IncrementalDestroyGarbage(bUseTimeLimit, TimeLimit);
 	}
 #endif // !UE_WITH_GC
