AsyncLoading: Add LoadPackage <package name> console command

diff --git a/Engine/Source/Runtime/CoreUObject/Private/Serialization/AsyncPackageLoader.cpp b/Engine/Source/Runtime/CoreUObject/Private/Serialization/AsyncPackageLoader.cpp
--- a/Engine/Source/Runtime/CoreUObject/Private/Serialization/AsyncPackageLoader.cpp
+++ b/Engine/Source/Runtime/CoreUObject/Private/Serialization/AsyncPackageLoader.cpp
@@ -13,6 +13,21 @@
 volatile int32 GIsLoaderCreated;
 TUniquePtr<IAsyncPackageLoader> GPackageLoader;
 
+#if !UE_BUILD_SHIPPING
+static void LoadPackageCommand(const TArray<FString>& Args)
+{
+	for (const FString& PackageName : Args)
+	{
+		LoadPackageAsync(PackageName);
+	}
+}
+
+static FAutoConsoleCommand CVar_LoadPackageCommand(
+	TEXT("LoadPackage"),
+	TEXT("Loads packages by names. Usage: LoadPackage <package name> [<package name> ...]"),
+	FConsoleCommandWithArgsDelegate::CreateStatic(LoadPackageCommand));
+#endif
+
 const FName PrestreamPackageClassNameLoad = FName("PrestreamPackage");
 
 struct FEDLBootObjectState
