AsyncLoading2: Add more verbose logging and a new -AsyncLoadingVerbosePackageNames="/Engine/Content/PackageA /Game/Content/PackageB" parameter to do verbose/veryverbose logging of selected packages only

diff --git a/Engine/Source/Runtime/CoreUObject/Private/Serialization/AsyncLoading2.cpp b/Engine/Source/Runtime/CoreUObject/Private/Serialization/AsyncLoading2.cpp
--- a/Engine/Source/Runtime/CoreUObject/Private/Serialization/AsyncLoading2.cpp
+++ b/Engine/Source/Runtime/CoreUObject/Private/Serialization/AsyncLoading2.cpp
@@ -157,22 +157,30 @@ FArchive& operator<<(FArchive& Ar, FExportMapEntry& ExportMapEntry)
 #define ALT2_LOG_VERBOSE DO_CHECK
 #endif
 
+static TSet<FName> GAsyncLoading2_VerbosePackageNames;
+
 #define UE_ASYNC_PACKAGE_LOG(Verbosity, PackageDesc, LogDesc, Format, ...) \
-if ((PackageDesc).Name != (PackageDesc).NameToLoad) \
-{ \
-	UE_LOG(LogStreaming, Verbosity, LogDesc TEXT(": %s (%d) %s (%d) - ") Format, \
-		*(PackageDesc).Name.ToString(), \
-		(PackageDesc).PackageId.ToIndexForDebugging(), \
-		*(PackageDesc).NameToLoad.ToString(), \
-		(PackageDesc).PackageIdToLoad.ToIndexForDebugging(), \
-		##__VA_ARGS__); \
-} \
-else \
+if (GAsyncLoading2_VerbosePackageNames.Num() == 0 || \
+	ELogVerbosity::Verbosity < ELogVerbosity::Verbose || \
+	GAsyncLoading2_VerbosePackageNames.Contains((PackageDesc).Name) || \
+	GAsyncLoading2_VerbosePackageNames.Contains((PackageDesc).NameToLoad)) \
 { \
-	UE_LOG(LogStreaming, Verbosity, LogDesc TEXT(": %s (%d) - ") Format, \
-		*(PackageDesc).Name.ToString(), \
-		(PackageDesc).PackageId.ToIndexForDebugging(), \
-		##__VA_ARGS__); \
+	if ((PackageDesc).Name != (PackageDesc).NameToLoad) \
+	{ \
+		UE_LOG(LogStreaming, Verbosity, LogDesc TEXT(": %s (%d) %s (%d) - ") Format, \
+			*(PackageDesc).Name.ToString(), \
+			(PackageDesc).PackageId.ToIndexForDebugging(), \
+			*(PackageDesc).NameToLoad.ToString(), \
+			(PackageDesc).PackageIdToLoad.ToIndexForDebugging(), \
+			##__VA_ARGS__); \
+	} \
+	else \
+	{ \
+		UE_LOG(LogStreaming, Verbosity, LogDesc TEXT(": %s (%d) - ") Format, \
+			*(PackageDesc).Name.ToString(), \
+			(PackageDesc).PackageId.ToIndexForDebugging(), \
+			##__VA_ARGS__); \
+	} \
 }
 
 #define UE_ASYNC_PACKAGE_CLOG(Condition, Verbosity, PackageDesc, LogDesc, Format, ...) \
@@ -383,8 +391,6 @@ struct FGlobalImportStore
 
 	inline FName GetName(FPackageObjectIndex GlobalIndex)
 	{
-		check(GlobalIndex.IsImport());
-
 		return GlobalIndex.IsScriptImport() && ScriptObjectEntries.Num() > 0
 			? MinimalNameToName(ScriptObjectEntries[GlobalIndex.GetIndex()].ObjectName)
 			: NAME_None;
@@ -467,24 +473,24 @@ struct FGlobalImportStore
 		return Object;
 	}
 
-	void StoreGlobalObject(FPackageId PackageId, FPackageObjectIndex GlobalIndex, UObject* InObject)
+	bool StoreGlobalObject(FPackageId PackageId, FPackageObjectIndex GlobalIndex, UObject* InObject)
 	{
 		if (GlobalIndex.IsScriptImport())
 		{
 			UObject*& Object = ScriptObjects[GlobalIndex.GetIndex()];
 			check(!Object || Object == InObject);
 			Object = InObject;
+			return true;
 		}
 		else if (GlobalIndex.IsPackageImport())
 		{
 			int32 ObjectIndex = GUObjectArray.ObjectToIndex(InObject);
 			PublicExportToObjectIndex.Add(GlobalIndex, {ObjectIndex, PackageId});
 			ObjectIndexToPublicExport.Add(ObjectIndex, GlobalIndex);
+			return true;
 		}
-		else
-		{
-			check(GlobalIndex.IsNull());
-		}
+		check(GlobalIndex.IsNull());
+		return false;
 	}
 
 	void FindAllScriptObjects();
@@ -1046,9 +1052,9 @@ struct FPackageImportStore
 		return FName();
 	}
 
-	void StoreGlobalObject(FPackageId InPackageId, FPackageObjectIndex InGlobalIndex, UObject* InObject)
+	inline bool StoreGlobalObject(FPackageId InPackageId, FPackageObjectIndex InGlobalIndex, UObject* InObject)
 	{
-		GlobalImportStore.StoreGlobalObject(InPackageId, InGlobalIndex, InObject);
+		return GlobalImportStore.StoreGlobalObject(InPackageId, InGlobalIndex, InObject);
 	}
 
 	void ClearReferences()
@@ -1830,7 +1836,7 @@ public:
 	static EAsyncPackageState::Type Event_Delete(FAsyncPackage2* Package, int32);
 
 	void EventDrivenCreateExport(int32 LocalExportIndex);
-	void EventDrivenSerializeExport(int32 LocalExportIndex, FExportArchive& Ar);
+	bool EventDrivenSerializeExport(int32 LocalExportIndex, FExportArchive& Ar);
 
 	UObject* EventDrivenIndexToObject(FPackageObjectIndex Index, bool bCheckSerialized);
 	template<class T>
@@ -2416,6 +2422,22 @@ struct FAsyncLoadingTickScope2
 
 void FAsyncLoadingThread2::InitializeLoading()
 {
+#if !UE_BUILD_SHIPPING
+	{
+		FString VerbosePackageNames;
+		FParse::Value(FCommandLine::Get(), TEXT("-AsyncLoadingVerbosePackageNames="), VerbosePackageNames);
+		if (VerbosePackageNames.Len() > 0)
+		{
+			TArray<FString> Args;
+			VerbosePackageNames.ParseIntoArray(Args, TEXT(" "));
+			for (FString& PackageName : Args)
+			{
+				GAsyncLoading2_VerbosePackageNames.Add(FName(*PackageName));
+			}
+		}
+	}
+#endif
+
 #if USE_NEW_BULKDATA
 	FBulkDataBase::SetIoDispatcher(&IoDispatcher);
 #endif
@@ -3238,10 +3260,7 @@ EAsyncPackageState::Type FAsyncPackage2::Event_ProcessExportBundle(FAsyncPackage
 
 			if (BundleEntry->CommandType == FExportBundleEntry::ExportCommandType_Create)
 			{
-				if (!(Export.bFiltered | Export.bExportLoadFailed))
-				{
-					Package->EventDrivenCreateExport(BundleEntry->LocalExportIndex);
-				}
+				Package->EventDrivenCreateExport(BundleEntry->LocalExportIndex);
 			}
 			else
 			{
@@ -3254,19 +3273,21 @@ EAsyncPackageState::Type FAsyncPackage2::Event_ProcessExportBundle(FAsyncPackage
 				check(Object || Export.bFiltered || Export.bExportLoadFailed);
 
 				Ar.ExportBufferBegin(ExportMapEntry.CookedSerialOffset, ExportMapEntry.CookedSerialSize);
-				if (!(Export.bFiltered | Export.bExportLoadFailed) && Object->HasAnyFlags(RF_NeedLoad))
-				{
-					TRACE_LOADTIME_SERIALIZE_EXPORT_SCOPE(Object, CookedSerialSize);
-					const int64 Pos = Ar.Tell();
-					check(CookedSerialSize <= uint64(Ar.TotalSize() - Pos));
 
-					Package->EventDrivenSerializeExport(BundleEntry->LocalExportIndex, Ar);
-					checkf(CookedSerialSize == uint64(Ar.Tell() - Pos), TEXT("Expect read size: %llu - Actual read size: %llu"), CookedSerialSize, uint64(Ar.Tell() - Pos));
-				}
-				else
+				const int64 Pos = Ar.Tell();
+				checkf(CookedSerialSize <= uint64(Ar.TotalSize() - Pos),
+					TEXT("Package %s: Expected read size: %llu - Remaining archive size: %llu"),
+					*Package->Desc.NameToLoad.ToString(), CookedSerialSize, uint64(Ar.TotalSize() - Pos));
+
+				const bool bSerialized = Package->EventDrivenSerializeExport(BundleEntry->LocalExportIndex, Ar);
+				if (!bSerialized)
 				{
 					Ar.Skip(CookedSerialSize);
 				}
+				checkf(CookedSerialSize == uint64(Ar.Tell() - Pos),
+					TEXT("Package %s: Expected read size: %llu - Actual read size: %llu"),
+					*Package->Desc.NameToLoad.ToString(), CookedSerialSize, uint64(Ar.Tell() - Pos));
+
 				Ar.ExportBufferEnd();
 
 				check((Object && !Object->HasAnyFlags(RF_NeedLoad)) || Export.bFiltered || Export.bExportLoadFailed);
@@ -3361,16 +3382,30 @@ void FAsyncPackage2::EventDrivenCreateExport(int32 LocalExportIndex)
 	TRACE_CPUPROFILER_EVENT_SCOPE(CreateExport);
 
 	const FExportMapEntry& Export = ExportMap[LocalExportIndex];
-	UObject*& Object = Exports[LocalExportIndex].Object;
+	FExportObject& ExportObject = Exports[LocalExportIndex];
+	UObject*& Object = ExportObject.Object;
 	check(!Object);
 
+	TRACE_LOADTIME_CREATE_EXPORT_SCOPE(this, &Object);
+
 	FName ObjectName;
 	{
 		TRACE_CPUPROFILER_EVENT_SCOPE(ObjectNameFixup);
 		ObjectName = NameMap->GetName(Export.ObjectName);
 	}
 
-	TRACE_LOADTIME_CREATE_EXPORT_SCOPE(this, &Object);
+	if (ExportObject.bFiltered | ExportObject.bExportLoadFailed)
+	{
+		if (ExportObject.bExportLoadFailed)
+		{
+			UE_ASYNC_PACKAGE_LOG(Warning, Desc, TEXT("CreateExport"), TEXT("Skipped failed export %s"), *ObjectName.ToString());
+		}
+		else
+		{
+			UE_ASYNC_PACKAGE_LOG_VERBOSE(Verbose, Desc, TEXT("CreateExport"), TEXT("Skipped filtered export %s"), *ObjectName.ToString());
+		}
+		return;
+	}
 
 	LLM_SCOPE(ELLMTag::AsyncLoading);
 	LLM_SCOPED_TAG_WITH_OBJECT_IN_SET(GetLinkerRoot(), ELLMTagSet::Assets);
@@ -3383,14 +3418,14 @@ void FAsyncPackage2::EventDrivenCreateExport(int32 LocalExportIndex)
 
 	if (!LoadClass)
 	{
-		UE_LOG(LogStreaming, Error, TEXT("Could not find class object for %s in %s"), *ObjectName.ToString(), *Desc.NameToLoad.ToString());
-		Exports[LocalExportIndex].bExportLoadFailed = true;
+		UE_ASYNC_PACKAGE_LOG(Error, Desc, TEXT("CreateExport"), TEXT("Could not find class object for %s"), *ObjectName.ToString());
+		ExportObject.bExportLoadFailed = true;
 		return;
 	}
 	if (!ThisParent)
 	{
-		UE_LOG(LogStreaming, Error, TEXT("Could not find outer object for %s in %s"), *ObjectName.ToString(), *Desc.NameToLoad.ToString());
-		Exports[LocalExportIndex].bExportLoadFailed = true;
+		UE_ASYNC_PACKAGE_LOG(Error, Desc, TEXT("CreateExport"), TEXT("Could not find outer object for %s"), *ObjectName.ToString());
+		ExportObject.bExportLoadFailed = true;
 		return;
 	}
 	check(!dynamic_cast<UObjectRedirector*>(ThisParent));
@@ -3410,7 +3445,6 @@ void FAsyncPackage2::EventDrivenCreateExport(int32 LocalExportIndex)
 		bIsCompleteyLoaded = !!(ObjectFlags & RF_LoadCompleted);
 		if (!bIsCompleteyLoaded)
 		{
-			UE_LOG(LogStreaming, VeryVerbose, TEXT("Note2: %s was constructed during load and is an export and so needs loading."), *Object->GetFullName());
 			check(!(ObjectFlags & (RF_NeedLoad | RF_WasLoaded))); // If export exist but is not completed, it is expected to have been created from a native constructor and not from EventDrivenCreateExport, but who knows...?
 			if (ObjectFlags & RF_ClassDefaultObject)
 			{
@@ -3437,8 +3471,8 @@ void FAsyncPackage2::EventDrivenCreateExport(int32 LocalExportIndex)
 		UObject* Template = EventDrivenIndexToObject(Export.TemplateIndex, true);
 		if (!Template)
 		{
-			UE_LOG(LogStreaming, Error, TEXT("Could not find template for %s in %s"), *ObjectName.ToString(), *Desc.NameToLoad.ToString());
-			Exports[LocalExportIndex].bExportLoadFailed = true;
+			UE_ASYNC_PACKAGE_LOG(Error, Desc, TEXT("CreateExport"), TEXT("Could not find template object for %s"), *ObjectName.ToString());
+			ExportObject.bExportLoadFailed = true;
 			return;
 		}
 		// we also need to ensure that the template has set up any instances
@@ -3523,25 +3557,55 @@ void FAsyncPackage2::EventDrivenCreateExport(int32 LocalExportIndex)
 	}
 
 	check(Object);
-	// don't store temp packages, nothing is importing them
-	if (AsyncLoadingThread.GlobalPackageStore.IsSerializedPackageId(Desc.PackageId))
+
+	// don't store exports temp packages, nothing is importing them
+	bool bStore = AsyncLoadingThread.GlobalPackageStore.IsSerializedPackageId(Desc.PackageId);
+	if (bStore)
+	{
+		bStore = ImportStore.StoreGlobalObject(Desc.PackageIdToLoad, Export.GlobalImportIndex, Object);
+	}
+
+	if (bStore)
+	{
+		UE_ASYNC_PACKAGE_LOG_VERBOSE(VeryVerbose, Desc, TEXT("CreateExport"),
+			TEXT("Created and stored public export %s (%d)"), *Object->GetPathName(), Export.GlobalImportIndex.GetIndex());
+	}
+	else
 	{
-		ImportStore.StoreGlobalObject(Desc.PackageId, Export.GlobalImportIndex, Object);
+		UE_ASYNC_PACKAGE_LOG_VERBOSE(VeryVerbose, Desc, TEXT("CreateExport"), TEXT("Created export %s"), *Object->GetPathName());
 	}
 }
 
-void FAsyncPackage2::EventDrivenSerializeExport(int32 LocalExportIndex, FExportArchive& Ar)
+bool FAsyncPackage2::EventDrivenSerializeExport(int32 LocalExportIndex, FExportArchive& Ar)
 {
 	TRACE_CPUPROFILER_EVENT_SCOPE(SerializeExport);
 
 	const FExportMapEntry& Export = ExportMap[LocalExportIndex];
-	UObject* Object = Exports[LocalExportIndex].Object;
-	check(Object);
+	FExportObject& ExportObject = Exports[LocalExportIndex];
+	UObject* Object = ExportObject.Object;
+	check(Object || (ExportObject.bFiltered | ExportObject.bExportLoadFailed));
 
-	LLM_SCOPE(ELLMTag::UObject);
-	LLM_SCOPED_TAG_WITH_OBJECT_IN_SET(GetLinkerRoot(), ELLMTagSet::Assets);
-	// LLM_SCOPED_TAG_WITH_OBJECT_IN_SET((Export.DynamicType == FObjectExport::EDynamicType::DynamicType) ? UDynamicClass::StaticClass() :
-	// 	CastEventDrivenIndexToObject<UClass>(Export.ClassIndex, false), ELLMTagSet::AssetClasses);
+	TRACE_LOADTIME_SERIALIZE_EXPORT_SCOPE(Object, Export.CookedSerialSize);
+
+	if ((ExportObject.bFiltered | ExportObject.bExportLoadFailed) || !(Object && Object->HasAnyFlags(RF_NeedLoad)))
+	{
+		if (ExportObject.bExportLoadFailed)
+		{
+			UE_ASYNC_PACKAGE_LOG(Warning, Desc, TEXT("SerializeExport"),
+				TEXT("Skipped failed export %s"), *NameMap->GetName(Export.ObjectName).ToString());
+		}
+		else if (ExportObject.bFiltered)
+		{
+			UE_ASYNC_PACKAGE_LOG_VERBOSE(Verbose, Desc, TEXT("SerializeExport"),
+				TEXT("Skipped filtered export %s"), *NameMap->GetName(Export.ObjectName).ToString());
+		}
+		else
+		{
+			UE_ASYNC_PACKAGE_LOG_VERBOSE(VeryVerbose, Desc, TEXT("SerializeExport"),
+				TEXT("Skipped already serialized export %s"), *NameMap->GetName(Export.ObjectName).ToString());
+		}
+		return false;
+	}
 
 	// If this is a struct, make sure that its parent struct is completely loaded
 	if (UStruct* Struct = dynamic_cast<UStruct*>(Object))
@@ -3551,9 +3615,10 @@ void FAsyncPackage2::EventDrivenSerializeExport(int32 LocalExportIndex, FExportA
 			UStruct* SuperStruct = CastEventDrivenIndexToObject<UStruct>(Export.SuperIndex, true);
 			if (!SuperStruct)
 			{
-				UE_LOG(LogStreaming, Fatal, TEXT("Could not find SuperStruct for %s"), *Object->GetFullName());
-				Exports[LocalExportIndex].bExportLoadFailed = true;
-				return;
+				UE_ASYNC_PACKAGE_LOG(Error, Desc, TEXT("SerializeExport"),
+					TEXT("Could not find SuperStruct object for %s"), *NameMap->GetName(Export.ObjectName).ToString());
+				ExportObject.bExportLoadFailed = true;
+				return false;
 			}
 			Struct->SetSuperStruct(SuperStruct);
 			if (UClass* ClassObject = dynamic_cast<UClass*>(Object))
@@ -3563,6 +3628,11 @@ void FAsyncPackage2::EventDrivenSerializeExport(int32 LocalExportIndex, FExportA
 		}
 	}
 
+	LLM_SCOPE(ELLMTag::UObject);
+	LLM_SCOPED_TAG_WITH_OBJECT_IN_SET(GetLinkerRoot(), ELLMTagSet::Assets);
+	// LLM_SCOPED_TAG_WITH_OBJECT_IN_SET((Export.DynamicType == FObjectExport::EDynamicType::DynamicType) ? UDynamicClass::StaticClass() :
+	// 	CastEventDrivenIndexToObject<UClass>(Export.ClassIndex, false), ELLMTagSet::AssetClasses);
+
 	// cache archetype
 	// prevents GetArchetype from hitting the expensive GetArchetypeFromRequiredInfoImpl
 	check(!Export.TemplateIndex.IsNull());
@@ -3601,8 +3671,12 @@ void FAsyncPackage2::EventDrivenSerializeExport(int32 LocalExportIndex, FExportA
 	}
 #endif
 
+	UE_ASYNC_PACKAGE_LOG_VERBOSE(VeryVerbose, Desc, TEXT("SerializeExport"), TEXT("Serialized export %s"), *Object->GetPathName());
+
 	// push stats so that we don't overflow number of tags per thread during blocking loading
 	LLM_PUSH_STATS_FOR_ASSET_TAGS();
+
+	return true;
 }
 
 EAsyncPackageState::Type FAsyncPackage2::Event_ExportsDone(FAsyncPackage2* Package, int32)
