EditorDomain: Save unversioned properties but not unversioned native.
#rnx

diff --git a/Engine/Source/Editor/UnrealEd/Private/EditorDomain/EditorDomainUtils.cpp b/Engine/Source/Editor/UnrealEd/Private/EditorDomain/EditorDomainUtils.cpp
--- a/Engine/Source/Editor/UnrealEd/Private/EditorDomain/EditorDomainUtils.cpp
+++ b/Engine/Source/Editor/UnrealEd/Private/EditorDomain/EditorDomainUtils.cpp
@@ -51,7 +51,7 @@ const TSet<FName>& GetEditorPackageBlockList()
 }
 
 // Change to a new guid when EditorDomain needs to be invalidated
-const TCHAR* EditorDomainVersion = TEXT("D1718C34CA7C47AEB87A1607568E25B0");
+const TCHAR* EditorDomainVersion = TEXT("86E449B024FD412BB80DB1397EFAE34D");
 // Identifier of the CacheBuckets for EditorDomain tables
 const TCHAR* EditorDomainPackageBucketName = TEXT("EditorDomainPackage");
 const TCHAR* EditorDomainBulkDataListBucketName = TEXT("EditorDomainBulkDataList");
@@ -754,7 +754,7 @@ bool TrySavePackage(UPackage* Package)
 				bSaveUnversioned = false;
 			}
 		}
-		SaveFlags |= bSaveUnversioned ? SAVE_Unversioned : 0;
+		SaveFlags |= bSaveUnversioned ? SAVE_Unversioned_Properties : 0;
 	}
 
 	FMemoryPackageWriter* PackageWriter = new FMemoryPackageWriter();
diff --git a/Engine/Source/Runtime/CoreUObject/Private/UObject/SavePackage.cpp b/Engine/Source/Runtime/CoreUObject/Private/UObject/SavePackage.cpp
--- a/Engine/Source/Runtime/CoreUObject/Private/UObject/SavePackage.cpp
+++ b/Engine/Source/Runtime/CoreUObject/Private/UObject/SavePackage.cpp
@@ -2363,7 +2363,8 @@ FSavePackageResultStruct UPackage::Save(UPackage* InOuter, UObject* Base, EObjec
 				/** If true, we are going to save to disk async to save time. */
 				const bool bSaveAsync = !!(SaveFlags & SAVE_Async);
 
-				const bool bSaveUnversioned = !!(SaveFlags & SAVE_Unversioned);
+				const bool bSaveUnversionedNative = !!(SaveFlags & SAVE_Unversioned_Native);
+				const bool bSaveUnversionedProperties = !!(SaveFlags & SAVE_Unversioned_Properties) && CanUseUnversionedPropertySerialization(TargetPlatform);
 
 				FArchiveFormatterType* Formatter = nullptr;
 				FArchive* TextFormatArchive = nullptr;
@@ -2415,13 +2416,13 @@ FSavePackageResultStruct UPackage::Save(UPackage* InOuter, UObject* Base, EObjec
 						// The entire package will be serialized to memory and then compared against package on disk.
 						// Each difference will be log with its Serialize call stack trace
 						FArchive* Saver = new FArchiveStackTrace(FindAssetInPackage(InOuter), *InOuter->GetLoadedPath().GetPackageName(), true, InOutDiffMap);
-						Linker = MakePimpl<FLinkerSave>(InOuter, Saver, bForceByteSwapping, bSaveUnversioned);
+						Linker = MakePimpl<FLinkerSave>(InOuter, Saver, bForceByteSwapping, bSaveUnversionedNative);
 					}
 					else if (TargetPlatform != nullptr && (SaveFlags & SAVE_DiffOnly))
 					{
 						// The entire package will be serialized to memory and then compared against package on disk
 						FArchive* Saver = new FArchiveStackTrace(FindAssetInPackage(InOuter), *InOuter->GetLoadedPath().GetPackageName(), false);
-						Linker = MakePimpl<FLinkerSave>(InOuter, Saver, bForceByteSwapping, bSaveUnversioned);
+						Linker = MakePimpl<FLinkerSave>(InOuter, Saver, bForceByteSwapping, bSaveUnversionedNative);
 					}
 					else if ((!!TargetPlatform) && FParse::Value(FCommandLine::Get(), TEXT("DiffCookedPackages="), DiffCookedPackagesPath))
 					{
@@ -2444,13 +2445,13 @@ FSavePackageResultStruct UPackage::Save(UPackage* InOuter, UObject* Base, EObjec
 						if (bSaveAsync)
 						{
 							// Allocate the linker with a memory writer, forcing byte swapping if wanted.
-							Linker = MakePimpl<FLinkerSave>(InOuter, bForceByteSwapping, bSaveUnversioned);
+							Linker = MakePimpl<FLinkerSave>(InOuter, bForceByteSwapping, bSaveUnversionedNative);
 						}
 						else
 						{
 							// Allocate the linker, forcing byte swapping if wanted.
 							TempFilename = FPaths::CreateTempFilename(*FPaths::ProjectSavedDir(), *BaseFilename.Left(32));
-							Linker = MakePimpl<FLinkerSave>(InOuter, *TempFilename.GetValue(), bForceByteSwapping, bSaveUnversioned);
+							Linker = MakePimpl<FLinkerSave>(InOuter, *TempFilename.GetValue(), bForceByteSwapping, bSaveUnversionedNative);
 						}
 					}
 					Linker->bProceduralSave = ObjectSaveContext.bProceduralSave;
@@ -2504,8 +2505,7 @@ FSavePackageResultStruct UPackage::Save(UPackage* InOuter, UObject* Base, EObjec
 				Linker->SetFilterEditorOnly( FilterEditorOnly );
 				Linker->SetCookingTarget(TargetPlatform);
 
-				bool bUseUnversionedProperties = bSaveUnversioned && CanUseUnversionedPropertySerialization(TargetPlatform);
-				Linker->SetUseUnversionedPropertySerialization(bUseUnversionedProperties);
+				Linker->SetUseUnversionedPropertySerialization(bSaveUnversionedProperties);
 
 				// Make sure the package has the same version as the linker
 				InOuter->LinkerPackageVersion = Linker->UEVer();
diff --git a/Engine/Source/Runtime/CoreUObject/Private/UObject/SavePackage/SaveContext.h b/Engine/Source/Runtime/CoreUObject/Private/UObject/SavePackage/SaveContext.h
--- a/Engine/Source/Runtime/CoreUObject/Private/UObject/SavePackage/SaveContext.h
+++ b/Engine/Source/Runtime/CoreUObject/Private/UObject/SavePackage/SaveContext.h
@@ -252,14 +252,14 @@ public:
 		return !!(SaveArgs.SaveFlags & SAVE_KeepDirty);
 	}
 
-	bool IsSaveUnversioned() const
+	bool IsSaveUnversionedNative() const
 	{
-		return !!(SaveArgs.SaveFlags & SAVE_Unversioned);
+		return !!(SaveArgs.SaveFlags & SAVE_Unversioned_Native);
 	}
 
-	bool IsUsingUnversionedProperties() const
+	bool IsSaveUnversionedProperties() const
 	{
-		return IsSaveUnversioned() && bCanUseUnversionedPropertySerialization;
+		return !!(SaveArgs.SaveFlags & SAVE_Unversioned_Properties) && bCanUseUnversionedPropertySerialization;
 	}
 
 	bool IsComputeHash() const
diff --git a/Engine/Source/Runtime/CoreUObject/Private/UObject/SavePackage2.cpp b/Engine/Source/Runtime/CoreUObject/Private/UObject/SavePackage2.cpp
--- a/Engine/Source/Runtime/CoreUObject/Private/UObject/SavePackage2.cpp
+++ b/Engine/Source/Runtime/CoreUObject/Private/UObject/SavePackage2.cpp
@@ -704,20 +704,20 @@ ESavePackageResult CreateLinker(FSaveContext& SaveContext)
 			// The entire package will be serialized to memory and then compared against package on disk.
 			// Each difference will be log with its Serialize call stack trace if IsDiffCallstack is true
 			FArchive* Saver = new FArchiveStackTrace(SaveContext.GetAsset(), *SaveContext.GetPackage()->GetLoadedPath().GetPackageName(), SaveContext.IsDiffCallstack(), SaveContext.GetDiffMapPtr());
-			SaveContext.Linker = MakePimpl<FLinkerSave>(SaveContext.GetPackage(), Saver, SaveContext.IsForceByteSwapping(), SaveContext.IsSaveUnversioned());
+			SaveContext.Linker = MakePimpl<FLinkerSave>(SaveContext.GetPackage(), Saver, SaveContext.IsForceByteSwapping(), SaveContext.IsSaveUnversionedNative());
 		}
 		else
 		{
 			if (SaveContext.IsSaveAsync())
 			{
 				// Allocate the linker with a memory writer, forcing byte swapping if wanted.
-				SaveContext.Linker = MakePimpl<FLinkerSave>(SaveContext.GetPackage(), SaveContext.IsForceByteSwapping(), SaveContext.IsSaveUnversioned());
+				SaveContext.Linker = MakePimpl<FLinkerSave>(SaveContext.GetPackage(), SaveContext.IsForceByteSwapping(), SaveContext.IsSaveUnversionedNative());
 			}
 			else
 			{
 				// Allocate the linker, forcing byte swapping if wanted.
 				SaveContext.TempFilename = FPaths::CreateTempFilename(*FPaths::ProjectSavedDir(), *BaseFilename.Left(32));
-				SaveContext.Linker = MakePimpl<FLinkerSave>(SaveContext.GetPackage(), *SaveContext.TempFilename.GetValue(), SaveContext.IsForceByteSwapping(), SaveContext.IsSaveUnversioned());
+				SaveContext.Linker = MakePimpl<FLinkerSave>(SaveContext.GetPackage(), *SaveContext.TempFilename.GetValue(), SaveContext.IsForceByteSwapping(), SaveContext.IsSaveUnversionedNative());
 			}
 		}
 		if (SaveContext.IsGenerateSaveError())
@@ -798,7 +798,7 @@ ESavePackageResult BuildLinker(FSaveContext& SaveContext)
 		SaveContext.Linker->SetFilterEditorOnly(SaveContext.IsFilterEditorOnly());
 		SaveContext.Linker->SetCookingTarget(SaveContext.GetTargetPlatform());
 
-		bool bUseUnversionedProperties = SaveContext.IsUsingUnversionedProperties();
+		bool bUseUnversionedProperties = SaveContext.IsSaveUnversionedProperties();
 		SaveContext.Linker->SetUseUnversionedPropertySerialization(bUseUnversionedProperties);
 
 #if WITH_EDITOR
diff --git a/Engine/Source/Runtime/CoreUObject/Public/UObject/ObjectMacros.h b/Engine/Source/Runtime/CoreUObject/Public/UObject/ObjectMacros.h
--- a/Engine/Source/Runtime/CoreUObject/Public/UObject/ObjectMacros.h
+++ b/Engine/Source/Runtime/CoreUObject/Public/UObject/ObjectMacros.h
@@ -88,7 +88,7 @@ enum ESaveFlags
 	SAVE_KeepDirty		= 0x00000004,	///< Do not clear the dirty flag when saving
 	SAVE_KeepGUID		= 0x00000008,	///< Keep the same guid, used to save cooked packages
 	SAVE_Async			= 0x00000010,	///< Save to a memory writer, then actually write to disk async
-	SAVE_Unversioned	= 0x00000020,	///< Save all versions as zero. Upon load this is changed to the current version. This is only reasonable to use with full cooked builds for distribution.
+	SAVE_Unversioned_Native	= 0x00000020,	///< Save all versions as zero. Upon load this is changed to the current version. This is only reasonable to use with full cooked builds for distribution.
 	SAVE_CutdownPackage	= 0x00000040,	///< Saving cutdown packages in a temp location WITHOUT renaming the package.
 	SAVE_KeepEditorOnlyCookedPackages = 0x00000080,  ///< Keep packages which are marked as editor only even though we are cooking
 	SAVE_Concurrent		= 0x00000100,	///< We are save packages in multiple threads at once and should not call non-threadsafe functions or rely on globals. GIsSavingPackage should be set and PreSave/Postsave functions should be called before/after the entire concurrent save.
@@ -97,6 +97,9 @@ enum ESaveFlags
 	SAVE_ComputeHash    = 0x00000800,	///< Compute the MD5 hash of the cooked data
 	SAVE_CompareLinker	= 0x00001000,	///< Return the linker save to compare against another
 	SAVE_BulkDataByReference = 0x00002000, ///< When saving to a different file than the package's LoadedPath, point bulkdata in the new file to be loaded from the original file.
+	SAVE_Unversioned_Properties = 0x00004000, /// Properties are saved without property name information, and are saved/loaded in the order of the current binary.
+											      ///This is only reasonable to use with full cooked builds for distribution, or with a domain that selects the payload matching a given version externally.
+	SAVE_Unversioned = SAVE_Unversioned_Native | SAVE_Unversioned_Properties,
 };
 
 /** Package flags, passed into UPackage::SetPackageFlags and related functions */
