CIS: Fix PVS warning in FFormatContainer.

#rnx

- PVS is claiming that delete is being called on a void* which is clearly incorrect (I have confirmed that the correct destructor is being called etc)
- Changing the type in the for loop to auto seems to clear up the confusion, so submitting this as a work around while we reach out to PVS about the false positive.

diff --git a/Engine/Source/Runtime/CoreUObject/Public/Serialization/BulkData.h b/Engine/Source/Runtime/CoreUObject/Public/Serialization/BulkData.h
--- a/Engine/Source/Runtime/CoreUObject/Public/Serialization/BulkData.h
+++ b/Engine/Source/Runtime/CoreUObject/Public/Serialization/BulkData.h
@@ -1085,7 +1085,8 @@ public:
 	{
 		for (const TPair<FName, FByteBulkData*>& Format : Formats)
 		{
-			delete Format.Value;
+			FByteBulkData* BulkData = Format.Value;
+			delete BulkData;
 		}
 		Formats.Empty();
 	}
