Improve the documentation for the bulkdata method ::RemoveBulkData to make it clear that calling it will prevent the bulkdata object from reloading payloads from disk in future.

#rnx

- Updated the comments to make the methods intent very clear (unload current payloads and prevent future payloads from being loaded from disk)
- Fix an inconsistency with BD2 running from IoStore, where calling ::RemoveBulkData was not preventing future loads in all cases. We now invalidate the PackageID when running from the IoStore in addition to unregistering the filepath token when running from pak files.
- No longer clear the flags when calling ::RemoveBulkData to be consistent with BD1.
- We might want to consider deprecating this method name and using something like ::Reset instead, but that would imply that we would reset all members and I do not see the value in making this much larger change if we do not need to do so.

diff --git a/Engine/Source/Runtime/CoreUObject/Private/Serialization/BulkData2.cpp b/Engine/Source/Runtime/CoreUObject/Private/Serialization/BulkData2.cpp
--- a/Engine/Source/Runtime/CoreUObject/Private/Serialization/BulkData2.cpp
+++ b/Engine/Source/Runtime/CoreUObject/Private/Serialization/BulkData2.cpp
@@ -1440,13 +1440,11 @@ void FBulkDataBase::RemoveBulkData()
 	FreeData();
 
 	if (!IsUsingIODispatcher())
-	{ 
+	{
 		PackageTokenSystem::UnRegisterToken(Data.Token);
-		Data.Token = InvalidToken;
-	}
-
-	BulkDataFlags = BULKDATA_None;
+	}	
 
+	Data.Token = InvalidToken;
 }
 
 bool FBulkDataBase::StartAsyncLoading()
diff --git a/Engine/Source/Runtime/CoreUObject/Public/Serialization/BulkData.h b/Engine/Source/Runtime/CoreUObject/Public/Serialization/BulkData.h
--- a/Engine/Source/Runtime/CoreUObject/Public/Serialization/BulkData.h
+++ b/Engine/Source/Runtime/CoreUObject/Public/Serialization/BulkData.h
@@ -578,7 +578,10 @@ public:
 	bool IsLocked() const { return LockStatus != LOCKSTATUS_Unlocked; }
 
 	/**
-	 * Clears/ removes the bulk data and resets element count to 0.
+	 * Clears/removes any currently allocated data payload and resets element count to 0.
+	 * 
+	 * Note that once this has been called, the bulkdata object will no longer be able to reload
+	 * it's payload from disk!
 	 */
 	void RemoveBulkData();
 
diff --git a/Engine/Source/Runtime/CoreUObject/Public/Serialization/BulkData2.h b/Engine/Source/Runtime/CoreUObject/Public/Serialization/BulkData2.h
--- a/Engine/Source/Runtime/CoreUObject/Public/Serialization/BulkData2.h
+++ b/Engine/Source/Runtime/CoreUObject/Public/Serialization/BulkData2.h
@@ -209,6 +209,12 @@ public:
 
 	static IBulkDataIORequest* CreateStreamingRequestForRange(const BulkDataRangeArray& RangeArray, EAsyncIOPriorityAndFlags Priority, FBulkDataIORequestCallBack* CompleteCallback);
 
+	/**
+	 * Clears/removes any currently allocated data payload.
+	 *
+	 * Note that once this has been called, the bulkdata object will no longer be able to reload
+	 * it's payload from disk!
+	 */
 	void RemoveBulkData();
 
 	/**
