AsyncLoading2: Add back missing fileopenlog generation functionality (unintentionally removed in cl 16481564)

#rnx

diff --git a/Engine/Source/Runtime/CoreUObject/Private/Serialization/AsyncLoading2.cpp b/Engine/Source/Runtime/CoreUObject/Private/Serialization/AsyncLoading2.cpp
--- a/Engine/Source/Runtime/CoreUObject/Private/Serialization/AsyncLoading2.cpp
+++ b/Engine/Source/Runtime/CoreUObject/Private/Serialization/AsyncLoading2.cpp
@@ -2749,6 +2749,12 @@ bool FAsyncLoadingThread2::CreateAsyncPackagesFromQueue(FAsyncLoadingThreadState
 				{
 					UE_ASYNC_PACKAGE_LOG(Verbose, PackageDesc, TEXT("CreateAsyncPackages: AddPackage"),
 						TEXT("Start loading package."));
+#if !UE_BUILD_SHIPPING
+						if (FileOpenLogWrapper)
+						{
+							FileOpenLogWrapper->AddPackageToOpenLog(*PackageDesc.UPackageName.ToString());
+						}
+#endif
 				}
 				else
 				{
