Mirage will now check for the existence of the requested backend graph in the ini files before attempting to parse it.

#rnx

- Previously the user requested a none existence graph we'd give an error message about being unable to find the 'LocalStorageHierarchy' entry in the graph, which was a bit misleading as it indicated that the graph existed but was malformed.
- Now we give a clear error message that the graph itself does not exist.

diff --git a/Engine/Source/Runtime/CoreUObject/Private/Virtualization/VirtualizationManager.cpp b/Engine/Source/Runtime/CoreUObject/Private/Virtualization/VirtualizationManager.cpp
--- a/Engine/Source/Runtime/CoreUObject/Private/Virtualization/VirtualizationManager.cpp
+++ b/Engine/Source/Runtime/CoreUObject/Private/Virtualization/VirtualizationManager.cpp
@@ -549,6 +549,11 @@ void FVirtualizationManager::MountBackends()
 
 	const TCHAR* GraphName = *BackendGraphName;
 
+	if(!GConfig->DoesSectionExist(GraphName, GEngineIni))
+	{
+		UE_LOG(LogVirtualization, Fatal, TEXT("Unable to find the backend graph: '%s' [ini=%s]."), GraphName, *GEngineIni);
+	}
+
 	UE_LOG(LogVirtualization, Log, TEXT("Using backend graph: '%s'"), GraphName);
 
 	// It is important to parse the local storage hierarchy first so those backends will show up before the
