Cache all mounted roots in FLongPackagePathsSingleton instead of generating a string array on each calls from various functions in FPackageName.
Greatly reduce temporary memory allocations from callers (worst offender during editor loads was FPackageName::IsValidLongPackageName).

diff --git a/Engine/Source/Runtime/CoreUObject/Private/Misc/PackageName.cpp b/Engine/Source/Runtime/CoreUObject/Private/Misc/PackageName.cpp
--- a/Engine/Source/Runtime/CoreUObject/Private/Misc/PackageName.cpp
+++ b/Engine/Source/Runtime/CoreUObject/Private/Misc/PackageName.cpp
@@ -205,6 +205,9 @@ struct FLongPackagePathsSingleton
 	TArray<FPathPair> ContentRootToPath;
 	TArray<FPathPair> ContentPathToRoot;
 
+	// Cached values
+	TArray<FString> ValidLongPackageRoots[2];
+
 	// singleton
 	static FLongPackagePathsSingleton& Get()
 	{
@@ -212,23 +215,22 @@ struct FLongPackagePathsSingleton
 		return Singleton;
 	}
 
-	void GetValidLongPackageRoots(TArray<FString>& OutRoots, bool bIncludeReadOnlyRoots) const
+	const TArray<FString>& GetValidLongPackageRoots(bool bIncludeReadOnlyRoots) const
 	{
-		OutRoots.Add(EngineRootPath);
-		OutRoots.Add(GameRootPath);
-
-		{
-			FReadScopeLock ScopeLock(ContentMountPointCriticalSection);
-			OutRoots += MountPointRootPaths;
-		}
+		return ValidLongPackageRoots[bIncludeReadOnlyRoots ? 1 : 0];
+	}
 
-		if (bIncludeReadOnlyRoots)
-		{
-			OutRoots.Add(ConfigRootPath);
-			OutRoots.Add(ScriptRootPath);
-			OutRoots.Add(MemoryRootPath);
-			OutRoots.Add(TempRootPath);
-		}
+	void UpdateValidLongPackageRoots()
+	{
+		ValidLongPackageRoots[0].Empty();
+		ValidLongPackageRoots[0].Add(EngineRootPath);
+		ValidLongPackageRoots[0].Add(GameRootPath);
+		ValidLongPackageRoots[0] += MountPointRootPaths;
+		ValidLongPackageRoots[1] = ValidLongPackageRoots[0];
+		ValidLongPackageRoots[1].Add(ConfigRootPath);
+		ValidLongPackageRoots[1].Add(ScriptRootPath);
+		ValidLongPackageRoots[1].Add(MemoryRootPath);
+		ValidLongPackageRoots[1].Add(TempRootPath);
 	}
 
 	// Given a content path ensure it is consistent, specifically with FileManager relative paths 
@@ -275,7 +277,9 @@ struct FLongPackagePathsSingleton
 			ContentRootToPath.Insert(Pair, 0);
 			ContentPathToRoot.Insert(Pair, 0);
 			MountPointRootPaths.Add(RootPath);
-		}
+
+			UpdateValidLongPackageRoots();
+		}		
 
 		// Let subscribers know that a new content path was mounted
 		FPackageName::OnContentPathMounted().Broadcast( RootPath, RelativeContentPath);
@@ -306,7 +310,9 @@ struct FLongPackagePathsSingleton
 
 				// Let subscribers know that a new content path was unmounted
 				bFirePathDismountedDelegate = true;
-			}
+
+				UpdateValidLongPackageRoots();
+			}			
 		}
 
 		if (bFirePathDismountedDelegate)
@@ -395,6 +401,8 @@ private:
 		// a Core class, but content path functionality is added at the CoreUObject level.
 		IPluginManager::Get().SetRegisterMountPointDelegate( IPluginManager::FRegisterMountPointDelegate::CreateStatic( &FPackageName::RegisterMountPoint ) );
 		IPluginManager::Get().SetUnRegisterMountPointDelegate( IPluginManager::FRegisterMountPointDelegate::CreateStatic( &FPackageName::UnRegisterMountPoint ) );
+
+		UpdateValidLongPackageRoots();
 	}
 };
 
@@ -705,8 +713,7 @@ bool FPackageName::SplitLongPackageName(const FString& InLongPackageName, FStrin
 	const FLongPackagePathsSingleton& Paths = FLongPackagePathsSingleton::Get();
 
 	const bool bIncludeReadOnlyRoots = true;
-	TArray<FString> ValidRoots;
-	Paths.GetValidLongPackageRoots(ValidRoots, bIncludeReadOnlyRoots);
+	const TArray<FString>& ValidRoots = Paths.GetValidLongPackageRoots(bIncludeReadOnlyRoots);
 
 	// Check to see whether our package came from a valid root
 	OutPackageRoot.Empty();
@@ -883,8 +890,7 @@ bool FPackageName::IsValidLongPackageName(FStringView InLongPackageName, bool bI
 			if (Reason == EErrorCode::PackageNamePathNotMounted)
 			{
 				const FLongPackagePathsSingleton& Paths = FLongPackagePathsSingleton::Get();
-				TArray<FString> ValidRoots;
-				Paths.GetValidLongPackageRoots(ValidRoots, bIncludeReadOnlyRoots);
+				const TArray<FString>& ValidRoots = Paths.GetValidLongPackageRoots(bIncludeReadOnlyRoots);
 				if (ValidRoots.Num() == 0)
 				{
 					*OutReason = NSLOCTEXT("Core", "LongPackageNames_NoValidRoots", "No valid roots exist!");
@@ -932,9 +938,8 @@ bool FPackageName::IsValidLongPackageName(FStringView InLongPackageName, bool bI
 
 	// Check valid roots
 	const FLongPackagePathsSingleton& Paths = FLongPackagePathsSingleton::Get();
-	TArray<FString> ValidRoots;
+	const TArray<FString>& ValidRoots = Paths.GetValidLongPackageRoots(bIncludeReadOnlyRoots);
 	bool bValidRoot = false;
-	Paths.GetValidLongPackageRoots(ValidRoots, bIncludeReadOnlyRoots);
 	for (int32 RootIdx = 0; RootIdx < ValidRoots.Num(); ++RootIdx)
 	{
 		const FString& Root = ValidRoots[RootIdx];
@@ -1047,8 +1052,7 @@ FName FPackageName::GetPackageMountPoint(const FString& InPackagePath, bool InWi
 {
 	FLongPackagePathsSingleton& Paths = FLongPackagePathsSingleton::Get();
 	
-	TArray<FString> MountPoints;
-	Paths.GetValidLongPackageRoots(MountPoints, true);
+	const TArray<FString>& MountPoints = Paths.GetValidLongPackageRoots(true);
 
 	int32 WithoutSlashes = InWithoutSlashes ? 1 : 0;
 	for (auto RootIt = MountPoints.CreateConstIterator(); RootIt; ++RootIt)
@@ -2025,7 +2029,7 @@ void FPackageName::IteratePackagesInDirectory(const FString& RootDir, const FPac
 void FPackageName::QueryRootContentPaths(TArray<FString>& OutRootContentPaths, bool bIncludeReadOnlyRoots, bool bWithoutLeadingSlashes, bool bWithoutTrailingSlashes)
 {
 	const FLongPackagePathsSingleton& Paths = FLongPackagePathsSingleton::Get();
-	Paths.GetValidLongPackageRoots( OutRootContentPaths, bIncludeReadOnlyRoots );
+	OutRootContentPaths = Paths.GetValidLongPackageRoots( bIncludeReadOnlyRoots );
 
 	if (bWithoutTrailingSlashes || bWithoutLeadingSlashes)
 	{
