FCoreUObjectDelegates::OnEndLoadPackage: Add broadcast from the early exits in LoadPackageInternal that were missing it.

diff --git a/Engine/Source/Runtime/CoreUObject/Private/UObject/UObjectGlobals.cpp b/Engine/Source/Runtime/CoreUObject/Private/UObject/UObjectGlobals.cpp
--- a/Engine/Source/Runtime/CoreUObject/Private/UObject/UObjectGlobals.cpp
+++ b/Engine/Source/Runtime/CoreUObject/Private/UObject/UObjectGlobals.cpp
@@ -1098,6 +1098,22 @@ public:
 // @todo: remove this in the new loader
 static int32 GGameThreadLoadCounter = 0;
 
+/** Notify delegate listeners of all the packages that loaded; called only once per explicit call to LoadPackage. */
+void BroadcastEndLoad(TArrayView<UPackage*> LoadedPackages)
+{
+#if WITH_EDITOR
+	// check(IsInGameThread()) was called by the caller, but we still need to test !IsInAsyncLoadingThread to exclude that callsite when the engine is single-threaded
+	if (GIsEditor && !IsInAsyncLoadingThread() && GGameThreadLoadCounter == 0)
+	{
+		for (UPackage* LoadedPackage : LoadedPackages)
+		{
+			LoadedPackage->bHasBeenEndLoaded = true;
+		}
+		FCoreUObjectDelegates::OnEndLoadPackage.Broadcast(LoadedPackages);
+	}
+#endif
+}
+
 UE_TRACE_EVENT_BEGIN(CUSTOM_LOADTIMER_LOG, LoadPackageInternal, NoSync)
 	UE_TRACE_EVENT_FIELD(UE::Trace::WideString, PackageName)
 UE_TRACE_EVENT_END()
@@ -1233,7 +1249,8 @@ UPackage* LoadPackageInternal(UPackage* InOuter, const FPackagePath& PackagePath
 
 		if (!Linker)
 		{
-			EndLoad(LoadContext);
+			EndLoad(LoadContext, &LoadedPackages);
+			BroadcastEndLoad(LoadedPackages);
 			return nullptr;
 		}
 
@@ -1278,7 +1295,8 @@ UPackage* LoadPackageInternal(UPackage* InOuter, const FPackagePath& PackagePath
 		{
 			// The linker is associated with a package that has already been loaded.
 			// Loading packages that have already been loaded is unsupported.
-			EndLoadAndCopyLocalizationGatherFlag();			
+			EndLoadAndCopyLocalizationGatherFlag();	
+			BroadcastEndLoad(LoadedPackages);
 			return Result;
 		}
 
@@ -1423,17 +1441,7 @@ UPackage* LoadPackageInternal(UPackage* InOuter, const FPackagePath& PackagePath
 		Result->SetFlags(RF_WasLoaded);
 	}
 
-#if WITH_EDITOR
-	if (GIsEditor && !IsInAsyncLoadingThread() && GGameThreadLoadCounter == 0)
-	{
-		// check(IsInGameThread() was called above, but we still need to test !IsInAsyncLoadingThread to exclude that callsite when the engine is single-threaded
-		for (UPackage* LoadedPackage : LoadedPackages)
-		{
-			LoadedPackage->bHasBeenEndLoaded = true;
-		}
-		FCoreUObjectDelegates::OnEndLoadPackage.Broadcast(LoadedPackages);
-	}
-#endif
+	BroadcastEndLoad(LoadedPackages);
 	return Result;
 }
 
