Fixed crash when trying to convert empty or invalid package paths to filenames
#rnx

diff --git a/Engine/Source/Runtime/CoreUObject/Private/Misc/PackagePath.cpp b/Engine/Source/Runtime/CoreUObject/Private/Misc/PackagePath.cpp
--- a/Engine/Source/Runtime/CoreUObject/Private/Misc/PackagePath.cpp
+++ b/Engine/Source/Runtime/CoreUObject/Private/Misc/PackagePath.cpp
@@ -1131,9 +1131,20 @@ FString FPackagePath::GetLocalFullPath(EPackageSegment PackageSegment) const
 
 FString FPackagePath::GetLocalBaseFilenameWithPath() const
 {
-	FString PackageNameString = PackageName.ToString();
-	FString LocalPath = FPackageName::LongPackageNameToFilename(PackageNameString);
-	return LocalPath;
+	if (!PackageName.IsNone())
+	{
+		FString PackageNameString = PackageName.ToString();
+		FString Result;
+		if (FPackageName::TryConvertLongPackageNameToFilename(PackageNameString, Result))
+		{
+			return Result;
+		}
+		else
+		{
+			UE_LOG(LogPackageName, Error, TEXT("GetLocalBaseFilenameWithPath: Failed converting package name \"%s\" to file name"), *PackageNameString);
+		}
+	}
+	return TEXT("");
 }
 
 FString FPackagePath::GetDebugName(EPackageSegment PackageSegment) const
