Strip unused names from package headers when cooking with unversioned property serialization
#rnx

diff --git a/Engine/Source/Runtime/CoreUObject/Private/UObject/SavePackage2.cpp b/Engine/Source/Runtime/CoreUObject/Private/UObject/SavePackage2.cpp
--- a/Engine/Source/Runtime/CoreUObject/Private/UObject/SavePackage2.cpp
+++ b/Engine/Source/Runtime/CoreUObject/Private/UObject/SavePackage2.cpp
@@ -278,6 +278,7 @@ ESavePackageResult HarvestPackage(FSaveContext& SaveContext)
 	SCOPED_SAVETIMER(UPackage_Save_HarvestPackage);
 
 	FPackageHarvester Harvester(SaveContext);
+	Harvester.SetUseUnversionedPropertySerialization(SaveContext.IsSaveUnversionedProperties());
 	EObjectFlags TopLevelFlags = SaveContext.GetTopLevelFlags();
 	UObject* Asset = SaveContext.GetAsset();
 
