Undo change 17444582
#rnx

diff --git a/Engine/Source/Runtime/CoreUObject/Private/UObject/SavePackage2.cpp b/Engine/Source/Runtime/CoreUObject/Private/UObject/SavePackage2.cpp
--- a/Engine/Source/Runtime/CoreUObject/Private/UObject/SavePackage2.cpp
+++ b/Engine/Source/Runtime/CoreUObject/Private/UObject/SavePackage2.cpp
@@ -278,7 +278,7 @@ ESavePackageResult HarvestPackage(FSaveContext& SaveContext)
 	SCOPED_SAVETIMER(UPackage_Save_HarvestPackage);
 
 	FPackageHarvester Harvester(SaveContext);
-	Harvester.SetUseUnversionedPropertySerialization(SaveContext.IsSaveUnversionedProperties());
+	//Harvester.SetUseUnversionedPropertySerialization(SaveContext.IsSaveUnversionedProperties());
 	EObjectFlags TopLevelFlags = SaveContext.GetTopLevelFlags();
 	UObject* Asset = SaveContext.GetAsset();
 
