Implemented FromLocalPath on slimmed down FPackagePath
#rnx

diff --git a/Engine/Source/Runtime/CoreUObject/Private/Misc/PackagePath.cpp b/Engine/Source/Runtime/CoreUObject/Private/Misc/PackagePath.cpp
--- a/Engine/Source/Runtime/CoreUObject/Private/Misc/PackagePath.cpp
+++ b/Engine/Source/Runtime/CoreUObject/Private/Misc/PackagePath.cpp
@@ -1054,8 +1054,23 @@ FPackagePath FPackagePath::FromPackageNameUnchecked(FName InPackageName)
 
 FPackagePath FPackagePath::FromLocalPath(FStringView InFilename, EPackageSegment& OutPackageSegment)
 {
-	unimplemented();
-	return FPackagePath();
+	int32 ExtensionStart;
+	EPackageExtension Extension = ParseExtension(InFilename, &ExtensionStart);
+	OutPackageSegment = ExtensionToSegment(Extension);
+	if (OutPackageSegment != EPackageSegment::Header)
+	{
+		Extension = EPackageExtension::Unspecified;
+	}
+	FString FileNameString(InFilename.Left(ExtensionStart));
+	FString PackageNameString;
+	if (!FPackageName::TryConvertFilenameToLongPackageName(FileNameString, PackageNameString))
+	{
+		UE_LOG(LogPackageName, Error, TEXT("FromLocalPath: Failed converting filename \"%s\" to package name"), *FileNameString);
+	}
+	FPackagePath PackagePath;
+	PackagePath.PackageName = FName(PackageNameString);
+	PackagePath.HeaderExtension = Extension;
+	return PackagePath;
 }
 
 FPackagePath FPackagePath::FromMountedComponents(FStringView PackageNameRoot, FStringView FilePathRoot, FStringView RelPath,
