PR #8378: Fix missing return value in PropertyAccessUtil.cpp, causing FEnumProperty copies to always fail (Contributed by [REDACTED_PERSON])

diff --git a/Engine/Source/Runtime/CoreUObject/Private/UObject/PropertyAccessUtil.cpp b/Engine/Source/Runtime/CoreUObject/Private/UObject/PropertyAccessUtil.cpp
--- a/Engine/Source/Runtime/CoreUObject/Private/UObject/PropertyAccessUtil.cpp
+++ b/Engine/Source/Runtime/CoreUObject/Private/UObject/PropertyAccessUtil.cpp
@@ -29,7 +29,7 @@ int64 GetPropertyEnumValue(const FProperty* InProp, const void* InPropValue)
 	}
 	if (const FEnumProperty* EnumProp = CastField<FEnumProperty>(InProp))
 	{
-		EnumProp->GetUnderlyingProperty()->GetSignedIntPropertyValue(InPropValue);
+		return EnumProp->GetUnderlyingProperty()->GetSignedIntPropertyValue(InPropValue);
 	}
 	return INDEX_NONE;
 }
