Fixing linux build
#rnx

diff --git a/Engine/Source/Runtime/CoreUObject/Private/UObject/StructVariant.cpp b/Engine/Source/Runtime/CoreUObject/Private/UObject/StructVariant.cpp
--- a/Engine/Source/Runtime/CoreUObject/Private/UObject/StructVariant.cpp
+++ b/Engine/Source/Runtime/CoreUObject/Private/UObject/StructVariant.cpp
@@ -1,6 +1,5 @@
 // Copyright Epic Games, Inc. All Rights Reserved.
 
-#pragma once
 
 #include "UObject/StructVariant.h"
 #include "UObject/PropertyPortFlags.h"
