COTF: Better lock scope when handling cooked packages (server) and moved subscription to cook-on-the-fly messages (client) when package store is initializing

#rnx

diff --git a/Engine/Source/Editor/UnrealEd/Private/Cooker/IoStoreCookOnTheFlyRequestManager.cpp b/Engine/Source/Editor/UnrealEd/Private/Cooker/IoStoreCookOnTheFlyRequestManager.cpp
--- a/Engine/Source/Editor/UnrealEd/Private/Cooker/IoStoreCookOnTheFlyRequestManager.cpp
+++ b/Engine/Source/Editor/UnrealEd/Private/Cooker/IoStoreCookOnTheFlyRequestManager.cpp
@@ -63,6 +63,9 @@ public:
 	{
 		TArray<int32> CookedEntryIndices;
 		TArray<FPackageId> FailedPackages;
+		int32 TotalCooked = 0;
+		int32 TotalFailed = 0;
+		int32 TotalTracked = 0;
 	};
 
 	TArrayView<const int32> GetCookedEntryIndices() const
@@ -172,6 +175,9 @@ public:
 		}
 
 		FCompletedPackages CompletedPackages;
+		CompletedPackages.TotalCooked = CookedEntries.Num();
+		CompletedPackages.TotalFailed = FailedPackages.Num();
+		CompletedPackages.TotalTracked = Packages.Num();
 
 		if (CookedEntries.Num() > NumCookedPackages || FailedPackages.Num() > NumFailedPackages)
 		{
@@ -186,7 +192,7 @@ public:
 			}
 		}
 
-		return MoveTemp(CompletedPackages);
+		return CompletedPackages;
 	}
 
 	void AddCompletedPackages(TArrayView<const FPackageStoreEntryResource> Entries)
@@ -592,9 +598,6 @@ private:
 
 		TRACE_CPUPROFILER_EVENT_SCOPE(CookOnTheFly::OnPackageCooked);
 
-		FPlatformContext& Context = GetContext(EventArgs.PlatformName);
-		FScopeLock _(&Context.CriticalSection);
-
 		if (EventArgs.AdditionalFiles.Num())
 		{
 			TArray<FString> Filenames;
@@ -614,12 +617,16 @@ private:
 				*Ar << ChunkIds;
 			}
 
-			ConnectionServer->BroadcastMessage(Message, Context.PlatformName);
+			ConnectionServer->BroadcastMessage(Message, EventArgs.PlatformName);
 		}
-
+		
 		FCookOnTheFlyPackageTracker::FCompletedPackages NewCompletedPackages;
 		{
 			TRACE_CPUPROFILER_EVENT_SCOPE(CookOnTheFly::MarkAsCooked);
+			
+			FPlatformContext& Context = GetContext(EventArgs.PlatformName);
+			FScopeLock _(&Context.CriticalSection);
+
 			NewCompletedPackages = Context.PackageTracker.MarkAsCompleted(EventArgs.PackageName, EventArgs.EntryIndex, EventArgs.Entries);
 		}
 
@@ -628,8 +635,8 @@ private:
 			TRACE_CPUPROFILER_EVENT_SCOPE(CookOnTheFly::SendCookedPackagesMessage);
 
 			FPackageStoreData PackageStoreData;
-			PackageStoreData.TotalCookedPackages = Context.PackageTracker.GetCookedEntryIndices().Num();
-			PackageStoreData.TotalFailedPackages = Context.PackageTracker.GetFailedPackages().Num();
+			PackageStoreData.TotalCookedPackages = NewCompletedPackages.TotalCooked;
+			PackageStoreData.TotalFailedPackages = NewCompletedPackages.TotalFailed;
 			PackageStoreData.FailedPackages = MoveTemp(NewCompletedPackages.FailedPackages);
 
 			for (int32 EntryIndex : NewCompletedPackages.CookedEntryIndices)
@@ -642,13 +649,13 @@ private:
 				LexToString(ECookOnTheFlyMessage::PackagesCooked),
 				PackageStoreData.CookedPackages.Num(),
 				PackageStoreData.FailedPackages.Num(),
-				Context.PackageTracker.TotalTracked(),
+				NewCompletedPackages.TotalTracked,
 				PackageStoreData.TotalCookedPackages,
 				PackageStoreData.TotalFailedPackages);
 
 			FCookOnTheFlyMessage Message(ECookOnTheFlyMessage::PackagesCooked);
 			Message.SetBodyTo<FPackagesCookedMessage>(FPackagesCookedMessage { MoveTemp(PackageStoreData) });
-			ConnectionServer->BroadcastMessage(Message, Context.PlatformName);
+			ConnectionServer->BroadcastMessage(Message, EventArgs.PlatformName);
 		}
 	}
 
diff --git a/Engine/Source/Runtime/CoreUObject/Private/Serialization/CookOnTheFlyPackageStore.cpp b/Engine/Source/Runtime/CoreUObject/Private/Serialization/CookOnTheFlyPackageStore.cpp
--- a/Engine/Source/Runtime/CoreUObject/Private/Serialization/CookOnTheFlyPackageStore.cpp
+++ b/Engine/Source/Runtime/CoreUObject/Private/Serialization/CookOnTheFlyPackageStore.cpp
@@ -58,7 +58,6 @@ public:
 	{
 		// Index zero is invalid
 		PackageEntries.Add();
-		CookOnTheFlyServerConnection.OnMessage().AddRaw(this, &FCookOnTheFlyPackageStore::OnCookOnTheFlyMessage);
 	}
 
 	~FCookOnTheFlyPackageStore()
@@ -71,6 +70,8 @@ private:
 		using namespace UE::Cook;
 		using namespace UE::PackageStore::Messaging;
 
+		CookOnTheFlyServerConnection.OnMessage().AddRaw(this, &FCookOnTheFlyPackageStore::OnCookOnTheFlyMessage);
+
 		FCookOnTheFlyRequest Request(ECookOnTheFlyMessage::GetCookedPackages);
 		FCookOnTheFlyResponse Response = CookOnTheFlyServerConnection.SendRequest(Request).Get();
 
@@ -82,13 +83,20 @@ private:
 			UE_LOG(LogCookOnTheFly, Log, TEXT("Got '%d' cooked and '%d' failed packages from server"),
 				PackageStoreData.CookedPackages.Num(), PackageStoreData.FailedPackages.Num());
 
-			AddPackages(MoveTemp(PackageStoreData.CookedPackages), MoveTemp(PackageStoreData.FailedPackages));
+			AddPackages(
+				MoveTemp(PackageStoreData.CookedPackages),
+				MoveTemp(PackageStoreData.FailedPackages),
+				PackageStoreData.TotalCookedPackages,
+				PackageStoreData.TotalFailedPackages);
 
 			LastServerActivtyTime = FPlatformTime::Seconds();
+
+			bInitialized = true;
 		}
 		else
 		{
 			UE_LOG(LogCookOnTheFly, Warning, TEXT("Failed to send '%s' request"), LexToString(Request.GetHeader().MessageType));
+			bInitialized = false;
 		}
 	}
 
@@ -188,7 +196,7 @@ private:
 		return false;
 	}
 
-	void AddPackages(TArray<FPackageStoreEntryResource> Entries, TArray<FPackageId> FailedPackageIds)
+	void AddPackages(TArray<FPackageStoreEntryResource> Entries, TArray<FPackageId> FailedPackageIds, int32 ServerTotalCookedPackages, int32 ServerTotalFailedPackages)
 	{
 		FScopeLock _(&CriticalSection);
 		
@@ -218,6 +226,11 @@ private:
 					*PackageEntries[EntryInfo.EntryIndex].PackageName.ToString(), PackageStats.Cooked.Load(), PackageStats.Failed.Load());
 			}
 		}
+		
+		UE_CLOG(bInitialized && (ServerTotalCookedPackages != PackageStats.Cooked.Load() || ServerTotalFailedPackages != PackageStats.Failed.Load()),
+			LogCookOnTheFly, Warning, TEXT("Client/Server package mismatch, Cooked='%d/%d', Failed='%d/%d' (Client/Server)"),
+			PackageStats.Cooked.Load(), ServerTotalCookedPackages, PackageStats.Failed.Load(), ServerTotalFailedPackages);
+
 	}
 
 	void OnCookOnTheFlyMessage(const UE::Cook::FCookOnTheFlyMessage& Message)
@@ -238,11 +251,11 @@ private:
 					PackageStoreData.TotalCookedPackages,
 					PackageStoreData.TotalFailedPackages);
 
-				AddPackages(MoveTemp(PackageStoreData.CookedPackages), MoveTemp(PackageStoreData.FailedPackages));
-
-				UE_CLOG(PackageStoreData.TotalCookedPackages != PackageStats.Cooked.Load() || PackageStoreData.TotalFailedPackages != PackageStats.Failed.Load(),
-					LogCookOnTheFly, Warning, TEXT("Client/Server package mismatch, Cooked='%d/%d', Failed='%d/%d' (Client/Server)"),
-					PackageStats.Cooked.Load(), PackageStoreData.TotalCookedPackages, PackageStats.Failed.Load(), PackageStoreData.TotalFailedPackages);
+				AddPackages(
+					MoveTemp(PackageStoreData.CookedPackages),
+					MoveTemp(PackageStoreData.FailedPackages),
+					PackageStoreData.TotalCookedPackages,
+					PackageStoreData.TotalFailedPackages);
 
 				EntriesAddedCallback();
 
@@ -305,6 +318,7 @@ private:
 	TMap<FPackageId, FEntryInfo> PackageIdToEntryInfo;
 	TChunkedArray<FPackageStoreEntryResource> PackageEntries;
 	FPackageStats PackageStats;
+	TAtomic<bool> bInitialized {false};
 
 	const double MaxInactivityTime = 20;
 	const double TimeBetweenWarning = 10;
