Fixed another crash when trying to convert empty or invalid package paths to filenames
#rnx

diff --git a/Engine/Source/Runtime/CoreUObject/Private/Misc/PackagePath.cpp b/Engine/Source/Runtime/CoreUObject/Private/Misc/PackagePath.cpp
--- a/Engine/Source/Runtime/CoreUObject/Private/Misc/PackagePath.cpp
+++ b/Engine/Source/Runtime/CoreUObject/Private/Misc/PackagePath.cpp
@@ -1138,10 +1138,21 @@ FName FPackagePath::GetPackageFName() const
 
 FString FPackagePath::GetLocalFullPath(EPackageSegment PackageSegment) const
 {
-	FString PackageNameString = PackageName.ToString();
-	EPackageExtension Extension = PackageSegment == EPackageSegment::Header ? HeaderExtension : SegmentToExtension(PackageSegment);
-	FString LocalPath = FPackageName::LongPackageNameToFilename(PackageNameString, LexToString(Extension));
-	return LocalPath;
+	if (!PackageName.IsNone())
+	{
+		FString PackageNameString = PackageName.ToString();
+		EPackageExtension Extension = PackageSegment == EPackageSegment::Header ? HeaderExtension : SegmentToExtension(PackageSegment);
+		FString Result;
+		if (FPackageName::TryConvertLongPackageNameToFilename(PackageNameString, Result, LexToString(Extension)))
+		{
+			return Result;
+		}
+		else
+		{
+			UE_LOG(LogPackageName, Error, TEXT("GetLocalFullPath: Failed converting package name \"%s\" to file name"), *PackageNameString);
+		}
+	}
+	return TEXT("");
 }
 
 FString FPackagePath::GetLocalBaseFilenameWithPath() const
