Remove PackageIdToLoad field from UPackage in favour of generating the id dynamically from the loaded path
#rnx

diff --git a/Engine/Source/Runtime/CoreUObject/Private/Serialization/AsyncLoading2.cpp b/Engine/Source/Runtime/CoreUObject/Private/Serialization/AsyncLoading2.cpp
--- a/Engine/Source/Runtime/CoreUObject/Private/Serialization/AsyncLoading2.cpp
+++ b/Engine/Source/Runtime/CoreUObject/Private/Serialization/AsyncLoading2.cpp
@@ -5739,7 +5739,6 @@ void FAsyncPackage2::CreateUPackage(const FPackageSummary* PackageSummary)
 		LinkerRoot->SetLoadedPath(FPackagePath::FromPackageNameUnchecked(Desc.PackageNameToLoad));
 		LinkerRoot->SetCanBeImportedFlag(Desc.bCanBeImported);
 		LinkerRoot->SetPackageId(Desc.UPackageId);
-		LinkerRoot->SetPackageIdToLoad(Desc.PackageIdToLoad);
 		LinkerRoot->SetPackageFlagsTo(PackageSummary->PackageFlags | PKG_Cooked);
 		LinkerRoot->LinkerPackageVersion = GPackageFileUEVersion;
 		LinkerRoot->LinkerLicenseeVersion = GPackageFileLicenseeUEVersion;
diff --git a/Engine/Source/Runtime/CoreUObject/Public/UObject/Package.h b/Engine/Source/Runtime/CoreUObject/Public/UObject/Package.h
--- a/Engine/Source/Runtime/CoreUObject/Public/UObject/Package.h
+++ b/Engine/Source/Runtime/CoreUObject/Public/UObject/Package.h
@@ -230,9 +230,6 @@ private:
 	/** Globally unique id */
 	FPackageId PackageId;
 
-	/** Globally unique id used to address I/O chunks within the package */
-	FPackageId PackageIdToLoad;
-
 	/** The PackagePath this package was loaded from */
 	FPackagePath LoadedPath;
 public:
@@ -603,13 +600,7 @@ public:
 	/** returns the unique package id to load */
 	FORCEINLINE FPackageId GetPackageIdToLoad() const
 	{
-		return PackageIdToLoad;
-	}
-
-	/** sets the unique package id to load */
-	FORCEINLINE void SetPackageIdToLoad(FPackageId InPackageIdToLoad)
-	{
-		PackageIdToLoad = InPackageIdToLoad;
+		return FPackageId::FromName(LoadedPath.GetPackageFName());
 	}
 
 	/**
