SavePackage2: Fix missing calculation of SerializedPackageFlags.
#rnx

diff --git a/Engine/Source/Runtime/CoreUObject/Private/UObject/SavePackage/SaveContext.h b/Engine/Source/Runtime/CoreUObject/Private/UObject/SavePackage/SaveContext.h
--- a/Engine/Source/Runtime/CoreUObject/Private/UObject/SavePackage/SaveContext.h
+++ b/Engine/Source/Runtime/CoreUObject/Private/UObject/SavePackage/SaveContext.h
@@ -539,8 +539,11 @@ public:
 			return Result;
 		}
 
-		ESavePackageResult FinalResult = IsStubRequested() ? ESavePackageResult::GenerateStub : (bDiffOnlyIdentical ? ESavePackageResult::Success : ESavePackageResult::DifferentContent);
-		return FSavePackageResultStruct(FinalResult, TotalPackageSizeUncompressed, AsyncWriteAndHashSequence.Finalize(EAsyncExecution::TaskGraph, MoveTemp(HashCompletionFunc)), 0, IsCompareLinker() ? MoveTemp(Linker) : nullptr);
+		ESavePackageResult FinalResult = IsStubRequested() ? ESavePackageResult::GenerateStub
+			: (bDiffOnlyIdentical ? ESavePackageResult::Success : ESavePackageResult::DifferentContent);
+		return FSavePackageResultStruct(FinalResult, TotalPackageSizeUncompressed,
+			AsyncWriteAndHashSequence.Finalize(EAsyncExecution::TaskGraph, MoveTemp(HashCompletionFunc)),
+			SerializedPackageFlags, IsCompareLinker() ? MoveTemp(Linker) : nullptr);
 	}
 
 	FObjectSaveContextData& GetObjectSaveContext()
@@ -573,6 +576,7 @@ public:
 	int32 OffsetAfterPackageFileSummary = 0;
 	int32 OffsetAfterImportMap = 0;
 	int32 OffsetAfterExportMap = 0;
+	int32 SerializedPackageFlags = 0;
 	TAsyncWorkSequence<FMD5> AsyncWriteAndHashSequence;
 	TArray<FLargeMemoryWriter, TInlineAllocator<4>> AdditionalFilesFromExports;
 	FSavePackageOutputFileArray AdditionalPackageFiles;
diff --git a/Engine/Source/Runtime/CoreUObject/Private/UObject/SavePackage2.cpp b/Engine/Source/Runtime/CoreUObject/Private/UObject/SavePackage2.cpp
--- a/Engine/Source/Runtime/CoreUObject/Private/UObject/SavePackage2.cpp
+++ b/Engine/Source/Runtime/CoreUObject/Private/UObject/SavePackage2.cpp
@@ -1976,6 +1976,26 @@ void PostSavePackage(FSaveContext& SaveContext)
 	// Package has been saved, so unmark the NewlyCreated flag.
 	Package->ClearPackageFlags(PKG_NewlyCreated);
 
+	if (SaveContext.Linker != nullptr)
+	{
+		// Copy and modify the output SerializedPackageFlags from the PackageFlags written into the file
+		SaveContext.SerializedPackageFlags = SaveContext.Linker->Summary.GetPackageFlags();
+		// Currently the PKG_ContainsNoAsset flag is not serialized as part of the summary
+		bool bContainsAsset = false;
+		for (FObjectExport& Export : SaveContext.Linker->ExportMap)
+		{
+			if (Export.bIsAsset)
+			{
+				bContainsAsset = true;
+				break;
+			}
+		}
+		if (!bContainsAsset)
+		{
+			SaveContext.SerializedPackageFlags |= PKG_ContainsNoAsset;
+		}
+	}
+
 	// Notify the soft reference collector about our harvested soft references during save. 
 	// This is currently needed only for cooking which does not require editor-only references 
 #if WITH_EDITOR
