Use MakeDialogDelayed vs MakeDialog in some of the Save / Load paths.  This will prevent the UI from disrupting workflows when in Multi-user mode.

diff --git a/Engine/Source/Editor/UnrealEd/Private/EditorServer.cpp b/Engine/Source/Editor/UnrealEd/Private/EditorServer.cpp
--- a/Engine/Source/Editor/UnrealEd/Private/EditorServer.cpp
+++ b/Engine/Source/Editor/UnrealEd/Private/EditorServer.cpp
@@ -2534,7 +2534,7 @@ bool UEditorEngine::Map_Load(const TCHAR* Str, FOutputDevice& Ar)
 				
 				// Don't show progress dialogs when loading one of our startup maps. They should load rather quickly.
 				FScopedSlowTask SlowTask(100, LocalizedLoadingMap, bShowProgress != 0);
-				SlowTask.MakeDialog();
+				SlowTask.MakeDialogDelayed(3.0f);
 
 				SlowTask.EnterProgressFrame(10, FText::Format( NSLOCTEXT("UnrealEd", "LoadingMapStatus_CleaningUp", "{0} (Clearing existing world)"), LocalizedLoadingMap ));
 
diff --git a/Engine/Source/Runtime/Core/Private/Misc/FeedbackContext.cpp b/Engine/Source/Runtime/Core/Private/Misc/FeedbackContext.cpp
--- a/Engine/Source/Runtime/Core/Private/Misc/FeedbackContext.cpp
+++ b/Engine/Source/Runtime/Core/Private/Misc/FeedbackContext.cpp
@@ -46,7 +46,7 @@ void FFeedbackContext::BeginSlowTask( const FText& Task, bool ShowProgressDialog
 	TUniquePtr<FSlowTask> NewScope(new FSlowTask(0, Task, true, *this));
 	if (ShowProgressDialog)
 	{
-		NewScope->MakeDialog(bShowCancelButton);
+		NewScope->MakeDialogDelayed(3.0f, bShowCancelButton);
 	}
 
 	NewScope->Initialize();
diff --git a/Engine/Source/Runtime/CoreUObject/Private/UObject/PackageReload.cpp b/Engine/Source/Runtime/CoreUObject/Private/UObject/PackageReload.cpp
--- a/Engine/Source/Runtime/CoreUObject/Private/UObject/PackageReload.cpp
+++ b/Engine/Source/Runtime/CoreUObject/Private/UObject/PackageReload.cpp
@@ -489,7 +489,7 @@ void ReloadPackages(const TArrayView<FReloadPackageData>& InPackagesToReload, TA
 	UE_LOG(LogUObjectGlobals, Log, TEXT("%s"), *Msg);
 
 	FScopedSlowTask ReloadingPackagesSlowTask(InPackagesToReload.Num(), NSLOCTEXT("CoreUObject", "ReloadingPackages", "Reloading Packages"));
-	ReloadingPackagesSlowTask.MakeDialog();
+	ReloadingPackagesSlowTask.MakeDialogDelayed(3.0f);
 
 	// Cache the current dirty state of all packages so we can restore it after the reload
 	TSet<FName> DirtyPackages;
diff --git a/Engine/Source/Runtime/CoreUObject/Private/UObject/SavePackage.cpp b/Engine/Source/Runtime/CoreUObject/Private/UObject/SavePackage.cpp
--- a/Engine/Source/Runtime/CoreUObject/Private/UObject/SavePackage.cpp
+++ b/Engine/Source/Runtime/CoreUObject/Private/UObject/SavePackage.cpp
@@ -2245,7 +2245,7 @@ FSavePackageResultStruct UPackage::Save(UPackage* InOuter, UObject* Base, EObjec
 
 		const int32 TotalSaveSteps = 33;
 		FScopedSlowTask SlowTask(TotalSaveSteps, StatusMessage, bSlowTask);
-		SlowTask.MakeDialog(SaveFlags & SAVE_FromAutosave ? true : false);
+		SlowTask.MakeDialogDelayed(3.0f, SaveFlags & SAVE_FromAutosave ? true : false);
 
 		SlowTask.EnterProgressFrame();
 
diff --git a/Engine/Source/Runtime/CoreUObject/Private/UObject/SavePackage2.cpp b/Engine/Source/Runtime/CoreUObject/Private/UObject/SavePackage2.cpp
--- a/Engine/Source/Runtime/CoreUObject/Private/UObject/SavePackage2.cpp
+++ b/Engine/Source/Runtime/CoreUObject/Private/UObject/SavePackage2.cpp
@@ -2040,7 +2040,7 @@ ESavePackageResult InnerSave(FSaveContext& SaveContext)
 	// Create slow task dialog if needed
 	const int32 TotalSaveSteps = 12;
 	FScopedSlowTask SlowTask(TotalSaveSteps, FText(), SaveContext.IsUsingSlowTask());
-	SlowTask.MakeDialog(SaveContext.IsFromAutoSave());
+	SlowTask.MakeDialogDelayed(3.0f, SaveContext.IsFromAutoSave());
 
 	// Harvest Package
 	SlowTask.EnterProgressFrame();
@@ -2279,7 +2279,7 @@ FSavePackageResultStruct UPackage::Save2(UPackage* InPackage, UObject* InAsset,
 	// Create the slow task dialog if needed
 	const int32 TotalSaveSteps = 7;
 	FScopedSlowTask SlowTask(TotalSaveSteps, GetSlowTaskStatusMessage(SaveContext), SaveContext.IsUsingSlowTask());
-	SlowTask.MakeDialog(SaveContext.IsFromAutoSave());
+	SlowTask.MakeDialogDelayed(3.0f, SaveContext.IsFromAutoSave());
 
 	SlowTask.EnterProgressFrame();
 	SaveContext.Result = ValidatePackage(SaveContext);
@@ -2365,7 +2365,7 @@ ESavePackageResult UPackage::SaveConcurrent(TArrayView<FPackageSaveInfo> InPacka
 {
 	const int32 TotalSaveSteps = 4;
 	FScopedSlowTask SlowTask(TotalSaveSteps, NSLOCTEXT("Core", "SavingFiles", "Saving files..."), SaveArgs.bSlowTask);
-	SlowTask.MakeDialog(!!(SaveArgs.SaveFlags & SAVE_FromAutosave));
+	SlowTask.MakeDialogDelayed(3.0f, !!(SaveArgs.SaveFlags & SAVE_FromAutosave));
 
 	// Create all the package save context and run pre save
 	SlowTask.EnterProgressFrame();
