Skip unreachable objects when gathering references some might have been partially destroyed already

diff --git a/Engine/Source/Runtime/CoreUObject/Private/UObject/ReferencerFinder.cpp b/Engine/Source/Runtime/CoreUObject/Private/UObject/ReferencerFinder.cpp
--- a/Engine/Source/Runtime/CoreUObject/Private/UObject/ReferencerFinder.cpp
+++ b/Engine/Source/Runtime/CoreUObject/Private/UObject/ReferencerFinder.cpp
@@ -102,10 +102,9 @@ TArray<UObject*> FReferencerFinder::GetAllReferencers(const TSet<UObject*>& Refe
 			for (int32 ObjectIndex = 0; ObjectIndex < NumObjects && (FirstObjectIndex + ObjectIndex) < MaxNumberOfObjects; ++ObjectIndex)
 			{
 				FUObjectItem& ObjectItem = GUObjectArray.GetObjectItemArrayUnsafe()[FirstObjectIndex + ObjectIndex];
-				if (ObjectItem.Object)
+				if (ObjectItem.Object && !ObjectItem.IsUnreachable())
 				{
 					UObject* PotentialReferencer = static_cast<UObject*>(ObjectItem.Object);
-
 					if (ObjectsToIgnore && ObjectsToIgnore->Contains(PotentialReferencer))
 					{
 						continue;
