Avoid a compile error when using MSVC with the /Zc:lambda or /experimental:newLambdaProcessor compiler options.

diff --git a/Engine/Source/Runtime/CoreUObject/Public/UObject/ObjectPtr.h b/Engine/Source/Runtime/CoreUObject/Public/UObject/ObjectPtr.h
--- a/Engine/Source/Runtime/CoreUObject/Public/UObject/ObjectPtr.h
+++ b/Engine/Source/Runtime/CoreUObject/Public/UObject/ObjectPtr.h
@@ -569,7 +569,7 @@ struct TContainerElementTypeCompatibility<TObjectPtr<T>>
 {
 	typedef T* ReinterpretType;
 	
-	template <typename IterBeginType, typename IterEndType, typename OperatorType = typename TRemoveReference<decltype(*DeclVal<IterBeginType>())>::Type& (*)(IterBeginType&)>
+	template <typename IterBeginType, typename IterEndType, typename OperatorType = std::remove_reference_t<decltype(*std::declval<IterBeginType>())>& (*)(IterBeginType&)>
 	UE_OBJPTR_DEPRECATED(5.0, "Reinterpretation between ranges of one type to another type is deprecated.")
 	static void ReinterpretRange(IterBeginType Iter, IterEndType IterEnd, OperatorType Operator = [](IterBeginType& InIt) -> decltype(auto) { return *InIt; })
 	{
@@ -593,9 +593,9 @@ struct TContainerElementTypeCompatibility<const TObjectPtr<T>>
 {
 	typedef T* const ReinterpretType;
 
-	template <typename IterBeginType, typename IterEndType, typename OperatorType = const TObjectPtr<T>&(*)(IterBeginType&)>
+	template <typename IterBeginType, typename IterEndType, typename OperatorType = const std::remove_cv_t<std::remove_reference_t<decltype(*std::declval<IterBeginType>())>>&(*)(IterBeginType&)>
 	UE_OBJPTR_DEPRECATED(5.0, "Reinterpretation between ranges of one type to another type is deprecated.")
-	static void ReinterpretRange(IterBeginType Iter, IterEndType IterEnd, OperatorType Operator = [](IterBeginType& InIt) -> const TObjectPtr<T>& { return *InIt; })
+	static void ReinterpretRange(IterBeginType Iter, IterEndType IterEnd, OperatorType Operator = [](IterBeginType& InIt) -> decltype(auto) { return *InIt; })
 	{
 #if UE_WITH_OBJECT_HANDLE_LATE_RESOLVE || UE_WITH_OBJECT_HANDLE_TRACKING
 		while (Iter != IterEnd)
