DerivedData: Added a prototype of FDerivedData

Serialization and the FBuildDefinition type of derived data have not been implemented yet. This version is functional in the editor when used with buffers or unstructured cache keys.

diff --git a/Engine/Source/Runtime/Core/Public/IO/IoDispatcher.h b/Engine/Source/Runtime/Core/Public/IO/IoDispatcher.h
--- a/Engine/Source/Runtime/Core/Public/IO/IoDispatcher.h
+++ b/Engine/Source/Runtime/Core/Public/IO/IoDispatcher.h
@@ -654,7 +654,9 @@ enum class EIoChunkType : uint8
 	ExternalFile = 7,
 	ShaderCodeLibrary = 8,
 	ShaderCode = 9,
-	PackageStoreEntry = 10
+	PackageStoreEntry = 10,
+	DerivedData = 11,
+	EditorDerivedData = 12,
 };
 
 /**
diff --git a/Engine/Source/Runtime/CoreUObject/CoreUObject.Build.cs b/Engine/Source/Runtime/CoreUObject/CoreUObject.Build.cs
--- a/Engine/Source/Runtime/CoreUObject/CoreUObject.Build.cs
+++ b/Engine/Source/Runtime/CoreUObject/CoreUObject.Build.cs
@@ -29,5 +29,10 @@ public class CoreUObject : ModuleRules
 		{
 			PublicDependencyModuleNames.Add("CookOnTheFly");
 		}
+
+		if (Target.bBuildWithEditorOnlyData)
+		{
+			PrivateDependencyModuleNames.Add("DerivedDataCache");
+		}
 	}
 }
diff --git a/Engine/Source/Runtime/CoreUObject/Private/Serialization/DerivedData.cpp b/Engine/Source/Runtime/CoreUObject/Private/Serialization/DerivedData.cpp
new file mode 100644
--- /dev/null
+++ b/Engine/Source/Runtime/CoreUObject/Private/Serialization/DerivedData.cpp
@@ -0,0 +1,161 @@
+// Copyright Epic Games, Inc. All Rights Reserved.
+
+#include "Serialization/DerivedData.h"
+
+#include "Memory/SharedBuffer.h"
+
+#if WITH_EDITORONLY_DATA
+
+#include "DerivedDataBuildDefinition.h"
+#include "DerivedDataPayloadId.h"
+#include "Serialization/EditorDerivedData.h"
+#include "Serialization/EditorDerivedDataIoStore.h"
+#include "Templates/SharedPointer.h"
+
+namespace UE::DerivedData::IoStore
+{
+
+TSharedPtr<IEditorDerivedDataIoStore> GEditorDerivedDataIoStore;
+
+} // UE::DerivedData::IoStore
+
+#endif // WITH_EDITORONLY_DATA
+
+///////////////////////////////////////////////////////////////////////////////////////////////////////////////////////
+
+namespace UE
+{
+
+void FDerivedData::Serialize(FArchive& Ar, UObject* Owner)
+{
+	unimplemented();
+}
+
+FUniqueBuffer FDerivedData::LoadData(FArchive& Ar, FDerivedDataBufferAllocator Allocator)
+{
+	check(Ar.IsLoading());
+	unimplemented();
+	return FUniqueBuffer();
+}
+
+#if WITH_EDITORONLY_DATA
+
+void FDerivedData::SaveData(FArchive& Ar, const FDerivedData& Data)
+{
+	checkf(Data.EditorData, TEXT("Inline serialization is not supported for non-editor derived data references."));
+	check(Ar.IsSaving());
+	unimplemented();
+}
+
+FDerivedData::FDerivedData() = default;
+FDerivedData::FDerivedData(FDerivedData&& Other) = default;
+FDerivedData& FDerivedData::operator=(FDerivedData&& Other) = default;
+
+FDerivedData::FDerivedData(const FDerivedData& Other)
+{
+	if (Other.EditorData)
+	{
+		EditorData = Other.EditorData->Clone();
+		ChunkId = DerivedData::IoStore::GEditorDerivedDataIoStore->AddData(EditorData.Get());
+	}
+	else
+	{
+		ChunkId = Other.ChunkId;
+	}
+}
+
+FDerivedData& FDerivedData::operator=(const FDerivedData& Other)
+{
+	if (EditorData)
+	{
+		DerivedData::IoStore::GEditorDerivedDataIoStore->RemoveData(ChunkId);
+	}
+	if (Other.EditorData)
+	{
+		EditorData = Other.EditorData->Clone();
+		ChunkId = DerivedData::IoStore::GEditorDerivedDataIoStore->AddData(EditorData.Get());
+	}
+	else
+	{
+		EditorData.Reset();
+		ChunkId = Other.ChunkId;
+	}
+	return *this;
+}
+
+FDerivedData::FDerivedData(const FSharedBuffer& Data)
+	: EditorData(DerivedData::Private::MakeEditorDerivedData(Data))
+	, ChunkId(DerivedData::IoStore::GEditorDerivedDataIoStore->AddData(EditorData.Get()))
+{
+}
+
+FDerivedData::FDerivedData(const FCompositeBuffer& Data)
+	: EditorData(DerivedData::Private::MakeEditorDerivedData(Data))
+	, ChunkId(DerivedData::IoStore::GEditorDerivedDataIoStore->AddData(EditorData.Get()))
+{
+}
+
+FDerivedData::FDerivedData(const FCompressedBuffer& Data)
+	: EditorData(DerivedData::Private::MakeEditorDerivedData(Data))
+	, ChunkId(DerivedData::IoStore::GEditorDerivedDataIoStore->AddData(EditorData.Get()))
+{
+}
+
+FDerivedData::FDerivedData(FStringView CacheKey, FStringView CacheContext)
+	: EditorData(DerivedData::Private::MakeEditorDerivedData(CacheKey, CacheContext))
+	, ChunkId(DerivedData::IoStore::GEditorDerivedDataIoStore->AddData(EditorData.Get()))
+{
+}
+
+FDerivedData::FDerivedData(const DerivedData::FBuildDefinition& BuildDefinition, const DerivedData::FPayloadId& PayloadId)
+	: EditorData(DerivedData::Private::MakeEditorDerivedData(BuildDefinition, PayloadId))
+	, ChunkId(DerivedData::IoStore::GEditorDerivedDataIoStore->AddData(EditorData.Get()))
+{
+}
+
+FDerivedData::~FDerivedData()
+{
+	if (EditorData)
+	{
+		using namespace DerivedData::IoStore;
+		if (IEditorDerivedDataIoStore* EditorBackend = GEditorDerivedDataIoStore.Get())
+		{
+			verify(EditorBackend->RemoveData(ChunkId) == EditorData.Get());
+		}
+	}
+}
+
+void FDerivedData::SetFlags(EDerivedDataFlags InFlags)
+{
+	checkf(!EnumHasAllFlags(InFlags, EDerivedDataFlags::Optional | EDerivedDataFlags::MemoryMapped),
+		TEXT("Optional and MemoryMapped cannot be used together on derived data references."));
+	Flags = InFlags;
+}
+
+#endif // WITH_EDITORONLY_DATA
+
+} // UE
+
+///////////////////////////////////////////////////////////////////////////////////////////////////////////////////////
+
+namespace UE::DerivedData::IoStore
+{
+
+void InitializeIoDispatcher()
+{
+#if WITH_EDITORONLY_DATA
+	FIoDispatcher& IoDispatcher = FIoDispatcher::Get();
+	TSharedRef<IEditorDerivedDataIoStore> EditorBackend = CreateEditorDerivedDataIoStore();
+	GEditorDerivedDataIoStore = EditorBackend;
+	IoDispatcher.Mount(EditorBackend);
+#endif // WITH_EDITORONLY_DATA
+}
+
+void TearDownIoDispatcher()
+{
+#if WITH_EDITORONLY_DATA
+	GEditorDerivedDataIoStore.Reset();
+#endif // WITH_EDITORONLY_DATA
+}
+
+} // UE::DerivedData::IoStore
diff --git a/Engine/Source/Runtime/CoreUObject/Private/Serialization/EditorDerivedData.cpp b/Engine/Source/Runtime/CoreUObject/Private/Serialization/EditorDerivedData.cpp
new file mode 100644
--- /dev/null
+++ b/Engine/Source/Runtime/CoreUObject/Private/Serialization/EditorDerivedData.cpp
@@ -0,0 +1,228 @@
+// Copyright Epic Games, Inc. All Rights Reserved.
+
+#include "Serialization/EditorDerivedData.h"
+
+#if WITH_EDITORONLY_DATA
+
+#include "Compression/CompressedBuffer.h"
+#include "DerivedDataBuildDefinition.h"
+#include "DerivedDataBuildSession.h"
+#include "DerivedDataCacheInterface.h"
+#include "DerivedDataPayloadId.h"
+#include "Memory/CompositeBuffer.h"
+#include "Memory/SharedBuffer.h"
+#include "Misc/CoreMisc.h"
+#include "Serialization/BulkDataRegistry.h"
+
+namespace UE::DerivedData::Private
+{
+
+class FEditorDerivedDataBuffer final : public FEditorDerivedData
+{
+public:
+	explicit FEditorDerivedDataBuffer(const FSharedBuffer& InData)
+		: Data(InData)
+	{
+	}
+
+	explicit FEditorDerivedDataBuffer(const FCompositeBuffer& InData)
+		: Data(InData)
+	{
+	}
+
+	TUniquePtr<FEditorDerivedData> Clone() const final { return MakeUnique<FEditorDerivedDataBuffer>(*this); }
+
+	void Read(IoStore::FDerivedDataIoRequest Request) const final;
+	bool TryGetSize(uint64& OutSize) const final;
+
+private:
+	FCompositeBuffer Data;
+};
+
+void FEditorDerivedDataBuffer::Read(IoStore::FDerivedDataIoRequest Request) const
+{
+	const uint64 DataSize = Data.GetSize();
+	const uint64 RequestOffset = Request.GetOffset();
+	const uint64 RequestSize = FMath::Min(Request.GetSize(), RequestOffset <= DataSize ? DataSize - RequestOffset : 0);
+	if (RequestSize)
+	{
+		Data.CopyTo(Request.CreateBuffer(RequestSize), RequestOffset);
+		Request.SetComplete();
+	}
+	else
+	{
+		Request.SetComplete();
+	}
+}
+
+bool FEditorDerivedDataBuffer::TryGetSize(uint64& OutSize) const
+{
+	OutSize = Data.GetSize();
+	return true;
+}
+
+///////////////////////////////////////////////////////////////////////////////////////////////////////////////////////
+
+class FEditorDerivedDataCompressedBuffer final : public FEditorDerivedData
+{
+public:
+	explicit FEditorDerivedDataCompressedBuffer(const FCompressedBuffer& InData)
+		: Data(InData)
+	{
+	}
+
+	TUniquePtr<FEditorDerivedData> Clone() const final { return MakeUnique<FEditorDerivedDataCompressedBuffer>(*this); }
+
+	void Read(IoStore::FDerivedDataIoRequest Request) const final;
+	bool TryGetSize(uint64& OutSize) const final;
+
+private:
+	FCompressedBuffer Data;
+};
+
+void FEditorDerivedDataCompressedBuffer::Read(IoStore::FDerivedDataIoRequest Request) const
+{
+	const uint64 DataSize = Data.GetRawSize();
+	const uint64 RequestOffset = Request.GetOffset();
+	const uint64 RequestSize = FMath::Min(Request.GetSize(), RequestOffset <= DataSize ? DataSize - RequestOffset : 0);
+	if (RequestSize)
+	{
+		if (const FSharedBuffer Buffer = Data.Decompress())
+		{
+			Request.CreateBuffer(RequestSize).CopyFrom(Buffer.GetView().Mid(RequestOffset, RequestSize));
+			Request.SetComplete();
+		}
+		else
+		{
+			Request.SetFailed();
+		}
+	}
+	else
+	{
+		Request.SetComplete();
+	}
+}
+
+bool FEditorDerivedDataCompressedBuffer::TryGetSize(uint64& OutSize) const
+{
+	OutSize = Data.GetRawSize();
+	return true;
+}
+
+///////////////////////////////////////////////////////////////////////////////////////////////////////////////////////
+
+class FEditorDerivedDataCache final : public FEditorDerivedData
+{
+public:
+	FEditorDerivedDataCache(FStringView InCacheKey, FStringView InCacheContext)
+		: CacheKey(InCacheKey)
+		, CacheContext(InCacheContext)
+	{
+	}
+
+	TUniquePtr<FEditorDerivedData> Clone() const final { return MakeUnique<FEditorDerivedDataCache>(*this); }
+
+	void Read(IoStore::FDerivedDataIoRequest Request) const final;
+	bool TryGetSize(uint64& OutSize) const final;
+
+private:
+	FString CacheKey;
+	FString CacheContext;
+};
+
+void FEditorDerivedDataCache::Read(IoStore::FDerivedDataIoRequest Request) const
+{
+	TArray<uint8> Data;
+	if (GetDerivedDataCacheRef().GetSynchronous(*CacheKey, Data, CacheContext))
+	{
+		const FMemoryView DataView = MakeMemoryView(Data);
+		const uint64 DataSize = DataView.GetSize();
+		const uint64 RequestOffset = Request.GetOffset();
+		const uint64 RequestSize = FMath::Min(Request.GetSize(), RequestOffset <= DataSize ? DataSize - RequestOffset : 0);
+		Request.CreateBuffer(RequestSize).CopyFrom(DataView.Mid(RequestOffset, RequestSize));
+		Request.SetComplete();
+	}
+	else
+	{
+		Request.SetFailed();
+	}
+}
+
+bool FEditorDerivedDataCache::TryGetSize(uint64& OutSize) const
+{
+	TArray<uint8> Data;
+	if (GetDerivedDataCacheRef().GetSynchronous(*CacheKey, Data, CacheContext))
+	{
+		OutSize = MakeMemoryView(Data).GetSize();
+		return true;
+	}
+	else
+	{
+		return false;
+	}
+}
+
+///////////////////////////////////////////////////////////////////////////////////////////////////////////////////////
+
+class FEditorDerivedDataBuild final : public FEditorDerivedData
+{
+public:
+	FEditorDerivedDataBuild(const FBuildDefinition& InBuildDefinition, const FPayloadId& InPayloadId)
+		: BuildDefinition(InBuildDefinition)
+		, PayloadId(InPayloadId)
+	{
+		unimplemented();
+	}
+
+	TUniquePtr<FEditorDerivedData> Clone() const final { return MakeUnique<FEditorDerivedDataBuild>(*this); }
+
+	void Read(IoStore::FDerivedDataIoRequest Request) const final;
+	bool TryGetSize(uint64& OutSize) const final;
+
+private:
+	FBuildDefinition BuildDefinition;
+	FPayloadId PayloadId;
+};
+
+void FEditorDerivedDataBuild::Read(IoStore::FDerivedDataIoRequest Request) const
+{
+	Request.SetFailed();
+}
+
+bool FEditorDerivedDataBuild::TryGetSize(uint64& OutSize) const
+{
+	return false;
+}
+
+///////////////////////////////////////////////////////////////////////////////////////////////////////////////////////
+
+TUniquePtr<FEditorDerivedData> MakeEditorDerivedData(const FSharedBuffer& Data)
+{
+	return MakeUnique<FEditorDerivedDataBuffer>(Data);
+}
+
+TUniquePtr<FEditorDerivedData> MakeEditorDerivedData(const FCompositeBuffer& Data)
+{
+	return MakeUnique<FEditorDerivedDataBuffer>(Data);
+}
+
+TUniquePtr<FEditorDerivedData> MakeEditorDerivedData(const FCompressedBuffer& Data)
+{
+	return MakeUnique<FEditorDerivedDataCompressedBuffer>(Data);
+}
+
+TUniquePtr<FEditorDerivedData> MakeEditorDerivedData(FStringView CacheKey, FStringView CacheContext)
+{
+	return MakeUnique<FEditorDerivedDataCache>(CacheKey, CacheContext);
+}
+
+TUniquePtr<FEditorDerivedData> MakeEditorDerivedData(
+	const FBuildDefinition& BuildDefinition,
+	const FPayloadId& PayloadId)
+{
+	return MakeUnique<FEditorDerivedDataBuild>(BuildDefinition, PayloadId);
+}
+
+} // UE::DerivedData::Private
+
+#endif // WITH_EDITORONLY_DATA
diff --git a/Engine/Source/Runtime/CoreUObject/Private/Serialization/EditorDerivedData.h b/Engine/Source/Runtime/CoreUObject/Private/Serialization/EditorDerivedData.h
new file mode 100644
--- /dev/null
+++ b/Engine/Source/Runtime/CoreUObject/Private/Serialization/EditorDerivedData.h
@@ -0,0 +1,72 @@
+// Copyright Epic Games, Inc. All Rights Reserved.
+
+#pragma once
+
+#include "CoreTypes.h"
+
+#if WITH_EDITORONLY_DATA
+
+#include "Containers/StringFwd.h"
+#include "Memory/MemoryFwd.h"
+#include "Templates/UniquePtr.h"
+
+class FCompressedBuffer;
+class FIoRequestImpl;
+
+namespace UE::DerivedData { class FBuildDefinition; }
+namespace UE::DerivedData { class FRequestOwner; }
+namespace UE::DerivedData { struct FPayloadId; }
+namespace UE::DerivedData::IoStore { class FDerivedDataIoRequestQueue; }
+
+namespace UE::DerivedData::IoStore
+{
+
+class FDerivedDataIoRequest
+{
+public:
+	FDerivedDataIoRequest(FIoRequestImpl* Request, FDerivedDataIoRequestQueue* Queue);
+	/** Create a request owner for requests that do not complete immediately. */
+	FRequestOwner& GetOwner();
+	/** Create a buffer of the given size and return a view to write into it. */
+	FMutableMemoryView CreateBuffer(uint64 Size);
+	/** Offset to start reading from. */
+	uint64 GetOffset() const;
+	/** Maximum number of bytes to read. */
+	uint64 GetSize() const;
+	/** Mark the request as complete. */
+	void SetComplete();
+	/** Mark the request as failed. */
+	void SetFailed();
+
+private:
+	FIoRequestImpl* Request;
+	FDerivedDataIoRequestQueue* Queue;
+};
+
+} // UE::DerivedData::IoStore
+
+namespace UE::DerivedData::Private
+{
+
+class FEditorDerivedData
+{
+public:
+	virtual ~FEditorDerivedData() = default;
+	virtual TUniquePtr<FEditorDerivedData> Clone() const = 0;
+	virtual void Read(IoStore::FDerivedDataIoRequest Request) const = 0;
+	virtual bool TryGetSize(uint64& OutSize) const = 0;
+};
+
+TUniquePtr<FEditorDerivedData> MakeEditorDerivedData(const FSharedBuffer& Data);
+TUniquePtr<FEditorDerivedData> MakeEditorDerivedData(const FCompositeBuffer& Data);
+TUniquePtr<FEditorDerivedData> MakeEditorDerivedData(const FCompressedBuffer& Data);
+
+TUniquePtr<FEditorDerivedData> MakeEditorDerivedData(FStringView CacheKey, FStringView CacheContext);
+
+TUniquePtr<FEditorDerivedData> MakeEditorDerivedData(
+	const DerivedData::FBuildDefinition& BuildDefinition,
+	const DerivedData::FPayloadId& PayloadId);
+
+} // UE::DerivedData::Private
+
+#endif // WITH_EDITORONLY_DATA
diff --git a/Engine/Source/Runtime/CoreUObject/Private/Serialization/EditorDerivedDataIoStore.cpp b/Engine/Source/Runtime/CoreUObject/Private/Serialization/EditorDerivedDataIoStore.cpp
new file mode 100644
--- /dev/null
+++ b/Engine/Source/Runtime/CoreUObject/Private/Serialization/EditorDerivedDataIoStore.cpp
@@ -0,0 +1,282 @@
+// Copyright Epic Games, Inc. All Rights Reserved.
+
+#include "Serialization/EditorDerivedDataIoStore.h"
+
+#if WITH_EDITORONLY_DATA
+
+#include "Containers/Map.h"
+#include "DerivedDataRequestOwner.h"
+#include "DerivedDataRequestTypes.h"
+#include "HAL/CriticalSection.h"
+#include "Misc/ScopeRWLock.h"
+#include "Serialization/EditorDerivedData.h"
+#include <atomic>
+
+namespace UE::DerivedData::IoStore
+{
+
+static EPriority GetPriority(int32 Priority)
+{
+	if (Priority < IoDispatcherPriority_Low)
+	{
+		return EPriority::Lowest;
+	}
+	if (Priority > IoDispatcherPriority_High)
+	{
+		return EPriority::Highest;
+	}
+	if (Priority < IoDispatcherPriority_Medium)
+	{
+		return EPriority::Low;
+	}
+	if (Priority > IoDispatcherPriority_Medium)
+	{
+		return EPriority::High;
+	}
+	return EPriority::Normal;
+}
+
+///////////////////////////////////////////////////////////////////////////////////////////////////////////////////////
+
+class FDerivedDataIoRequestQueue
+{
+public:
+	void Add(FIoRequestImpl* Request);
+
+	FIoRequestImpl* Steal();
+
+	void SetContext(const TSharedRef<const FIoDispatcherBackendContext>& Context);
+
+	void SetSkipWakeDispatcher(bool bValue) { bSkipWakeDispatcher.store(bValue, std::memory_order_relaxed); }
+
+private:
+	FRWLock Lock;
+	FIoRequestImpl* Head = nullptr;
+	FIoRequestImpl* Tail = nullptr;
+	TSharedPtr<const FIoDispatcherBackendContext> Context;
+	std::atomic<bool> bSkipWakeDispatcher{false};
+};
+
+void FDerivedDataIoRequestQueue::Add(FIoRequestImpl* Request)
+{
+	check(Request);
+	bool bWakeDispatcher;
+	{
+		FWriteScopeLock WriteLock(Lock);
+		if (!Tail)
+		{
+			Head = Request;
+			Tail = Request;
+		}
+		else
+		{
+			Tail->NextRequest = Request;
+			Tail = Request;
+		}
+		Tail->NextRequest = nullptr;
+		bWakeDispatcher = !bSkipWakeDispatcher.load(std::memory_order_relaxed);
+	}
+	if (bWakeDispatcher)
+	{
+		Context->WakeUpDispatcherThreadDelegate.Execute();
+	}
+}
+
+FIoRequestImpl* FDerivedDataIoRequestQueue::Steal()
+{
+	FWriteScopeLock WriteLock(Lock);
+	FIoRequestImpl* Queue = Head;
+	Head = Tail = nullptr;
+	return Queue;
+}
+
+void FDerivedDataIoRequestQueue::SetContext(const TSharedRef<const FIoDispatcherBackendContext>& InContext)
+{
+	FWriteScopeLock WriteLock(Lock);
+	Context = InContext;
+}
+
+///////////////////////////////////////////////////////////////////////////////////////////////////////////////////////
+
+FDerivedDataIoRequest::FDerivedDataIoRequest(FIoRequestImpl* InRequest, FDerivedDataIoRequestQueue* InQueue)
+	: Request(InRequest)
+	, Queue(InQueue)
+{
+}
+
+FRequestOwner& FDerivedDataIoRequest::GetOwner()
+{
+	if (!Request->BackendData)
+	{
+		static_assert(sizeof(Request->BackendData) == sizeof(FRequestOwner));
+		new(&Request->BackendData) FRequestOwner(GetPriority(Request->Priority));
+	}
+	return reinterpret_cast<FRequestOwner&>(Request->BackendData);
+}
+
+FMutableMemoryView FDerivedDataIoRequest::CreateBuffer(uint64 Size)
+{
+	if (!Request->HasBuffer())
+	{
+		Request->CreateBuffer(Size);
+	}
+	FIoBuffer& Buffer = Request->GetBuffer();
+	return MakeMemoryView(Buffer.Data(), Buffer.DataSize());
+}
+
+uint64 FDerivedDataIoRequest::GetOffset() const
+{
+	return Request->Options.GetOffset();
+}
+
+uint64 FDerivedDataIoRequest::GetSize() const
+{
+	return Request->Options.GetSize();
+}
+
+void FDerivedDataIoRequest::SetComplete()
+{
+	Queue->Add(Request);
+}
+
+void FDerivedDataIoRequest::SetFailed()
+{
+	Request->SetFailed();
+	Queue->Add(Request);
+}
+
+///////////////////////////////////////////////////////////////////////////////////////////////////////////////////////
+
+class FEditorDerivedDataIoStore final : public IEditorDerivedDataIoStore
+{
+public:
+	void Initialize(TSharedRef<const FIoDispatcherBackendContext> Context) final;
+	bool Resolve(FIoRequestImpl* Request) final;
+	void CancelIoRequest(FIoRequestImpl* Request) final;
+	void UpdatePriorityForIoRequest(FIoRequestImpl* Request) final;
+	bool DoesChunkExist(const FIoChunkId& ChunkId) const final;
+	TIoStatusOr<uint64> GetSizeForChunk(const FIoChunkId& ChunkId) const final;
+	FIoRequestImpl* GetCompletedRequests() final;
+	TIoStatusOr<FIoMappedRegion> OpenMapped(const FIoChunkId& ChunkId, const FIoReadOptions& Options) final;
+	void AppendMountedContainers(TSet<FIoContainerId>& OutContainers) final;
+
+	FIoChunkId AddData(const Private::FEditorDerivedData* Data) final;
+	const Private::FEditorDerivedData* RemoveData(const FIoChunkId& ChunkId) final;
+
+private:
+	mutable FRWLock DataLock;
+	TMap<FIoChunkId, const Private::FEditorDerivedData*> DataById;
+	uint64 CurrentChunkId = 0;
+	uint16 LastChunkIndex = 0;
+	FDerivedDataIoRequestQueue CompletedQueue;
+};
+
+void FEditorDerivedDataIoStore::Initialize(TSharedRef<const FIoDispatcherBackendContext> Context)
+{
+	CompletedQueue.SetContext(Context);
+}
+
+bool FEditorDerivedDataIoStore::Resolve(FIoRequestImpl* Request)
+{
+	if (Request->ChunkId.GetChunkType() != EIoChunkType::EditorDerivedData)
+	{
+		return false;
+	}
+
+	FReadScopeLock ReadLock(DataLock);
+	if (const Private::FEditorDerivedData* Data = DataById.FindRef(Request->ChunkId))
+	{
+		CompletedQueue.SetSkipWakeDispatcher(true);
+		Data->Read(FDerivedDataIoRequest(Request, &CompletedQueue));
+		CompletedQueue.SetSkipWakeDispatcher(false);
+		return true;
+	}
+	return false;
+}
+
+void FEditorDerivedDataIoStore::CancelIoRequest(FIoRequestImpl* Request)
+{
+	FDerivedDataIoRequest(Request, &CompletedQueue).GetOwner().Cancel();
+}
+
+void FEditorDerivedDataIoStore::UpdatePriorityForIoRequest(FIoRequestImpl* Request)
+{
+	FDerivedDataIoRequest(Request, &CompletedQueue).GetOwner().SetPriority(GetPriority(Request->Priority));
+}
+
+bool FEditorDerivedDataIoStore::DoesChunkExist(const FIoChunkId& ChunkId) const
+{
+	if (ChunkId.GetChunkType() != EIoChunkType::EditorDerivedData)
+	{
+		return false;
+	}
+
+	FReadScopeLock ReadLock(DataLock);
+	return DataById.Contains(ChunkId);
+}
+
+TIoStatusOr<uint64> FEditorDerivedDataIoStore::GetSizeForChunk(const FIoChunkId& ChunkId) const
+{
+	FReadScopeLock ReadLock(DataLock);
+	if (const Private::FEditorDerivedData* Data = DataById.FindRef(ChunkId))
+	{
+		uint64 Size;
+		if (Data->TryGetSize(Size))
+		{
+			return Size;
+		}
+	}
+	return FIoStatus(EIoErrorCode::NotFound);
+}
+
+FIoRequestImpl* FEditorDerivedDataIoStore::GetCompletedRequests()
+{
+	FIoRequestImpl* const Head = CompletedQueue.Steal();
+	for (FIoRequestImpl* Request = Head; Request; Request = Request->NextRequest)
+	{
+		if (Request->BackendData)
+		{
+			FDerivedDataIoRequest(Request, &CompletedQueue).GetOwner().~FRequestOwner();
+			Request->BackendData = nullptr;
+		}
+	}
+	return Head;
+}
+
+TIoStatusOr<FIoMappedRegion> FEditorDerivedDataIoStore::OpenMapped(const FIoChunkId& ChunkId, const FIoReadOptions& Options)
+{
+	return FIoStatus(EIoErrorCode::NotFound);
+}
+
+void FEditorDerivedDataIoStore::AppendMountedContainers(TSet<FIoContainerId>& OutContainers)
+{
+}
+
+FIoChunkId FEditorDerivedDataIoStore::AddData(const Private::FEditorDerivedData* Data)
+{
+	FWriteScopeLock WriteLock(DataLock);
+	if (LastChunkIndex == MAX_uint16)
+	{
+		++CurrentChunkId;
+	}
+	const FIoChunkId ChunkId = CreateIoChunkId(CurrentChunkId, ++LastChunkIndex, EIoChunkType::EditorDerivedData);
+	DataById.Add(ChunkId, Data);
+	return ChunkId;
+}
+
+const Private::FEditorDerivedData* FEditorDerivedDataIoStore::RemoveData(const FIoChunkId& ChunkId)
+{
+	FWriteScopeLock WriteLock(DataLock);
+	return DataById.FindAndRemoveChecked(ChunkId);
+}
+
+///////////////////////////////////////////////////////////////////////////////////////////////////////////////////////
+
+TSharedRef<IEditorDerivedDataIoStore> CreateEditorDerivedDataIoStore()
+{
+	return MakeShared<FEditorDerivedDataIoStore>();
+}
+
+} // UE::DerivedData::IoStore
+
+#endif // WITH_EDITORONLY_DATA
diff --git a/Engine/Source/Runtime/CoreUObject/Private/Serialization/EditorDerivedDataIoStore.h b/Engine/Source/Runtime/CoreUObject/Private/Serialization/EditorDerivedDataIoStore.h
new file mode 100644
--- /dev/null
+++ b/Engine/Source/Runtime/CoreUObject/Private/Serialization/EditorDerivedDataIoStore.h
@@ -0,0 +1,29 @@
+// Copyright Epic Games, Inc. All Rights Reserved.
+
+#pragma once
+
+#include "CoreTypes.h"
+
+#if WITH_EDITORONLY_DATA
+
+#include "IO/IoDispatcher.h"
+#include "IO/IoDispatcherBackend.h"
+#include "Templates/SharedPointer.h"
+
+namespace UE::DerivedData::Private { class FEditorDerivedData; }
+
+namespace UE::DerivedData::IoStore
+{
+
+class IEditorDerivedDataIoStore : public IIoDispatcherBackend
+{
+public:
+	virtual FIoChunkId AddData(const Private::FEditorDerivedData* Data) = 0;
+	virtual const Private::FEditorDerivedData* RemoveData(const FIoChunkId& ChunkId) = 0;
+};
+
+TSharedRef<IEditorDerivedDataIoStore> CreateEditorDerivedDataIoStore();
+
+} // UE::DerivedData::IoStore
+
+#endif // WITH_EDITORONLY_DATA
diff --git a/Engine/Source/Runtime/CoreUObject/Public/Serialization/DerivedData.h b/Engine/Source/Runtime/CoreUObject/Public/Serialization/DerivedData.h
new file mode 100644
--- /dev/null
+++ b/Engine/Source/Runtime/CoreUObject/Public/Serialization/DerivedData.h
@@ -0,0 +1,91 @@
+// Copyright Epic Games, Inc. All Rights Reserved.
+
+#pragma once
+
+#include "Containers/StringFwd.h"
+#include "IO/IoDispatcher.h"
+#include "Memory/MemoryFwd.h"
+#include "Misc/EnumClassFlags.h"
+#include "Templates/UniquePtr.h"
+
+#define UE_API COREUOBJECT_API
+
+class FArchive;
+class FCompressedBuffer;
+class UObject;
+
+template <typename FuncType> class TFunctionRef;
+
+namespace UE::DerivedData { class FBuildDefinition; }
+namespace UE::DerivedData { struct FPayloadId; }
+namespace UE::DerivedData::Private { class FEditorDerivedData; }
+
+namespace UE
+{
+
+enum class EDerivedDataFlags : uint32
+{
+	None            = 0,
+	Required        = 1 << 0,
+	Optional        = 1 << 1,
+	MemoryMapped    = 1 << 2,
+};
+
+ENUM_CLASS_FLAGS(EDerivedDataFlags);
+
+using FDerivedDataBufferAllocator = TFunctionRef<FUniqueBuffer (uint64 Size)>;
+
+class FDerivedData
+{
+public:
+	inline bool HasData() const { return ChunkId != FIoChunkId::InvalidChunkId; }
+
+	inline const FIoChunkId& GetChunkId() const { return ChunkId; }
+
+	inline EDerivedDataFlags GetFlags() const { return Flags; }
+
+	UE_API void Serialize(FArchive& Ar, UObject* Owner);
+
+	UE_API static FUniqueBuffer LoadData(FArchive& Ar, FDerivedDataBufferAllocator Allocator);
+
+#if WITH_EDITORONLY_DATA
+	UE_API static void SaveData(FArchive& Ar, const FDerivedData& Data);
+
+	UE_API FDerivedData();
+	UE_API ~FDerivedData();
+
+	UE_API FDerivedData(FDerivedData&& Other);
+	UE_API FDerivedData(const FDerivedData& Other);
+	UE_API FDerivedData& operator=(FDerivedData&& Other);
+	UE_API FDerivedData& operator=(const FDerivedData& Other);
+
+	UE_API explicit FDerivedData(const FSharedBuffer& Data);
+	UE_API explicit FDerivedData(const FCompositeBuffer& Data);
+	UE_API explicit FDerivedData(const FCompressedBuffer& Data);
+
+	UE_API FDerivedData(FStringView CacheKey, FStringView CacheContext);
+
+	UE_API FDerivedData(const DerivedData::FBuildDefinition& BuildDefinition, const DerivedData::FPayloadId& PayloadId);
+
+	UE_API void SetFlags(EDerivedDataFlags Flags);
+
+private:
+	TUniquePtr<DerivedData::Private::FEditorDerivedData> EditorData;
+#endif // WITH_EDITORONLY_DATA
+
+private:
+	FIoChunkId ChunkId = FIoChunkId::InvalidChunkId;
+	EDerivedDataFlags Flags = EDerivedDataFlags::Required;
+};
+
+} // UE
+
+namespace UE::DerivedData::IoStore
+{
+
+UE_API void InitializeIoDispatcher();
+UE_API void TearDownIoDispatcher();
+
+} // UE::DerivedData::IoStore
+
+#undef UE_API
diff --git a/Engine/Source/Runtime/Launch/Private/LaunchEngineLoop.cpp b/Engine/Source/Runtime/Launch/Private/LaunchEngineLoop.cpp
--- a/Engine/Source/Runtime/Launch/Private/LaunchEngineLoop.cpp
+++ b/Engine/Source/Runtime/Launch/Private/LaunchEngineLoop.cpp
@@ -130,14 +130,15 @@
 	#include "EngineStats.h"
 	#include "EngineGlobals.h"
 	#include "AudioThread.h"
-#if WITH_ENGINE && !UE_BUILD_SHIPPING
+#if !UE_BUILD_SHIPPING
 	#include "IAutomationControllerModule.h"
-#endif // WITH_ENGINE && !UE_BUILD_SHIPPING
+#endif // !UE_BUILD_SHIPPING
 #if WITH_EDITORONLY_DATA
 	#include "DerivedDataBuild.h"
 	#include "DerivedDataCache.h"
 #endif // WITH_EDITORONLY_DATA
 	#include "DerivedDataCacheInterface.h"
+	#include "Serialization/DerivedData.h"
 	#include "ShaderCompiler.h"
 	#include "DistanceFieldAtlas.h"
 	#include "MeshCardRepresentation.h"
@@ -2550,6 +2551,13 @@ PRAGMA_ENABLE_DEPRECATION_WARNINGS
 		FPlatformMemory::Init();
 	}
 
+#if WITH_ENGINE
+	{
+		SCOPED_BOOT_TIMING("InitDerivedData");
+		UE::DerivedData::IoStore::InitializeIoDispatcher();
+	}
+#endif // WITH_ENGINE
+
 #if USE_IO_DISPATCHER
 	{
 		SCOPED_BOOT_TIMING("InitIoDispatcher");
@@ -2852,10 +2860,13 @@ PRAGMA_ENABLE_DEPRECATION_WARNINGS
 	if (!FPlatformProperties::RequiresCookedData())
 	{
 #if WITH_EDITORONLY_DATA
-		// Ensure that DDC is initialized from the game thread.
-		UE::DerivedData::GetBuild();
-		UE::DerivedData::GetCache();
-		GetDerivedDataCacheRef();
+		{
+			// Ensure that DDC is initialized from the game thread.
+			SCOPED_BOOT_TIMING("InitDerivedData");
+			UE::DerivedData::GetBuild();
+			UE::DerivedData::GetCache();
+			GetDerivedDataCacheRef();
+		}
 #endif
 
 #if WITH_EDITOR
@@ -4554,12 +4565,14 @@ void FEngineLoop::Exit()
 	// Close shader code map, if any
 	FShaderCodeLibrary::Shutdown();
 
+	// Stop IoDispatcher after FShaderCodeLibrary, as it holds on to file requests
+#if WITH_ENGINE
+	UE::DerivedData::IoStore::TearDownIoDispatcher();
+#endif
 #if USE_IO_DISPATCHER
-	// has to happen after the FShaderCodeLibrary, as it holds on to file requests
 	FIoDispatcher::Shutdown();
 #endif
 
-
 #if !PLATFORM_ANDROID // UnloadModules doesn't work on Android
 #if WITH_ENGINE
 	// Save the hot reload state
