Partial pass of fixing type truncation warnings in CoreUObject

diff --git a/Engine/Source/Runtime/CoreUObject/CoreUObject.Build.cs b/Engine/Source/Runtime/CoreUObject/CoreUObject.Build.cs
--- a/Engine/Source/Runtime/CoreUObject/CoreUObject.Build.cs
+++ b/Engine/Source/Runtime/CoreUObject/CoreUObject.Build.cs
@@ -30,6 +30,9 @@ public class CoreUObject : ModuleRules
 			PublicDependencyModuleNames.Add("CookOnTheFly");
 		}
 
+		//@TODO: UE-127233
+		// UnsafeTypeCastWarningLevel = WarningLevel.Warning;
+
 		if (Target.bBuildWithEditorOnlyData)
 		{
 			PrivateDependencyModuleNames.Add("DerivedDataCache");
diff --git a/Engine/Source/Runtime/CoreUObject/Private/Misc/PackageName.cpp b/Engine/Source/Runtime/CoreUObject/Private/Misc/PackageName.cpp
--- a/Engine/Source/Runtime/CoreUObject/Private/Misc/PackageName.cpp
+++ b/Engine/Source/Runtime/CoreUObject/Private/Misc/PackageName.cpp
@@ -451,7 +451,7 @@ void FPackageName::InternalFilenameToLongPackageName(FStringView InFilename, FSt
 	// between /Path/.ext and /Path/
 	if (!SplitOutName.IsEmpty() || SplitOutExtension.IsEmpty())
 	{
-		Result = FStringView(SplitOutPath.GetData(), SplitOutName.GetData() + SplitOutName.Len() - SplitOutPath.GetData());
+		Result = FStringView(SplitOutPath.GetData(), UE_PTRDIFF_TO_INT32(SplitOutName.GetData() + SplitOutName.Len() - SplitOutPath.GetData()));
 	}
 
 	if (bIsValidLongPackageName && Result.Len() != Filename.Len())
@@ -1682,7 +1682,7 @@ bool FPackageName::SearchForPackageOnDisk(const FString& PackageName, FString* O
 			}
 		}
 	}
-	float ThisTime = FPlatformTime::Seconds() - StartTime;
+	const double ThisTime = FPlatformTime::Seconds() - StartTime;
 
 	if ( bResult )
 	{
@@ -1818,7 +1818,7 @@ FString FPackageName::GetSourcePackagePath(const FString& InLocalizedPackagePath
 		if (FCString::Strnicmp(CurChar, TEXT("L10N/"), 5) == 0) // StartsWith "L10N/"
 		{
 			CurChar -= 1; // -1 because we need to eat the slash before L10N
-			OutL10NStart = (CurChar - *InPath);
+			OutL10NStart = UE_PTRDIFF_TO_INT32(CurChar - *InPath);
 			OutL10NLength = 6; // "/L10N/"
 
 			// Walk to the next slash as that will be the end of the culture code
@@ -1830,7 +1830,7 @@ FString FPackageName::GetSourcePackagePath(const FString& InLocalizedPackagePath
 		else if (FCString::Stricmp(CurChar, TEXT("L10N")) == 0) // Is "L10N"
 		{
 			CurChar -= 1; // -1 because we need to eat the slash before L10N
-			OutL10NStart = (CurChar - *InPath);
+			OutL10NStart = UE_PTRDIFF_TO_INT32(CurChar - *InPath);
 			OutL10NLength = 5; // "/L10N"
 
 			return true;
@@ -2250,7 +2250,7 @@ bool FPackageName::IsLocalizedPackage(FStringView InPackageName)
 	}
 
 	// Are we part of the L10N folder?
-	FStringView Remaining(CurChar, EndChar - CurChar);
+	FStringView Remaining(CurChar, UE_PTRDIFF_TO_INT32(EndChar - CurChar));
 	// Is "L10N" or StartsWith "L10N/" 
 	return Remaining.StartsWith(TEXT("L10N"_SV), ESearchCase::IgnoreCase) && (Remaining.Len() == 4 || Remaining[4] == '/');
 }
diff --git a/Engine/Source/Runtime/CoreUObject/Private/Serialization/ArchiveStackTrace.cpp b/Engine/Source/Runtime/CoreUObject/Private/Serialization/ArchiveStackTrace.cpp
--- a/Engine/Source/Runtime/CoreUObject/Private/Serialization/ArchiveStackTrace.cpp
+++ b/Engine/Source/Runtime/CoreUObject/Private/Serialization/ArchiveStackTrace.cpp
@@ -612,7 +612,7 @@ void FArchiveStackTrace::CompareWithInternal(const FPackageData& SourcePackage,
 					if (SourceSize == DestSize && ShouldDumpPropertyValueState(SerProp))
 					{
 						// Walk backwards until we find a callstack which wasn't from the given property
-						int32 OffsetX = DestAbsoluteOffset;
+						int64 OffsetX = DestAbsoluteOffset;
 						for (;;)
 						{
 							if (OffsetX == 0)
@@ -868,7 +868,7 @@ void FArchiveStackTrace::CompareWith(const TCHAR* InFilename, const int64 TotalH
 bool FArchiveStackTrace::GenerateDiffMapInternal(const FPackageData& SourcePackage, const FPackageData& DestPackage, int32 MaxDiffsToFind, FArchiveDiffMap& OutDiffMap)
 {
 	bool bIdentical = true;
-	int32 LastDifferenceCallstackOffsetIndex = -1;
+	int64 LastDifferenceCallstackOffsetIndex = -1;
 	FCallstackData* DifferenceCallstackData = nullptr;
 
 	const int64 SourceSize = SourcePackage.Size - SourcePackage.StartOffset;
@@ -884,7 +884,7 @@ bool FArchiveStackTrace::GenerateDiffMapInternal(const FPackageData& SourcePacka
 			bIdentical = false;
 			if (OutDiffMap.Num() < MaxDiffsToFind)
 			{
-				int64 DifferenceCallstackoffsetIndex = GetCallstackAtOffset(DestAbsoluteOffset, FMath::Max(LastDifferenceCallstackOffsetIndex, 0));
+				int64 DifferenceCallstackoffsetIndex = GetCallstackAtOffset(DestAbsoluteOffset, FMath::Max<int64>(LastDifferenceCallstackOffsetIndex, 0));
 				if (DifferenceCallstackoffsetIndex >= 0 && DifferenceCallstackoffsetIndex != LastDifferenceCallstackOffsetIndex)
 				{
 					const FCallstactAtOffset& CallstackAtOffset = CallstackAtOffsetMap[DifferenceCallstackoffsetIndex];
diff --git a/Engine/Source/Runtime/CoreUObject/Private/Serialization/AsyncLoading.cpp b/Engine/Source/Runtime/CoreUObject/Private/Serialization/AsyncLoading.cpp
--- a/Engine/Source/Runtime/CoreUObject/Private/Serialization/AsyncLoading.cpp
+++ b/Engine/Source/Runtime/CoreUObject/Private/Serialization/AsyncLoading.cpp
@@ -7046,7 +7046,7 @@ void FAsyncLoadingThread::FlushLoading(int32 PackageID)
 	}
 }
 
-EAsyncPackageState::Type FAsyncLoadingThread::ProcessLoadingUntilComplete(TFunctionRef<bool()> CompletionPredicate, float TimeLimit)
+EAsyncPackageState::Type FAsyncLoadingThread::ProcessLoadingUntilComplete(TFunctionRef<bool()> CompletionPredicate, double TimeLimit)
 {
 	if (!IsAsyncLoadingPackages())
 	{
@@ -7058,10 +7058,10 @@ EAsyncPackageState::Type FAsyncLoadingThread::ProcessLoadingUntilComplete(TFunct
 	// Flushing async loading while loading is suspend will result in infinite stall
 	UE_CLOG(IsAsyncLoadingSuspendedInternal(), LogStreaming, Fatal, TEXT("Cannot Flush Async Loading while async loading is suspended (%d)"), GetAsyncLoadingSuspendedCount());
 
-	if (TimeLimit <= 0.0f)
+	if (TimeLimit <= 0.0)
 	{
 		// Set to one hour if no time limit
-		TimeLimit = 60 * 60;
+		TimeLimit = 60.0 * 60.0;
 	}
 
 	while (IsAsyncLoadingPackages() && TimeLimit > 0 && !CompletionPredicate())
@@ -7082,10 +7082,10 @@ EAsyncPackageState::Type FAsyncLoadingThread::ProcessLoadingUntilComplete(TFunct
 		TimeLimit -= (FPlatformTime::Seconds() - TickStartTime);
 	}
 
-	return TimeLimit <= 0 ? EAsyncPackageState::TimeOut : EAsyncPackageState::Complete;
+	return TimeLimit <= 0.0 ? EAsyncPackageState::TimeOut : EAsyncPackageState::Complete;
 }
 
-EAsyncPackageState::Type FAsyncLoadingThread::ProcessLoading(bool bUseTimeLimit, bool bUseFullTimeLimit, float TimeLimit)
+EAsyncPackageState::Type FAsyncLoadingThread::ProcessLoading(bool bUseTimeLimit, bool bUseFullTimeLimit, double TimeLimit)
 {
 	SCOPE_CYCLE_COUNTER(STAT_AsyncLoadingTime);
 	CSV_SCOPED_TIMING_STAT_EXCLUSIVE(AsyncLoading);
@@ -7455,7 +7455,7 @@ void FAsyncArchive::Seek(int64 InPos)
 #endif
 }
 
-bool FAsyncArchive::WaitRead(float TimeLimit)
+bool FAsyncArchive::WaitRead(double TimeLimit)
 {
 	if (ReadRequestPtr)
 	{
@@ -7545,7 +7545,7 @@ void FAsyncArchive::CancelRead()
 	ReadRequestSize = 0;
 }
 
-bool FAsyncArchive::WaitForIntialPhases(float InTimeLimit)
+bool FAsyncArchive::WaitForIntialPhases(double InTimeLimit)
 {
 	if (SizeRequestPtr 
 		|| GEventDrivenLoaderEnabled || SummaryRequestPtr || SummaryPrecacheRequestPtr
@@ -7570,8 +7570,8 @@ bool FAsyncArchive::WaitForIntialPhases(float InTimeLimit)
 		{
 			if (SummaryRequestPtr)
 			{
-				float TimeLimit = 0.0f;
-				if (InTimeLimit > 0.0f)
+				double TimeLimit = 0.0;
+				if (InTimeLimit > 0.0)
 				{
 					TimeLimit = InTimeLimit - (FPlatformTime::Seconds() - StartTime);
 					if (TimeLimit < MIN_REMAIN_TIME)
@@ -7586,14 +7586,14 @@ bool FAsyncArchive::WaitForIntialPhases(float InTimeLimit)
 				}
 				else
 				{
-					check(InTimeLimit > 0.0f);
+					check(InTimeLimit > 0.0);
 					return false;
 				}
 			}
 			if (SummaryPrecacheRequestPtr)
 			{
-				float TimeLimit = 0.0f;
-				if (InTimeLimit > 0.0f)
+				double TimeLimit = 0.0;
+				if (InTimeLimit > 0.0)
 				{
 					TimeLimit = InTimeLimit - (FPlatformTime::Seconds() - StartTime);
 					if (TimeLimit < MIN_REMAIN_TIME)
@@ -7608,7 +7608,7 @@ bool FAsyncArchive::WaitForIntialPhases(float InTimeLimit)
 				}
 				else
 				{
-					check(InTimeLimit > 0.0f);
+					check(InTimeLimit > 0.0);
 					return false;
 				}
 			}
@@ -7800,7 +7800,7 @@ IAsyncReadRequest* FAsyncArchive::MakeEventDrivenPrecacheRequest(int64 Offset, i
 }
 
 
-bool FAsyncArchive::PrecacheWithTimeLimit(int64 RequestOffset, int64 RequestSize, bool bUseTimeLimit, bool bUseFullTimeLimit, double TickStartTime, float TimeLimit)
+bool FAsyncArchive::PrecacheWithTimeLimit(int64 RequestOffset, int64 RequestSize, bool bUseTimeLimit, bool bUseFullTimeLimit, double TickStartTime, double TimeLimit)
 {
 #if defined(ASYNC_WATCH_FILE)
 	if (FileName.Contains(TEXT(ASYNC_WATCH_FILE)) && RequestOffset == 878129)
@@ -7824,7 +7824,7 @@ bool FAsyncArchive::PrecacheWithTimeLimit(int64 RequestOffset, int64 RequestSize
 	bool bResult = PrecacheInternal(RequestOffset, RequestSize);
 	if (!bResult && bUseFullTimeLimit)
 	{
-		float RemainingTime = TimeLimit - float(FPlatformTime::Seconds() - TickStartTime);
+		double RemainingTime = TimeLimit - (FPlatformTime::Seconds() - TickStartTime);
 		if (RemainingTime > MIN_REMAIN_TIME && WaitRead(RemainingTime))
 		{
 			bResult = true;
@@ -7881,7 +7881,7 @@ void FAsyncArchive::EndReadingHeader()
 	}
 }
 
-bool FAsyncArchive::ReadyToStartReadingHeader(bool bUseTimeLimit, bool bUseFullTimeLimit, double TickStartTime, float TimeLimit)
+bool FAsyncArchive::ReadyToStartReadingHeader(bool bUseTimeLimit, bool bUseFullTimeLimit, double TickStartTime, double TimeLimit)
 {
 	if (SummaryReadTime == 0.0)
 	{
@@ -7895,7 +7895,7 @@ bool FAsyncArchive::ReadyToStartReadingHeader(bool bUseTimeLimit, bool bUseFullT
 	{
 		if (bUseFullTimeLimit)
 		{
-			float RemainingTime = TimeLimit - float(FPlatformTime::Seconds() - TickStartTime);
+			double RemainingTime = TimeLimit - (FPlatformTime::Seconds() - TickStartTime);
 			if (RemainingTime < MIN_REMAIN_TIME || !WaitForIntialPhases(RemainingTime))
 			{
 				return false; // not ready
diff --git a/Engine/Source/Runtime/CoreUObject/Private/Serialization/AsyncLoading2.cpp b/Engine/Source/Runtime/CoreUObject/Private/Serialization/AsyncLoading2.cpp
--- a/Engine/Source/Runtime/CoreUObject/Private/Serialization/AsyncLoading2.cpp
+++ b/Engine/Source/Runtime/CoreUObject/Private/Serialization/AsyncLoading2.cpp
@@ -2354,17 +2354,17 @@ public:
 		int32 InPackagePriority,
 		const FLinkerInstancingContext* InstancingContext = nullptr) override;
 
-	EAsyncPackageState::Type ProcessLoadingFromGameThread(FAsyncLoadingThreadState2& ThreadState, bool bUseTimeLimit, bool bUseFullTimeLimit, float TimeLimit);
+	EAsyncPackageState::Type ProcessLoadingFromGameThread(FAsyncLoadingThreadState2& ThreadState, bool bUseTimeLimit, bool bUseFullTimeLimit, double TimeLimit);
 
-	inline virtual EAsyncPackageState::Type ProcessLoading(bool bUseTimeLimit, bool bUseFullTimeLimit, float TimeLimit) override
+	inline virtual EAsyncPackageState::Type ProcessLoading(bool bUseTimeLimit, bool bUseFullTimeLimit, double TimeLimit) override
 	{
 		FAsyncLoadingThreadState2& ThreadState = *FAsyncLoadingThreadState2::Get();
 		return ProcessLoadingFromGameThread(ThreadState, bUseTimeLimit, bUseFullTimeLimit, TimeLimit);
 	}
 
-	EAsyncPackageState::Type ProcessLoadingUntilCompleteFromGameThread(FAsyncLoadingThreadState2& ThreadState, TFunctionRef<bool()> CompletionPredicate, float TimeLimit);
+	EAsyncPackageState::Type ProcessLoadingUntilCompleteFromGameThread(FAsyncLoadingThreadState2& ThreadState, TFunctionRef<bool()> CompletionPredicate, double TimeLimit);
 
-	inline virtual EAsyncPackageState::Type ProcessLoadingUntilComplete(TFunctionRef<bool()> CompletionPredicate, float TimeLimit) override
+	inline virtual EAsyncPackageState::Type ProcessLoadingUntilComplete(TFunctionRef<bool()> CompletionPredicate, double TimeLimit) override
 	{
 		FAsyncLoadingThreadState2& ThreadState = *FAsyncLoadingThreadState2::Get();
 		return ProcessLoadingUntilCompleteFromGameThread(ThreadState, CompletionPredicate, TimeLimit);
@@ -5986,7 +5986,7 @@ void FAsyncLoadingThread2::QueueMissingPackage(FAsyncPackageDesc2& PackageDesc,
 	}
 }
 
-EAsyncPackageState::Type FAsyncLoadingThread2::ProcessLoadingFromGameThread(FAsyncLoadingThreadState2& ThreadState, bool bUseTimeLimit, bool bUseFullTimeLimit, float TimeLimit)
+EAsyncPackageState::Type FAsyncLoadingThread2::ProcessLoadingFromGameThread(FAsyncLoadingThreadState2& ThreadState, bool bUseTimeLimit, bool bUseFullTimeLimit, double TimeLimit)
 {
 	SCOPE_CYCLE_COUNTER(STAT_AsyncLoadingTime);
 	CSV_SCOPED_TIMING_STAT_EXCLUSIVE(AsyncLoading);
@@ -6052,7 +6052,7 @@ void FAsyncLoadingThread2::FlushLoading(int32 RequestId)
 	}
 }
 
-EAsyncPackageState::Type FAsyncLoadingThread2::ProcessLoadingUntilCompleteFromGameThread(FAsyncLoadingThreadState2& ThreadState, TFunctionRef<bool()> CompletionPredicate, float TimeLimit)
+EAsyncPackageState::Type FAsyncLoadingThread2::ProcessLoadingUntilCompleteFromGameThread(FAsyncLoadingThreadState2& ThreadState, TFunctionRef<bool()> CompletionPredicate, double TimeLimit)
 {
 	if (!IsAsyncLoadingPackages())
 	{
diff --git a/Engine/Source/Runtime/CoreUObject/Private/Serialization/AsyncLoadingPrivate.h b/Engine/Source/Runtime/CoreUObject/Private/Serialization/AsyncLoadingPrivate.h
--- a/Engine/Source/Runtime/CoreUObject/Private/Serialization/AsyncLoadingPrivate.h
+++ b/Engine/Source/Runtime/CoreUObject/Private/Serialization/AsyncLoadingPrivate.h
@@ -45,10 +45,10 @@ public:
 	}
 
 	/** AsyncArchive interface */
-	bool PrecacheWithTimeLimit(int64 PrecacheOffset, int64 PrecacheSize, bool bUseTimeLimit, bool bUseFullTimeLimit, double TickStartTime, float TimeLimit);
+	bool PrecacheWithTimeLimit(int64 PrecacheOffset, int64 PrecacheSize, bool bUseTimeLimit, bool bUseFullTimeLimit, double TickStartTime, double TimeLimit);
 	bool PrecacheForEvent(IAsyncReadRequest* Read, int64 PrecacheOffset, int64 PrecacheSize);
 	void FlushPrecacheBlock();
-	bool ReadyToStartReadingHeader(bool bUseTimeLimit, bool bUseFullTimeLimit, double TickStartTime, float TimeLimit);
+	bool ReadyToStartReadingHeader(bool bUseTimeLimit, bool bUseFullTimeLimit, double TickStartTime, double TimeLimit);
 	void StartReadingHeader();
 	void EndReadingHeader();
 	IAsyncReadRequest* MakeEventDrivenPrecacheRequest(int64 Offset, int64 BytesToRead, FAsyncFileCallBack* CompleteCallback);
@@ -70,11 +70,11 @@ private:
 
 #endif
 	void FirstExportStarting();
-	bool WaitRead(float TimeLimit = 0.0f);
+	bool WaitRead(double TimeLimit = 0.0);
 	void CompleteRead();
 	void CancelRead();
 	void CompleteCancel();
-	bool WaitForIntialPhases(float TimeLimit = 0.0f);
+	bool WaitForIntialPhases(double TimeLimit = 0.0);
 	void ReadCallback(bool bWasCancelled, IAsyncReadRequest*);
 	bool PrecacheInternal(int64 PrecacheOffset, int64 PrecacheSize, bool bApplyMinReadSize = true, IAsyncReadRequest* Read = nullptr);
 	
diff --git a/Engine/Source/Runtime/CoreUObject/Private/Serialization/AsyncLoadingThread.h b/Engine/Source/Runtime/CoreUObject/Private/Serialization/AsyncLoadingThread.h
--- a/Engine/Source/Runtime/CoreUObject/Private/Serialization/AsyncLoadingThread.h
+++ b/Engine/Source/Runtime/CoreUObject/Private/Serialization/AsyncLoadingThread.h
@@ -275,9 +275,9 @@ public:
 			int32 InPackagePriority,
 			const FLinkerInstancingContext* InstancingContext) override;
 
-	EAsyncPackageState::Type ProcessLoading(bool bUseTimeLimit, bool bUseFullTimeLimit, float TimeLimit) override;
+	EAsyncPackageState::Type ProcessLoading(bool bUseTimeLimit, bool bUseFullTimeLimit, double TimeLimit) override;
 
-	EAsyncPackageState::Type ProcessLoadingUntilComplete(TFunctionRef<bool()> CompletionPredicate, float TimeLimit) override;
+	EAsyncPackageState::Type ProcessLoadingUntilComplete(TFunctionRef<bool()> CompletionPredicate, double TimeLimit) override;
 
 	void FlushLoading(int32 PackageId) override;
 
diff --git a/Engine/Source/Runtime/CoreUObject/Private/Serialization/AsyncPackageLoader.cpp b/Engine/Source/Runtime/CoreUObject/Private/Serialization/AsyncPackageLoader.cpp
--- a/Engine/Source/Runtime/CoreUObject/Private/Serialization/AsyncPackageLoader.cpp
+++ b/Engine/Source/Runtime/CoreUObject/Private/Serialization/AsyncPackageLoader.cpp
@@ -657,7 +657,7 @@ void FlushAsyncLoading(int32 PackageID /* = INDEX_NONE */)
 	}
 }
 
-EAsyncPackageState::Type ProcessAsyncLoadingUntilComplete(TFunctionRef<bool()> CompletionPredicate, float TimeLimit)
+EAsyncPackageState::Type ProcessAsyncLoadingUntilComplete(TFunctionRef<bool()> CompletionPredicate, double TimeLimit)
 {
 	LLM_SCOPE(ELLMTag::AsyncLoading);
 	return GetAsyncPackageLoader().ProcessLoadingUntilComplete(CompletionPredicate, TimeLimit);
@@ -668,7 +668,7 @@ int32 GetNumAsyncPackages()
 	return GetAsyncPackageLoader().GetNumAsyncPackages();
 }
 
-EAsyncPackageState::Type ProcessAsyncLoading(bool bUseTimeLimit, bool bUseFullTimeLimit, float TimeLimit)
+EAsyncPackageState::Type ProcessAsyncLoading(bool bUseTimeLimit, bool bUseFullTimeLimit, double TimeLimit)
 {
 	LLM_SCOPE(ELLMTag::AsyncLoading);
 	TRACE_CPUPROFILER_EVENT_SCOPE(ProcessAsyncLoading);
@@ -1096,10 +1096,10 @@ void IsTimeLimitExceededPrint(
 		(CurrentTime - InTickStartTime) > GTimeLimitExceededMinTime &&
 		(CurrentTime - InTickStartTime) > (GTimeLimitExceededMultiplier * InTimeLimit))
 	{
-		float EstimatedTimeForThisStep = (CurrentTime - InTickStartTime) * 1000;
+		double EstimatedTimeForThisStep = (CurrentTime - InTickStartTime) * 1000.0;
 		if (LastTestTime > InTickStartTime)
 		{
-			EstimatedTimeForThisStep = (CurrentTime - LastTestTime) * 1000;
+			EstimatedTimeForThisStep = (CurrentTime - LastTestTime) * 1000.0;
 		}
 		LastPrintStartTime = InTickStartTime;
 		UE_LOG(LogStreaming, Warning, TEXT("IsTimeLimitExceeded: %s %s Load Time %5.2fms   Last Step Time %5.2fms"),
diff --git a/Engine/Source/Runtime/CoreUObject/Private/Serialization/EditorPackageLoader.cpp b/Engine/Source/Runtime/CoreUObject/Private/Serialization/EditorPackageLoader.cpp
--- a/Engine/Source/Runtime/CoreUObject/Private/Serialization/EditorPackageLoader.cpp
+++ b/Engine/Source/Runtime/CoreUObject/Private/Serialization/EditorPackageLoader.cpp
@@ -66,7 +66,7 @@ public:
 		}
 	}
 
-	virtual EAsyncPackageState::Type ProcessLoading(bool bUseTimeLimit, bool bUseFullTimeLimit, float TimeLimit) override
+	virtual EAsyncPackageState::Type ProcessLoading(bool bUseTimeLimit, bool bUseFullTimeLimit, double TimeLimit) override
 	{
 		EAsyncPackageState::Type CookedLoadingState = CookedPackageLoader->ProcessLoading(bUseTimeLimit, bUseFullTimeLimit, TimeLimit);
 		EAsyncPackageState::Type UncookedLoadingState = UncookedPackageLoader->ProcessLoading(bUseTimeLimit, bUseFullTimeLimit, TimeLimit);
@@ -76,7 +76,7 @@ public:
 			: EAsyncPackageState::TimeOut;
 	}
 
-	virtual EAsyncPackageState::Type ProcessLoadingUntilComplete(TFunctionRef<bool()> CompletionPredicate, float TimeLimit) override
+	virtual EAsyncPackageState::Type ProcessLoadingUntilComplete(TFunctionRef<bool()> CompletionPredicate, double TimeLimit) override
 	{
 		const EAsyncPackageState::Type LoadingState = CookedPackageLoader->ProcessLoadingUntilComplete(CompletionPredicate, TimeLimit);
 		if (LoadingState != EAsyncPackageState::Complete)
diff --git a/Engine/Source/Runtime/CoreUObject/Private/UObject/GarbageCollection.cpp b/Engine/Source/Runtime/CoreUObject/Private/UObject/GarbageCollection.cpp
--- a/Engine/Source/Runtime/CoreUObject/Private/UObject/GarbageCollection.cpp
+++ b/Engine/Source/Runtime/CoreUObject/Private/UObject/GarbageCollection.cpp
@@ -1440,7 +1440,7 @@ struct FConditionalGCLock
 	}
 };
 
-static bool IncrementalDestroyGarbage(bool bUseTimeLimit, float TimeLimit);
+static bool IncrementalDestroyGarbage(bool bUseTimeLimit, double TimeLimit);
 
 /**
  * Incrementally purge garbage by deleting all unreferenced objects after routing Destroy.
@@ -1451,7 +1451,7 @@ static bool IncrementalDestroyGarbage(bool bUseTimeLimit, float TimeLimit);
  * @param	bUseTimeLimit	whether the time limit parameter should be used
  * @param	TimeLimit		soft time limit for this function call
  */
-void IncrementalPurgeGarbage(bool bUseTimeLimit, float TimeLimit)
+void IncrementalPurgeGarbage(bool bUseTimeLimit, double TimeLimit)
 {
 #if !UE_WITH_GC
 	return;
@@ -1545,7 +1545,7 @@ static FAutoConsoleVariableRef CVarAdditionalFinishDestroyTimeGC(
 	ECVF_Default
 );
 
-bool IncrementalDestroyGarbage(bool bUseTimeLimit, float TimeLimit)
+bool IncrementalDestroyGarbage(bool bUseTimeLimit, double TimeLimit)
 {
 	TRACE_CPUPROFILER_EVENT_SCOPE(IncrementalDestroyGarbage);
 	const bool bMultithreadedPurge = !ShouldForceSingleThreadedGC() && GMultithreadedDestructionEnabled;
@@ -2158,7 +2158,7 @@ bool IsIncrementalUnhashPending()
 	return GUnrechableObjectIndex < GUnreachableObjects.Num();
 }
 
-bool UnhashUnreachableObjects(bool bUseTimeLimit, float TimeLimit)
+bool UnhashUnreachableObjects(bool bUseTimeLimit, double TimeLimit)
 {
 	TRACE_CPUPROFILER_EVENT_SCOPE(UnhashUnreachableObjects);
 	DECLARE_SCOPE_CYCLE_COUNTER(TEXT("UnhashUnreachableObjects"), STAT_UnhashUnreachableObjects, STATGROUP_GC);
diff --git a/Engine/Source/Runtime/CoreUObject/Private/UObject/Linker.cpp b/Engine/Source/Runtime/CoreUObject/Private/UObject/Linker.cpp
--- a/Engine/Source/Runtime/CoreUObject/Private/UObject/Linker.cpp
+++ b/Engine/Source/Runtime/CoreUObject/Private/UObject/Linker.cpp
@@ -160,9 +160,9 @@ FName FLinker::GetExportClassName( int32 i )
 	return NAME_Class;
 }
 
-uint32 FLinkerTables::GetAllocatedSize() const
+SIZE_T FLinkerTables::GetAllocatedSize() const
 {
-	uint32 Result = 0;
+	SIZE_T Result = 0;
 	Result += ImportMap.GetAllocatedSize();
 	Result += ExportMap.GetAllocatedSize();
 	Result += DependsMap.GetAllocatedSize();
diff --git a/Engine/Source/Runtime/CoreUObject/Private/UObject/LinkerLoad.cpp b/Engine/Source/Runtime/CoreUObject/Private/UObject/LinkerLoad.cpp
--- a/Engine/Source/Runtime/CoreUObject/Private/UObject/LinkerLoad.cpp
+++ b/Engine/Source/Runtime/CoreUObject/Private/UObject/LinkerLoad.cpp
@@ -1187,7 +1187,7 @@ FLinkerLoad::ELinkerStatus FLinkerLoad::CreateLoader(
 				// and decompress it.
 				static int64 MinimumReadSize = 32 * 1024;
 				checkSlow(MinimumReadSize >= 2048 && MinimumReadSize <= 1024 * 1024); // not a hard limit, but we should be loading at least a reasonable amount of data
-				int32 PrecacheSize = FMath::Min(MinimumReadSize, Size);
+				int64 PrecacheSize = FMath::Min(MinimumReadSize, Size);
 				check( PrecacheSize > 0 );
 				// Wait till we're finished precaching before executing the next step.
 				bExecuteNextStep = Loader->Precache(0, PrecacheSize);
diff --git a/Engine/Source/Runtime/CoreUObject/Private/UObject/PackageReload.cpp b/Engine/Source/Runtime/CoreUObject/Private/UObject/PackageReload.cpp
--- a/Engine/Source/Runtime/CoreUObject/Private/UObject/PackageReload.cpp
+++ b/Engine/Source/Runtime/CoreUObject/Private/UObject/PackageReload.cpp
@@ -488,7 +488,7 @@ void ReloadPackages(const TArrayView<FReloadPackageData>& InPackagesToReload, TA
 	}
 	UE_LOG(LogUObjectGlobals, Log, TEXT("%s"), *Msg);
 
-	FScopedSlowTask ReloadingPackagesSlowTask(InPackagesToReload.Num(), NSLOCTEXT("CoreUObject", "ReloadingPackages", "Reloading Packages"));
+	FScopedSlowTask ReloadingPackagesSlowTask((float)InPackagesToReload.Num(), NSLOCTEXT("CoreUObject", "ReloadingPackages", "Reloading Packages"));
 	ReloadingPackagesSlowTask.MakeDialogDelayed(3.0f);
 
 	// Cache the current dirty state of all packages so we can restore it after the reload
@@ -506,7 +506,7 @@ void ReloadPackages(const TArrayView<FReloadPackageData>& InPackagesToReload, TA
 	PackageReloadInternal::FExistingPackageReferences ExistingPackages;
 	ExistingPackages.Refs.Reserve(InPackagesToReload.Num());
 	{
-		FScopedSlowTask PreparingPackagesForReloadSlowTask(InPackagesToReload.Num(), NSLOCTEXT("CoreUObject", "PreparingPackagesForReload", "Preparing Packages for Reload"));
+		FScopedSlowTask PreparingPackagesForReloadSlowTask((float)InPackagesToReload.Num(), NSLOCTEXT("CoreUObject", "PreparingPackagesForReload", "Preparing Packages for Reload"));
 
 		for (const FReloadPackageData& PackageToReloadData : InPackagesToReload)
 		{
@@ -562,7 +562,7 @@ void ReloadPackages(const TArrayView<FReloadPackageData>& InPackagesToReload, TA
 
 			const int32 NumPackagesInBatch = PackageIndex - BatchStartIndex;
 
-			FScopedSlowTask FixingUpReferencesSlowTask((NumPackagesInBatch * 4) + GUObjectArray.GetObjectArrayNum(), NSLOCTEXT("CoreUObject", "FixingUpReferences", "Fixing-Up References"));
+			FScopedSlowTask FixingUpReferencesSlowTask((float)((NumPackagesInBatch * 4) + GUObjectArray.GetObjectArrayNum()), NSLOCTEXT("CoreUObject", "FixingUpReferences", "Fixing-Up References"));
 
 			// Pre-pass to notify things that the package old package is about to be fixed-up
 			TMap<UObject*, PackageReloadInternal::FObjectAndPackageIndex> OldObjectToNewData;
diff --git a/Engine/Source/Runtime/CoreUObject/Private/UObject/Property.cpp b/Engine/Source/Runtime/CoreUObject/Private/UObject/Property.cpp
--- a/Engine/Source/Runtime/CoreUObject/Private/UObject/Property.cpp
+++ b/Engine/Source/Runtime/CoreUObject/Private/UObject/Property.cpp
@@ -579,7 +579,7 @@ FORCEINLINE constexpr FStringView ParsePropertyToken(const TCHAR* Str, bool Dott
 	FAsciiSet CurrentNonTokenChars = DottedNames ? DottedNonTokenChars : RegularNonTokenChars;
 
 	const TCHAR* TokenEnd = FAsciiSet::FindFirstOrEnd(Str, CurrentNonTokenChars);
-	return FStringView(Str, TokenEnd - Str);
+	return FStringView(Str, UE_PTRDIFF_TO_INT32(TokenEnd - Str));
 }
 
 //
@@ -1301,7 +1301,7 @@ const TCHAR* FProperty::ImportSingleProperty( const TCHAR* Str, void* DestData,
 	if (*Str)
 	{
 		// strip trailing whitespace
-		int32 Len = Str - Start;
+		int32 Len = UE_PTRDIFF_TO_INT32(Str - Start);
 		while (Len > 0 && Whitespaces.Contains(Start[Len - 1]))
 		{
 			--Len;
diff --git a/Engine/Source/Runtime/CoreUObject/Private/UObject/PropertyByte.cpp b/Engine/Source/Runtime/CoreUObject/Private/UObject/PropertyByte.cpp
--- a/Engine/Source/Runtime/CoreUObject/Private/UObject/PropertyByte.cpp
+++ b/Engine/Source/Runtime/CoreUObject/Private/UObject/PropertyByte.cpp
@@ -183,7 +183,7 @@ struct TConvertIntToEnumProperty
 		OldIntType OldValue;
 		Slot << OldValue;
 
-		uint8 NewValue = OldValue;
+		uint8 NewValue = (uint8)OldValue;
 		if (OldValue > (OldIntType)TNumericLimits<uint8>::Max() || !Enum->IsValidEnumValue(NewValue))
 		{
 			UE_LOG(
@@ -350,8 +350,8 @@ void FByteProperty::ExportTextItem( FString& ValueStr, const void* PropertyValue
 		if (Enum)
 		{
 			const int32 ActualValue = *(const uint8*)PropertyValue;
-			const int32 MaxValue = Enum->GetMaxEnumValue();
-			const int32 GoodValue = Enum->IsValidEnumValue(ActualValue) ? ActualValue : MaxValue;
+			const int64 MaxValue = Enum->GetMaxEnumValue();
+			const int64 GoodValue = Enum->IsValidEnumValue(ActualValue) ? ActualValue : MaxValue;
 			const bool bNonNativeEnum = Enum->GetClass() != UEnum::StaticClass();
 			ensure(!bNonNativeEnum || Enum->CppType.IsEmpty());
 			const FString FullyQualifiedEnumName = bNonNativeEnum ? ::UnicodeToCPPIdentifier(Enum->GetName(), false, TEXT("E__"))
diff --git a/Engine/Source/Runtime/CoreUObject/Private/UObject/PropertyTag.cpp b/Engine/Source/Runtime/CoreUObject/Private/UObject/PropertyTag.cpp
--- a/Engine/Source/Runtime/CoreUObject/Private/UObject/PropertyTag.cpp
+++ b/Engine/Source/Runtime/CoreUObject/Private/UObject/PropertyTag.cpp
@@ -213,7 +213,7 @@ void FPropertyTag::SerializeTaggedProperty(FArchive& Ar, FProperty* Property, ui
 void FPropertyTag::SerializeTaggedProperty(FStructuredArchive::FSlot Slot, FProperty* Property, uint8* Value, uint8* Defaults) const
 {
 	FArchive& UnderlyingArchive = Slot.GetUnderlyingArchive();
-	const int32 StartOfProperty = UnderlyingArchive.Tell();
+	const int64 StartOfProperty = UnderlyingArchive.Tell();
 
 	if (!UnderlyingArchive.IsTextFormat() && Property->GetClass() == FBoolProperty::StaticClass())
 	{
@@ -242,10 +242,10 @@ void FPropertyTag::SerializeTaggedProperty(FStructuredArchive::FSlot Slot, FProp
 	}
 
 	// Ensure that we serialize what we expected to serialize.
-	const int32 EndOfProperty = UnderlyingArchive.Tell();
+	const int64 EndOfProperty = UnderlyingArchive.Tell();
 	if (Size && (EndOfProperty - StartOfProperty != Size))
 	{
-		UE_LOG(LogClass, Error, TEXT("Failed loading tagged %s. Read %dB, expected %dB."), *GetFullNameSafe(Property), EndOfProperty - StartOfProperty, Size);
+		UE_LOG(LogClass, Error, TEXT("Failed loading tagged %s. Read %" INT64_FMT "B, expected %dB."), *GetFullNameSafe(Property), EndOfProperty - StartOfProperty, Size);
 		UnderlyingArchive.Seek(StartOfProperty + Size);
 		Property->ClearValue(Value);
 	}
diff --git a/Engine/Source/Runtime/CoreUObject/Private/UObject/SavePackage/SavePackageUtilities.cpp b/Engine/Source/Runtime/CoreUObject/Private/UObject/SavePackage/SavePackageUtilities.cpp
--- a/Engine/Source/Runtime/CoreUObject/Private/UObject/SavePackage/SavePackageUtilities.cpp
+++ b/Engine/Source/Runtime/CoreUObject/Private/UObject/SavePackage/SavePackageUtilities.cpp
@@ -2273,7 +2273,7 @@ ESavePackageResult SaveBulkData(FLinkerSave* Linker, int64& InOutStartOffset, co
 
 	COOK_STAT(FScopedDurationTimer SaveTimer(FSavePackageStats::SerializeBulkDataTimeSec));
 
-	FScopedSlowTask BulkDataFeedback(Linker->BulkDataToAppend.Num());
+	FScopedSlowTask BulkDataFeedback((float)Linker->BulkDataToAppend.Num());
 
 	TUniquePtr<FLargeMemoryWriterWithRegions> BulkArchive;
 	TUniquePtr<FLargeMemoryWriterWithRegions> OptionalBulkArchive;
diff --git a/Engine/Source/Runtime/CoreUObject/Private/UObject/UObjectGlobals.cpp b/Engine/Source/Runtime/CoreUObject/Private/UObject/UObjectGlobals.cpp
--- a/Engine/Source/Runtime/CoreUObject/Private/UObject/UObjectGlobals.cpp
+++ b/Engine/Source/Runtime/CoreUObject/Private/UObject/UObjectGlobals.cpp
@@ -1191,7 +1191,7 @@ UPackage* LoadPackageInternal(UPackage* InOuter, const FPackagePath& PackagePath
 	if (ShouldCreateThrottledSlowTask())
 	{
 		static const FTextFormat LoadingPackageTextFormat = NSLOCTEXT("Core", "LoadingPackage_Scope", "Loading Package '{0}'");
-		SlowTask.Emplace(100, FText::Format(LoadingPackageTextFormat, PackagePath.GetDebugNameText()));
+		SlowTask.Emplace(100.0f, FText::Format(LoadingPackageTextFormat, PackagePath.GetDebugNameText()));
 		SlowTask->Visibility = ESlowTaskVisibility::Invisible;
 		SlowTask->EnterProgressFrame(10);
 	}
diff --git a/Engine/Source/Runtime/CoreUObject/Private/UObject/UObjectHash.cpp b/Engine/Source/Runtime/CoreUObject/Private/UObject/UObjectHash.cpp
--- a/Engine/Source/Runtime/CoreUObject/Private/UObject/UObjectHash.cpp
+++ b/Engine/Source/Runtime/CoreUObject/Private/UObject/UObjectHash.cpp
@@ -543,13 +543,13 @@ struct FObjectSearchPath
 		const TCHAR* End = FAsciiSet::FindFirstOrEnd(Begin, DotColon);
 		while (*End != '\0')
 		{
-			Outers.Add(FName(End - Begin, Begin));
+			Outers.Add(FName(UE_PTRDIFF_TO_INT32(End - Begin), Begin));
 
 			Begin = End + 1;
 			End = FAsciiSet::FindFirstOrEnd(Begin, DotColon);
 		}
 
-		Inner = Outers.Num() == 0 ? InPath : FName(FStringView(Begin, End - Begin), InPath.GetNumber());
+		Inner = Outers.Num() == 0 ? InPath : FName(FStringView(Begin, UE_PTRDIFF_TO_INT32(End - Begin)), InPath.GetNumber());
 	}
 
 	bool MatchOuterNames(UObject* Outer) const
diff --git a/Engine/Source/Runtime/CoreUObject/Public/Serialization/AsyncPackageLoader.h b/Engine/Source/Runtime/CoreUObject/Public/Serialization/AsyncPackageLoader.h
--- a/Engine/Source/Runtime/CoreUObject/Public/Serialization/AsyncPackageLoader.h
+++ b/Engine/Source/Runtime/CoreUObject/Public/Serialization/AsyncPackageLoader.h
@@ -137,7 +137,7 @@ public:
 	 * @param bUseFullTimeLimit	
 	 * @param TimeLimit				Time limit
 	 */
-	virtual EAsyncPackageState::Type ProcessLoading(bool bUseTimeLimit, bool bUseFullTimeLimit, float TimeLimit) = 0;
+	virtual EAsyncPackageState::Type ProcessLoading(bool bUseTimeLimit, bool bUseFullTimeLimit, double TimeLimit) = 0;
 	
 	/**
 	 * Process all loading package requests until completion predicate is satisfied.
@@ -145,7 +145,7 @@ public:
 	 * @param CompletionPredicate		Completion predicate
 	 * @param TimeLimit					Time limit
 	 */
-	virtual EAsyncPackageState::Type ProcessLoadingUntilComplete(TFunctionRef<bool()> CompletionPredicate, float TimeLimit) = 0;
+	virtual EAsyncPackageState::Type ProcessLoadingUntilComplete(TFunctionRef<bool()> CompletionPredicate, double TimeLimit) = 0;
 
 	/**
 	* Cancels streaming.
diff --git a/Engine/Source/Runtime/CoreUObject/Public/UObject/FastReferenceCollector.h b/Engine/Source/Runtime/CoreUObject/Public/UObject/FastReferenceCollector.h
--- a/Engine/Source/Runtime/CoreUObject/Public/UObject/FastReferenceCollector.h
+++ b/Engine/Source/Runtime/CoreUObject/Public/UObject/FastReferenceCollector.h
@@ -98,7 +98,7 @@ public:
 		checkSlow(CheckUsedArrays == 0);
 #endif // UE_BUILD_DEBUG
 
-		uint32 FreedMemory = 0;
+		SIZE_T FreedMemory = 0;
 		TArray<FGCArrayStruct*> AllArrays;
 		Pool.PopAll(AllArrays);
 		for (FGCArrayStruct* ArrayStruct : AllArrays)
@@ -109,7 +109,7 @@ public:
 			FreedMemory += ArrayStruct->WeakReferences.GetAllocatedSize();
 			delete ArrayStruct;
 		}
-		UE_LOG(LogGarbage, Log, TEXT("Freed %ub from %d GC array pools."), FreedMemory, AllArrays.Num());
+		UE_LOG(LogGarbage, Log, TEXT("Freed %" SIZE_T_FMT "b from %d GC array pools."), FreedMemory, AllArrays.Num());
 	}
 
 	/**
diff --git a/Engine/Source/Runtime/CoreUObject/Public/UObject/Linker.h b/Engine/Source/Runtime/CoreUObject/Public/UObject/Linker.h
--- a/Engine/Source/Runtime/CoreUObject/Public/UObject/Linker.h
+++ b/Engine/Source/Runtime/CoreUObject/Public/UObject/Linker.h
@@ -168,7 +168,7 @@ public:
 	COREUOBJECT_API void SerializeSearchableNamesMap(FStructuredArchive::FSlot Slot);
 
 	/** Returns the amount of memory allocated by this container, not including sizeof(*this). */
-	COREUOBJECT_API uint32 GetAllocatedSize() const;
+	COREUOBJECT_API SIZE_T GetAllocatedSize() const;
 };
 
 
diff --git a/Engine/Source/Runtime/CoreUObject/Public/UObject/UObjectGlobals.h b/Engine/Source/Runtime/CoreUObject/Public/UObject/UObjectGlobals.h
--- a/Engine/Source/Runtime/CoreUObject/Public/UObject/UObjectGlobals.h
+++ b/Engine/Source/Runtime/CoreUObject/Public/UObject/UObjectGlobals.h
@@ -559,7 +559,7 @@ COREUOBJECT_API bool TryCollectGarbage(EObjectFlags KeepFlags, bool bPerformFull
 *
 * @return true if the time limit passed and there's still objects pending to be unhashed
 */
-COREUOBJECT_API bool UnhashUnreachableObjects(bool bUseTimeLimit, float TimeLimit = 0.0f);
+COREUOBJECT_API bool UnhashUnreachableObjects(bool bUseTimeLimit, double TimeLimit = 0.0);
 
 /**
 * Checks if there's objects pending to be unhashed when running incremental purge
@@ -591,7 +591,7 @@ COREUOBJECT_API void GatherUnreachableObjects(bool bForceSingleThreaded);
  * @param	bUseTimeLimit	whether the time limit parameter should be used
  * @param	TimeLimit		soft time limit for this function call
  */
-COREUOBJECT_API void IncrementalPurgeGarbage( bool bUseTimeLimit, float TimeLimit = 0.002f );
+COREUOBJECT_API void IncrementalPurgeGarbage( bool bUseTimeLimit, double TimeLimit = 0.002 );
 
 /**
  * Create a unique name by combining a base name and an arbitrary number string.
@@ -684,7 +684,7 @@ namespace EAsyncPackageState
  * @param	TimeLimit		Soft limit of time this function is allowed to consume
  * @return The minimum state of any of the queued packages.
  */
-COREUOBJECT_API EAsyncPackageState::Type ProcessAsyncLoading( bool bUseTimeLimit, bool bUseFullTimeLimit, float TimeLimit);
+COREUOBJECT_API EAsyncPackageState::Type ProcessAsyncLoading( bool bUseTimeLimit, bool bUseFullTimeLimit, double TimeLimit);
 
 /**
  * Blocks and runs ProcessAsyncLoading until the time limit is hit, the completion predicate returns true, or all async loading is done
@@ -693,7 +693,7 @@ COREUOBJECT_API EAsyncPackageState::Type ProcessAsyncLoading( bool bUseTimeLimit
  * @param	TimeLimit			Hard time limit. 0 means infinite length
  * @return The minimum state of any of the queued packages.
  */
-COREUOBJECT_API EAsyncPackageState::Type ProcessAsyncLoadingUntilComplete(TFunctionRef<bool()> CompletionPredicate, float TimeLimit);
+COREUOBJECT_API EAsyncPackageState::Type ProcessAsyncLoadingUntilComplete(TFunctionRef<bool()> CompletionPredicate, double TimeLimit);
 
 /** UObjects are being loaded between these calls */
 COREUOBJECT_API void BeginLoad(FUObjectSerializeContext* LoadContext, const TCHAR* DebugContext = nullptr);
