Optimize FPackageLocalizationCache::ConditionalUpdateCache

1) Batch up calls to AssetRegistry->GetAssetsByPath instead of calling it once per plugin, as the function has a high overhead.
2) Remove trace scope from UE:AssetRegistry FFiltering::ShouldSkipAsset as the function is called very frequently (large projects can have a million assets).

#rnx

diff --git a/Engine/Source/Runtime/AssetRegistry/Private/AssetRegistry.cpp b/Engine/Source/Runtime/AssetRegistry/Private/AssetRegistry.cpp
--- a/Engine/Source/Runtime/AssetRegistry/Private/AssetRegistry.cpp
+++ b/Engine/Source/Runtime/AssetRegistry/Private/AssetRegistry.cpp
@@ -1176,6 +1176,15 @@ bool UAssetRegistryImpl::GetAssetsByPath(FName PackagePath, TArray<FAssetData>&
 	return GetAssets(Filter, OutAssetData);
 }
 
+bool UAssetRegistryImpl::GetAssetsByPaths(TArray<FName> PackagePaths, TArray<FAssetData>& OutAssetData, bool bRecursive, bool bIncludeOnlyOnDiskAssets) const
+{
+	FARFilter Filter;
+	Filter.bRecursivePaths = bRecursive;
+	Filter.PackagePaths = MoveTemp(PackagePaths);
+	Filter.bIncludeOnlyOnDiskAssets = bIncludeOnlyOnDiskAssets;
+	return GetAssets(Filter, OutAssetData);
+}
+
 namespace UE::AssetRegistry
 {
 
diff --git a/Engine/Source/Runtime/AssetRegistry/Private/AssetRegistry.h b/Engine/Source/Runtime/AssetRegistry/Private/AssetRegistry.h
--- a/Engine/Source/Runtime/AssetRegistry/Private/AssetRegistry.h
+++ b/Engine/Source/Runtime/AssetRegistry/Private/AssetRegistry.h
@@ -31,6 +31,7 @@ public:
 	virtual bool HasAssets(const FName PackagePath, const bool bRecursive = false) const override;
 	virtual bool GetAssetsByPackageName(FName PackageName, TArray<FAssetData>& OutAssetData, bool bIncludeOnlyOnDiskAssets = false) const override;
 	virtual bool GetAssetsByPath(FName PackagePath, TArray<FAssetData>& OutAssetData, bool bRecursive = false, bool bIncludeOnlyOnDiskAssets = false) const override;
+	virtual bool GetAssetsByPaths(TArray<FName> PackagePath, TArray<FAssetData>& OutAssetData, bool bRecursive = false, bool bIncludeOnlyOnDiskAssets = false) const override;
 	virtual bool GetAssetsByClass(FName ClassName, TArray<FAssetData>& OutAssetData, bool bSearchSubClasses = false) const override;
 	virtual bool GetAssetsByTags(const TArray<FName>& AssetTags, TArray<FAssetData>& OutAssetData) const override;
 	virtual bool GetAssetsByTagValues(const TMultiMap<FName, FString>& AssetTagsAndValues, TArray<FAssetData>& OutAssetData) const override;
diff --git a/Engine/Source/Runtime/AssetRegistry/Public/AssetRegistry/IAssetRegistry.h b/Engine/Source/Runtime/AssetRegistry/Public/AssetRegistry/IAssetRegistry.h
--- a/Engine/Source/Runtime/AssetRegistry/Public/AssetRegistry/IAssetRegistry.h
+++ b/Engine/Source/Runtime/AssetRegistry/Public/AssetRegistry/IAssetRegistry.h
@@ -141,6 +141,16 @@ public:
 	UFUNCTION(BlueprintCallable, BlueprintPure=false, Category = "AssetRegistry")
 	virtual bool GetAssetsByPath(FName PackagePath, TArray<FAssetData>& OutAssetData, bool bRecursive = false, bool bIncludeOnlyOnDiskAssets = false) const = 0;
 
+	/**
+	 * Gets asset data for all assets in any of the supplied folder paths
+	 *
+	 * @param PackagePaths the paths to query asset data in (eg, /Game/MyFolder)
+	 * @param OutAssetData the list of assets in this path
+	 * @param bRecursive if true, all supplied paths will be searched recursively
+	 */
+	UFUNCTION(BlueprintCallable, BlueprintPure = false, Category = "AssetRegistry")
+	virtual bool GetAssetsByPaths(TArray<FName> PackagePaths, TArray<FAssetData>& OutAssetData, bool bRecursive = false, bool bIncludeOnlyOnDiskAssets = false) const = 0;
+
 	/**
 	 * Gets asset data for all assets with the supplied class
 	 *
diff --git a/Engine/Source/Runtime/CoreUObject/Private/Internationalization/PackageLocalizationCache.cpp b/Engine/Source/Runtime/CoreUObject/Private/Internationalization/PackageLocalizationCache.cpp
--- a/Engine/Source/Runtime/CoreUObject/Private/Internationalization/PackageLocalizationCache.cpp
+++ b/Engine/Source/Runtime/CoreUObject/Private/Internationalization/PackageLocalizationCache.cpp
@@ -37,6 +37,7 @@ void FPackageLocalizationCultureCache::ConditionalUpdateCache_NoLock()
 	SCOPED_BOOT_TIMING("FPackageLocalizationCultureCache::ConditionalUpdateCache_NoLock");
 	const double CacheStartTime = FPlatformTime::Seconds();
 
+	TMap<FString, TArray<FString>> NewSourceToLocalizedPaths;
 	for (const FString& SourceRootPath : PendingSourceRootPathsToSearch)
 	{
 		TArray<FString>& LocalizedRootPaths = SourcePathsToLocalizedPaths.FindOrAdd(SourceRootPath);
@@ -46,10 +47,11 @@ void FPackageLocalizationCultureCache::ConditionalUpdateCache_NoLock()
 			if (!LocalizedRootPaths.Contains(LocalizedRootPath))
 			{
 				LocalizedRootPaths.Add(LocalizedRootPath);
-				OwnerCache->FindLocalizedPackages(SourceRootPath, LocalizedRootPath, SourcePackagesToLocalizedPackages);
+				NewSourceToLocalizedPaths.FindOrAdd(SourceRootPath).Add(LocalizedRootPath);
 			}
 		}
 	}
+	OwnerCache->FindLocalizedPackages(NewSourceToLocalizedPaths, SourcePackagesToLocalizedPackages);
 
 	UE_LOG(LogPackageLocalizationCache, Log, TEXT("Processed %d localized package path(s) for %d prioritized culture(s) in %0.6f seconds"), PendingSourceRootPathsToSearch.Num(), PrioritizedCultureNames.Num(), FPlatformTime::Seconds() - CacheStartTime);
 
diff --git a/Engine/Source/Runtime/CoreUObject/Private/Internationalization/PackageLocalizationManager.cpp b/Engine/Source/Runtime/CoreUObject/Private/Internationalization/PackageLocalizationManager.cpp
--- a/Engine/Source/Runtime/CoreUObject/Private/Internationalization/PackageLocalizationManager.cpp
+++ b/Engine/Source/Runtime/CoreUObject/Private/Internationalization/PackageLocalizationManager.cpp
@@ -17,30 +17,37 @@ public:
 
 protected:
 	//~ FPackageLocalizationCache interface
-	virtual void FindLocalizedPackages(const FString& InSourceRoot, const FString& InLocalizedRoot, TMap<FName, TArray<FName>>& InOutSourcePackagesToLocalizedPackages) override;
+	virtual void FindLocalizedPackages(const TMap<FString, TArray<FString>>& NewSourceToLocalizedPaths, TMap<FName, TArray<FName>>& InOutSourcePackagesToLocalizedPackages) override;
 	virtual void FindAssetGroupPackages(const FName InAssetGroupName, const FName InAssetClassName) override;
 };
 
-void FDefaultPackageLocalizationCache::FindLocalizedPackages(const FString& InSourceRoot, const FString& InLocalizedRoot, TMap<FName, TArray<FName>>& InOutSourcePackagesToLocalizedPackages)
+void FDefaultPackageLocalizationCache::FindLocalizedPackages(const TMap<FString, TArray<FString>>& NewSourceToLocalizedPaths, TMap<FName, TArray<FName>>& InOutSourcePackagesToLocalizedPackages)
 {
 	// Convert the package path to a filename with no extension (directory)
-	FString LocalizedPackageFilePath;
-	if (!FPackageName::TryConvertLongPackageNameToFilename(InLocalizedRoot / TEXT(""), LocalizedPackageFilePath))
+	for (const TPair<FString, TArray<FString>>& Pair : NewSourceToLocalizedPaths)
 	{
-		return;
-	}
+		const FString& SourceRoot = Pair.Key;
+		for (const FString& LocalizedRoot : Pair.Value)
+		{
+			FString LocalizedPackageFilePath;
+			if (!FPackageName::TryConvertLongPackageNameToFilename(LocalizedRoot / TEXT(""), LocalizedPackageFilePath))
+			{
+				continue;
+			}
 
-	FPackageName::IteratePackagesInDirectory(LocalizedPackageFilePath, FPackageName::FPackageNameVisitor([&](const TCHAR* InPackageFileName) -> bool
-	{
-		const FString PackageSubPath = FPaths::ChangeExtension(InPackageFileName + LocalizedPackageFilePath.Len(), FString());
-		const FName SourcePackageName = *(InSourceRoot / PackageSubPath);
-		const FName LocalizedPackageName = *(InLocalizedRoot / PackageSubPath);
+			FPackageName::IteratePackagesInDirectory(LocalizedPackageFilePath, FPackageName::FPackageNameVisitor([&](const TCHAR* InPackageFileName) -> bool
+			{
+				const FString PackageSubPath = FPaths::ChangeExtension(InPackageFileName + LocalizedPackageFilePath.Len(), FString());
+				const FName SourcePackageName = *(SourceRoot / PackageSubPath);
+				const FName LocalizedPackageName = *(LocalizedRoot / PackageSubPath);
 
-		TArray<FName>& PrioritizedLocalizedPackageNames = InOutSourcePackagesToLocalizedPackages.FindOrAdd(SourcePackageName);
-		PrioritizedLocalizedPackageNames.AddUnique(LocalizedPackageName);
+				TArray<FName>& PrioritizedLocalizedPackageNames = InOutSourcePackagesToLocalizedPackages.FindOrAdd(SourcePackageName);
+				PrioritizedLocalizedPackageNames.AddUnique(LocalizedPackageName);
 
-		return true;
-	}));
+				return true;
+			}));
+		}
+	}
 }
 
 void FDefaultPackageLocalizationCache::FindAssetGroupPackages(const FName InAssetGroupName, const FName InAssetClassName)
diff --git a/Engine/Source/Runtime/CoreUObject/Private/Misc/AssetRegistryInterface.cpp b/Engine/Source/Runtime/CoreUObject/Private/Misc/AssetRegistryInterface.cpp
--- a/Engine/Source/Runtime/CoreUObject/Private/Misc/AssetRegistryInterface.cpp
+++ b/Engine/Source/Runtime/CoreUObject/Private/Misc/AssetRegistryInterface.cpp
@@ -136,8 +136,6 @@ namespace Utils
 	bool FFiltering::ShouldSkipAsset(FName AssetClass, uint32 PackageFlags)
 	{
 #if WITH_ENGINE && WITH_EDITOR
-		TRACE_CPUPROFILER_EVENT_SCOPE(AssetRegistry::FFiltering::ShouldSkipAsset)
-
 		// We do not yet support having UBlueprintGeneratedClasses be assets when the UBlueprint is also
 		// an asset; the content browser does not handle the multiple assets correctly and displays this
 		// class asset as if it is in a separate package. Revisit when we have removed the UBlueprint as an asset
diff --git a/Engine/Source/Runtime/CoreUObject/Public/Internationalization/PackageLocalizationCache.h b/Engine/Source/Runtime/CoreUObject/Public/Internationalization/PackageLocalizationCache.h
--- a/Engine/Source/Runtime/CoreUObject/Public/Internationalization/PackageLocalizationCache.h
+++ b/Engine/Source/Runtime/CoreUObject/Public/Internationalization/PackageLocalizationCache.h
@@ -2,8 +2,14 @@
 
 #pragma once
 
-#include "CoreMinimal.h"
+#include "Containers/Array.h"
+#include "Containers/Map.h"
+#include "Containers/UnrealString.h"
+#include "HAL/CriticalSection.h"
 #include "Internationalization/IPackageLocalizationCache.h"
+#include "Templates/SharedPointer.h"
+#include "Templates/Tuple.h"
+#include "UObject/NameTypes.h"
 
 class FPackageLocalizationCache;
 
@@ -112,11 +118,12 @@ protected:
 	/**
 	 * Find all of the localized packages under the given roots, and update the map with the result.
 	 *
-	 * @param InSourceRoot		The root package path that contains the source packages we're finding localized packages for, eg) /Game
-	 * @param InLocalizedRoot	The root package path to search for localized packages in, eg) /Game/L10N/fr
-	 * @param InOutSourcePackagesToLocalizedPackages The map to update. This will be passed to multiple calls of FindLocalizedPackages in order to build a mapping between a source package, and an array of prioritized localized packages.
+	 * @param NewSourceToLocalizedPaths Map containing a key for each of the the source root paths we're finding localized packages for, e.g. /Game
+	 *                                  The value for each key is an array of the roots to search for localized packages for that source, e.g. { /Game/L10/hu, /Game/L10/fr }
+	 * @param InOutSourcePackagesToLocalizedPackages The map to update. This will accumulate results from each root in order to build
+	 *                          a mapping between each source package and its array of prioritized localized packages.
 	 */
-	virtual void FindLocalizedPackages(const FString& InSourceRoot, const FString& InLocalizedRoot, TMap<FName, TArray<FName>>& InOutSourcePackagesToLocalizedPackages) = 0;
+	virtual void FindLocalizedPackages(const TMap<FString, TArray<FString>>& NewSourceToLocalizedPaths, TMap<FName, TArray<FName>>& InOutSourcePackagesToLocalizedPackages) = 0;
 
 	/**
 	 * Find all of the packages using the given asset group class, and update the PackageNameToAssetGroup map with the result.
diff --git a/Engine/Source/Runtime/Engine/Private/Internationalization/EnginePackageLocalizationCache.cpp b/Engine/Source/Runtime/Engine/Private/Internationalization/EnginePackageLocalizationCache.cpp
--- a/Engine/Source/Runtime/Engine/Private/Internationalization/EnginePackageLocalizationCache.cpp
+++ b/Engine/Source/Runtime/Engine/Private/Internationalization/EnginePackageLocalizationCache.cpp
@@ -32,17 +32,38 @@ FEnginePackageLocalizationCache::~FEnginePackageLocalizationCache()
 	}
 }
 
-void FEnginePackageLocalizationCache::FindLocalizedPackages(const FString& InSourceRoot, const FString& InLocalizedRoot, TMap<FName, TArray<FName>>& InOutSourcePackagesToLocalizedPackages)
+void FEnginePackageLocalizationCache::FindLocalizedPackages(const TMap<FString, TArray<FString>>& NewSourceToLocalizedPaths, TMap<FName, TArray<FName>>& InOutSourcePackagesToLocalizedPackages)
 {
+	if (NewSourceToLocalizedPaths.Num() == 0)
+	{
+		return;
+	}
+
 	FAssetRegistryModule& AssetRegistryModule = FModuleManager::LoadModuleChecked<FAssetRegistryModule>(TEXT("AssetRegistry"));
 	IAssetRegistry& AssetRegistry = AssetRegistryModule.Get();
 
+	int32 NumElementsPerKeyGuess = NewSourceToLocalizedPaths.CreateConstIterator()->Value.Num();
+	int32 SizeGuess = NewSourceToLocalizedPaths.Num() * NumElementsPerKeyGuess;
 #if WITH_EDITOR
-	// Make sure the asset registry has the data we need
+	TArray<FString> LocalizedPackagePaths;
+	LocalizedPackagePaths.Reserve(SizeGuess);
+#endif
+	TArray<FName> LocalizedPackagePathFNames;
+	LocalizedPackagePathFNames.Reserve(SizeGuess);
+	for (const TPair<FString, TArray<FString>>& Pair : NewSourceToLocalizedPaths)
 	{
-		TArray<FString> LocalizedPackagePaths;
-		LocalizedPackagePaths.Add(InLocalizedRoot);
+		for (const FString& LocalizedRoot : Pair.Value)
+		{
+#if WITH_EDITOR
+			LocalizedPackagePaths.Add(LocalizedRoot);
+#endif
+			LocalizedPackagePathFNames.Add(*LocalizedRoot);
+		}
+	}
 
+#if WITH_EDITOR
+	// Make sure the asset registry has the data we need
+	{
 		// Set bIsScanningPath to avoid us processing newly added assets from this scan
 		TGuardValue<bool> SetIsScanningPath(bIsScanningPath, true);
 		AssetRegistry.ScanPathsSynchronous(LocalizedPackagePaths);
@@ -52,7 +73,7 @@ void FEnginePackageLocalizationCache::FindLocalizedPackages(const FString& InSou
 	TArray<FAssetData> LocalizedAssetDataArray;
 	bool bIncludeOnlyOnDiskAssets = !GIsEditor;
 	bool bRecursive = true;
-	AssetRegistry.GetAssetsByPath(*InLocalizedRoot, LocalizedAssetDataArray, bRecursive, bIncludeOnlyOnDiskAssets);
+	AssetRegistry.GetAssetsByPaths(MoveTemp(LocalizedPackagePathFNames), LocalizedAssetDataArray, bRecursive, bIncludeOnlyOnDiskAssets);
 
 	for (const FAssetData& LocalizedAssetData : LocalizedAssetDataArray)
 	{
diff --git a/Engine/Source/Runtime/Engine/Public/Internationalization/EnginePackageLocalizationCache.h b/Engine/Source/Runtime/Engine/Public/Internationalization/EnginePackageLocalizationCache.h
--- a/Engine/Source/Runtime/Engine/Public/Internationalization/EnginePackageLocalizationCache.h
+++ b/Engine/Source/Runtime/Engine/Public/Internationalization/EnginePackageLocalizationCache.h
@@ -16,7 +16,7 @@ public:
 
 protected:
 	//~ FPackageLocalizationCache interface
-	virtual void FindLocalizedPackages(const FString& InSourceRoot, const FString& InLocalizedRoot, TMap<FName, TArray<FName>>& InOutSourcePackagesToLocalizedPackages) override;
+	virtual void FindLocalizedPackages(const TMap<FString, TArray<FString>>& NewSourceToLocalizedPaths, TMap<FName, TArray<FName>>& InOutSourcePackagesToLocalizedPackages) override;
 	virtual void FindAssetGroupPackages(const FName InAssetGroupName, const FName InAssetClassName) override;
 
 private:
