FAsyncPackage::FlushObjectLinkerCache: Optimize by calling FlushCache once per linker instead of once per object.
#rnx

diff --git a/Engine/Source/Runtime/CoreUObject/Private/Serialization/AsyncLoading.cpp b/Engine/Source/Runtime/CoreUObject/Private/Serialization/AsyncLoading.cpp
--- a/Engine/Source/Runtime/CoreUObject/Private/Serialization/AsyncLoading.cpp
+++ b/Engine/Source/Runtime/CoreUObject/Private/Serialization/AsyncLoading.cpp
@@ -5376,6 +5376,7 @@ void FAsyncPackage::DetachLinker()
 
 void FAsyncPackage::FlushObjectLinkerCache()
 {
+	TSet<FLinkerLoad*> Linkers;
 	for (UObject* Obj : PackageObjLoaded)
 	{
 		if (Obj)
@@ -5383,10 +5384,14 @@ void FAsyncPackage::FlushObjectLinkerCache()
 			FLinkerLoad* ObjLinker = Obj->GetLinker();
 			if (ObjLinker)
 			{
-				ObjLinker->FlushCache();
+				Linkers.Add(ObjLinker);
 			}
 		}
 	}
+	for (FLinkerLoad* LocalLinker : Linkers)
+	{
+		LocalLinker->FlushCache();
+	}
 }
 
 #if WITH_EDITOR 
