Fix  potential  crash when syncing textures from within the editor.

#rnx

### Problem
- The crash would occur when the texture being since contains a reference to another texture (via composite) which is also in need of syncing on the users machine.
- After the files on disk have been updated the current packages are reloaded but the original version of the texture will still be in memory when the composite texture notifies that it has been reloaded causing a texture compilation job with the older now invalid texture.
-- VirtualizedBulkData relied on out of date packages not trying to access their payloads on disk once the file has been changed, but with the above scenario it was possible. There could also be other paths causing this (even if the access of out of date package payloads at this point is a waste of time)

### Fix
- The virtualized bulkdata objects now reference themselves with the FLinkerLoader if the payload is on disk, just as the old bulkdata system did. When file changing operations are invoked, the FLinkerLoader will inform it's referenced bulkdata to load the payloads off disk and hold them in memory.
-- Note that virtualized bulkdata was trying to avoid this memory bloat by design, but if we cannot be certain that the bulkdata will not try to access the payload after the file changes (as was thought) then we must take the payload into memory.
- Added a log message to VBD if we detect that the loaded payload does not match the VBD members expectations. This log message can either be an error or a fatal error depending on the define 'VBD_CORRUPTED_PAYLOAD_IS_FATAL' (defaulted to on)
- The registration to FLinkerLoader can be easily disabled by setting the define 'VBD_ALLOW_LINKERLOADER_ATTACHMENT' to 0

### Outstanding Issues
- Cloned bulkdata objects (via assignment) do not currently register themselves to the FLinkerLoader (the old bulkdata system did not) if we were to add this then we need to make sure that the FLinkerLoader registration is thread safe. This should be done as a separate work item so that we can unblock the current workflows that this changelist addresses.
-- Note that if we did not update cloned bulkdata to work with this system then we could drop the registration and just iterate over the package's bulkdata and force them to load into memory when required. However that would leave the cloned virtual bulkdata in an uncertain state and is not desirable.
- When saving a package, the virtualized bulkdata objects that it contains will change their members so that future payload access can come from the newly saved files, where as the old bulkdata system would force the payload to remain in memory for the rest of the editor session. If we want to keep this then we will need to try and make sure that the newly updated bulkdata can be registered so that it receives future notifications about the file being changed.

diff --git a/Engine/Source/Runtime/Core/Public/Serialization/Archive.h b/Engine/Source/Runtime/Core/Public/Serialization/Archive.h
--- a/Engine/Source/Runtime/Core/Public/Serialization/Archive.h
+++ b/Engine/Source/Runtime/Core/Public/Serialization/Archive.h
@@ -28,6 +28,10 @@ class ITargetPlatform;
 class UObject;
 class FProperty;
 struct FUntypedBulkData;
+namespace UE::Virtualization
+{
+class FVirtualizedUntypedBulkData;
+}
 struct FArchiveSerializedPropertyChain;
 template<class TEnum> class TEnumAsByte;
 typedef TFunction<bool (double RemainingTime)> FExternalReadCallback;
@@ -1614,6 +1618,7 @@ public:
 	 * @param	BulkData	Bulk data object to associate
 	 */
 	virtual void AttachBulkData(UObject* Owner, FUntypedBulkData* BulkData) { }
+	virtual void AttachBulkData(UE::Virtualization::FVirtualizedUntypedBulkData* BulkData) {}
 
 	/**
 	 * Detaches the passed in bulk data object from the linker.
@@ -1622,6 +1627,7 @@ public:
 	 * @param	bEnsureBulkDataIsLoaded	Whether to ensure that the bulk data is loaded before detaching
 	 */
 	virtual void DetachBulkData(FUntypedBulkData* BulkData, bool bEnsureBulkDataIsLoaded) { }
+	virtual void DetachBulkData(UE::Virtualization::FVirtualizedUntypedBulkData* BulkData, bool bEnsureBulkDataIsLoaded) {}
 
 	/**
 	* Determine if the given archive is a valid "child" of this archive. In general, this means "is exactly the same" but
diff --git a/Engine/Source/Runtime/Core/Public/Virtualization/PayloadId.h b/Engine/Source/Runtime/Core/Public/Virtualization/PayloadId.h
--- a/Engine/Source/Runtime/Core/Public/Virtualization/PayloadId.h
+++ b/Engine/Source/Runtime/Core/Public/Virtualization/PayloadId.h
@@ -53,8 +53,22 @@ public:
 	{
 		return Identifier;
 	}
+
 	bool operator == (const FPayloadId& InPayloadId) const;
-	bool operator != (const FPayloadId& InPayloadId) const { return !(*this == InPayloadId); }
+	bool operator != (const FPayloadId& InPayloadId) const 
+	{ 
+		return !(*this == InPayloadId); 
+	}
+
+	bool operator == (const FIoHash& IoHash) const
+	{
+		return GetIdentifier() == IoHash;
+	}
+
+	bool operator != (const FIoHash& IoHash) const 
+	{ 
+		return !(*this == IoHash); 
+	}
 
 	/** Serialization specialization */
 	CORE_API friend FArchive& operator<<(FArchive& Ar, FPayloadId& PayloadId);
diff --git a/Engine/Source/Runtime/CoreUObject/Private/UObject/LinkerLoad.cpp b/Engine/Source/Runtime/CoreUObject/Private/UObject/LinkerLoad.cpp
--- a/Engine/Source/Runtime/CoreUObject/Private/UObject/LinkerLoad.cpp
+++ b/Engine/Source/Runtime/CoreUObject/Private/UObject/LinkerLoad.cpp
@@ -46,6 +46,7 @@
 #include "Misc/StringBuilder.h"
 #include "Misc/EngineBuildSettings.h"
 #include "Internationalization/GatherableTextData.h"
+#include "Virtualization/VirtualizedBulkData.h"
 
 class FTexture2DResourceMem;
 
@@ -5579,6 +5580,12 @@ void FLinkerLoad::AttachBulkData( UObject* Owner, FUntypedBulkData* BulkData )
 	BulkDataLoaders.Add( BulkData );
 }
 
+void FLinkerLoad::AttachBulkData(UE::Virtualization::FVirtualizedUntypedBulkData* BulkData)
+{
+	check(VirtualizedBulkDataLoaders.Find(BulkData) == INDEX_NONE);
+	VirtualizedBulkDataLoaders.Add(BulkData);	
+}
+
 /**
  * Detaches the passed in bulk data object from the linker.
  *
@@ -5595,6 +5602,16 @@ void FLinkerLoad::DetachBulkData( FUntypedBulkData* BulkData, bool bEnsureBulkDa
 	BulkData->DetachFromArchive( this, bEnsureBulkDataIsLoaded );
 }
 
+void FLinkerLoad::DetachBulkData(UE::Virtualization::FVirtualizedUntypedBulkData* BulkData, bool bEnsureBulkDataIsLoaded)
+{
+	int32 RemovedCount = VirtualizedBulkDataLoaders.Remove(BulkData);
+	if (RemovedCount != 1)
+	{
+		UE_LOG(LogLinker, Fatal, TEXT("Detachment inconsistency: %i (%s)"), RemovedCount, *GetDebugName());
+	}
+	BulkData->DetachFromDisk(this, bEnsureBulkDataIsLoaded);
+}
+
 /**
  * Detaches all attached bulk  data objects.
  *
@@ -5602,13 +5619,27 @@ void FLinkerLoad::DetachBulkData( FUntypedBulkData* BulkData, bool bEnsureBulkDa
  */
 void FLinkerLoad::DetachAllBulkData(bool bEnsureAllBulkDataIsLoaded)
 {
-	auto BulkDataToDetach = BulkDataLoaders;
-	for (auto BulkData : BulkDataToDetach)
+	// Old style bulkdata first
 	{
-		check( BulkData );
-		BulkData->DetachFromArchive(this, bEnsureAllBulkDataIsLoaded);
+		TArray<FUntypedBulkData*> BulkDataToDetach = BulkDataLoaders;
+		for (FUntypedBulkData* BulkData : BulkDataToDetach)
+		{
+			check(BulkData != nullptr);
+			BulkData->DetachFromArchive(this, bEnsureAllBulkDataIsLoaded);
+		}
+		BulkDataLoaders.Empty();
+	}
+
+	// Then virtualized bulkdata
+	{
+		TArray<UE::Virtualization::FVirtualizedUntypedBulkData*> BulkDataToDetach = VirtualizedBulkDataLoaders;
+		for (UE::Virtualization::FVirtualizedUntypedBulkData* BulkData : BulkDataToDetach)
+		{
+			check(BulkData != nullptr);
+			BulkData->DetachFromDisk(this, bEnsureAllBulkDataIsLoaded);
+		}
+		VirtualizedBulkDataLoaders.Empty();
 	}
-	BulkDataLoaders.Empty();
 }
 
 #endif // WITH_EDITOR
diff --git a/Engine/Source/Runtime/CoreUObject/Private/Virtualization/VirtualizedBulkData.cpp b/Engine/Source/Runtime/CoreUObject/Private/Virtualization/VirtualizedBulkData.cpp
--- a/Engine/Source/Runtime/CoreUObject/Private/Virtualization/VirtualizedBulkData.cpp
+++ b/Engine/Source/Runtime/CoreUObject/Private/Virtualization/VirtualizedBulkData.cpp
@@ -20,6 +20,23 @@
 
 //#if WITH_EDITORONLY_DATA
 
+/** When enabled the non-virtualized bulkdata objects will attach to the FLinkerLoader for the package that they are loaded from */
+#if WITH_EDITOR
+	#define VBD_ALLOW_LINKERLOADER_ATTACHMENT 1
+#else
+	#define VBD_ALLOW_LINKERLOADER_ATTACHMENT 0
+#endif //WITH_EDITOR
+
+/** When enabled we will fatal log if we detect corrupted data rather than logging an error and returning a null FCompressedBuffer/FSharedBuffer. */
+#define VBD_CORRUPTED_PAYLOAD_IS_FATAL 1
+
+
+#if VBD_CORRUPTED_PAYLOAD_IS_FATAL
+	#define VBD_CORRUPTED_DATA_SEVERITY Fatal
+#else
+	#define VBD_CORRUPTED_DATA_SEVERITY Error
+#endif // VBD_CORRUPTED_PAYLOAD_IS_FATAL
+
 namespace UE::Virtualization
 {
 
@@ -74,6 +91,27 @@ void LogPackageOpenFailureMessage(const FPackagePath& PackagePath, EPackageSegme
 	}
 }
 
+/** Utility for checking if a payload (in FCompressedBuffer format) matches the expectations of a bulkdata's members or not. */
+bool IsValid(const FVirtualizedUntypedBulkData& BulkData, const FCompressedBuffer& Payload)
+{
+	if (Payload.IsNull() && BulkData.GetPayloadSize() > 0)
+	{
+		return false;
+	}
+
+	if (!BulkData.HasPlaceholderPayloadId() && BulkData.GetPayloadId() != FIoHash(Payload.GetRawHash()))
+	{
+		return false;
+	}
+
+	if (Payload.GetRawSize() != BulkData.GetPayloadSize())
+	{
+		return false;
+	}
+
+	return true;
+}
+
 /** Utility for accessing IVirtualizationSourceControlUtilities from the modular feature system. */
 IVirtualizationSourceControlUtilities* GetSourceControlInterface()
 {
@@ -248,6 +286,11 @@ FVirtualizedUntypedBulkData& FVirtualizedUntypedBulkData::operator=(const FVirtu
 
 FVirtualizedUntypedBulkData::~FVirtualizedUntypedBulkData()
 {
+	if (AttachedAr != nullptr)
+	{
+		AttachedAr->DetachBulkData(this, false);
+	}
+
 	Unregister();
 }
 
@@ -321,6 +364,14 @@ void FVirtualizedUntypedBulkData::CreateFromBulkData(FUntypedBulkData& InBulkDat
 
 	Reset();
 
+#if VBD_ALLOW_LINKERLOADER_ATTACHMENT
+	AttachedAr = InBulkData.AttachedAr;
+	if (AttachedAr != nullptr)
+	{
+		AttachedAr->AttachBulkData(this);
+	}
+#endif //VBD_ALLOW_LINKERLOADER_ATTACHMENT
+
 	// We only need to set up the bulkdata/content identifiers if we have a valid payload
 	bool bWasKeyGuidDerived = false;
 	if (InBulkData.GetBulkDataSize() > 0)
@@ -621,7 +672,8 @@ void FVirtualizedUntypedBulkData::Serialize(FArchive& Ar, UObject* Owner, bool b
 			{
 				// If we can lazy load then find the PackagePath, otherwise we will want
 				// to serialize immediately.
-				if (Ar.IsAllowingLazyLoading())
+				FArchive* CacheableArchive = Ar.GetCacheableArchive();
+				if (Ar.IsAllowingLazyLoading() && CacheableArchive != nullptr)
 				{
 					PackagePath = GetPackagePathFromOwner(Owner, PackageSegment);
 				}
@@ -634,7 +686,17 @@ void FVirtualizedUntypedBulkData::Serialize(FArchive& Ar, UObject* Owner, bool b
 				OffsetInFile = INDEX_NONE;
 				Ar << OffsetInFile;
 
-				if (PackagePath.IsEmpty())
+				if (!PackagePath.IsEmpty())
+				{
+#if VBD_ALLOW_LINKERLOADER_ATTACHMENT
+					AttachedAr = CacheableArchive;
+					if (AttachedAr != nullptr)
+					{
+						AttachedAr->AttachBulkData(this);
+					}
+#endif //VBD_ALLOW_LINKERLOADER_ATTACHMENT
+				}
+				else
 				{
 					// If we have no packagepath then we need to load the data immediately
 					// as we will not be able to load it on demand.
@@ -1014,6 +1076,11 @@ bool FVirtualizedUntypedBulkData::IsMemoryOnlyPayload() const
 void FVirtualizedUntypedBulkData::Reset()
 {
 	// Note that we do not reset the BulkDataId
+	if (AttachedAr != nullptr)
+	{
+		AttachedAr->DetachBulkData(this, false);
+	}
+
 	Unregister();
 	PayloadContentId.Reset();
 	Payload.Reset();
@@ -1034,6 +1101,30 @@ void FVirtualizedUntypedBulkData::UnloadData()
 	}
 }
 
+void FVirtualizedUntypedBulkData::DetachFromDisk(FArchive* Ar, bool bEnsurePayloadIsLoaded)
+{
+	check(Ar);
+	check(Ar == AttachedAr || AttachedAr == nullptr || AttachedAr->IsProxyOf(Ar));
+
+	if (bEnsurePayloadIsLoaded && !IsDataVirtualized() && !PackagePath.IsEmpty())
+	{
+		if (Payload.IsNull())
+		{
+			FCompressedBuffer CompressedPayload = GetDataInternal();
+
+			UE_CLOG(!IsValid(*this, CompressedPayload), LogVirtualization, VBD_CORRUPTED_DATA_SEVERITY, TEXT("%s"), *GetCorruptedPayloadErrorMessage());
+
+			FSharedBuffer DecompressedPayload = CompressedPayload.Decompress();
+
+			AttachedAr = nullptr; // Prevent UpdatePayloadImpl from sending more detach notifications to AttachedAr
+
+			UpdatePayloadImpl(MoveTemp(DecompressedPayload), MoveTemp(PayloadContentId));
+		}
+	}
+
+	AttachedAr = nullptr;	
+}
+
 FGuid FVirtualizedUntypedBulkData::GetIdentifier() const
 {
 	checkf(GetPayloadSize() == 0 || BulkDataId.IsValid(), TEXT("If bulkdata has a valid payload then it should have a valid BulkDataId"));
@@ -1046,6 +1137,13 @@ void FVirtualizedUntypedBulkData::UpdatePayloadImpl(FSharedBuffer&& InPayload, F
 
 	UnloadData();
 
+	if (AttachedAr != nullptr)
+	{
+		AttachedAr->DetachBulkData(this, false);
+	}
+
+	check(AttachedAr == nullptr);
+
 	// Make sure that we own the memory in the shared buffer
 	Payload = MoveTemp(InPayload).MakeOwned();
 	PayloadSize = (int64)Payload.GetSize();
@@ -1087,20 +1185,27 @@ FCompressedBuffer FVirtualizedUntypedBulkData::GetDataInternal() const
 	// Check if we already have the data in memory
 	if (Payload)
 	{
+		// Note that this doesn't actually compress the data!
 		return FCompressedBuffer::Compress(Payload, ECompressedBufferCompressor::NotSet, ECompressedBufferCompressionLevel::None);
 	}
 
 	if (IsDataVirtualized())
 	{
-		FCompressedBuffer Buffer = PullData();
+		FCompressedBuffer CompressedPayload = PullData();
+		
 		checkf(Payload.IsNull(), TEXT("Pulling data somehow assigned it to the bulk data object!")); //Make sure that we did not assign the buffer internally
-		return Buffer;
+		UE_CLOG(!IsValid(*this, CompressedPayload), LogVirtualization, VBD_CORRUPTED_DATA_SEVERITY, TEXT("%s"), *GetCorruptedPayloadErrorMessage());
+		
+		return CompressedPayload;
 	}
 	else
 	{
-		FCompressedBuffer Buffer = LoadFromDisk();
+		FCompressedBuffer CompressedPayload = LoadFromDisk();
+		
 		check(Payload.IsNull()); //Make sure that we did not assign the buffer internally
-		return Buffer;
+		UE_CLOG(!IsValid(*this, CompressedPayload), LogVirtualization, VBD_CORRUPTED_DATA_SEVERITY, TEXT("%s"), *GetCorruptedPayloadErrorMessage());
+		
+		return CompressedPayload;
 	}
 }
 
@@ -1115,10 +1220,12 @@ TFuture<FSharedBuffer> FVirtualizedUntypedBulkData::GetPayload() const
 	}
 	else
 	{
-		FSharedBuffer UncompressedPayload = GetDataInternal().Decompress();
+		FCompressedBuffer CompressedPayload = GetDataInternal();
 
+		UE_CLOG(!IsValid(*this, CompressedPayload), LogVirtualization, VBD_CORRUPTED_DATA_SEVERITY, TEXT("%s"), *GetCorruptedPayloadErrorMessage());
+		
 		// TODO: Not actually async yet!
-		Promise.SetValue(MoveTemp(UncompressedPayload));
+		Promise.SetValue(CompressedPayload.Decompress());
 	}
 
 	return Promise.GetFuture();
@@ -1297,6 +1404,25 @@ FVirtualizedUntypedBulkData::EFlags FVirtualizedUntypedBulkData::BuildFlagsForSe
 	}
 }
 
+FString FVirtualizedUntypedBulkData::GetCorruptedPayloadErrorMessage() const
+{
+	TStringBuilder<512> Message;
+	
+	if (IsDataVirtualized())
+	{
+		Message << TEXT("Virtualized payload '") << PayloadContentId << TEXT("' is corrupt! Check the backend storage.");
+	}
+	else
+	{
+		Message << TEXT("Payload ' ") << PayloadContentId << TEXT(" loaded from '") << *PackagePath.GetDebugName() << TEXT("' is corrupt! Check the file on disk.");
+	}
+
+	return Message.ToString();
+}
+
 } // namespace UE::Virtualization
 
+#undef VBD_ALLOW_LINKERLOADER_ATTACHMENT
+#undef VBD_CORRUPTED_DATA_SEVERITY
+
 //#endif //WITH_EDITORONLY_DATA
diff --git a/Engine/Source/Runtime/CoreUObject/Public/Serialization/BulkData.h b/Engine/Source/Runtime/CoreUObject/Public/Serialization/BulkData.h
--- a/Engine/Source/Runtime/CoreUObject/Public/Serialization/BulkData.h
+++ b/Engine/Source/Runtime/CoreUObject/Public/Serialization/BulkData.h
@@ -293,6 +293,8 @@ private:
 
 public:
 	friend class FLinkerLoad;
+	friend class UE::Virtualization::FVirtualizedUntypedBulkData; // To allow access to AttachedAr
+
 	using BulkDataRangeArray = TArray<FBulkDataStreamingToken*, TInlineAllocator<8>>;
 
 	/*-----------------------------------------------------------------------------
diff --git a/Engine/Source/Runtime/CoreUObject/Public/UObject/LinkerLoad.h b/Engine/Source/Runtime/CoreUObject/Public/UObject/LinkerLoad.h
--- a/Engine/Source/Runtime/CoreUObject/Public/UObject/LinkerLoad.h
+++ b/Engine/Source/Runtime/CoreUObject/Public/UObject/LinkerLoad.h
@@ -17,6 +17,10 @@ class IPakFile;
 class ULinkerPlaceholderExportObject;
 struct FScopedSlowTask;
 struct FUntypedBulkData;
+namespace UE::Virtualization
+{
+class FVirtualizedUntypedBulkData;
+}
 
 /*----------------------------------------------------------------------------
 	FLinkerLoad.
@@ -220,6 +224,7 @@ public:
 #if WITH_EDITOR
 	/** Bulk data that does not need to be loaded when the linker is loaded.												*/
 	TArray<FUntypedBulkData*> BulkDataLoaders;
+	TArray<UE::Virtualization::FVirtualizedUntypedBulkData*> VirtualizedBulkDataLoaders;
 #endif // WITH_EDITOR
 
 	/** Hash table for exports.																								*/
@@ -819,6 +824,7 @@ private:
 	 * @param	BulkData	Bulk data object to associate
 	 */
 	virtual void AttachBulkData(UObject* Owner, FUntypedBulkData* BulkData) override;
+	virtual void AttachBulkData(UE::Virtualization::FVirtualizedUntypedBulkData* BulkData) override;
 	/**
 	 * Detaches the passed in bulk data object from the linker.
 	 *
@@ -826,6 +832,7 @@ private:
 	 * @param	bEnsureBulkDataIsLoaded	Whether to ensure that the bulk data is loaded before detaching
 	 */
 	virtual void DetachBulkData(FUntypedBulkData* BulkData, bool bEnsureBulkDataIsLoaded) override;
+	virtual void DetachBulkData(UE::Virtualization::FVirtualizedUntypedBulkData* BulkData, bool bEnsureBulkDataIsLoaded) override;
 	/**
 	 * Detaches all attached bulk  data objects.
 	 *
diff --git a/Engine/Source/Runtime/CoreUObject/Public/Virtualization/VirtualizedBulkData.h b/Engine/Source/Runtime/CoreUObject/Public/Virtualization/VirtualizedBulkData.h
--- a/Engine/Source/Runtime/CoreUObject/Public/Virtualization/VirtualizedBulkData.h
+++ b/Engine/Source/Runtime/CoreUObject/Public/Virtualization/VirtualizedBulkData.h
@@ -134,6 +134,16 @@ public:
 	/** Unloads the data (if possible) but leaves it in a state where the data can be reloaded */
 	void UnloadData();
 
+	/** 
+	 * Removes the ability for the bulkdata object to load it's payload from disk (if it was doing so) 
+	 * 
+	 * @param Ar						The archive that the bulkdata is being detached from. This should match AttachedAr.
+	 * @param bEnsurePayloadIsLoaded	If true and the bulkdata currently is storing it's payload as a package on disk then
+	 *									the payload should be loaded into memory so that it can be accessed in the future.
+	 *									If false then the payload is not important and does not need to be loaded.
+	 */
+	void DetachFromDisk(FArchive* Ar, bool bEnsurePayloadIsLoaded);
+
 	/** Returns a unique identifier for the object itself. */
 	FGuid GetIdentifier() const;
 
@@ -343,6 +353,9 @@ private:
 	void Register(UObject* Owner);
 	void Unregister();
 
+	/** Utility to return an apt error message if a payload loaded from this object is considered corrupt. Handles the formatting depending on the various payload flags etc. */
+	FString GetCorruptedPayloadErrorMessage() const;
+
 	/** Unique identifier for the bulkdata object itself */
 	FGuid BulkDataId;
 
@@ -357,6 +370,9 @@ private:
 
 	//---- The remaining members are used when the payload is not virtualized.
 	
+	/** The archive representing the file on disk containing the payload (if there is one), we keep the pointer so that the bulkdata object can be detached if needed. */
+	FArchive* AttachedAr = nullptr;
+
 	/** Offset of the payload in the file that contains it (INDEX_NONE if the payload does not come from a file)*/
 	int64 OffsetInFile = INDEX_NONE;
 
