Do not check for end-of-file tag when loading unversioned assets from the editor domain.

diff --git a/Engine/Source/Runtime/CoreUObject/Private/UObject/LinkerLoad.cpp b/Engine/Source/Runtime/CoreUObject/Private/UObject/LinkerLoad.cpp
--- a/Engine/Source/Runtime/CoreUObject/Private/UObject/LinkerLoad.cpp
+++ b/Engine/Source/Runtime/CoreUObject/Private/UObject/LinkerLoad.cpp
@@ -1321,7 +1321,7 @@ FLinkerLoad::ELinkerStatus FLinkerLoad::SerializePackageFileSummaryInternal()
 	if (!FPlatformProperties::RequiresCookedData() &&
 		// We can't check the post tag if the file is an EDL cooked package
 		!((Summary.GetPackageFlags() & PKG_FilterEditorOnly) && Summary.PreloadDependencyCount > 0 && Summary.PreloadDependencyOffset > 0)
-		&& !IsTextFormat())
+		&& !IsTextFormat() && bLoaderVersionCheck)
 	{
 		// check if this package version stored the 4-byte magic post tag
 		// get the offset of the post tag
