ReferenceChainSearch: added optional output device to the class and updated console command "obj refs" to provide its output device so the results will be available directly in the console (not only in log file or debugger output). This behavior is now consistent with console command "obj list"

diff --git a/Engine/Source/Runtime/CoreUObject/Private/UObject/Obj.cpp b/Engine/Source/Runtime/CoreUObject/Private/UObject/Obj.cpp
--- a/Engine/Source/Runtime/CoreUObject/Private/UObject/Obj.cpp
+++ b/Engine/Source/Runtime/CoreUObject/Private/UObject/Obj.cpp
@@ -4080,7 +4080,7 @@ bool StaticExec( UWorld* InWorld, const TCHAR* Cmd, FOutputDevice& Ar )
 					}
 				}
 				
-				FReferenceChainSearch RefChainSearch(Object, SearchModeFlags);
+				FReferenceChainSearch RefChainSearch(Object, SearchModeFlags, Ar);
 			}
 			else
 			{
diff --git a/Engine/Source/Runtime/CoreUObject/Private/UObject/ReferenceChainSearch.cpp b/Engine/Source/Runtime/CoreUObject/Private/UObject/ReferenceChainSearch.cpp
--- a/Engine/Source/Runtime/CoreUObject/Private/UObject/ReferenceChainSearch.cpp
+++ b/Engine/Source/Runtime/CoreUObject/Private/UObject/ReferenceChainSearch.cpp
@@ -3,7 +3,6 @@
 #include "UObject/ReferenceChainSearch.h"
 #include "HAL/PlatformStackWalk.h"
 #include "UObject/UObjectIterator.h"
-#include "UObject/UnrealType.h"
 #include "UObject/FastReferenceCollector.h"
 #include "UObject/GCObject.h"
 #include "HAL/ThreadHeartBeat.h"
@@ -13,7 +12,7 @@ DEFINE_LOG_CATEGORY_STATIC(LogReferenceChain, Log, All);
 // Returns true if the object can't be collected by GC
 static FORCEINLINE bool IsNonGCObject(UObject* Object, EReferenceChainSearchMode SearchMode)
 {
-	FUObjectItem* ObjectItem = GUObjectArray.ObjectToObjectItem(Object);
+	const FUObjectItem* ObjectItem = GUObjectArray.ObjectToObjectItem(Object);
 	return (ObjectItem->IsRootSet() ||
 		ObjectItem->HasAnyFlags(EInternalObjectFlags::GarbageCollectionKeepFlags) ||
 		(GARBAGE_COLLECTION_KEEPFLAGS != RF_NoFlags && Object->HasAnyFlags(GARBAGE_COLLECTION_KEEPFLAGS) && !(SearchMode & EReferenceChainSearchMode::FullChain))
@@ -54,8 +53,8 @@ int32 FReferenceChainSearch::BuildReferenceChains(FGraphNode* TargetNode, TArray
 				// For each of the referencers of this node, duplicate the current chain and continue processing
 				if (ReferencedByNode->Visited != VisitCounter)
 				{
-					int32 OldChainsCount = ProducedChains.Num();
-					int32 NewChainsCount = BuildReferenceChains(ReferencedByNode, ProducedChains, ChainDepth + 1, VisitCounter, SearchMode);
+					const int32 OldChainsCount = ProducedChains.Num();
+					const int32 NewChainsCount = BuildReferenceChains(ReferencedByNode, ProducedChains, ChainDepth + 1, VisitCounter, SearchMode);
 					// Insert the current node to all chains produced recursively
 					for (int32 NewChainIndex = OldChainsCount; NewChainIndex < (NewChainsCount + OldChainsCount); ++NewChainIndex)
 					{
@@ -140,7 +139,7 @@ void FReferenceChainSearch::BuildReferenceChains(FGraphNode* TargetNode, TArray<
 		TargetNode->Visited = ++VisitCounter;
 		
 		AllChains.Reset();
-		const int32 MinChainDepth = 2; // The chain will contain at least the TargetNode and the ReferencedByNode
+		constexpr int32 MinChainDepth = 2; // The chain will contain at least the TargetNode and the ReferencedByNode
 		BuildReferenceChains(ReferencedByNode, AllChains, MinChainDepth, VisitCounter, SearchMode);
 		for (FReferenceChain* Chain : AllChains)
 		{
@@ -152,7 +151,7 @@ void FReferenceChainSearch::BuildReferenceChains(FGraphNode* TargetNode, TArray<
 		{
 			for (int32 ChainIndex = AllChains.Num() - 1; ChainIndex >= 0; --ChainIndex)
 			{
-				FReferenceChain* Chain = AllChains[ChainIndex];	
+				const FReferenceChain* Chain = AllChains[ChainIndex];
 				if (!Chain->IsExternal())
 				{
 					// Discard the chain
@@ -248,7 +247,7 @@ FString FReferenceChainSearch::GetObjectFlags(UObject* InObject)
 		Flags += TEXT("(NeverGCed) ");
 	}
 
-	FUObjectItem* ReferencedByObjectItem = GUObjectArray.ObjectToObjectItem(InObject);
+	const FUObjectItem* ReferencedByObjectItem = GUObjectArray.ObjectToObjectItem(InObject);
 	if (ReferencedByObjectItem->HasAnyFlags(EInternalObjectFlags::ClusterRoot))
 	{
 		Flags += TEXT("(ClusterRoot) ");
@@ -260,7 +259,7 @@ FString FReferenceChainSearch::GetObjectFlags(UObject* InObject)
 	return Flags;
 }
 
-void FReferenceChainSearch::DumpChain(FReferenceChain* Chain)
+void FReferenceChainSearch::DumpChain(const FReferenceChain* Chain, FOutputDevice& OutputDevice)
 {
 	if (Chain->Num())
 	{
@@ -270,18 +269,18 @@ void FReferenceChainSearch::DumpChain(FReferenceChain* Chain)
 			UObject* Object = Chain->GetNode(NodeIndex)->Object;
 			const FNodeReferenceInfo& ReferenceInfo = Chain->GetReferenceInfo(NodeIndex);
 
-			UE_LOG(LogReferenceChain, Log, TEXT("%s%s%s%s"),
+			OutputDevice.CategorizedLogf(LogReferenceChain.GetCategoryName(), ELogVerbosity::Log, TEXT("%s%s%s%s"),
 				FCString::Spc(FMath::Min<int32>(TCStringSpcHelper<TCHAR>::MAX_SPACES, Chain->Num() - NodeIndex - 1)),
 				*GetObjectFlags(Object),
 				*Object->GetFullName(),
 				*ReferenceInfo.ToString()
 			);
 		}
-		UE_LOG(LogReferenceChain, Log, TEXT("  "));
+		OutputDevice.CategorizedLogf(LogReferenceChain.GetCategoryName(), ELogVerbosity::Log, TEXT("  "));
 	}
 }
 
-void FReferenceChainSearch::WriteChain(FReferenceChain* Chain, FString& OutString)
+void FReferenceChainSearch::WriteChain(const FReferenceChain* Chain, FString& OutString)
 {
 	if (Chain->Num())
 	{
@@ -360,7 +359,7 @@ public:
 		, ReferencedObjects(InReferencedObjects)		
 	{
 	}
-	FORCEINLINE void HandleTokenStreamObjectReference(FGCArrayStruct& ObjectsToSerializeStruct, UObject* ReferencingObject, UObject*& Object, const int32 TokenIndex, bool bAllowReferenceElimination)
+	FORCEINLINE void HandleTokenStreamObjectReference(FGCArrayStruct& ObjectsToSerializeStruct, UObject* ReferencingObject, UObject*& Object, const int32 TokenIndex, bool bAllowReferenceElimination) const
 	{
 		FReferenceChainSearch::FObjectReferenceInfo RefInfo(Object);
 		if (Object && !ReferencedObjects.Contains(RefInfo))
@@ -368,7 +367,7 @@ public:
 #if ENABLE_GC_OBJECT_CHECKS
 			if (TokenIndex >= 0)
 			{
-				FTokenInfo TokenInfo = ReferencingObject->GetClass()->ReferenceTokenStream.GetTokenInfo(TokenIndex);
+				const FTokenInfo TokenInfo = ReferencingObject->GetClass()->ReferenceTokenStream.GetTokenInfo(TokenIndex);
 				RefInfo.ReferencerName = TokenInfo.Name;
 				RefInfo.Type = FReferenceChainSearch::EReferenceType::Property;
 			}
@@ -400,8 +399,9 @@ public:
 
 typedef TDefaultReferenceCollector<FDirectReferenceProcessor> FDirectReferenceCollector;
 
-FReferenceChainSearch::FReferenceChainSearch(UObject* InObjectToFindReferencesTo, EReferenceChainSearchMode Mode /*= EReferenceChainSearchMode::PrintResults*/)
+FReferenceChainSearch::FReferenceChainSearch(UObject* InObjectToFindReferencesTo, EReferenceChainSearchMode Mode /*= EReferenceChainSearchMode::PrintResults*/, FOutputDevice& InOutputDevice /*= *GLog*/)
 	: ObjectToFindReferencesTo(InObjectToFindReferencesTo)
+	, OutputDevice(InOutputDevice)
 {
 	check(InObjectToFindReferencesTo);
 
@@ -456,7 +456,7 @@ void FReferenceChainSearch::FindDirectReferencesForObjects()
 
 	for (FRawObjectIterator It; It; ++It)
 	{
-		FUObjectItem* ObjItem = *It;
+		const FUObjectItem* ObjItem = *It;
 		UObject* Object = static_cast<UObject*>(ObjItem->Object);
 		FGraphNode* ObjectNode = FindOrAddNode(AllNodes, Object);
 
@@ -482,19 +482,19 @@ void FReferenceChainSearch::PrintResults(bool bDumpAllChains /*= false*/) const
 	{
 		FSlowHeartBeatScope DisableHangDetection; // This function can be very slow
 
-		const int32 MaxChainsToPrint = 100;
+		constexpr int32 MaxChainsToPrint = 100;
 		int32 NumPrintedChains = 0;
 
-		for (FReferenceChain* Chain : ReferenceChains)
+		for (const FReferenceChain* Chain : ReferenceChains)
 		{
 			if (bDumpAllChains || NumPrintedChains < MaxChainsToPrint)
 			{
-				DumpChain(Chain);
+				DumpChain(Chain, OutputDevice);
 				NumPrintedChains++;
 			}
 			else
 			{
-				UE_LOG(LogReferenceChain, Log, TEXT("Referenced by %d more reference chain(s)."), ReferenceChains.Num() - NumPrintedChains);
+				OutputDevice.CategorizedLogf(LogReferenceChain.GetCategoryName(), ELogVerbosity::Log, TEXT("Referenced by %d more reference chain(s)."), ReferenceChains.Num() - NumPrintedChains);
 				break;
 			}
 		}
@@ -502,7 +502,8 @@ void FReferenceChainSearch::PrintResults(bool bDumpAllChains /*= false*/) const
 	else
 	{
 		check(ObjectToFindReferencesTo);
-		UE_LOG(LogReferenceChain, Log, TEXT("%s%s is not currently reachable."),
+
+		OutputDevice.CategorizedLogf(LogReferenceChain.GetCategoryName(), ELogVerbosity::Log, TEXT("%s%s is not currently reachable."),
 			*GetObjectFlags(ObjectToFindReferencesTo),
 			*ObjectToFindReferencesTo->GetFullName()
 		);
@@ -513,7 +514,7 @@ FString FReferenceChainSearch::GetRootPath() const
 {
 	if (ReferenceChains.Num())
 	{
-		FReferenceChain* Chain = ReferenceChains[0];
+		const FReferenceChain* Chain = ReferenceChains[0];
 		FString OutString;
 
 		WriteChain(Chain, OutString);
@@ -536,7 +537,7 @@ void FReferenceChainSearch::Cleanup()
 	}
 	ReferenceChains.Empty();
 
-	for (TPair<UObject*, FGraphNode*>& ObjectNodePair : AllNodes)
+	for (const TPair<UObject*, FGraphNode*>& ObjectNodePair : AllNodes)
 	{
 		delete ObjectNodePair.Value;
 	}
diff --git a/Engine/Source/Runtime/CoreUObject/Public/UObject/ReferenceChainSearch.h b/Engine/Source/Runtime/CoreUObject/Public/UObject/ReferenceChainSearch.h
--- a/Engine/Source/Runtime/CoreUObject/Public/UObject/ReferenceChainSearch.h
+++ b/Engine/Source/Runtime/CoreUObject/Public/UObject/ReferenceChainSearch.h
@@ -2,10 +2,7 @@
 
 #pragma once
 
-#include "CoreMinimal.h"
-#include "UObject/UObjectGlobals.h"
-#include "UObject/Class.h"
-#include "UObject/GarbageCollection.h"
+#include "Misc/OutputDeviceRedirector.h"
 #include "HAL/ThreadHeartBeat.h"
 
 /** Search mode flags */
@@ -171,12 +168,7 @@ public:
 		{
 			return Nodes.Num();
 		}
-		/** Returns a duplicate of this chain */
-		FReferenceChain* Split()
-		{
-			FReferenceChain* NewChain = new FReferenceChain(*this);
-			return NewChain;
-		}
+
 		/** Checks if this chain contains the specified node */
 		bool Contains(const FGraphNode* InNode) const
 		{
@@ -192,7 +184,7 @@ public:
 	};
 
 	/** Constructs a new search engine and finds references to the specified object */
-	COREUOBJECT_API FReferenceChainSearch(UObject* InObjectToFindReferencesTo, EReferenceChainSearchMode Mode = EReferenceChainSearchMode::PrintResults);
+	COREUOBJECT_API FReferenceChainSearch(UObject* InObjectToFindReferencesTo, EReferenceChainSearchMode Mode = EReferenceChainSearchMode::PrintResults, FOutputDevice& InOutputDevice = *GLog);
 
 	/** Destructor */
 	COREUOBJECT_API ~FReferenceChainSearch();
@@ -221,6 +213,9 @@ private:
 	/** All nodes created during the search */
 	TMap<UObject*, FGraphNode*> AllNodes;
 
+	/** Output device to log to */
+	FOutputDevice& OutputDevice;
+
 	/** Performs the search */
 	void PerformSearch(EReferenceChainSearchMode SearchMode);
 
@@ -246,7 +241,7 @@ private:
 	/** Returns a string with all flags (we care about) set on an object */
 	static FString GetObjectFlags(UObject* InObject);
 	/** Dumps a reference chain to log */
-	static void DumpChain(FReferenceChain* Chain);
+	static void DumpChain(const FReferenceChain* Chain, FOutputDevice& OutputDevice);
 	/** Writes a reference chain to a string */
-	static void WriteChain(FReferenceChain* Chain, FString& OutString);
+	static void WriteChain(const FReferenceChain* Chain, FString& OutString);
 };
