[Backout] - CL17665886
Some changes need to be integrated from UE5-Release-Engine-Staging to reduce amount of conflicts
Original CL Desc
-----------------------------------------------------------------
ReferenceChainSearch: added optional output device to the class and updated console command "obj refs" to provide its output device so the results will be available directly in the console (not only in log file or debugger output). This behavior is now consistent with console command "obj list"

diff --git a/Engine/Source/Runtime/CoreUObject/Private/UObject/Obj.cpp b/Engine/Source/Runtime/CoreUObject/Private/UObject/Obj.cpp
--- a/Engine/Source/Runtime/CoreUObject/Private/UObject/Obj.cpp
+++ b/Engine/Source/Runtime/CoreUObject/Private/UObject/Obj.cpp
@@ -4080,7 +4080,7 @@ bool StaticExec( UWorld* InWorld, const TCHAR* Cmd, FOutputDevice& Ar )
 					}
 				}
 				
-				FReferenceChainSearch RefChainSearch(Object, SearchModeFlags, Ar);
+				FReferenceChainSearch RefChainSearch(Object, SearchModeFlags);
 			}
 			else
 			{
diff --git a/Engine/Source/Runtime/CoreUObject/Private/UObject/ReferenceChainSearch.cpp b/Engine/Source/Runtime/CoreUObject/Private/UObject/ReferenceChainSearch.cpp
--- a/Engine/Source/Runtime/CoreUObject/Private/UObject/ReferenceChainSearch.cpp
+++ b/Engine/Source/Runtime/CoreUObject/Private/UObject/ReferenceChainSearch.cpp
@@ -3,6 +3,7 @@
 #include "UObject/ReferenceChainSearch.h"
 #include "HAL/PlatformStackWalk.h"
 #include "UObject/UObjectIterator.h"
+#include "UObject/UnrealType.h"
 #include "UObject/FastReferenceCollector.h"
 #include "UObject/GCObject.h"
 #include "HAL/ThreadHeartBeat.h"
@@ -12,7 +13,7 @@ DEFINE_LOG_CATEGORY_STATIC(LogReferenceChain, Log, All);
 // Returns true if the object can't be collected by GC
 static FORCEINLINE bool IsNonGCObject(UObject* Object, EReferenceChainSearchMode SearchMode)
 {
-	const FUObjectItem* ObjectItem = GUObjectArray.ObjectToObjectItem(Object);
+	FUObjectItem* ObjectItem = GUObjectArray.ObjectToObjectItem(Object);
 	return (ObjectItem->IsRootSet() ||
 		ObjectItem->HasAnyFlags(EInternalObjectFlags::GarbageCollectionKeepFlags) ||
 		(GARBAGE_COLLECTION_KEEPFLAGS != RF_NoFlags && Object->HasAnyFlags(GARBAGE_COLLECTION_KEEPFLAGS) && !(SearchMode & EReferenceChainSearchMode::FullChain))
@@ -53,8 +54,8 @@ int32 FReferenceChainSearch::BuildReferenceChains(FGraphNode* TargetNode, TArray
 				// For each of the referencers of this node, duplicate the current chain and continue processing
 				if (ReferencedByNode->Visited != VisitCounter)
 				{
-					const int32 OldChainsCount = ProducedChains.Num();
-					const int32 NewChainsCount = BuildReferenceChains(ReferencedByNode, ProducedChains, ChainDepth + 1, VisitCounter, SearchMode);
+					int32 OldChainsCount = ProducedChains.Num();
+					int32 NewChainsCount = BuildReferenceChains(ReferencedByNode, ProducedChains, ChainDepth + 1, VisitCounter, SearchMode);
 					// Insert the current node to all chains produced recursively
 					for (int32 NewChainIndex = OldChainsCount; NewChainIndex < (NewChainsCount + OldChainsCount); ++NewChainIndex)
 					{
@@ -139,7 +140,7 @@ void FReferenceChainSearch::BuildReferenceChains(FGraphNode* TargetNode, TArray<
 		TargetNode->Visited = ++VisitCounter;
 		
 		AllChains.Reset();
-		constexpr int32 MinChainDepth = 2; // The chain will contain at least the TargetNode and the ReferencedByNode
+		const int32 MinChainDepth = 2; // The chain will contain at least the TargetNode and the ReferencedByNode
 		BuildReferenceChains(ReferencedByNode, AllChains, MinChainDepth, VisitCounter, SearchMode);
 		for (FReferenceChain* Chain : AllChains)
 		{
@@ -151,7 +152,7 @@ void FReferenceChainSearch::BuildReferenceChains(FGraphNode* TargetNode, TArray<
 		{
 			for (int32 ChainIndex = AllChains.Num() - 1; ChainIndex >= 0; --ChainIndex)
 			{
-				const FReferenceChain* Chain = AllChains[ChainIndex];
+				FReferenceChain* Chain = AllChains[ChainIndex];	
 				if (!Chain->IsExternal())
 				{
 					// Discard the chain
@@ -247,7 +248,7 @@ FString FReferenceChainSearch::GetObjectFlags(UObject* InObject)
 		Flags += TEXT("(NeverGCed) ");
 	}
 
-	const FUObjectItem* ReferencedByObjectItem = GUObjectArray.ObjectToObjectItem(InObject);
+	FUObjectItem* ReferencedByObjectItem = GUObjectArray.ObjectToObjectItem(InObject);
 	if (ReferencedByObjectItem->HasAnyFlags(EInternalObjectFlags::ClusterRoot))
 	{
 		Flags += TEXT("(ClusterRoot) ");
@@ -259,7 +260,7 @@ FString FReferenceChainSearch::GetObjectFlags(UObject* InObject)
 	return Flags;
 }
 
-void FReferenceChainSearch::DumpChain(const FReferenceChain* Chain, FOutputDevice& OutputDevice)
+void FReferenceChainSearch::DumpChain(FReferenceChain* Chain)
 {
 	if (Chain->Num())
 	{
@@ -269,18 +270,18 @@ void FReferenceChainSearch::DumpChain(const FReferenceChain* Chain, FOutputDevic
 			UObject* Object = Chain->GetNode(NodeIndex)->Object;
 			const FNodeReferenceInfo& ReferenceInfo = Chain->GetReferenceInfo(NodeIndex);
 
-			OutputDevice.CategorizedLogf(LogReferenceChain.GetCategoryName(), ELogVerbosity::Log, TEXT("%s%s%s%s"),
+			UE_LOG(LogReferenceChain, Log, TEXT("%s%s%s%s"),
 				FCString::Spc(FMath::Min<int32>(TCStringSpcHelper<TCHAR>::MAX_SPACES, Chain->Num() - NodeIndex - 1)),
 				*GetObjectFlags(Object),
 				*Object->GetFullName(),
 				*ReferenceInfo.ToString()
 			);
 		}
-		OutputDevice.CategorizedLogf(LogReferenceChain.GetCategoryName(), ELogVerbosity::Log, TEXT("  "));
+		UE_LOG(LogReferenceChain, Log, TEXT("  "));
 	}
 }
 
-void FReferenceChainSearch::WriteChain(const FReferenceChain* Chain, FString& OutString)
+void FReferenceChainSearch::WriteChain(FReferenceChain* Chain, FString& OutString)
 {
 	if (Chain->Num())
 	{
@@ -359,7 +360,7 @@ public:
 		, ReferencedObjects(InReferencedObjects)		
 	{
 	}
-	FORCEINLINE void HandleTokenStreamObjectReference(FGCArrayStruct& ObjectsToSerializeStruct, UObject* ReferencingObject, UObject*& Object, const int32 TokenIndex, bool bAllowReferenceElimination) const
+	FORCEINLINE void HandleTokenStreamObjectReference(FGCArrayStruct& ObjectsToSerializeStruct, UObject* ReferencingObject, UObject*& Object, const int32 TokenIndex, bool bAllowReferenceElimination)
 	{
 		FReferenceChainSearch::FObjectReferenceInfo RefInfo(Object);
 		if (Object && !ReferencedObjects.Contains(RefInfo))
@@ -367,7 +368,7 @@ public:
 #if ENABLE_GC_OBJECT_CHECKS
 			if (TokenIndex >= 0)
 			{
-				const FTokenInfo TokenInfo = ReferencingObject->GetClass()->ReferenceTokenStream.GetTokenInfo(TokenIndex);
+				FTokenInfo TokenInfo = ReferencingObject->GetClass()->ReferenceTokenStream.GetTokenInfo(TokenIndex);
 				RefInfo.ReferencerName = TokenInfo.Name;
 				RefInfo.Type = FReferenceChainSearch::EReferenceType::Property;
 			}
@@ -399,9 +400,8 @@ public:
 
 typedef TDefaultReferenceCollector<FDirectReferenceProcessor> FDirectReferenceCollector;
 
-FReferenceChainSearch::FReferenceChainSearch(UObject* InObjectToFindReferencesTo, EReferenceChainSearchMode Mode /*= EReferenceChainSearchMode::PrintResults*/, FOutputDevice& InOutputDevice /*= *GLog*/)
+FReferenceChainSearch::FReferenceChainSearch(UObject* InObjectToFindReferencesTo, EReferenceChainSearchMode Mode /*= EReferenceChainSearchMode::PrintResults*/)
 	: ObjectToFindReferencesTo(InObjectToFindReferencesTo)
-	, OutputDevice(InOutputDevice)
 {
 	check(InObjectToFindReferencesTo);
 
@@ -456,7 +456,7 @@ void FReferenceChainSearch::FindDirectReferencesForObjects()
 
 	for (FRawObjectIterator It; It; ++It)
 	{
-		const FUObjectItem* ObjItem = *It;
+		FUObjectItem* ObjItem = *It;
 		UObject* Object = static_cast<UObject*>(ObjItem->Object);
 		FGraphNode* ObjectNode = FindOrAddNode(AllNodes, Object);
 
@@ -482,19 +482,19 @@ void FReferenceChainSearch::PrintResults(bool bDumpAllChains /*= false*/) const
 	{
 		FSlowHeartBeatScope DisableHangDetection; // This function can be very slow
 
-		constexpr int32 MaxChainsToPrint = 100;
+		const int32 MaxChainsToPrint = 100;
 		int32 NumPrintedChains = 0;
 
-		for (const FReferenceChain* Chain : ReferenceChains)
+		for (FReferenceChain* Chain : ReferenceChains)
 		{
 			if (bDumpAllChains || NumPrintedChains < MaxChainsToPrint)
 			{
-				DumpChain(Chain, OutputDevice);
+				DumpChain(Chain);
 				NumPrintedChains++;
 			}
 			else
 			{
-				OutputDevice.CategorizedLogf(LogReferenceChain.GetCategoryName(), ELogVerbosity::Log, TEXT("Referenced by %d more reference chain(s)."), ReferenceChains.Num() - NumPrintedChains);
+				UE_LOG(LogReferenceChain, Log, TEXT("Referenced by %d more reference chain(s)."), ReferenceChains.Num() - NumPrintedChains);
 				break;
 			}
 		}
@@ -502,8 +502,7 @@ void FReferenceChainSearch::PrintResults(bool bDumpAllChains /*= false*/) const
 	else
 	{
 		check(ObjectToFindReferencesTo);
-
-		OutputDevice.CategorizedLogf(LogReferenceChain.GetCategoryName(), ELogVerbosity::Log, TEXT("%s%s is not currently reachable."),
+		UE_LOG(LogReferenceChain, Log, TEXT("%s%s is not currently reachable."),
 			*GetObjectFlags(ObjectToFindReferencesTo),
 			*ObjectToFindReferencesTo->GetFullName()
 		);
@@ -514,7 +513,7 @@ FString FReferenceChainSearch::GetRootPath() const
 {
 	if (ReferenceChains.Num())
 	{
-		const FReferenceChain* Chain = ReferenceChains[0];
+		FReferenceChain* Chain = ReferenceChains[0];
 		FString OutString;
 
 		WriteChain(Chain, OutString);
@@ -537,7 +536,7 @@ void FReferenceChainSearch::Cleanup()
 	}
 	ReferenceChains.Empty();
 
-	for (const TPair<UObject*, FGraphNode*>& ObjectNodePair : AllNodes)
+	for (TPair<UObject*, FGraphNode*>& ObjectNodePair : AllNodes)
 	{
 		delete ObjectNodePair.Value;
 	}
diff --git a/Engine/Source/Runtime/CoreUObject/Public/UObject/ReferenceChainSearch.h b/Engine/Source/Runtime/CoreUObject/Public/UObject/ReferenceChainSearch.h
--- a/Engine/Source/Runtime/CoreUObject/Public/UObject/ReferenceChainSearch.h
+++ b/Engine/Source/Runtime/CoreUObject/Public/UObject/ReferenceChainSearch.h
@@ -2,7 +2,10 @@
 
 #pragma once
 
-#include "Misc/OutputDeviceRedirector.h"
+#include "CoreMinimal.h"
+#include "UObject/UObjectGlobals.h"
+#include "UObject/Class.h"
+#include "UObject/GarbageCollection.h"
 #include "HAL/ThreadHeartBeat.h"
 
 /** Search mode flags */
@@ -168,7 +171,12 @@ public:
 		{
 			return Nodes.Num();
 		}
-
+		/** Returns a duplicate of this chain */
+		FReferenceChain* Split()
+		{
+			FReferenceChain* NewChain = new FReferenceChain(*this);
+			return NewChain;
+		}
 		/** Checks if this chain contains the specified node */
 		bool Contains(const FGraphNode* InNode) const
 		{
@@ -184,7 +192,7 @@ public:
 	};
 
 	/** Constructs a new search engine and finds references to the specified object */
-	COREUOBJECT_API FReferenceChainSearch(UObject* InObjectToFindReferencesTo, EReferenceChainSearchMode Mode = EReferenceChainSearchMode::PrintResults, FOutputDevice& InOutputDevice = *GLog);
+	COREUOBJECT_API FReferenceChainSearch(UObject* InObjectToFindReferencesTo, EReferenceChainSearchMode Mode = EReferenceChainSearchMode::PrintResults);
 
 	/** Destructor */
 	COREUOBJECT_API ~FReferenceChainSearch();
@@ -213,9 +221,6 @@ private:
 	/** All nodes created during the search */
 	TMap<UObject*, FGraphNode*> AllNodes;
 
-	/** Output device to log to */
-	FOutputDevice& OutputDevice;
-
 	/** Performs the search */
 	void PerformSearch(EReferenceChainSearchMode SearchMode);
 
@@ -241,7 +246,7 @@ private:
 	/** Returns a string with all flags (we care about) set on an object */
 	static FString GetObjectFlags(UObject* InObject);
 	/** Dumps a reference chain to log */
-	static void DumpChain(const FReferenceChain* Chain, FOutputDevice& OutputDevice);
+	static void DumpChain(FReferenceChain* Chain);
 	/** Writes a reference chain to a string */
-	static void WriteChain(const FReferenceChain* Chain, FString& OutString);
+	static void WriteChain(FReferenceChain* Chain, FString& OutString);
 };
