Fix false positives when checking if a payload loaded from disk is corrupt or not.

#rnx
#tests Loaded production map in editor (cold DDC), full EngineTest cook, retested the original use case for CL 17648118

### Problem
- If a virtualized bulkdata object, created from an old bulkdata object (so that the content id was derived from the old guid rather than the hash of the payload content) , was forced into memory by having it's loader reset, it incorrectly have it's LegacyKeyWasGuidDerived flag removed.
- The flag was removed because we'd assign the newly loaded payload via ::UpdatePayloadImpl which assumes that the provided FPayloadId is correct, but in this case it will not be (it's still derived from the old guid) and given that in most use cases the bulkdata will be destroyed after this call (making loading it from disk a bit pointless, but we currently need to do so for safety) we do not want to pay the cost of hashing the payload at this point.
- If the payload was then accessed then the corruption check would fail as we would assume that the content id should match the hash found in the FCompressedBuffer.

### The Fix
- Instead of calling ::UpdatePayloadImpl when detaching the virtualized bulkdata from disk we now just assign it directly and remove the flags/members relating to disk access.
- Added a cpu trace scope to ::DetachFromDisk as it would be great to find out how much time we are wasting here.
- I also noticed that when ::DetachFromDisk is called, we would not remove the links to the file on disk if bEnsurePayloadIsLoaded was true which will now be fixed as well.
- Since there is no easy way for a user to fix this mismatch and it would only be a problem for them if virtualization is enabled I have reduced the validation severity to a warning rather than error/fatal error until I can provide a way for a package to be 'fixed' so people don't waste work.

diff --git a/Engine/Source/Runtime/CoreUObject/Private/Virtualization/VirtualizedBulkData.cpp b/Engine/Source/Runtime/CoreUObject/Private/Virtualization/VirtualizedBulkData.cpp
--- a/Engine/Source/Runtime/CoreUObject/Private/Virtualization/VirtualizedBulkData.cpp
+++ b/Engine/Source/Runtime/CoreUObject/Private/Virtualization/VirtualizedBulkData.cpp
@@ -28,13 +28,13 @@
 #endif //WITH_EDITOR
 
 /** When enabled we will fatal log if we detect corrupted data rather than logging an error and returning a null FCompressedBuffer/FSharedBuffer. */
-#define VBD_CORRUPTED_PAYLOAD_IS_FATAL 1
+#define VBD_CORRUPTED_PAYLOAD_IS_FATAL 0
 
 
 #if VBD_CORRUPTED_PAYLOAD_IS_FATAL
 	#define VBD_CORRUPTED_DATA_SEVERITY Fatal
 #else
-	#define VBD_CORRUPTED_DATA_SEVERITY Error
+	#define VBD_CORRUPTED_DATA_SEVERITY Warning
 #endif // VBD_CORRUPTED_PAYLOAD_IS_FATAL
 
 namespace UE::Virtualization
@@ -1103,23 +1103,27 @@ void FVirtualizedUntypedBulkData::UnloadData()
 
 void FVirtualizedUntypedBulkData::DetachFromDisk(FArchive* Ar, bool bEnsurePayloadIsLoaded)
 {
+	TRACE_CPUPROFILER_EVENT_SCOPE(FVirtualizedUntypedBulkData::DetachFromDisk);
+
 	check(Ar);
 	check(Ar == AttachedAr || AttachedAr == nullptr || AttachedAr->IsProxyOf(Ar));
 
-	if (bEnsurePayloadIsLoaded && !IsDataVirtualized() && !PackagePath.IsEmpty())
+	if (!IsDataVirtualized() && !PackagePath.IsEmpty())
 	{
-		if (Payload.IsNull())
+		if (Payload.IsNull() && bEnsurePayloadIsLoaded)
 		{
 			FCompressedBuffer CompressedPayload = GetDataInternal();
 
 			UE_CLOG(!IsValid(*this, CompressedPayload), LogVirtualization, VBD_CORRUPTED_DATA_SEVERITY, TEXT("%s"), *GetCorruptedPayloadErrorMessage());
 
-			FSharedBuffer DecompressedPayload = CompressedPayload.Decompress();
+			Payload = CompressedPayload.Decompress();
+		}
 
-			AttachedAr = nullptr; // Prevent UpdatePayloadImpl from sending more detach notifications to AttachedAr
+		PackagePath.Empty();
+		PackageSegment = EPackageSegment::Header;
+		OffsetInFile = INDEX_NONE;
 
-			UpdatePayloadImpl(MoveTemp(DecompressedPayload), MoveTemp(PayloadContentId));
-		}
+		EnumRemoveFlags(Flags, EFlags::ReferencesLegacyFile | EFlags::LegacyFileIsCompressed);
 	}
 
 	AttachedAr = nullptr;	
