Remove unused names from package headers, reducing the amount of data and number of names that have to be serialzed
* Strip names that are not used when cooking for UPS
* Strip names that are not used by the loader when staging for IoStore

diff --git a/Engine/Source/Developer/IoStoreUtilities/Private/PackageStoreOptimizer.cpp b/Engine/Source/Developer/IoStoreUtilities/Private/PackageStoreOptimizer.cpp
--- a/Engine/Source/Developer/IoStoreUtilities/Private/PackageStoreOptimizer.cpp
+++ b/Engine/Source/Developer/IoStoreUtilities/Private/PackageStoreOptimizer.cpp
@@ -127,9 +127,9 @@ FPackageStorePackage* FPackageStoreOptimizer::CreatePackageFromCookedHeader(cons
 	FCookedHeaderData CookedHeaderData = LoadCookedHeader(CookedHeaderBuffer);
 	Package->PackageFlags = CookedHeaderData.Summary.GetPackageFlags();
 	Package->CookedHeaderSize = CookedHeaderData.Summary.TotalHeaderSize;
-	for (const FName& SummaryName : CookedHeaderData.SummaryNames)
+	for (int32 I = 0; I < CookedHeaderData.Summary.NamesReferencedFromExportDataCount; ++I)
 	{
-		Package->NameMapBuilder.AddName(SummaryName);
+		Package->NameMapBuilder.AddName(CookedHeaderData.SummaryNames[I]);
 	}
 	ProcessImports(CookedHeaderData, Package);
 	ProcessExports(CookedHeaderData, Package);
@@ -1363,6 +1363,7 @@ void FPackageStoreOptimizer::FinalizePackageHeader(FPackageStorePackage* Package
 		FExportMapEntry ExportMapEntry;
 		ExportMapEntry.CookedSerialOffset = Export.CookedSerialOffset;
 		ExportMapEntry.CookedSerialSize = Export.SerialSize;
+		Package->NameMapBuilder.MarkNameAsReferenced(Export.ObjectName);
 		ExportMapEntry.ObjectName = Package->NameMapBuilder.MapName(Export.ObjectName);
 		ExportMapEntry.ExportHash = Export.ExportHash;
 		ExportMapEntry.OuterIndex = Export.OuterIndex;
diff --git a/Engine/Source/Runtime/Core/Public/UObject/ObjectVersion.h b/Engine/Source/Runtime/Core/Public/UObject/ObjectVersion.h
--- a/Engine/Source/Runtime/Core/Public/UObject/ObjectVersion.h
+++ b/Engine/Source/Runtime/Core/Public/UObject/ObjectVersion.h
@@ -45,6 +45,9 @@ enum class EUnrealEngineObjectUE5Version : uint32
 	// The original UE5 version, at the time this was added the UE4 version was 522, so UE5 will start from 1000 to show a clear difference
 	INITIAL_VERSION = 1000,
 
+	// Support stripping names that are not referenced from export data
+	NAMES_REFERENCED_FROM_EXPORT_DATA,
+
 	// -----<new versions can be added before this line>-------------------------------------------------
 	// - this needs to be the last line (see note below)
 	AUTOMATIC_VERSION_PLUS_ONE,
diff --git a/Engine/Source/Runtime/CoreUObject/Private/UObject/LinkerSave.cpp b/Engine/Source/Runtime/CoreUObject/Private/UObject/LinkerSave.cpp
--- a/Engine/Source/Runtime/CoreUObject/Private/UObject/LinkerSave.cpp
+++ b/Engine/Source/Runtime/CoreUObject/Private/UObject/LinkerSave.cpp
@@ -334,6 +334,24 @@ FArchive& FLinkerSave::operator<<( FName& InName )
 		}
 	}
 
+	if (!CurrentlySavingExport.IsNull())
+	{
+		if (Save >= Summary.NamesReferencedFromExportDataCount)
+		{
+			SetCriticalError();
+			FString ErrorMessage = FString::Printf(TEXT("Name \"%s\" is referenced from an export but not mapped in the export data names region when saving %s (object: %s, property: %s)."),
+				*InName.ToString(),
+				*GetArchiveName(),
+				*GetSerializeContext()->SerializedObject->GetFullName(),
+				*GetFullNameSafe(GetSerializedProperty()));
+			ensureMsgf(false, TEXT("%s"), *ErrorMessage);
+			if (LogOutput)
+			{
+				LogOutput->Logf(ELogVerbosity::Error, TEXT("%s"), *ErrorMessage);
+			}
+		}
+	}
+
 	int32 Number = InName.GetNumber();
 	FArchive& Ar = *this;
 	return Ar << Save << Number;
diff --git a/Engine/Source/Runtime/CoreUObject/Private/UObject/PackageFileSummary.cpp b/Engine/Source/Runtime/CoreUObject/Private/UObject/PackageFileSummary.cpp
--- a/Engine/Source/Runtime/CoreUObject/Private/UObject/PackageFileSummary.cpp
+++ b/Engine/Source/Runtime/CoreUObject/Private/UObject/PackageFileSummary.cpp
@@ -395,6 +395,15 @@ void operator<<(FStructuredArchive::FSlot Slot, FPackageFileSummary& Sum)
 			Sum.PreloadDependencyCount = -1;
 			Sum.PreloadDependencyOffset = 0;
 		}
+
+		if (BaseArchive.IsSaving() || Sum.FileVersionUE >= EUnrealEngineObjectUE5Version::NAMES_REFERENCED_FROM_EXPORT_DATA)
+		{
+			Record << SA_VALUE(TEXT("NamesReferencedFromExportDataCount"), Sum.NamesReferencedFromExportDataCount);
+		}
+		else
+		{
+			Sum.NamesReferencedFromExportDataCount = Sum.NameCount;
+		}
 	}
 }
 
diff --git a/Engine/Source/Runtime/CoreUObject/Private/UObject/SavePackage.cpp b/Engine/Source/Runtime/CoreUObject/Private/UObject/SavePackage.cpp
--- a/Engine/Source/Runtime/CoreUObject/Private/UObject/SavePackage.cpp
+++ b/Engine/Source/Runtime/CoreUObject/Private/UObject/SavePackage.cpp
@@ -724,6 +724,7 @@ void FPackageNameMapSaver::UpdateLinker(FLinkerSave& Linker, FLinkerLoad* Confor
 	if (BinarySaver)
 	{
 		Linker.Summary.NameCount = Linker.NameMap.Num();
+		Linker.Summary.NamesReferencedFromExportDataCount = Linker.Summary.NameCount; // Only supported with SavePackage2
 		for (int32 i = 0; i < Linker.NameMap.Num(); i++)
 		{
 			FName::GetEntry(Linker.NameMap[i])->Write(Linker);
diff --git a/Engine/Source/Runtime/CoreUObject/Private/UObject/SavePackage/PackageHarvester.cpp b/Engine/Source/Runtime/CoreUObject/Private/UObject/SavePackage/PackageHarvester.cpp
--- a/Engine/Source/Runtime/CoreUObject/Private/UObject/SavePackage/PackageHarvester.cpp
+++ b/Engine/Source/Runtime/CoreUObject/Private/UObject/SavePackage/PackageHarvester.cpp
@@ -282,7 +282,7 @@ void FPackageHarvester::TryHarvestExport(UObject* InObject)
 			SaveContext.AddExport(InObject, !DoesObjectNeedLoadForEditorGame(InObject));
 
 			// Harvest the export name
-			HarvestName(InObject->GetFName());
+			HarvestPackageHeaderName(InObject->GetFName());
 
 			ExportsToProcess.Enqueue(InObject);
 		}
@@ -350,7 +350,7 @@ void FPackageHarvester::TryHarvestImport(UObject* InObject)
 		}
 
 		// Harvest the import name
-		HarvestName(ObjName);
+		HarvestPackageHeaderName(ObjName);
 
 		// Recurse into outer, package override and non native class
 		if (ObjOuter)
@@ -370,8 +370,8 @@ void FPackageHarvester::TryHarvestImport(UObject* InObject)
 		}	
 		else
 		{
-			HarvestName(ObjClass->GetFName());
-			HarvestName(ObjClass->GetOuter()->GetFName());
+			HarvestPackageHeaderName(ObjClass->GetFName());
+			HarvestPackageHeaderName(ObjClass->GetOuter()->GetFName());
 		}
 	}
 }
@@ -406,7 +406,7 @@ FArchive& FPackageHarvester::operator<<(UObject*& Obj)
 	// if the package we are saving is referenced, just harvest its name
 	if (Obj == SaveContext.GetPackage())
 	{
-		HarvestName(Obj->GetFName());
+		HarvestPackageHeaderName(Obj->GetFName());
 		return *this;
 	}
 
@@ -466,10 +466,10 @@ FArchive& FPackageHarvester::operator<<(FLazyObjectPtr& LazyObjectPtr)
 
 FArchive& FPackageHarvester::operator<<(FSoftObjectPath& Value)
 {
+	// We need to harvest NAME_None even if the path isn't valid
+	Value.SerializePath(*this);
 	if (Value.IsValid())
 	{
-		Value.SerializePath(*this);
-
 		FSoftObjectPathThreadContext& ThreadContext = FSoftObjectPathThreadContext::Get();
 		FName ReferencingPackageName, ReferencingPropertyName;
 		ESoftObjectPathCollectType CollectType = ESoftObjectPathCollectType::AlwaysCollect;
@@ -482,7 +482,7 @@ FArchive& FPackageHarvester::operator<<(FSoftObjectPath& Value)
 			// Don't track if this is a never collect path
 			FString Path = Value.ToString();
 			FName PackageName = FName(*FPackageName::ObjectPathToPackageName(Path));
-			HarvestName(PackageName);
+			HarvestPackageHeaderName(PackageName);
 			SaveContext.SoftPackageReferenceList.AddUnique(PackageName);
 #if WITH_EDITORONLY_DATA
 			if (CollectType != ESoftObjectPathCollectType::EditorOnlyCollect && !bIsEditorOnlyExportOnStack)
@@ -497,7 +497,7 @@ FArchive& FPackageHarvester::operator<<(FSoftObjectPath& Value)
 
 FArchive& FPackageHarvester::operator<<(FName& Name)
 {
-	HarvestName(Name);
+	HarvestExportDataName(Name);
 	return *this;
 }
 
@@ -526,9 +526,14 @@ bool FPackageHarvester::CurrentExportHasDependency(UObject* InObj) const
 	return SaveContext.ExportObjectDependencies.Contains(InObj) || SaveContext.ExportNativeObjectDependencies.Contains(InObj);
 }
 
-void FPackageHarvester::HarvestName(FName Name)
+void FPackageHarvester::HarvestExportDataName(FName Name)
+{
+	SaveContext.NamesReferencedFromExportData.Add(Name.GetDisplayIndex());
+}
+
+void FPackageHarvester::HarvestPackageHeaderName(FName Name)
 {
-	SaveContext.ReferencedNames.Add(Name.GetDisplayIndex());
+	SaveContext.NamesReferencedFromPackageHeader.Add(Name.GetDisplayIndex());
 }
 
 void FPackageHarvester::HarvestSearchableName(UObject* TypeObject, FName Name)
@@ -539,7 +544,7 @@ void FPackageHarvester::HarvestSearchableName(UObject* TypeObject, FName Name)
 		(*this) << TypeObject;
 	}
 
-	HarvestName(Name);
+	HarvestPackageHeaderName(Name);
 	SaveContext.SearchableNamesObjectMap.FindOrAdd(TypeObject).AddUnique(Name);
 }
 
diff --git a/Engine/Source/Runtime/CoreUObject/Private/UObject/SavePackage/PackageHarvester.h b/Engine/Source/Runtime/CoreUObject/Private/UObject/SavePackage/PackageHarvester.h
--- a/Engine/Source/Runtime/CoreUObject/Private/UObject/SavePackage/PackageHarvester.h
+++ b/Engine/Source/Runtime/CoreUObject/Private/UObject/SavePackage/PackageHarvester.h
@@ -68,7 +68,8 @@ public:
 	void TryHarvestExport(UObject* InObject);
 	void TryHarvestImport(UObject* InObject);
 
-	void HarvestName(FName Name);
+	void HarvestExportDataName(FName Name);
+	void HarvestPackageHeaderName(FName Name);
 	void HarvestSearchableName(UObject* TypeObject, FName Name);
 	void HarvestDependency(UObject* InObj, bool bIsNative);
 
diff --git a/Engine/Source/Runtime/CoreUObject/Private/UObject/SavePackage/SaveContext.h b/Engine/Source/Runtime/CoreUObject/Private/UObject/SavePackage/SaveContext.h
--- a/Engine/Source/Runtime/CoreUObject/Private/UObject/SavePackage/SaveContext.h
+++ b/Engine/Source/Runtime/CoreUObject/Private/UObject/SavePackage/SaveContext.h
@@ -450,9 +450,14 @@ public:
 		return SearchableNamesObjectMap;
 	}
 
-	const TSet<FNameEntryId>& GetReferencedNames() const
+	const TSet<FNameEntryId>& GetNamesReferencedFromExportData() const
 	{
-		return ReferencedNames;
+		return NamesReferencedFromExportData;
+	}
+
+	const TSet<FNameEntryId>& GetNamesReferencedFromPackageHeader() const
+	{
+		return NamesReferencedFromPackageHeader;
 	}
 
 	const FCustomVersionContainer& GetCustomVersions() const
@@ -492,7 +497,14 @@ public:
 
 	bool NameExists(FNameEntryId ComparisonId) const
 	{
-		for (FNameEntryId DisplayId : ReferencedNames)
+		for (FNameEntryId DisplayId : NamesReferencedFromExportData)
+		{
+			if (FName::GetComparisonIdFromDisplayId(DisplayId) == ComparisonId)
+			{
+				return true;
+			}
+		}
+		for (FNameEntryId DisplayId : NamesReferencedFromPackageHeader)
 		{
 			if (FName::GetComparisonIdFromDisplayId(DisplayId) == ComparisonId)
 			{
@@ -617,8 +629,10 @@ private:
 	TSet<UObject*> Imports;
 	// Subset of this->Imports which are referenced from not-editoronly properties
 	TSet<UObject*> ImportsUsedInGame;
-	// Set of names referenced
-	TSet<FNameEntryId> ReferencedNames;
+	// Set of names referenced from export serialization
+	TSet<FNameEntryId> NamesReferencedFromExportData;
+	// Set of names referenced from the package header (import and export table object names etc)
+	TSet<FNameEntryId> NamesReferencedFromPackageHeader;
 	// List of soft package reference found
 	TArray<FName> SoftPackageReferenceList;
 	// Subset of this->SoftPackageReferenceList which are referenced from not-editoronly properties
diff --git a/Engine/Source/Runtime/CoreUObject/Private/UObject/SavePackage2.cpp b/Engine/Source/Runtime/CoreUObject/Private/UObject/SavePackage2.cpp
--- a/Engine/Source/Runtime/CoreUObject/Private/UObject/SavePackage2.cpp
+++ b/Engine/Source/Runtime/CoreUObject/Private/UObject/SavePackage2.cpp
@@ -278,7 +278,7 @@ ESavePackageResult HarvestPackage(FSaveContext& SaveContext)
 	SCOPED_SAVETIMER(UPackage_Save_HarvestPackage);
 
 	FPackageHarvester Harvester(SaveContext);
-	//Harvester.SetUseUnversionedPropertySerialization(SaveContext.IsSaveUnversionedProperties());
+	Harvester.SetUseUnversionedPropertySerialization(SaveContext.IsSaveUnversionedProperties());
 	EObjectFlags TopLevelFlags = SaveContext.GetTopLevelFlags();
 	UObject* Asset = SaveContext.GetAsset();
 
@@ -327,7 +327,7 @@ ESavePackageResult HarvestPackage(FSaveContext& SaveContext)
 	// Harvest the PrestreamPackage class name if needed
 	if (PrestreamPackages.Num() > 0)
 	{
-		Harvester.HarvestName(SavePackageUtilities::NAME_PrestreamPackage);		
+		Harvester.HarvestPackageHeaderName(SavePackageUtilities::NAME_PrestreamPackage);		
 	}
 
 	// if we have a WorldTileInfo, we need to harvest its dependencies as well, i.e. Custom Version
@@ -827,10 +827,28 @@ ESavePackageResult BuildLinker(FSaveContext& SaveContext)
 	// Build Name Map
 	{
 		SCOPED_SAVETIMER(UPackage_Save_BuildNameMap);
+		const TSet<FNameEntryId>& NamesReferencedFromExportData = SaveContext.GetNamesReferencedFromExportData();
+		const TSet<FNameEntryId>& NamesReferencedFromPackageHeader = SaveContext.GetNamesReferencedFromPackageHeader();
+		Linker->NameMap.Reserve(NamesReferencedFromExportData.Num() + NamesReferencedFromPackageHeader.Num());
+		for (FNameEntryId NameEntryId : NamesReferencedFromExportData)
+		{
+			Linker->NameMap.Add(NameEntryId);
+		}
+		for (FNameEntryId NameEntryId : NamesReferencedFromPackageHeader)
+		{
+			if (!NamesReferencedFromExportData.Contains(NameEntryId))
+			{
+				Linker->NameMap.Add(NameEntryId);
+			}
+		}
 		Linker->Summary.NameOffset = 0;
-		Linker->Summary.NameCount = 0;
-		Linker->NameMap.Append(SaveContext.GetReferencedNames().Array());
-		Sort(&Linker->NameMap[0], Linker->NameMap.Num(), FNameEntryIdSortHelper());
+		Linker->Summary.NameCount = Linker->NameMap.Num();
+		Linker->Summary.NamesReferencedFromExportDataCount = NamesReferencedFromExportData.Num();
+
+		Sort(Linker->NameMap.GetData(), Linker->Summary.NamesReferencedFromExportDataCount, FNameEntryIdSortHelper());
+		Sort(Linker->NameMap.GetData() + Linker->Summary.NamesReferencedFromExportDataCount,
+			 Linker->Summary.NameCount - Linker->Summary.NamesReferencedFromExportDataCount,
+			 FNameEntryIdSortHelper());
 
 		if (!SaveContext.IsTextFormat())
 		{
@@ -1397,7 +1415,7 @@ ESavePackageResult WritePackageHeader(FStructuredArchive::FRecord& StructuredArc
 	Linker->Summary.NameOffset = SaveContext.OffsetAfterPackageFileSummary;
 	{
 		SCOPED_SAVETIMER(UPackage_Save_BuildNameMap);
-		Linker->Summary.NameCount = Linker->NameMap.Num();
+		checkf(Linker->Summary.NameCount == Linker->NameMap.Num(), TEXT("Summary NameCount didn't match linker name map count when saving package header for '%s'"), *Linker->LinkerRoot->GetName());
 		for (const FNameEntryId NameEntryId : Linker->NameMap)
 		{
 			FName::GetEntry(NameEntryId)->Write(*Linker);
diff --git a/Engine/Source/Runtime/CoreUObject/Public/UObject/PackageFileSummary.h b/Engine/Source/Runtime/CoreUObject/Public/UObject/PackageFileSummary.h
--- a/Engine/Source/Runtime/CoreUObject/Public/UObject/PackageFileSummary.h
+++ b/Engine/Source/Runtime/CoreUObject/Public/UObject/PackageFileSummary.h
@@ -215,6 +215,11 @@ public:
 	*/
 	int32		PreloadDependencyOffset;
 
+	/**
+	* Number of names that are referenced from serialized export data (sorted first in the name map)
+	*/
+	int32		NamesReferencedFromExportDataCount;
+
 	/** Constructor */
 	COREUOBJECT_API FPackageFileSummary();
 
