Added copy and move operators to FGCObject, to go along the already existing copy and move constructors.
They are empty as the value of bReferenceAdded shouldn't be copied from one instance to another as each instance needs to register itself.

diff --git a/Engine/Source/Runtime/CoreUObject/Public/UObject/GCObject.h b/Engine/Source/Runtime/CoreUObject/Public/UObject/GCObject.h
--- a/Engine/Source/Runtime/CoreUObject/Public/UObject/GCObject.h
+++ b/Engine/Source/Runtime/CoreUObject/Public/UObject/GCObject.h
@@ -154,6 +154,9 @@ public:
 		}
 	}
 
+	FGCObject& operator=(const FGCObject&) {return *this;}
+	FGCObject& operator=(FGCObject&&) {return *this;}
+
 	/**
 	 * Pure virtual that must be overloaded by the inheriting class. Use this
 	 * method to serialize any UObjects contained that you wish to keep around.
