Add a table of contents detailing the payload layout in a package providing an easy way to find info about all payloads in a package file.

#rnx

### Current Changes
- By storing the location of the toc in the package file summary we can access the toc by reading the summary and then jumping to it. It should also be possible for external scripts to be able to do this as well.
- We currently store the identifier, the offset into the package file and the uncompressed size of payload, which will be stored in FCompressedBuffer format.
-- We don't need all this info right now but it means that the structure will be compatible with the experimental sidecar feature. It would also (hypothetically) allow us to move the offset/size info from the bulkdata objects themselves and simply rely on the toc, in which case the loading from within the package file or a sidecar file will be the same.
- When saving a package's exports we will gather a reference to each virtualized bulkdata serialized. The toc should appear before the payloads in the file but we do not know the correct info for the toc until after the payloads have been appended to the package file. So we need to write out a dummy toc that will reserve the space in the file and then after we have appended the payloads, seek back in the file and write out the correct info.

### Future Work
- Although this CL does have a number of code changes to allow this to work with text based assets there are some problems with how that format is currently saved which will prevent it from working (and currently prevents VBD from working with it as well). This will be fixed as a separate issue.
- The experimental sidecar feature uses FTocEntry but serializes a TArray to disk. This should be updated to use FPayloadToc in a later submit.

diff --git a/Engine/Source/Runtime/Core/Public/UObject/ObjectVersion.h b/Engine/Source/Runtime/Core/Public/UObject/ObjectVersion.h
--- a/Engine/Source/Runtime/Core/Public/UObject/ObjectVersion.h
+++ b/Engine/Source/Runtime/Core/Public/UObject/ObjectVersion.h
@@ -48,6 +48,9 @@ enum class EUnrealEngineObjectUE5Version : uint32
 	// Support stripping names that are not referenced from export data
 	NAMES_REFERENCED_FROM_EXPORT_DATA,
 
+	// Added a payload table of contents to the package summary 
+	PAYLOAD_TOC,
+	
 	// -----<new versions can be added before this line>-------------------------------------------------
 	// - this needs to be the last line (see note below)
 	AUTOMATIC_VERSION_PLUS_ONE,
diff --git a/Engine/Source/Runtime/CoreUObject/Private/UObject/PackageFileSummary.cpp b/Engine/Source/Runtime/CoreUObject/Private/UObject/PackageFileSummary.cpp
--- a/Engine/Source/Runtime/CoreUObject/Private/UObject/PackageFileSummary.cpp
+++ b/Engine/Source/Runtime/CoreUObject/Private/UObject/PackageFileSummary.cpp
@@ -404,6 +404,15 @@ void operator<<(FStructuredArchive::FSlot Slot, FPackageFileSummary& Sum)
 		{
 			Sum.NamesReferencedFromExportDataCount = Sum.NameCount;
 		}
+
+		if (BaseArchive.IsSaving() || Sum.FileVersionUE >= EUnrealEngineObjectUE5Version::PAYLOAD_TOC)
+		{
+			Record << SA_VALUE(TEXT("PayloadTocOffset"), Sum.PayloadTocOffset);
+		}
+		else
+		{
+			Sum.PayloadTocOffset = INDEX_NONE;
+		}
 	}
 }
 
diff --git a/Engine/Source/Runtime/CoreUObject/Private/UObject/SavePackage.cpp b/Engine/Source/Runtime/CoreUObject/Private/UObject/SavePackage.cpp
--- a/Engine/Source/Runtime/CoreUObject/Private/UObject/SavePackage.cpp
+++ b/Engine/Source/Runtime/CoreUObject/Private/UObject/SavePackage.cpp
@@ -69,6 +69,7 @@
 #include "Misc/ScopeExit.h"
 #include "Misc/PackageAccessTracking.h"
 #include "Misc/PackageAccessTrackingOps.h"
+#include "Virtualization/VirtualizedBulkData.h"
 
 #if ENABLE_COOK_STATS
 #include "ProfilingDebugging/ScopedTimers.h"
@@ -1881,6 +1882,79 @@ static bool ValidateConformCompatibility(UPackage* NewPackage, FLinkerLoad* OldL
 	return !bHadCompatibilityErrors;
 }
 
+/**
+ * Writes out the initial state of the payload table of contents.
+ *
+ * However because we store the table of contents before the payloads in the package
+ * format the data will not be correct as we cannot know the correct payload offsets.
+ * The correct data will be written out by a subsequent call to 'UpdatePayloadToC'.
+ *
+ * Because of this we only write out the table of contents if we are currently writing
+ * to a binary format package. This will reserve the space in the file that we need
+ * to later write the correct data too.
+ * If the package is in text format then we do not need to write anything as we can
+ * simple write out the correct data later when we know it.
+ */
+static ESavePackageResult WritePayloadToC(FLinkerSave* Linker, bool bTextFormat)
+{
+	if (!bTextFormat)
+	{
+		Linker->Summary.PayloadTocOffset = Linker->Tell();
+
+		UE::Virtualization::FPayloadToc PayloadTableOfContents;
+		for (const UE::Virtualization::FVirtualizedUntypedBulkData* BulkData : Linker->BulkDataInPackage)
+		{
+			if (BulkData != nullptr)
+			{
+				PayloadTableOfContents.AddEntry(*BulkData);
+			}
+		}
+
+		*Linker << PayloadTableOfContents;
+	}
+
+	return ESavePackageResult::Success;
+}
+
+/**
+ * Either update the existing table of contents (if writing out in binary format) or write the table
+ * of contents (if writing out in text format).
+ *
+ * @see WritePayloadToC
+ */
+ESavePackageResult UpdatePayloadToC(FLinkerSave* Linker,  FStructuredArchive::FRecord& StructuredArchiveRoot, bool bTextFormat, int64 OffsetAfterPayloadToc)
+{
+	UE::Virtualization::FPayloadToc PayloadTableOfContents;
+	for (const UE::Virtualization::FVirtualizedUntypedBulkData* BulkData : Linker->BulkDataInPackage)
+	{
+		if (BulkData != nullptr)
+		{
+			PayloadTableOfContents.AddEntry(*BulkData);
+		}
+	}
+
+	if (bTextFormat)
+	{
+		FStructuredArchive::FStream Stream = StructuredArchiveRoot.EnterStream(SA_FIELD_NAME(TEXT("PayloadTableOfContents")));
+		Stream.EnterElement() << PayloadTableOfContents;
+	}
+	else
+	{
+		check(Linker->Summary.PayloadTocOffset != INDEX_NONE);
+
+		int64 CurrentPos = Linker->Tell();
+		Linker->Seek(Linker->Summary.PayloadTocOffset);
+
+		*Linker << PayloadTableOfContents;
+
+		checkf(OffsetAfterPayloadToc == Linker->Tell(), TEXT("Mismatch between ::WritePayloadToC and ::UpdatePayloadToC, the package will be corrupt"));
+
+		Linker->Seek(CurrentPos);
+	}
+
+	return ESavePackageResult::Success;
+}
+
 #if WITH_EDITOR
 /**
  * Helper function to sort export objects by fully qualified names.
@@ -2244,7 +2318,7 @@ FSavePackageResultStruct UPackage::Save(UPackage* InOuter, UObject* Base, EObjec
 
 		FText StatusMessage = FText::Format(NSLOCTEXT("Core", "SavingFile", "Saving file: {CleanFilename}..."), Args);
 
-		const int32 TotalSaveSteps = 33;
+		const int32 TotalSaveSteps = 35;
 		FScopedSlowTask SlowTask(TotalSaveSteps, StatusMessage, bSlowTask);
 		SlowTask.MakeDialogDelayed(3.0f, SaveFlags & SAVE_FromAutosave ? true : false);
 
@@ -4161,6 +4235,15 @@ FSavePackageResultStruct UPackage::Save(UPackage* InOuter, UObject* Base, EObjec
 					return ESavePackageResult::Canceled;
 				}
 
+				SlowTask.EnterProgressFrame();
+				ESavePackageResult WritePayloadTocResult = WritePayloadToC(Linker.Get(), bTextFormat);
+				if (WritePayloadTocResult != ESavePackageResult::Success)
+				{
+					return WritePayloadTocResult;
+				}
+
+				const int64 OffsetAfterPayloadToc = Linker->Tell();
+
 				SlowTask.EnterProgressFrame(1, NSLOCTEXT("Core", "SerializingBulkData", "Serializing bulk data"));
 
 				IPackageWriter* PackageWriter = SavePackageContext ? SavePackageContext->PackageWriter : nullptr;
@@ -4229,6 +4312,13 @@ FSavePackageResultStruct UPackage::Save(UPackage* InOuter, UObject* Base, EObjec
 					}
 				}
 
+				SlowTask.EnterProgressFrame();
+				ESavePackageResult UpdatePayloadTocResult = UpdatePayloadToC(Linker.Get(), StructuredArchiveRoot, bTextFormat, OffsetAfterPayloadToc);
+				if (UpdatePayloadTocResult != ESavePackageResult::Success)
+				{
+					return UpdatePayloadTocResult;
+				}
+
 				// write the package post tag
 				if (!bTextFormat)
 				{
diff --git a/Engine/Source/Runtime/CoreUObject/Private/UObject/SavePackage/SaveContext.h b/Engine/Source/Runtime/CoreUObject/Private/UObject/SavePackage/SaveContext.h
--- a/Engine/Source/Runtime/CoreUObject/Private/UObject/SavePackage/SaveContext.h
+++ b/Engine/Source/Runtime/CoreUObject/Private/UObject/SavePackage/SaveContext.h
@@ -588,7 +588,8 @@ public:
 	int32 OffsetAfterPackageFileSummary = 0;
 	int32 OffsetAfterImportMap = 0;
 	int32 OffsetAfterExportMap = 0;
-	int32 SerializedPackageFlags = 0;
+	int64 OffsetAfterPayloadToc = 0;
+	int32 SerializedPackageFlags = 0;	
 	TAsyncWorkSequence<FMD5> AsyncWriteAndHashSequence;
 	TArray<FLargeMemoryWriter, TInlineAllocator<4>> AdditionalFilesFromExports;
 	FSavePackageOutputFileArray AdditionalPackageFiles;
diff --git a/Engine/Source/Runtime/CoreUObject/Private/UObject/SavePackage2.cpp b/Engine/Source/Runtime/CoreUObject/Private/UObject/SavePackage2.cpp
--- a/Engine/Source/Runtime/CoreUObject/Private/UObject/SavePackage2.cpp
+++ b/Engine/Source/Runtime/CoreUObject/Private/UObject/SavePackage2.cpp
@@ -40,6 +40,7 @@
 #include "UObject/SavePackage/SavePackageUtilities.h"
 #include "UObject/UObjectGlobals.h"
 #include "UObject/UObjectHash.h"
+#include "Virtualization/VirtualizedBulkData.h"
 
 #if ENABLE_COOK_STATS
 #include "ProfilingDebugging/ScopedTimers.h"
@@ -1644,6 +1645,81 @@ ESavePackageResult WriteExports(FStructuredArchive::FRecord& StructuredArchiveRo
 	return Linker->IsError() ? ESavePackageResult::Error : ReturnSuccessOrCancel();
 }
 
+/** 
+ * Writes out the initial state of the payload table of contents.
+ * 
+ * However because we store the table of contents before the payloads in the package
+ * format the data will not be correct as we cannot know the correct payload offsets.
+ * The correct data will be written out by a subsequent call to 'UpdatePayloadToC'.
+ * 
+ * Because of this we only write out the table of contents if we are currently writing
+ * to a binary format package. This will reserve the space in the file that we need 
+ * to later write the correct data too.
+ * If the package is in text format then we do not need to write anything as we can
+ * simple write out the correct data later when we know it.
+ */
+ESavePackageResult WritePayloadToC(FSaveContext& SaveContext)
+{
+	if (!SaveContext.IsTextFormat())
+	{
+		SaveContext.Linker->Summary.PayloadTocOffset = SaveContext.Linker->Tell();
+
+		UE::Virtualization::FPayloadToc PayloadTableOfContents;
+		for (const UE::Virtualization::FVirtualizedUntypedBulkData* BulkData : SaveContext.Linker->BulkDataInPackage)
+		{
+			if (BulkData != nullptr)
+			{
+				PayloadTableOfContents.AddEntry(*BulkData);
+			}
+		}
+
+		*SaveContext.Linker << PayloadTableOfContents;
+
+		SaveContext.OffsetAfterPayloadToc = SaveContext.Linker->Tell();
+	}
+	
+	return ESavePackageResult::Success;
+}
+
+/** 
+ * Either update the existing table of contents (if writing out in binary format) or write the table
+ * of contents (if writing out in text format).
+ * 
+ * @see WritePayloadToC
+ */
+ESavePackageResult UpdatePayloadToC(FStructuredArchive::FRecord& StructuredArchiveRoot, FSaveContext& SaveContext)
+{
+	UE::Virtualization::FPayloadToc PayloadTableOfContents;
+	for (const UE::Virtualization::FVirtualizedUntypedBulkData* BulkData : SaveContext.Linker->BulkDataInPackage)
+	{
+		if (BulkData != nullptr)
+		{
+			PayloadTableOfContents.AddEntry(*BulkData);
+		}
+	}
+
+	if (SaveContext.IsTextFormat())
+	{
+		FStructuredArchive::FStream Stream = StructuredArchiveRoot.EnterStream(SA_FIELD_NAME(TEXT("PayloadTableOfContents")));
+		Stream.EnterElement() << PayloadTableOfContents;
+	}
+	else
+	{
+		check(SaveContext.Linker->Summary.PayloadTocOffset != INDEX_NONE);
+
+		int64 CurrentPos = SaveContext.Linker->Tell();
+		SaveContext.Linker->Seek(SaveContext.Linker->Summary.PayloadTocOffset);
+
+		*SaveContext.Linker << PayloadTableOfContents;
+
+		checkf(SaveContext.OffsetAfterPayloadToc == SaveContext.Linker->Tell(), TEXT("Mismatch between ::WritePayloadToC and ::UpdatePayloadToC, the package will be corrupt"));
+
+		SaveContext.Linker->Seek(CurrentPos);
+	}
+
+	return ESavePackageResult::Success;
+}
+
 ESavePackageResult WriteAdditionalExportFiles(FSaveContext& SaveContext)
 {
 	FSavePackageContext* SavePackageContext = SaveContext.GetSavePackageContext();
@@ -2056,7 +2132,7 @@ ESavePackageResult InnerSave(FSaveContext& SaveContext)
 	SaveContext.SetEDLCookChecker(&FEDLCookChecker::Get());
 
 	// Create slow task dialog if needed
-	const int32 TotalSaveSteps = 12;
+	const int32 TotalSaveSteps = 14;
 	FScopedSlowTask SlowTask(TotalSaveSteps, FText(), SaveContext.IsUsingSlowTask());
 	SlowTask.MakeDialogDelayed(3.0f, SaveContext.IsFromAutoSave());
 
@@ -2148,6 +2224,13 @@ ESavePackageResult InnerSave(FSaveContext& SaveContext)
 		}
 	}
 
+	SlowTask.EnterProgressFrame();
+	SaveContext.Result = WritePayloadToC(SaveContext);
+	if (SaveContext.Result != ESavePackageResult::Success)
+	{
+		return SaveContext.Result;
+	}
+
 	bool bAdditionalFilesNeedLinkerSize = SaveContext.IsAdditionalFilesNeedLinkerSize();
 	if (!bAdditionalFilesNeedLinkerSize)
 	{
@@ -2157,7 +2240,14 @@ ESavePackageResult InnerSave(FSaveContext& SaveContext)
 			return SaveContext.Result;
 		}
 	}
-	
+
+	SlowTask.EnterProgressFrame();
+	SaveContext.Result = UpdatePayloadToC(StructuredArchiveRoot, SaveContext);
+	if (SaveContext.Result != ESavePackageResult::Success)
+	{
+		return SaveContext.Result;
+	}
+
 	// Write Package Post Tag
 	if (!SaveContext.IsTextFormat())
 	{
diff --git a/Engine/Source/Runtime/CoreUObject/Private/Virtualization/VirtualizedBulkData.cpp b/Engine/Source/Runtime/CoreUObject/Private/Virtualization/VirtualizedBulkData.cpp
--- a/Engine/Source/Runtime/CoreUObject/Private/Virtualization/VirtualizedBulkData.cpp
+++ b/Engine/Source/Runtime/CoreUObject/Private/Virtualization/VirtualizedBulkData.cpp
@@ -629,8 +629,9 @@ void FVirtualizedUntypedBulkData::Serialize(FArchive& Ar, UObject* Owner, bool b
 
 				// If we have a valid linker then we will defer serialization of the payload so that it will
 				// be placed at the end of the output file so we don't have to seek past the payload on load.
-				// If we do not have a linker then we might as well just serialize right away.
-				if(LinkerSave != nullptr)
+				// If we do not have a linker OR the linker is in text format then we should just serialize
+				// the payload directly to the archive.
+				if(LinkerSave != nullptr && !LinkerSave->IsTextFormat())
 				{
 					auto AdditionalDataCallback = [SerializePayload = MoveTemp(SerializePayload)](FLinkerSave& ExportsArchive, FArchive& DataArchive, int64 DataStartOffset) mutable
 					{
@@ -644,6 +645,11 @@ void FVirtualizedUntypedBulkData::Serialize(FArchive& Ar, UObject* Owner, bool b
 				}
 			}
 
+			if (LinkerSave != nullptr)
+			{
+				LinkerSave->BulkDataInPackage.Add(this);
+			}
+
 			if (CanUnloadData())
 			{
 				this->CompressionSettings.Reset();
@@ -1423,6 +1429,138 @@ FString FVirtualizedUntypedBulkData::GetCorruptedPayloadErrorMessage() const
 
 	return Message.ToString();
 }
+	
+FArchive& operator<<(FArchive& Ar, FTocEntry& Entry)
+{
+	Ar << Entry.Identifier;
+	Ar << Entry.OffsetInFile;
+	Ar << Entry.UncompressedSize;
+
+	return Ar;
+}
+
+ void operator<<(FStructuredArchive::FSlot Slot, FTocEntry& Entry)
+{
+	FStructuredArchive::FRecord Record = Slot.EnterRecord();
+
+	Record << SA_VALUE(TEXT("Identifier"), Entry.Identifier);
+	Record << SA_VALUE(TEXT("OffsetInFile"), Entry.OffsetInFile);
+	Record << SA_VALUE(TEXT("UncompressedSize"), Entry.UncompressedSize);
+}
+
+void FPayloadToc::AddEntry(const FVirtualizedUntypedBulkData& BulkData)
+{
+	if (BulkData.GetPayloadId().IsValid())
+	{
+		Contents.Emplace(BulkData);
+	}
+}
+
+bool FPayloadToc::FindEntry(const FPayloadId& Identifier, FTocEntry& OutEntry)
+{
+	for (const FTocEntry& Entry : Contents)
+	{
+		if (Entry.Identifier == Identifier)
+		{
+			OutEntry = Entry;
+			return true;
+		}
+	}
+
+	return false;
+}
+
+const TArray<UE::Virtualization::FTocEntry>& FPayloadToc::GetContents() const
+{
+	return Contents;
+}
+
+FArchive& operator<<(FArchive& Ar, FPayloadToc& TableOfContents)
+{
+	FPayloadToc::EVersion Version = FPayloadToc::EVersion::AUTOMATIC_VERSION;
+	Ar << Version;
+
+	Ar << TableOfContents.Contents;
+
+	return Ar;
+}
+
+void operator<<(FStructuredArchive::FSlot Slot, FPayloadToc& TableOfContents)
+{
+	FStructuredArchive::FRecord Record = Slot.EnterRecord();
+
+	FPayloadToc::EVersion Version = FPayloadToc::EVersion::AUTOMATIC_VERSION;
+
+	Record << SA_VALUE(TEXT("Version"), Version);
+	Record << SA_VALUE(TEXT("Entries"), TableOfContents.Contents);
+}
+
+bool FindPayloadsInPackageFile(const FPackagePath& PackagePath, EPayloadType Filter, TArray<FPayloadId>& OutPayloadIds)
+{
+	if (FPackageName::IsTextPackageExtension(PackagePath.GetHeaderExtension()))
+	{
+		UE_LOG(LogVirtualization, Warning, TEXT("Attempting to call 'FindPayloadsInPackageFile' on a text based asset '%s' this is not currently supported"), *PackagePath.GetDebugName());
+		return false;
+	}
+
+	TUniquePtr<FArchive> Ar = IPackageResourceManager::Get().OpenReadExternalResource(EPackageExternalResource::WorkspaceDomainFile, PackagePath.GetPackageName());
+
+	if (Ar.IsValid())
+	{
+		FPackageFileSummary PackageFileSummary;
+		*Ar << PackageFileSummary;
+
+		if (Ar->IsError())
+		{
+			UE_LOG(LogVirtualization, Warning, TEXT("Failed to parse the FPackageFileSummary for '%s'"), *PackagePath.GetDebugName());
+			return false;
+		}
+
+		if (PackageFileSummary.Tag == PACKAGE_FILE_TAG)
+		{
+			if (PackageFileSummary.PayloadTocOffset != INDEX_NONE)
+			{
+				Ar->Seek(PackageFileSummary.PayloadTocOffset);
+
+				FPayloadToc PayloadToC;
+				*Ar << PayloadToC;
+
+				for (const UE::Virtualization::FTocEntry& Entry : PayloadToC.GetContents())
+				{
+					switch (Filter)
+					{
+					case EPayloadType::Local:
+						if (Entry.OffsetInFile != INDEX_NONE)
+						{
+							OutPayloadIds.Add(Entry.Identifier);
+						}
+						break;
+					case EPayloadType::Virtualized:
+						if (Entry.OffsetInFile == INDEX_NONE)
+						{
+							OutPayloadIds.Add(Entry.Identifier);
+						}
+						break;
+					case EPayloadType::All:
+						OutPayloadIds.Add(Entry.Identifier);
+						break;
+					}
+				}
+			}
+		}
+		else
+		{
+			UE_LOG(LogVirtualization, Warning, TEXT("Failed to find PACKAGE_FILE_TAG in '%s', not a valid package!"), *PackagePath.GetDebugName());
+		}
+
+		return true;
+	}
+	else
+	{
+		UE_LOG(LogVirtualization, Warning, TEXT("Unable to open '%s' for reading"), *PackagePath.GetDebugName());
+		return false;
+	}
+}
 
 } // namespace UE::Virtualization
 
diff --git a/Engine/Source/Runtime/CoreUObject/Public/UObject/LinkerSave.h b/Engine/Source/Runtime/CoreUObject/Public/UObject/LinkerSave.h
--- a/Engine/Source/Runtime/CoreUObject/Public/UObject/LinkerSave.h
+++ b/Engine/Source/Runtime/CoreUObject/Public/UObject/LinkerSave.h
@@ -14,6 +14,7 @@
 
 class FObjectPostSaveContext;
 struct FUntypedBulkData;
+namespace UE::Virtualization { class FVirtualizedUntypedBulkData; }
 
 /*----------------------------------------------------------------------------
 	FLinkerSave.
@@ -103,15 +104,19 @@ public:
 	 * This is used e.g. to decide whether to update the in-memory file offsets for BulkData.
 	 */
 	bool bUpdatingLoadedPath = false;
-
-	/** Used by FVirtualizedUntypedBulkData to add payloads to be added to the payload sidecar file */
+	
 	struct FSidecarStorageInfo
 	{
 		UE::Virtualization::FPayloadId Identifier;
 		FCompressedBuffer Payload;
 	};
+
+	/** Used by FVirtualizedUntypedBulkData to add payloads to be added to the payload sidecar file (currently an experimental feature) */
 	TArray<FSidecarStorageInfo> SidecarDataToAppend;
 
+	/** A list of all virtualized bulkdata objects that serialized while saving the package. Note that we do not expect nullptr values to be added! */
+	TArray<UE::Virtualization::FVirtualizedUntypedBulkData*> BulkDataInPackage;
+
 	/** 
 	 * Array of callbacks that will be invoked when the package has successfully saved to disk.
 	 * The callbacks will not be invoked if the package fails to save for some reason.
diff --git a/Engine/Source/Runtime/CoreUObject/Public/UObject/PackageFileSummary.h b/Engine/Source/Runtime/CoreUObject/Public/UObject/PackageFileSummary.h
--- a/Engine/Source/Runtime/CoreUObject/Public/UObject/PackageFileSummary.h
+++ b/Engine/Source/Runtime/CoreUObject/Public/UObject/PackageFileSummary.h
@@ -219,6 +219,11 @@ public:
 	* Number of names that are referenced from serialized export data (sorted first in the name map)
 	*/
 	int32		NamesReferencedFromExportDataCount;
+	
+	/**
+	* Location into the file on disk for the payload table of contents data
+	*/
+	int64		PayloadTocOffset;
 
 	/** Constructor */
 	COREUOBJECT_API FPackageFileSummary();
diff --git a/Engine/Source/Runtime/CoreUObject/Public/Virtualization/VirtualizedBulkData.h b/Engine/Source/Runtime/CoreUObject/Public/Virtualization/VirtualizedBulkData.h
--- a/Engine/Source/Runtime/CoreUObject/Public/Virtualization/VirtualizedBulkData.h
+++ b/Engine/Source/Runtime/CoreUObject/Public/Virtualization/VirtualizedBulkData.h
@@ -286,6 +286,8 @@ protected:
 	FVirtualizedUntypedBulkData(const FVirtualizedUntypedBulkData& Other, ETornOff);
 
 private:
+	friend struct FTocEntry;
+
 	/** Flags used to store additional meta information about the bulk data */
 	enum class EFlags : uint32
 	{
@@ -423,27 +425,81 @@ using FFloatVirtualizedBulkData	= TVirtualizedBulkData<float>;
   * This might be moved to it's own header and the table of contents made into a proper class 
   * if we decide that we want to make access of the payload sidecar file a generic feature.
   */
-struct FTocEntry
+struct COREUOBJECT_API FTocEntry
 {
+	FTocEntry() = default;
+	FTocEntry(const FVirtualizedUntypedBulkData& BulkData)
+		: Identifier(BulkData.PayloadContentId)
+		, OffsetInFile(BulkData.OffsetInFile)
+		, UncompressedSize(BulkData.PayloadSize)
+	{
+
+	}
+
+	friend FArchive& operator<<(FArchive& Ar, FTocEntry& Entry);
+	friend void operator<<(FStructuredArchive::FSlot Slot, FTocEntry& Entry);
+
 	static constexpr uint32 PayloadSidecarFileVersion = 1;
 
 	/** Identifier for the payload */
 	FPayloadId Identifier;
-	/** The offset into the file where we can find the payload */
+	/** The offset into the file where we can find the payload, note that a virtualized payload will have an offset of INDEX_NONE */
 	int64 OffsetInFile = INDEX_NONE;
-	/** The size of the payload WHEN uncompressed. */
+	/** The size of the payload when uncompressed. */
 	int64 UncompressedSize = INDEX_NONE;
+};
+
+/** A table of contents showing the location of all virtualized bulkdata payloads in a file. */
+class COREUOBJECT_API FPayloadToc
+{
+public:
+
+	void AddEntry(const FVirtualizedUntypedBulkData& BulkData);
+	bool FindEntry(const FPayloadId& Identifier, FTocEntry& OutEntry);
+
+	const TArray<UE::Virtualization::FTocEntry>& GetContents() const;
+
+	friend FArchive& operator<<(FArchive& Ar, FPayloadToc& TableOfContents);
+	friend void operator<<(FStructuredArchive::FSlot Slot, FPayloadToc& TableOfContents);
 
-	friend FArchive& operator<<(FArchive& Ar, FTocEntry& Entry)
+private:
+	enum class EVersion : uint32
 	{
-		Ar << Entry.Identifier;
-		Ar << Entry.OffsetInFile;
-		Ar << Entry.UncompressedSize;
+		INITIAL = 0,
 
-		return Ar;
-	}
+		// -----<new versions can be added before this line>-------------------------------------------------
+		// - this needs to be the last line (see note below)
+		AUTOMATIC_VERSION_PLUS_ONE,
+		AUTOMATIC_VERSION = AUTOMATIC_VERSION_PLUS_ONE - 1
+	};
+
+	TArray<UE::Virtualization::FTocEntry> Contents;
+};
+
+/** Used to filter a request to a specific type of payload */
+enum class EPayloadType
+{
+	/** All payload types. */
+	All,
+	/** All payloads stored locally on disk. */
+	Local,
+	/** All payloads stored in a virtualized backend. */
+	Virtualized
 };
 
+/**
+ * Used to find the identifiers of the payload in a given package. Note that this will return the payloads included in the
+ * package on disk and will not take into account any edits to the package if they are in memory and unsaved.
+ *
+ * @param PackagePath	The package to look in.
+ * @param Filter		What sort of payloads should be returned. @see EPayloadType
+ * @param OutPayloadIds	This array will be filled with the FPayloadId values found in the package that passed the filter.
+ *						Note that existing content in the array will be preserved. It is up to the caller to empty it.
+ *
+ * @return 				True if the package was parsed successfully (although it might not have contained any payloads) and false if opening or parsing the package file failed.
+ */
+COREUOBJECT_API bool FindPayloadsInPackageFile(const FPackagePath& PackagePath, EPayloadType Filter, TArray<FPayloadId>& OutPayloadIds);
+
 } // namespace UE::Virtualization
 
 //#endif //WITH_EDITORONLY_DATA
