Clean up some minor code comments.

#rnx

- Cleaned up some of the code comments documenting the utility macros and functions.
- Removed ::IsIoDispatcherEnabled as it was no longer used.

diff --git a/Engine/Source/Runtime/CoreUObject/Private/Serialization/BulkData2.cpp b/Engine/Source/Runtime/CoreUObject/Private/Serialization/BulkData2.cpp
--- a/Engine/Source/Runtime/CoreUObject/Private/Serialization/BulkData2.cpp
+++ b/Engine/Source/Runtime/CoreUObject/Private/Serialization/BulkData2.cpp
@@ -31,13 +31,14 @@ IMPLEMENT_TYPE_LAYOUT(FBulkDataBase);
 // the new loader to keep the results consistent.
 #define UE_KEEP_INLINE_RELOADING_CONSISTENT 0
 
-// Handy macro to validate FIoStatus return values
+// Used to validate FIoStatus return values and assert if there is an error
 #define CHECK_IOSTATUS( InIoStatus, InMethodName ) checkf(InIoStatus.IsOk(), TEXT("%s failed: %s"), InMethodName, *InIoStatus.ToString());
 
 namespace
 {
 	const uint16 InvalidBulkDataIndex = ~uint16(0);
 
+	/** Return true if bulkdata is allowed to be stored in the IoStore system and false if it must be stored in the older PakFile system */
 	bool ShouldAllowBulkDataInIoStore()
 	{
 		static struct FAllowBulkDataInIoStore
@@ -58,18 +59,11 @@ namespace
 		return AllowBulkDataInIoStore.bEnabled;
 	}
 
-	FORCEINLINE bool IsIoDispatcherEnabled()
-	{
-		if (ShouldAllowBulkDataInIoStore())
-		{
-			return FIoDispatcher::IsInitialized();
-		}
-		else
-		{
-			return false;
-		}
-	}
-
+	/** 
+	 * Returns true if we should not trigger an ensure if we detect an inline bulkdata reload request that will not work with the 
+	 * IoStore system. This can be done by setting [Core.System]IgnoreInlineBulkDataReloadEnsures to true in the config file.
+	 * It is NOT recommended that you do this, but is provded in case of unforseen use cases.
+	 */
 	bool ShouldIgnoreInlineDataReloadEnsures()
 	{
 		static struct FIgnoreInlineDataReloadEnsures
@@ -134,9 +128,6 @@ FIoFilenameHash MakeIoFilenameHash(const FIoChunkId& ChunkID)
 	}
 }
 
-// TODO: The code in the PackageTokenSystem namespace is a temporary system so that FBulkDataBase can hold
-// all of it's info about where the data is on disk in a single integer value. This can all be removed when
-// we switch this over to the new packing system.
 namespace PackageTokenSystem
 {
 	// Internal to the PackageTokenSystem namespace
@@ -294,9 +285,9 @@ private:
 	}
 };
 
-// TODO: Currently shared between all FReadChunkIdRequest as the PS4/Pak implementation do but it would be
-// worth profiling on some different platforms to see if we lose more perf from the potential increase in 
-// locks vs the gain we get from not creating so many CriticalSections.
+// Currently we share the single critical section between all FReadChunkIdRequest to match how the older 
+// implementation for reading from pak files was managing multi-threaded access.
+// TODO: Profile this to see if it is worth replacing with a single critical section per FReadChunkIdRequest.
 static FCriticalSection FReadChunkIdRequestEvent;
 
 class FReadChunkIdRequest : public IAsyncReadRequest
