AsyncLoading2 now calls FCoreUObjectDelegates::OnAssetLoaded when WITH_IOSTORE_IN_EDITOR is set

FCoreUObjectDelegates::OnEndLoadPackage:
Fix the implementation in AsyncLoading1 to handle recursive loads; FAsyncPackage->PackageObjLoaded includes UObjects loaded from packages that were import dependencies of the package, and we need to send OnEndLoadPackage for those as well.

Also filter out the transient package and script and memory packages from all OnEndLoadPackage calls; this was affecting all three of AsyncLoading1, AsyncLoading2, and UObjectGlobals.

Did some comment improvement and removal of dead code in AsyncLoading1 while I was working in the area.

#rnx

diff --git a/Engine/Source/Editor/UnrealEd/Private/EditorDomain/BulkDataRegistryEditorDomain.cpp b/Engine/Source/Editor/UnrealEd/Private/EditorDomain/BulkDataRegistryEditorDomain.cpp
--- a/Engine/Source/Editor/UnrealEd/Private/EditorDomain/BulkDataRegistryEditorDomain.cpp
+++ b/Engine/Source/Editor/UnrealEd/Private/EditorDomain/BulkDataRegistryEditorDomain.cpp
@@ -166,7 +166,7 @@ void FBulkDataRegistryEditorDomain::Register(UPackage* Owner, const UE::Virtuali
 	UE::Virtualization::FVirtualizedUntypedBulkData CopyBulk(BulkData.CopyTornOff());
 	if (Owner
 		&& Owner->GetFileSize() // We only record the BulkDataList for disk packages
-		&& !Owner->bHasBeenEndLoaded // We only record BulkDats that are loaded before the package finishes loading
+		&& !Owner->GetHasBeenEndLoaded() // We only record BulkDats that are loaded before the package finishes loading
 		&& CopyBulk.CanSaveForRegistry()
 		)
 	{
diff --git a/Engine/Source/Runtime/CoreUObject/Private/Serialization/AsyncLoading.cpp b/Engine/Source/Runtime/CoreUObject/Private/Serialization/AsyncLoading.cpp
--- a/Engine/Source/Runtime/CoreUObject/Private/Serialization/AsyncLoading.cpp
+++ b/Engine/Source/Runtime/CoreUObject/Private/Serialization/AsyncLoading.cpp
@@ -4185,12 +4185,9 @@ EAsyncPackageState::Type FAsyncLoadingThread::ProcessAsyncLoading(int32& OutPack
 	else if (AsyncPackages.Num())
 	{
 		SCOPE_CYCLE_COUNTER(STAT_FAsyncLoadingThread_ProcessAsyncLoading);
-
-		bool bDepthFirst = false;
-
 		// We need to loop as the function has to handle finish loading everything given no time limit
 		// like e.g. when called from FlushAsyncLoading.
-		for (int32 PackageIndex = 0; ((bDepthFirst && LoadingState == EAsyncPackageState::Complete) || (!bDepthFirst && LoadingState != EAsyncPackageState::TimeOut)) && PackageIndex < AsyncPackages.Num(); ++PackageIndex)
+		for (int32 PackageIndex = 0; LoadingState != EAsyncPackageState::TimeOut && PackageIndex < AsyncPackages.Num(); ++PackageIndex)
 		{
 			SCOPED_LOADTIMER(ProcessAsyncLoadingTime);
 			OutPackagesProcessed++;
@@ -4203,17 +4200,9 @@ EAsyncPackageState::Type FAsyncLoadingThread::ProcessAsyncLoading(int32& OutPack
 			}
 			else if (Package->HasFinishedLoading() == false)
 			{
-				if (GEventDrivenLoaderEnabled)
-				{
-					LoadingState = EAsyncPackageState::PendingImports;
-				}
-				else
-				// @todo: Guard against recursively re-entering?
-				{
-					// Package tick returns EAsyncPackageState::Complete on completion.
-					// We only tick packages that have not yet been loaded.
-					LoadingState = Package->TickAsyncPackage(bUseTimeLimit, bUseFullTimeLimit, TimeLimit, FlushTree);
-				}
+				// Package tick returns EAsyncPackageState::Complete on completion.
+				// We only tick packages that have not yet been loaded.
+				LoadingState = Package->TickAsyncPackage(bUseTimeLimit, bUseFullTimeLimit, TimeLimit, FlushTree);
 			}
 			else
 			{
@@ -4352,8 +4341,8 @@ EAsyncPackageState::Type FAsyncLoadingThread::ProcessLoadedPackages(bool bUseTim
 	}
 #endif
 
-		
-	bDidSomething = LoadedPackagesToProcess.Num() > 0;
+
+	bDidSomething = bDidSomething || LoadedPackagesToProcess.Num() > 0;
 
 	TArray<FAsyncPackage*, TInlineAllocator<4>> CompletedPackages;
 	for (int32 PackageIndex = 0; PackageIndex < LoadedPackagesToProcess.Num() && !IsAsyncLoadingSuspendedInternal(); ++PackageIndex)
@@ -4391,9 +4380,9 @@ EAsyncPackageState::Type FAsyncLoadingThread::ProcessLoadedPackages(bool bUseTim
 							// This may be slow, hence we only do it in the editor
 							Package->FlushObjectLinkerCache();
 						}
-						// Detach linker in mutex scope to make sure that if something requests this package
-						// before it's been deleted does not try to associate the new async package with the old linker
-						// while this async package is still bound to it.
+						// Detach linker in mutex scope to make sure that if something requests the packagename again
+						// before this old async package has been deleted we do not try to associate the new async package
+						// with the linker while this async package is still bound to it.
 						Package->DetachLinker();
 					}
 
@@ -4430,24 +4419,40 @@ EAsyncPackageState::Type FAsyncLoadingThread::ProcessLoadedPackages(bool bUseTim
 		}
 	}
 
-	// Call callbacks in a batch in a stack-local array. This is to ensure that callbacks that trigger
-	// on each package load and call FlushAsyncLoading do not stack overflow by adding one FlushAsyncLoading
-	// call per LoadedPackageToProcess onto the stack
+	// Call callbacks in a batch in a stack-local array. This is to ensure that callbacks that recursively enter this function
+	// via FlushAsyncLoading do not stack overflow by doing one recursive call per LoadedPackageToProcess
+	TSet<UPackage*> CompletedUPackages;
 	for (FAsyncPackage* Package : CompletedPackages)
 	{
 		// Call external callbacks
 		const bool bInternalCallbacks = false;
 		const EAsyncLoadingResult::Type LoadingResult = Package->HasLoadFailed() ? EAsyncLoadingResult::Failed : EAsyncLoadingResult::Succeeded;
 		Package->CallCompletionCallbacks(bInternalCallbacks, LoadingResult);
+		
 #if WITH_EDITOR
-		// In the editor we need to find any assets and add them to list for later callback
-		Package->GetLoadedAssets(LoadedAssets);
+		if (GIsEditor)
+		{
+			// In the editor we need to find any assets and packages and add them to list for later callback
+			Package->GetLoadedAssetsAndPackages(LoadedAssets, CompletedUPackages);
+		}
 #endif
 		// We don't need the package anymore
 		PackagesToDelete.AddUnique(Package);
 		Package->MarkRequestIDsAsComplete();
-
 	}
+#if WITH_EDITOR
+	if (GIsEditor)
+	{
+		// Call the global delegate for package endloads and set the bHasBeenLoaded flag that is used
+		// to check which packages have reached this state
+		for (UPackage* CompletedUPackage : CompletedUPackages)
+		{
+			CompletedUPackage->SetHasBeenEndLoaded(true);
+		}
+		FCoreUObjectDelegates::OnEndLoadPackage.Broadcast(CompletedUPackages.Array());
+	}
+#endif
+
 	bDidSomething = bDidSomething || PackagesToDelete.Num() > 0;
 
 	// Delete packages we're done processing and are no longer dependencies of anything else
@@ -4509,7 +4514,7 @@ EAsyncPackageState::Type FAsyncLoadingThread::ProcessLoadedPackages(bool bUseTim
 
 #if WITH_EDITOR
 		// In editor builds, call the asset load callback. This happens in both editor and standalone to match EndLoad
-		TArray<FWeakObjectPtr> TempLoadedAssets = LoadedAssets;
+		TSet<FWeakObjectPtr> TempLoadedAssets = LoadedAssets;
 		LoadedAssets.Reset();
 
 		// Make a copy because LoadedAssets could be modified by one of the OnAssetLoaded callbacks
@@ -4548,7 +4553,7 @@ EAsyncPackageState::Type FAsyncLoadingThread::TickAsyncLoading(bool bUseTimeLimi
 		double TimeLimitUsedForProcessLoaded = 0;
 
 		// First make sure there's no objects pending to be unhashed. This is important in uncooked builds since we don't 
-		// detach linkers immediately there and we may end up in getting unreachable objects from Linkers in CreateImports
+		// detach linkers immediately and we may otherwise end up getting unreachable objects from Linkers in CreateImports
 		if (FPlatformProperties::RequiresCookedData() == false && IsIncrementalUnhashPending() && IsAsyncLoadingPackages())
 		{
 			// Call ConditionalBeginDestroy on all pending objects. CBD is where linkers get detached from objects.
@@ -5395,13 +5400,22 @@ void FAsyncPackage::FlushObjectLinkerCache()
 }
 
 #if WITH_EDITOR 
-void FAsyncPackage::GetLoadedAssets(TArray<FWeakObjectPtr>& AssetList)
+void FAsyncPackage::GetLoadedAssetsAndPackages(TSet<FWeakObjectPtr>& AssetList, TSet<UPackage*>& PackageList)
 {
 	for (UObject* Obj : PackageObjLoaded)
 	{
-		if (IsValid(Obj) && Obj->IsAsset())
+		if (!IsValid(Obj))
+		{
+			continue;
+		}
+		if (Obj->IsAsset())
 		{
-			AssetList.AddUnique(Obj);
+			AssetList.Add(Obj);
+		}
+		UPackage* Package = Obj->GetPackage();
+		if (Package && !Package->HasAnyFlags(RF_Transient) && !Package->HasAnyPackageFlags(PKG_InMemoryOnly))
+		{
+			PackageList.Add(Package);
 		}
 	}
 }
@@ -6088,7 +6102,7 @@ EAsyncPackageState::Type FAsyncPackage::LoadImports(FFlushTree* FlushTree)
 			continue;
 		}
 
-		// Our import package name is the import name or the specified package name when the object isn't a package
+		// Our import package name is the import name itself if import is a package, else the specified package name
 		const FLinkerInstancingContext& InstancingContext = Linker->GetInstancingContext();
 		FName ImportToLoad = !Import->HasPackageName() ? Import->ObjectName : Import->GetPackageName();
 		FName ImportPackageFName = InstancingContext.Remap(ImportToLoad);
@@ -6791,7 +6805,7 @@ EAsyncPackageState::Type FAsyncPackage::FinishObjects()
 
 void FAsyncPackage::CloseDelayedLinkers()
 {
-	// Close any linkers that have been open as a result of blocking load while async loading
+	// Close any linkers that have been opened as a result of blocking loads while async loading
 	for (FLinkerLoad* LinkerToClose : DelayedLinkerClosePackages)
 	{
 		if (LinkerToClose->LinkerRoot != nullptr)
@@ -6832,15 +6846,6 @@ void FAsyncPackage::CallCompletionCallbacks(bool bInternal, EAsyncLoadingResult:
 	checkSlow(bInternal || !IsInAsyncLoadingThread());
 
 	UPackage* LoadedPackage = (!bLoadHasFailed) ? LinkerRoot : nullptr;
-#if WITH_EDITOR
-	if (!bInternal && GIsEditor && LoadedPackage)
-	{
-		// Call the global delegate for package endloads, which is a kind of external completioncallback,
-		// and set the bHasBeenLoaded flag that is used to check which packages have reached this state
-		LoadedPackage->bHasBeenEndLoaded = true;
-		FCoreUObjectDelegates::OnEndLoadPackage.Broadcast(TConstArrayView<UPackage*> { LoadedPackage });
-	}
-#endif
 	for (FCompletionCallback& CompletionCallback : CompletionCallbacks)
 	{
 		if (CompletionCallback.bIsInternal == bInternal && !CompletionCallback.bCalled)
@@ -6883,7 +6888,7 @@ void FAsyncPackage::Cancel()
 
 	FUObjectSerializeContext* LoadContext = GetSerializeContext();
 	if (LoadContext)
-		{			
+	{
 		TArray<UObject*>& ThreadObjLoaded = LoadContext->PRIVATE_GetObjectsLoadedInternalUseOnly();
 		if (ThreadObjLoaded.Num())
 		{
@@ -7012,7 +7017,7 @@ void FAsyncLoadingThread::FlushLoading(int32 PackageID)
 
 		double StartTime = FPlatformTime::Seconds();
 
-		// Flush async loaders by not using a time limit. Needed for e.g. garbage collection.
+		// Flush async loaders without using a time limit. Needed for e.g. garbage collection.
 		{
 			TUniquePtr<FFlushTree> FlushTree;
 			if (PackageID != INDEX_NONE)
diff --git a/Engine/Source/Runtime/CoreUObject/Private/Serialization/AsyncLoading.h b/Engine/Source/Runtime/CoreUObject/Private/Serialization/AsyncLoading.h
--- a/Engine/Source/Runtime/CoreUObject/Private/Serialization/AsyncLoading.h
+++ b/Engine/Source/Runtime/CoreUObject/Private/Serialization/AsyncLoading.h
@@ -496,8 +496,8 @@ struct FAsyncPackage : public FGCObject
 	UPackage* GetLoadedPackage();
 
 #if WITH_EDITOR
-	/** Gets all assets loaded by this async package, used in the editor */
-	void GetLoadedAssets(TArray<FWeakObjectPtr>& AssetList);
+	/** Gets all assets and packages loaded by this async package, used in the editor */
+	void GetLoadedAssetsAndPackages(TSet<FWeakObjectPtr>& AssetList, TSet<UPackage*>& PackageList);
 #endif
 
 	/** Checks if all dependencies (imported packages) of this package have been fully loaded */
diff --git a/Engine/Source/Runtime/CoreUObject/Private/Serialization/AsyncLoading2.cpp b/Engine/Source/Runtime/CoreUObject/Private/Serialization/AsyncLoading2.cpp
--- a/Engine/Source/Runtime/CoreUObject/Private/Serialization/AsyncLoading2.cpp
+++ b/Engine/Source/Runtime/CoreUObject/Private/Serialization/AsyncLoading2.cpp
@@ -1743,7 +1743,7 @@ struct FAsyncPackage2
 	void StartLoading(FIoBatch& IoBatch);
 
 #if WITH_IOSTORE_IN_EDITOR
-	void GetLoadedAssets(TArray<FWeakObjectPtr>& AssetList);
+	void GetLoadedAssetsAndPackages(TSet<FWeakObjectPtr>& AssetList, TSet<UPackage*>& PackageList);
 #endif
 
 private:
@@ -2096,6 +2096,10 @@ private:
 	TArray<FAsyncPackage2*> LoadedPackagesToProcess;
 	/** [GAME THREAD] Game thread CompletedPackages list */
 	TArray<FAsyncPackage2*> CompletedPackages;
+#if WITH_IOSTORE_IN_EDITOR
+	/** [GAME THREAD] Game thread LoadedAssets list */
+	TSet<FWeakObjectPtr> LoadedAssets;
+#endif
 	/** [ASYNC/GAME THREAD] Packages to be deleted from async thread */
 	TQueue<FAsyncPackage2*, EQueueMode::Spsc> DeferredDeletePackages;
 	
@@ -4578,6 +4582,9 @@ EAsyncPackageState::Type FAsyncLoadingThread2::ProcessLoadedPackagesFromGameThre
 
 		bLocalDidSomething |= LoadedPackagesToProcess.Num() > 0;
 		TArray<FAsyncPackage2*, TInlineAllocator<4>> PackagesReadyForCallback;
+#if WITH_IOSTORE_IN_EDITOR
+		TSet<UPackage*> CompletedUPackages;
+#endif
 		for (int32 PackageIndex = 0; PackageIndex < LoadedPackagesToProcess.Num(); ++PackageIndex)
 		{
 			SCOPED_LOADTIMER(ProcessLoadedPackagesTime);
@@ -4657,6 +4664,10 @@ EAsyncPackageState::Type FAsyncLoadingThread2::ProcessLoadedPackagesFromGameThre
 				AsyncPackageLookup.Remove(Package->Desc.UPackageId);
 				if (!Package->bLoadHasFailed)
 				{
+#if WITH_IOSTORE_IN_EDITOR
+					// In the editor we need to find any assets and packages and add them to list for later callback
+					Package->GetLoadedAssetsAndPackages(LoadedAssets, CompletedUPackages);
+#endif
 					Package->ClearConstructedObjects();
 				}
 			}
@@ -4701,6 +4712,15 @@ EAsyncPackageState::Type FAsyncLoadingThread2::ProcessLoadedPackagesFromGameThre
 			UE_ASYNC_PACKAGE_LOG(Verbose, Package->Desc, TEXT("GameThread: LoadCompleted"),
 				TEXT("All loading of package is done, and the async package and load request will be deleted."));
 		}
+#if WITH_IOSTORE_IN_EDITOR
+		// Call the global delegate for package endloads and set the bHasBeenLoaded flag that is used to
+		// check which packages have reached this state
+		for (UPackage* CompletedUPackage : CompletedUPackages)
+		{
+			CompletedUPackage->SetHasBeenEndLoaded(true);
+		}
+		FCoreUObjectDelegates::OnEndLoadPackage.Broadcast(CompletedUPackages.Array());
+#endif
 		
 		{
 			TArray<FFailedPackageRequest> LocalFailedPackageRequests;
@@ -4779,6 +4799,22 @@ EAsyncPackageState::Type FAsyncLoadingThread2::ProcessLoadedPackagesFromGameThre
 
 	if (Result == EAsyncPackageState::Complete)
 	{
+#if WITH_IOSTORE_IN_EDITOR
+		// In editor builds, call the asset load callback. This happens in both editor and standalone to match EndLoad
+		TSet<FWeakObjectPtr> TempLoadedAssets = LoadedAssets;
+		LoadedAssets.Reset();
+
+		// Make a copy because LoadedAssets could be modified by one of the OnAssetLoaded callbacks
+		for (const FWeakObjectPtr& WeakAsset : TempLoadedAssets)
+		{
+			// It may have been unloaded/marked pending kill since being added, ignore those cases
+			if (UObject* LoadedAsset = WeakAsset.Get())
+			{
+				FCoreUObjectDelegates::OnAssetLoaded.Broadcast(LoadedAsset);
+			}
+		}
+#endif
+
 		// We're not done until all packages have been deleted
 		Result = CompletedPackages.Num() ? EAsyncPackageState::PendingImports  : EAsyncPackageState::Complete;
 		if (Result == EAsyncPackageState::Complete && ThreadState.HasDeferredFrees())
@@ -5597,8 +5633,21 @@ double FAsyncPackage2::GetLoadStartTime() const
 }
 
 #if WITH_IOSTORE_IN_EDITOR
-void FAsyncPackage2::GetLoadedAssets(TArray<FWeakObjectPtr>& AssetList)
+void FAsyncPackage2::GetLoadedAssetsAndPackages(TSet<FWeakObjectPtr>& AssetList, TSet<UPackage*>& PackageList)
 {
+	for (UObject* Object : ConstructedObjects)
+	{
+		if (IsValid(Object) && Object->IsAsset())
+		{
+			AssetList.Add(Object);
+		}
+	}
+
+	// All ConstructedObjects belong to this package, so we only have to consider the single package in this->LinkerRoot
+	if (LinkerRoot && !LinkerRoot->HasAnyFlags(RF_Transient) && !LinkerRoot->HasAnyPackageFlags(PKG_InMemoryOnly))
+	{
+		PackageList.Add(LinkerRoot);
+	}
 }
 #endif
 
@@ -5791,15 +5840,6 @@ void FAsyncPackage2::CallCompletionCallbacks(EAsyncLoadingResult::Type LoadingRe
 	checkSlow(IsInGameThread());
 
 	UPackage* LoadedPackage = (!bLoadHasFailed) ? LinkerRoot : nullptr;
-#if WITH_EDITOR
-	if (GIsEditor && LoadedPackage)
-	{
-		// Call the global delegate for package endloads, which is a kind of completioncallback,
-		// and set the bHasBeenLoaded flag that is used to check which packages have reached this state
-		LoadedPackage->bHasBeenEndLoaded = true;
-		FCoreUObjectDelegates::OnEndLoadPackage.Broadcast(TConstArrayView<UPackage*> { LoadedPackage });
-	}
-#endif
 	for (FCompletionCallback& CompletionCallback : CompletionCallbacks)
 	{
 		CompletionCallback->ExecuteIfBound(Desc.UPackageName, LoadedPackage, LoadingResult);
diff --git a/Engine/Source/Runtime/CoreUObject/Private/Serialization/AsyncLoadingThread.h b/Engine/Source/Runtime/CoreUObject/Private/Serialization/AsyncLoadingThread.h
--- a/Engine/Source/Runtime/CoreUObject/Private/Serialization/AsyncLoadingThread.h
+++ b/Engine/Source/Runtime/CoreUObject/Private/Serialization/AsyncLoadingThread.h
@@ -189,7 +189,7 @@ class FAsyncLoadingThread final : public FRunnable, public IAsyncPackageLoader
 	TArray<FAsyncPackage*> PackagesToDelete;
 	TMap<FName, FAsyncPackage*> LoadedPackagesToProcessNameLookup;
 #if WITH_EDITOR
-	TArray<FWeakObjectPtr> LoadedAssets;
+	TSet<FWeakObjectPtr> LoadedAssets;
 #endif
 #if THREADSAFE_UOBJECTS
 	/** [ASYNC/GAME THREAD] Critical section for LoadedPackagesToProcess list. 
diff --git a/Engine/Source/Runtime/CoreUObject/Private/UObject/UObjectGlobals.cpp b/Engine/Source/Runtime/CoreUObject/Private/UObject/UObjectGlobals.cpp
--- a/Engine/Source/Runtime/CoreUObject/Private/UObject/UObjectGlobals.cpp
+++ b/Engine/Source/Runtime/CoreUObject/Private/UObject/UObjectGlobals.cpp
@@ -1099,15 +1099,19 @@ public:
 static int32 GGameThreadLoadCounter = 0;
 
 /** Notify delegate listeners of all the packages that loaded; called only once per explicit call to LoadPackage. */
-void BroadcastEndLoad(TArrayView<UPackage*> LoadedPackages)
+void BroadcastEndLoad(TArray<UPackage*>&& LoadedPackages)
 {
 #if WITH_EDITOR
 	// check(IsInGameThread()) was called by the caller, but we still need to test !IsInAsyncLoadingThread to exclude that callsite when the engine is single-threaded
 	if (GIsEditor && !IsInAsyncLoadingThread() && GGameThreadLoadCounter == 0)
 	{
+		LoadedPackages.RemoveAllSwap([](UPackage* Package)
+			{
+				return Package->HasAnyFlags(RF_Transient) || Package->HasAnyPackageFlags(PKG_InMemoryOnly);
+			});
 		for (UPackage* LoadedPackage : LoadedPackages)
 		{
-			LoadedPackage->bHasBeenEndLoaded = true;
+			LoadedPackage->SetHasBeenEndLoaded(true);
 		}
 		FCoreUObjectDelegates::OnEndLoadPackage.Broadcast(LoadedPackages);
 	}
@@ -1250,7 +1254,7 @@ UPackage* LoadPackageInternal(UPackage* InOuter, const FPackagePath& PackagePath
 		if (!Linker)
 		{
 			EndLoad(LoadContext, &LoadedPackages);
-			BroadcastEndLoad(LoadedPackages);
+			BroadcastEndLoad(MoveTemp(LoadedPackages));
 			return nullptr;
 		}
 
@@ -1296,7 +1300,7 @@ UPackage* LoadPackageInternal(UPackage* InOuter, const FPackagePath& PackagePath
 			// The linker is associated with a package that has already been loaded.
 			// Loading packages that have already been loaded is unsupported.
 			EndLoadAndCopyLocalizationGatherFlag();	
-			BroadcastEndLoad(LoadedPackages);
+			BroadcastEndLoad(MoveTemp(LoadedPackages));
 			return Result;
 		}
 
@@ -1441,7 +1445,7 @@ UPackage* LoadPackageInternal(UPackage* InOuter, const FPackagePath& PackagePath
 		Result->SetFlags(RF_WasLoaded);
 	}
 
-	BroadcastEndLoad(LoadedPackages);
+	BroadcastEndLoad(MoveTemp(LoadedPackages));
 	return Result;
 }
 
diff --git a/Engine/Source/Runtime/CoreUObject/Public/UObject/Package.h b/Engine/Source/Runtime/CoreUObject/Public/UObject/Package.h
--- a/Engine/Source/Runtime/CoreUObject/Public/UObject/Package.h
+++ b/Engine/Source/Runtime/CoreUObject/Public/UObject/Package.h
@@ -177,7 +177,17 @@ public:
 
 #if WITH_EDITORONLY_DATA
 	/** Set to true after serialization and postload and before returning from LoadPackage or calling load completion delegate, for loaded packages.			*/
+private:
 	uint8 bHasBeenEndLoaded : 1;
+public:
+	bool GetHasBeenEndLoaded() const
+	{
+		return bHasBeenEndLoaded != 0;
+	}
+	void SetHasBeenEndLoaded(bool bValue)
+	{
+		bHasBeenEndLoaded = bValue ? 1 : 0;
+	}
 #endif
 
 	/**
