Fix explicit range/bound UHT-only type decls to match member access specifiers in their native, non-UHT C++ counterparts.

diff --git a/Engine/Source/Runtime/CoreUObject/Public/UObject/NoExportTypes.h b/Engine/Source/Runtime/CoreUObject/Public/UObject/NoExportTypes.h
--- a/Engine/Source/Runtime/CoreUObject/Public/UObject/NoExportTypes.h
+++ b/Engine/Source/Runtime/CoreUObject/Public/UObject/NoExportTypes.h
@@ -1388,92 +1388,104 @@ namespace ERangeBoundTypes
 /**
  * Defines a single bound for a range of values.
  * @note This is a mirror of TRangeBound<float>, defined in RangeBound.h
+ * @note Fields are private to match the C++ declaration in the header above.
  */
 USTRUCT(noexport, BlueprintType)
 struct FFloatRangeBound
 {
+private:
 	/** Holds the type of the bound. */
-	UPROPERTY(EditAnywhere, BlueprintReadWrite, Category=Range)
+	UPROPERTY(EditAnywhere, BlueprintReadWrite, Category=Range, meta=(AllowPrivateAccess="true"))
 	TEnumAsByte<ERangeBoundTypes::Type> Type;
 
 	/** Holds the bound's value. */
-	UPROPERTY(EditAnywhere, BlueprintReadWrite, Category=Range)
+	UPROPERTY(EditAnywhere, BlueprintReadWrite, Category=Range, meta=(AllowPrivateAccess="true"))
 	float Value;
 };
 
 /**
  * A contiguous set of floats described by lower and upper bound values.
  * @note This is a mirror of TRange<float>, defined in Range.h
+ * @note Fields are private to match the C++ declaration in the header above.
  */
 USTRUCT(noexport, BlueprintType)
 struct FFloatRange
 {
+private:
 	/** Holds the range's lower bound. */
-	UPROPERTY(EditAnywhere, BlueprintReadWrite, Category=Range)
+	UPROPERTY(EditAnywhere, BlueprintReadWrite, Category=Range, meta=(AllowPrivateAccess="true"))
 	FFloatRangeBound LowerBound;
 
 	/** Holds the range's upper bound. */
-	UPROPERTY(EditAnywhere, BlueprintReadWrite, Category=Range)
+	UPROPERTY(EditAnywhere, BlueprintReadWrite, Category=Range, meta=(AllowPrivateAccess="true"))
 	FFloatRangeBound UpperBound;
 };
 
 /**
  * Defines a single bound for a range of values.
  * @note This is a mirror of TRangeBound<int32>, defined in RangeBound.h
+ * @note Fields are private to match the C++ declaration in the header above.
  */
 USTRUCT(noexport, BlueprintType)
 struct FInt32RangeBound
 {
+private:
 	/** Holds the type of the bound. */
-	UPROPERTY(EditAnywhere, BlueprintReadWrite, Category = Range)
+	UPROPERTY(EditAnywhere, BlueprintReadWrite, Category=Range, meta=(AllowPrivateAccess="true"))
 	TEnumAsByte<ERangeBoundTypes::Type> Type;
 
 	/** Holds the bound's value. */
-	UPROPERTY(EditAnywhere, BlueprintReadWrite, Category = Range)
+	UPROPERTY(EditAnywhere, BlueprintReadWrite, Category=Range, meta=(AllowPrivateAccess="true"))
 	int32 Value;
 };
 
 /**
  * A contiguous set of floats described by lower and upper bound values.
  * @note This is a mirror of TRange<int32>, defined in Range.h
+ * @note Fields are private to match the C++ declaration in the header above.
  */
 USTRUCT(noexport, BlueprintType)
 struct FInt32Range
 {
+private:
 	/** Holds the range's lower bound. */
-	UPROPERTY(EditAnywhere, BlueprintReadWrite, Category = Range)
+	UPROPERTY(EditAnywhere, BlueprintReadWrite, Category=Range, meta=(AllowPrivateAccess="true"))
 	FInt32RangeBound LowerBound;
 
 	/** Holds the range's upper bound. */
-	UPROPERTY(EditAnywhere, BlueprintReadWrite, Category = Range)
+	UPROPERTY(EditAnywhere, BlueprintReadWrite, Category=Range, meta=(AllowPrivateAccess="true"))
 	FInt32RangeBound UpperBound;
 };
 
 /**
  * Defines a single bound for a range of frame numbers.
  * @note This is a mirror of TRangeBound<FFrameNumber>, defined in RangeBound.h
+ * @note Fields are private to match the C++ declaration in the header above.
  */
 USTRUCT(noexport, BlueprintType)
 struct FFrameNumberRangeBound
 {
-	UPROPERTY(EditAnywhere, BlueprintReadWrite, Category = Range)
+private:
+	UPROPERTY(EditAnywhere, BlueprintReadWrite, Category=Range, meta=(AllowPrivateAccess="true"))
 	TEnumAsByte<ERangeBoundTypes::Type> Type;
 
-	UPROPERTY(EditAnywhere, BlueprintReadWrite, Category = Range)
+	UPROPERTY(EditAnywhere, BlueprintReadWrite, Category=Range, meta=(AllowPrivateAccess="true"))
 	FFrameNumber Value;
 };
 
 /**
  * A contiguous set of frame numbers described by lower and upper bound values.
  * @note This is a mirror of TRange<FFrameNumber>, defined in Range.h
+ * @note Fields are private to match the C++ declaration in the header above.
  */
 USTRUCT(noexport, BlueprintType)
 struct FFrameNumberRange
 {
-	UPROPERTY(EditAnywhere, BlueprintReadWrite, Category = Range)
+private:
+	UPROPERTY(EditAnywhere, BlueprintReadWrite, Category=Range, meta=(AllowPrivateAccess="true"))
 	FFrameNumberRangeBound LowerBound;
 
-	UPROPERTY(EditAnywhere, BlueprintReadWrite, Category = Range)
+	UPROPERTY(EditAnywhere, BlueprintReadWrite, Category=Range, meta=(AllowPrivateAccess="true"))
 	FFrameNumberRangeBound UpperBound;
 };
 
