Integrate CL 16307919 from Dev-EngineMerge: Prevent async loading during hotfixing / content patching.

diff --git a/Engine/Plugins/Online/OnlineFramework/Source/Hotfix/Private/OnlineHotfixManager.cpp b/Engine/Plugins/Online/OnlineFramework/Source/Hotfix/Private/OnlineHotfixManager.cpp
--- a/Engine/Plugins/Online/OnlineFramework/Source/Hotfix/Private/OnlineHotfixManager.cpp
+++ b/Engine/Plugins/Online/OnlineFramework/Source/Hotfix/Private/OnlineHotfixManager.cpp
@@ -23,6 +23,8 @@
 #include "Curves/CurveFloat.h"
 #include "Engine/BlueprintGeneratedClass.h"
 
+#include "Serialization/AsyncLoadingLock.h"
+
 #ifdef WITH_ONLINETRACING
 #include "OnlineTracingModule.h"
 #endif
@@ -177,6 +179,15 @@ UOnlineHotfixManager::UOnlineHotfixManager() :
 	GameContentPath = FString() / FApp::GetProjectName() / TEXT("Content");
 }
 
+UOnlineHotfixManager::UOnlineHotfixManager(FVTableHelper& Helper)
+	: Super(Helper)
+{
+}
+
+UOnlineHotfixManager::~UOnlineHotfixManager()
+{
+}
+
 UOnlineHotfixManager* UOnlineHotfixManager::Get(UWorld* World)
 {
 	UOnlineHotfixManager* DefaultObject = UOnlineHotfixManager::StaticClass()->GetDefaultObject<UOnlineHotfixManager>();
@@ -258,6 +269,7 @@ void UOnlineHotfixManager::Cleanup()
 	}
 	OnlineTitleFile = nullptr;
 	bHotfixingInProgress = false;
+	AsyncLock = nullptr;
 }
 
 void UOnlineHotfixManager::StartHotfixProcess()
@@ -419,14 +431,34 @@ void UOnlineHotfixManager::OnEnumerateFilesComplete(bool bWasSuccessful, const F
 		{
 			if (RemovedHotfixFileList.Num() > 0)
 			{
-				// No changes, just reverts
-				// Perform any undo operations needed
-				RestoreBackupIniFiles();
-				UnmountHotfixFiles();
-			}
+				UE_LOG(LogHotfixManager, Display, TEXT("Files have been removed since last check. Reverting."));
+
+				// Prevent async loading while reverting hotfixes.
+				check(!AsyncLock);
+				AsyncLock = MakeUnique<FAsyncLoadingLock>(TEXT("RevertHotfix"));
+				AsyncLock->Acquire(
+					FOnAsyncLoadingLockAcquired::CreateWeakLambda(
+						this,
+						[this]()
+						{
+							{
+								// No changes, just reverts
+								// Perform any undo operations needed
+								RestoreBackupIniFiles();
+								UnmountHotfixFiles();
+							}
+
+							AsyncLock->Release();
 
-			UE_LOG(LogHotfixManager, Display, TEXT("Returned hotfix data is the same as last application, skipping the apply phase"));
-			TriggerHotfixComplete(EHotfixResult::SuccessNoChange);
+							// Trigger completion outside of async lock.
+							TriggerHotfixComplete(EHotfixResult::SuccessNoChange);
+						}));
+			}
+			else
+			{
+				UE_LOG(LogHotfixManager, Display, TEXT("Returned hotfix data is the same as last application, skipping the apply phase"));
+				TriggerHotfixComplete(EHotfixResult::SuccessNoChange);
+			}
 		}
 	}
 	else
@@ -648,8 +680,25 @@ void UOnlineHotfixManager::OnReadFileComplete(bool bWasSuccessful, const FString
 			PendingHotfixFiles.Remove(FileName);
 			if (PendingHotfixFiles.Num() == 0)
 			{
-				// All files have been downloaded so now apply the files
-				ApplyHotfix();
+				// Prevent async loading while applying hotfixes.
+				check(!AsyncLock);
+				AsyncLock = MakeUnique<FAsyncLoadingLock>(TEXT("ApplyHotfix"));
+				AsyncLock->Acquire(
+					FOnAsyncLoadingLockAcquired::CreateWeakLambda(
+						this,
+						[this]()
+						{
+							EHotfixResult Result = EHotfixResult::Success;
+
+							{
+								Result = ApplyHotfix();
+							}
+
+							AsyncLock->Release();
+
+							// Trigger completion outside of async lock.
+							TriggerHotfixComplete(Result);
+						}));
 			}
 		}
 		else
@@ -668,7 +717,7 @@ void UOnlineHotfixManager::UpdateProgress(uint32 FileCount, uint64 UpdateSize)
 	TriggerOnHotfixProgressDelegates(NumDownloaded, TotalFiles, NumBytes, TotalBytes);
 }
 
-void UOnlineHotfixManager::ApplyHotfix()
+EHotfixResult UOnlineHotfixManager::ApplyHotfix()
 {
 	// Perform any undo operations needed
 	// This occurs same frame as the application of new hotfixes
@@ -680,8 +729,7 @@ void UOnlineHotfixManager::ApplyHotfix()
 		if (!ApplyHotfixProcessing(FileHeader))
 		{
 			UE_LOG(LogHotfixManager, Error, TEXT("Couldn't apply hotfix file (%s)"), *FileHeader.FileName);
-			TriggerHotfixComplete(EHotfixResult::Failed);
-			return;
+			return EHotfixResult::Failed;
 		}
 		// Let anyone listening know we just processed this file
 		TriggerOnHotfixProcessedFileDelegates(FileHeader.FileName, GetCachedDirectory() / FileHeader.DLName);
@@ -698,11 +746,16 @@ void UOnlineHotfixManager::ApplyHotfix()
 		UE_LOG(LogHotfixManager, Display, TEXT("Hotfix has detected PAK files containing currently loaded maps, so a level load is needed"));
 		Result = EHotfixResult::SuccessNeedsReload;
 	}
-	TriggerHotfixComplete(Result);
+
+	return Result;
 }
 
 void UOnlineHotfixManager::TriggerHotfixComplete(EHotfixResult HotfixResult)
 {
+	// Flush the async loader to prevent updating data while it is accessed for loading.
+	// This may cause a hitch in some failure cases.
+	FlushAsyncLoading();
+
 	if (HotfixResult != EHotfixResult::Failed && HotfixResult != EHotfixResult::SuccessNoChange)
 	{
 		PatchAssetsFromIniFiles();
diff --git a/Engine/Plugins/Online/OnlineFramework/Source/Hotfix/Public/OnlineHotfixManager.h b/Engine/Plugins/Online/OnlineFramework/Source/Hotfix/Public/OnlineHotfixManager.h
--- a/Engine/Plugins/Online/OnlineFramework/Source/Hotfix/Public/OnlineHotfixManager.h
+++ b/Engine/Plugins/Online/OnlineFramework/Source/Hotfix/Public/OnlineHotfixManager.h
@@ -15,6 +15,8 @@
 
 HOTFIX_API DECLARE_LOG_CATEGORY_EXTERN(LogHotfixManager, Display, All);
 
+class FAsyncLoadingLock;
+
 UENUM()
 enum class EHotfixResult : uint8
 {
@@ -157,6 +159,8 @@ protected:
 	uint64 NumBytes;
 	/** Some title file interfaces aren't re-entrant so handle it ourselves */
 	bool bHotfixingInProgress;
+	/** Lock to prevent async loading during the hotfixing process. */
+	TUniquePtr<FAsyncLoadingLock> AsyncLock;
 	/** Set to true if any PAK file contains an update to a level that is currently loaded */
 	bool bHotfixNeedsMapReload;
 #if !UE_BUILD_SHIPPING
@@ -174,7 +178,7 @@ protected:
 	virtual void Init();
 	virtual void Cleanup();
 	/** Looks at each file returned via the hotfix and processes them */
-	void ApplyHotfix();
+	EHotfixResult ApplyHotfix();
 	/** Cleans up and fires the delegate indicating it's done */
 	void TriggerHotfixComplete(EHotfixResult HotfixResult);
 	/** Checks each file listed to see if it is a hotfix file to process */
@@ -311,6 +315,8 @@ protected:
 
 public:
 	UOnlineHotfixManager();
+	UOnlineHotfixManager(FVTableHelper& Helper);
+	virtual ~UOnlineHotfixManager();
 
 	/** Tells the hotfix manager which OSS to use. Uses the default if empty */
 	UPROPERTY(Config)
diff --git a/Engine/Source/Runtime/CoreUObject/Private/Serialization/AsyncLoadingLock.cpp b/Engine/Source/Runtime/CoreUObject/Private/Serialization/AsyncLoadingLock.cpp
new file mode 100644
--- /dev/null
+++ b/Engine/Source/Runtime/CoreUObject/Private/Serialization/AsyncLoadingLock.cpp
@@ -0,0 +1,170 @@
+// Copyright Epic Games, Inc. All Rights Reserved.
+
+#include "Serialization/AsyncLoadingLock.h"
+
+#include "Containers/Ticker.h"
+#include "Serialization/AsyncPackageLoader.h"
+
+DEFINE_LOG_CATEGORY(AsyncLoadingLock);
+
+bool IsAsyncLoadingCoreUObjectInternal();
+void SuspendAsyncLoadingInternal();
+void ResumeAsyncLoadingInternal();
+
+int32 FAsyncLoadingLock::SuspendCount = 0;
+double FAsyncLoadingLock::SuspendStartTime = 0;
+int32 FAsyncLoadingLock::NextLockId = 0;
+
+FAsyncLoadingLock::FAsyncLoadingLock(const FString& Context)
+	: Context(Context)
+	, LockId(NextLockId++)
+{
+}
+
+FAsyncLoadingLock::FAsyncLoadingLock(FString&& Context)
+	: Context(MoveTemp(Context))
+	, LockId(NextLockId++)
+{
+}
+
+FAsyncLoadingLock::~FAsyncLoadingLock()
+{
+	if (State == LockState::Acquiring)
+	{
+		Abandon();
+	}
+	else if (State == LockState::Acquired)
+	{
+		Release();
+	}
+}
+
+void FAsyncLoadingLock::Acquire(const FOnAsyncLoadingLockAcquired& OnLockAcquired)
+{
+	check(State == LockState::Released);
+	check(IsInGameThread() && !IsInSlateThread());
+
+	OnLockAcquiredDelegate = OnLockAcquired;
+	check(OnLockAcquiredDelegate.IsBound());
+
+	// Check each frame whether async loading has completed.
+	LoadingCompleteCheckDelegateHandle = FTicker::GetCoreTicker().AddTicker(FTickerDelegate::CreateRaw(this, &FAsyncLoadingLock::OnAsyncLoadingCheck));
+
+	// Periodically warn when waiting for async loading to complete.
+	WaitingWarnDelegateHandle = FTicker::GetCoreTicker().AddTicker(FTickerDelegate::CreateRaw(this, &FAsyncLoadingLock::OnAsyncLoadingWarn), LockAcquireWarnIntervalSeconds);
+
+	StartTime = FPlatformTime::Seconds();
+	UE_LOG(AsyncLoadingLock, Verbose, TEXT("Context[%s:%d] Acquiring loading scoped lock."), *Context, LockId);
+
+	State = LockState::Acquiring;
+}
+
+void FAsyncLoadingLock::Release()
+{
+	check(State == LockState::Acquired);
+	check(IsInGameThread() && !IsInSlateThread());
+
+	CleanupTickers();
+
+	const double EndTime = FPlatformTime::Seconds();
+
+	UE_LOG(AsyncLoadingLock, Verbose, TEXT("Context[%s:%d] Releasing async loading lock. TimeLockHeld: %f TotalTime: %f"),
+		*Context, LockId, EndTime - LockAcquiredTime, EndTime - StartTime);
+
+	if (--SuspendCount == 0)
+	{
+		UE_LOG(AsyncLoadingLock, Log, TEXT("Context[%s:%d] Resuming async loading after suspension of %f seconds."), *Context, LockId, EndTime - SuspendStartTime);
+		ResumeAsyncLoadingInternal();
+	}
+	else
+	{
+		UE_LOG(AsyncLoadingLock, Verbose, TEXT("Context[%s:%d] Async loading suspension remains active."), *Context, LockId);
+	}
+
+	State = LockState::Released;
+}
+
+void FAsyncLoadingLock::Abandon()
+{
+	check(State == LockState::Acquiring);
+	check(IsInGameThread() && !IsInSlateThread());
+
+	CleanupTickers();
+
+	const double EndTime = FPlatformTime::Seconds();
+
+	UE_LOG(AsyncLoadingLock, Verbose, TEXT("Context[%s:%d] Abandoning async loading lock. TotalTime: %f"),
+		*Context, LockId, EndTime - StartTime);
+
+	State = LockState::Released;
+}
+
+void FAsyncLoadingLock::CleanupTickers()
+{
+	FTicker::GetCoreTicker().RemoveTicker(LoadingCompleteCheckDelegateHandle);
+	FTicker::GetCoreTicker().RemoveTicker(WaitingWarnDelegateHandle);
+	FTicker::GetCoreTicker().RemoveTicker(LockHeldWarnDelegateHandle);
+}
+
+bool FAsyncLoadingLock::OnAsyncLoadingCheck(float DeltaTime)
+{
+	// Abandon lock if delegate is no longer valid to be called.
+	// Delegates bound to a UObject will fail the below check once they have been marked for delete.
+	if (!OnLockAcquiredDelegate.IsBound())
+	{
+		Abandon();
+		return false;
+	}
+
+	if (IsAsyncLoadingCoreUObjectInternal())
+	{
+		// Continue rescheduling the check.
+		return true;
+	}
+
+	// Async loading completed, remove warning logger.
+	FTicker::GetCoreTicker().RemoveTicker(WaitingWarnDelegateHandle);
+
+	const double CurrentTimeSeconds = FPlatformTime::Seconds();
+	UE_LOG(AsyncLoadingLock, Verbose, TEXT("Context[%s:%d] Lock acquired in %f seconds."), *Context, LockId, CurrentTimeSeconds - StartTime);
+
+	// Mark lock as acquired so the suspension count will be decremented correctly.
+	State = LockState::Acquired;
+	LockAcquiredTime = CurrentTimeSeconds;
+
+	// Async loading is now idle. Suspend async loads.
+	if (SuspendCount++ == 0)
+	{
+		UE_LOG(AsyncLoadingLock, Log, TEXT("Context[%s:%d] Suspending async loading."), *Context, LockId);
+		SuspendStartTime = CurrentTimeSeconds;
+		SuspendAsyncLoadingInternal();
+		UE_LOG(AsyncLoadingLock, Verbose, TEXT("Context[%s:%d] Async loading suspended in %f seconds."), *Context, LockId, FPlatformTime::Seconds() - SuspendStartTime);
+	}
+	else
+	{
+		UE_LOG(AsyncLoadingLock, Verbose, TEXT("Context[%s:%d] Async loading previously suspended."), *Context, LockId);
+	}
+
+	// Periodically warn when the lock has been held for a significant amount of time.
+	LockHeldWarnDelegateHandle = FTicker::GetCoreTicker().AddTicker(FTickerDelegate::CreateRaw(this, &FAsyncLoadingLock::OnLockHeldWarn), LockHeldWarnIntervalSeconds);
+
+	// Signal to user that lock has been acquired.
+	FOnAsyncLoadingLockAcquired OnLockAcquired = MoveTemp(OnLockAcquiredDelegate);
+	OnLockAcquired.ExecuteIfBound();
+	// do not access member variables below this point in case firing the delegate deleted this lock.
+
+	// Unregister check.
+	return false;
+}
+
+bool FAsyncLoadingLock::OnAsyncLoadingWarn(float DeltaTime)
+{
+	UE_LOG(AsyncLoadingLock, Warning, TEXT("Context[%s:%d] Waiting on async loading to complete. Total lock wait time: %f seconds."), *Context, LockId, FPlatformTime::Seconds() - StartTime);
+	return true;
+}
+
+bool FAsyncLoadingLock::OnLockHeldWarn(float DeltaTime)
+{
+	UE_LOG(AsyncLoadingLock, Warning, TEXT("Context[%s:%d] Async loading has been suspended for %f seconds."), *Context, LockId, FPlatformTime::Seconds() - SuspendStartTime);
+	return true;
+}
diff --git a/Engine/Source/Runtime/CoreUObject/Public/Serialization/AsyncLoadingLock.h b/Engine/Source/Runtime/CoreUObject/Public/Serialization/AsyncLoadingLock.h
new file mode 100644
--- /dev/null
+++ b/Engine/Source/Runtime/CoreUObject/Public/Serialization/AsyncLoadingLock.h
@@ -0,0 +1,78 @@
+// Copyright Epic Games, Inc. All Rights Reserved.
+
+#pragma once
+
+#include "Delegates/Delegate.h"
+#include "Logging/LogMacros.h"
+
+DECLARE_LOG_CATEGORY_EXTERN(AsyncLoadingLock, Log, All);
+
+DECLARE_DELEGATE(FOnAsyncLoadingLockAcquired);
+
+/**
+ * Lock for suspending the async loading thread without hitching the game thread. To work properly the lock must be released
+ * from within the lock acquired delegate. Calls to FlushLoading will fail while async loading is suspended.
+ */
+class COREUOBJECT_API FAsyncLoadingLock final : FNoncopyable
+{
+public:
+	FAsyncLoadingLock(const FString& Context);
+	FAsyncLoadingLock(FString&& Context);
+	~FAsyncLoadingLock();
+
+	void Acquire(const FOnAsyncLoadingLockAcquired& OnLockAcquired);
+	void Release();
+
+	int32 GetId() const { return LockId; }
+
+private:
+	FAsyncLoadingLock() = delete;
+
+	void Abandon();
+	void CleanupTickers();
+
+	bool OnAsyncLoadingCheck(float DeltaTime);
+	bool OnAsyncLoadingWarn(float DeltaTime);
+	bool OnLockHeldWarn(float DeltaTime);
+
+	enum class LockState : uint8
+	{
+		Acquiring,
+		Acquired,
+		Released,
+	};
+
+	/** How often to warn while waiting for async loading to complete. */
+	static constexpr float LockAcquireWarnIntervalSeconds = 10.f;
+
+	/** How often to warn while holding the async loading lock. */
+	static constexpr float LockHeldWarnIntervalSeconds = 30.f;
+
+	/** Context for logging */
+	FString Context;
+
+	/** Unique id for lock. */
+	int32 LockId;
+
+	/** Delegate to be fired on lock. */
+	FOnAsyncLoadingLockAcquired OnLockAcquiredDelegate;
+
+	/** Whether async loading was suspended */
+	LockState State = LockState::Released;
+
+	/** Record event timing for logging. */
+	double StartTime = 0;
+	double LockAcquiredTime = 0;
+
+	/** Tick delegates. */
+	FDelegateHandle LoadingCompleteCheckDelegateHandle;
+	FDelegateHandle WaitingWarnDelegateHandle;
+	FDelegateHandle LockHeldWarnDelegateHandle;
+
+	/** Track the number of suspensions to preserve suspended state with multiple active FAsyncLoadingLock instances. */
+	static int32 SuspendCount;
+	static double SuspendStartTime;
+
+	/** Increment lock ids so each instance has a unique id. */
+	static int32 NextLockId;
+};
