#UE If there are are both editor-only and non-editor-only hard refs to an asset, the asset registry will now properly show the package reference as non-editor-only (as opposed to whatever reference it processes first)

diff --git a/Engine/Source/Runtime/CoreUObject/Private/UObject/SavePackage.cpp b/Engine/Source/Runtime/CoreUObject/Private/UObject/SavePackage.cpp
--- a/Engine/Source/Runtime/CoreUObject/Private/UObject/SavePackage.cpp
+++ b/Engine/Source/Runtime/CoreUObject/Private/UObject/SavePackage.cpp
@@ -471,6 +471,18 @@ FArchive& FArchiveSaveTagImports::operator<<( UObject*& Obj )
 				Outer = Outer->GetOuter();
 			}
 
+			bool bIsImport = !Obj->HasAnyMarks(OBJECTMARK_TagExp);
+			if (bIsImport)
+			{
+#if WITH_EDITORONLY_DATA
+				if (!bReferencerIsEditorOnly && !IsEditorOnlyPropertyOnTheStack())
+#endif
+				{
+					// Check every reference to an import for whether the reference is used-in-game and upgrade the _Package_'s reference to used-in-game. Do this before we early exit for imports we have seen before
+					ImportsUsedInGame.Add(Obj);
+				}
+			}
+
 			// We add objects as dependencies even if they're also exports
 			if (!bIsTopLevelPackage && !bIgnoreDependencies)
 			{
@@ -482,7 +494,7 @@ FArchive& FArchiveSaveTagImports::operator<<( UObject*& Obj )
 				DependencyArray.Add(Obj);
 			}
 			
-			if (!Obj->HasAnyMarks(OBJECTMARK_TagExp))  
+			if (bIsImport)
 			{
 				// Add into other imports list unless it's already there
 				if (bIsTopLevelPackage || bIgnoreDependencies)
@@ -497,12 +509,7 @@ FArchive& FArchiveSaveTagImports::operator<<( UObject*& Obj )
 
 				// Mark this object as an import
 				Obj->Mark(OBJECTMARK_TagImp);
-#if WITH_EDITORONLY_DATA
-				if (!bReferencerIsEditorOnly && !IsEditorOnlyPropertyOnTheStack())
-#endif
-				{
-					ImportsUsedInGame.Add(Obj);
-				}
+
 				UClass* ClassObj = Cast<UClass>(Obj);
 
 				// Don't recurse into CDOs if we're already ignoring dependencies, we only want to recurse into our outer chain in that case
