- Change how we look for uncooked packages on disk to load instead of IOStore cooked packages, for cooked editor situations
- Made FPackageName::DoesPackageExistEx, removed DoesPackageExistInIoStore

diff --git a/Engine/Source/Runtime/CoreUObject/Private/Misc/PackageName.cpp b/Engine/Source/Runtime/CoreUObject/Private/Misc/PackageName.cpp
--- a/Engine/Source/Runtime/CoreUObject/Private/Misc/PackageName.cpp
+++ b/Engine/Source/Runtime/CoreUObject/Private/Misc/PackageName.cpp
@@ -1410,6 +1410,132 @@ bool FPackageName::DoesPackageExist(const FPackagePath& PackagePath, const FGuid
 	}
 }
 
+FPackageName::EPackageLocationFilter FPackageName::DoesPackageExistEx(const FPackagePath& PackagePath, EPackageLocationFilter Filter, const FGuid* Guid, bool bMatchCaseOnDisk, FPackagePath* OutPackagePath)
+{
+	TRACE_CPUPROFILER_EVENT_SCOPE(FPackageName::DoesPackageExistEx);
+
+	TStringBuilder<256> PackageName;
+	PackagePath.AppendPackageName(PackageName);
+
+	// Once we have the real Package Name, we can exit early if it's a script package - they exist only in memory.
+	if (IsScriptPackage(PackageName))
+	{
+		return EPackageLocationFilter::None;
+	}
+
+	if (IsMemoryPackage(PackageName))
+	{
+		return EPackageLocationFilter::None;
+	}
+
+	FText Reason;
+	if ( !FPackageName::IsValidTextForLongPackageName( PackageName, &Reason ) )
+	{
+		UE_LOG(LogPackageName, Error, TEXT( "DoesPackageExist: DoesPackageExist FAILED: '%s' is not a long packagename name. Reason: %s"), PackageName.ToString(), *Reason.ToString() );
+		return EPackageLocationFilter::None;
+	}
+
+	EPackageLocationFilter Result = EPackageLocationFilter::None;
+	// Used when I/O dispatcher is enabled
+	if (((uint8)Filter & (uint8)EPackageLocationFilter::Cooked))
+	{
+	    // For Cooked locations, return false for local filenames that are in unmounted directories
+	    if (PackagePath.IsMountedPath())
+	    {
+	        // Used when I/O dispatcher is enabled
+	        if (DoesPackageExistOverrideDelegate.IsBound())
+	        {
+		        if (DoesPackageExistOverrideDelegate.Execute(FName(*PackageName)))
+		        {
+			        if (OutPackagePath)
+			        {
+				        *OutPackagePath = PackagePath;
+				        if (OutPackagePath->GetHeaderExtension() == EPackageExtension::Unspecified)
+				        {
+					        OutPackagePath->SetHeaderExtension(EPackageExtension::EmptyString);
+				        }
+			        }
+
+					// note that the file was found in Cooked (IOStore) location
+				    Result = EPackageLocationFilter((uint8)Result | (uint8)EPackageLocationFilter::Cooked);
+
+				    // if we just want to find any existence, then we are done and we can skip the on disk check lower
+				    if (Filter == EPackageLocationFilter::Any)
+				    {
+					    return Result;
+				    }
+		        }
+			}
+		}
+	}
+
+	if ((uint8)Filter & (uint8)EPackageLocationFilter::Uncooked)
+	{
+		bool bFoundUncooked = false;
+	    // On consoles, we don't support package downloading, so no need to waste any extra cycles/disk io dealing with it
+	    if (!FPlatformProperties::RequiresCookedData() && Guid != nullptr)
+	    {
+		    // @todo: If we could get to list of linkers here, it would be faster to check
+		    // then to open the file and read it
+		    FPackagePath LocalPackagePath;
+		    FOpenPackageResult OpenPackageResult = IPackageResourceManager::Get().OpenReadPackage(PackagePath, &LocalPackagePath);
+		    TUniquePtr<FArchive>& PackageReader = OpenPackageResult.Archive;
+		    if (!PackageReader || OpenPackageResult.Format != EPackageFormat::Binary)
+		    {
+			    if (PackageReader)
+			    {
+				    UE_LOG(LogPackageName, Error, TEXT("DoesPackageExist: DoesPackageExist with Guid FAILED: '%s' exists on disk with TextFormat, and we cannot read guids from TextFormat packages."),
+					    *PackagePath.GetDebugName());
+			    }
+		    }
+			else
+			{
+		        // Read in the package summary
+		        FPackageFileSummary Summary;
+		        *PackageReader << Summary;
+
+		        // Compare Guids
+		        PRAGMA_DISABLE_DEPRECATION_WARNINGS
+		        if (!PackageReader->IsError() && Summary.Guid == *Guid)
+		        PRAGMA_ENABLE_DEPRECATION_WARNINGS
+				{
+		            PackageReader.Reset();
+            
+		            if (bMatchCaseOnDisk)
+		            {
+			            IPackageResourceManager::Get().TryMatchCaseOnDisk(LocalPackagePath, &LocalPackagePath);
+		            }
+		            if (OutPackagePath)
+		            {
+			            *OutPackagePath = LocalPackagePath;
+		            }
+				}
+			    
+				// note that we found it uncooked (on disk/pak) location
+				bFoundUncooked = true;
+			}
+	    }
+	    else
+	    {
+		    if (bMatchCaseOnDisk)
+		    {
+			    bFoundUncooked = IPackageResourceManager::Get().TryMatchCaseOnDisk(PackagePath, OutPackagePath);
+		    }
+		    else
+		    {
+				bFoundUncooked = IPackageResourceManager::Get().DoesPackageExist(PackagePath, OutPackagePath);
+			}
+		}
+
+		if (bFoundUncooked)
+		{
+			Result = EPackageLocationFilter((uint8)Result | (uint8)EPackageLocationFilter::Uncooked);
+		}
+	}
+	
+	return Result;
+}
+
 bool FPackageName::SearchForPackageOnDisk(const FString& PackageName, FString* OutLongPackageName, FString* OutFilename)
 {
 	TRACE_CPUPROFILER_EVENT_SCOPE(FPackageName::SearchForPackageOnDisk);
diff --git a/Engine/Source/Runtime/CoreUObject/Private/Serialization/AsyncLoading.cpp b/Engine/Source/Runtime/CoreUObject/Private/Serialization/AsyncLoading.cpp
--- a/Engine/Source/Runtime/CoreUObject/Private/Serialization/AsyncLoading.cpp
+++ b/Engine/Source/Runtime/CoreUObject/Private/Serialization/AsyncLoading.cpp
@@ -4550,7 +4550,7 @@ EAsyncPackageState::Type FAsyncLoadingThread::TickAsyncLoading(bool bUseTimeLimi
 		}
 
 		const bool bIsMultithreaded = FAsyncLoadingThread::IsMultithreaded();
-		
+
 		bool bDidSomething = false;
 		{
 			Result = ProcessLoadedPackages(bUseTimeLimit, bUseFullTimeLimit, TimeLimit, bDidSomething, FlushTree);
@@ -5767,10 +5767,11 @@ EAsyncPackageState::Type FAsyncPackage::CreateLinker()
 				bool bDoesPackageExist = false;
 				{
 					SCOPED_LOADTIMER(CreateLinker_DoesExist);
-					bDoesPackageExist = FPackageName::DoesPackageExist(PackagePath, Guid, false /* bMatchCaseOnDisk */, &PackagePath);
-#if WITH_IOSTORE_IN_EDITOR
+#if WITH_IOSTORE_IN_EDITOR 
 					// Only look for non cooked packages on disk
-					bDoesPackageExist &= !DoesPackageExistInIoStore(NameToLoadFName);
+					bDoesPackageExist = FPackageName::DoesPackageExistEx(PackagePath, FPackageName::EPackageLocationFilter::Uncooked, Guid, false /* bMatchCaseOnDisk */, &PackagePath) != FPackageName::EPackageLocationFilter::None;
+#else
+					bDoesPackageExist = FPackageName::DoesPackageExist(PackagePath, Guid, false /* bMatchCaseOnDisk */, &PackagePath);
 #endif
 				}
 
diff --git a/Engine/Source/Runtime/CoreUObject/Private/Serialization/AsyncPackageLoader.cpp b/Engine/Source/Runtime/CoreUObject/Private/Serialization/AsyncPackageLoader.cpp
--- a/Engine/Source/Runtime/CoreUObject/Private/Serialization/AsyncPackageLoader.cpp
+++ b/Engine/Source/Runtime/CoreUObject/Private/Serialization/AsyncPackageLoader.cpp
@@ -833,19 +833,6 @@ void NotifyUnreachableObjects(const TArrayView<FUObjectItem*>& UnreachableObject
 	GetAsyncPackageLoader().NotifyUnreachableObjects(UnreachableObjects);
 }
 
-#if WITH_IOSTORE_IN_EDITOR
-bool DoesPackageExistInIoStore(FName InPackageName)
-{
-	if (FIoDispatcher::IsInitialized())
-	{
-		FIoChunkId PackageChunkId = CreateIoChunkId(FPackageId::FromName(InPackageName).Value(), 0, EIoChunkType::ExportBundleData);
-		return FIoDispatcher::Get().DoesChunkExist(PackageChunkId);
-	}
-
-	return false;
-}
-#endif
-
 double GFlushAsyncLoadingTime = 0.0;
 uint32 GFlushAsyncLoadingCount = 0;
 uint32 GSyncLoadCount = 0;
diff --git a/Engine/Source/Runtime/CoreUObject/Private/Serialization/EditorPackageLoader.cpp b/Engine/Source/Runtime/CoreUObject/Private/Serialization/EditorPackageLoader.cpp
--- a/Engine/Source/Runtime/CoreUObject/Private/Serialization/EditorPackageLoader.cpp
+++ b/Engine/Source/Runtime/CoreUObject/Private/Serialization/EditorPackageLoader.cpp
@@ -53,7 +53,7 @@ public:
 		const FLinkerInstancingContext* InstancingContext) override
 	{
 		// Use the old loader if an uncooked package exists on disk
-		const bool bDoesUncookedPackageExist = FPackageName::DoesPackageExist(PackagePath) && !DoesPackageExistInIoStore(PackagePath.GetPackageFName());
+		const bool bDoesUncookedPackageExist = FPackageName::DoesPackageExistEx(PackagePath, FPackageName::EPackageLocationFilter::Uncooked) != FPackageName::EPackageLocationFilter::None;
 		if (bDoesUncookedPackageExist)
 		{
 			UE_LOG(LogEditorPackageLoader, Verbose, TEXT("Loading uncooked package '%s' from filesystem"), *PackagePath.GetDebugName());
diff --git a/Engine/Source/Runtime/CoreUObject/Private/UObject/Linker.cpp b/Engine/Source/Runtime/CoreUObject/Private/UObject/Linker.cpp
--- a/Engine/Source/Runtime/CoreUObject/Private/UObject/Linker.cpp
+++ b/Engine/Source/Runtime/CoreUObject/Private/UObject/Linker.cpp
@@ -627,7 +627,7 @@ static bool IsConvertedDynamicPackage(FName PackageName)
 {
 	return
 #if WITH_EDITORONLY_DATA
-		GLinkerAllowDynamicClasses &&
+		GLinkerAllowDynamicClasses && 
 #endif
 		GetConvertedDynamicPackageNameToTypeName().Contains(PackageName);
 }
@@ -650,11 +650,14 @@ FString GetPrestreamPackageLinkerName(const TCHAR* InLongPackageName, bool bSkip
 	{
 		return FString(); // we won't load this anyway, don't prestream
 	}
-	bool DoesPackageExist = FPackageName::DoesPackageExist(PackagePath, &PackagePath);
-#if WITH_IOSTORE_IN_EDITOR
+
+#if WITH_IOSTORE_IN_EDITOR 
 	// Only look for non cooked packages on disk
-	DoesPackageExist &= !DoesPackageExistInIoStore(PackageFName);
+	bool DoesPackageExist = FPackageName::DoesPackageExistEx(PackagePath, FPackageName::EPackageLocationFilter::Uncooked, nullptr, false, &PackagePath) != FPackageName::EPackageLocationFilter::None;
+#else
+	bool DoesPackageExist = FPackageName::DoesPackageExist(PackagePath, &PackagePath);
 #endif
+
 	if (!IsConvertedDynamicPackage(PackageFName) && !DoesPackageExist)
 	{
 		return FString();
@@ -689,7 +692,7 @@ static FPackagePath GetPackagePath(UPackage* InOuter, const TCHAR* InPackageName
 		return FPackagePath();
 	}
 }
-
+		
 COREUOBJECT_API FLinkerLoad* GetPackageLinker(UPackage* InOuter, const TCHAR* InLongPackageName, uint32 LoadFlags,
 	UPackageMap* Sandbox, FGuid* CompatibleGuid, FArchive* InReaderOverride, FUObjectSerializeContext** InOutLoadContext,
 	FLinkerLoad* ImportLinker, const FLinkerInstancingContext* InstancingContext)
diff --git a/Engine/Source/Runtime/CoreUObject/Private/UObject/UObjectGlobals.cpp b/Engine/Source/Runtime/CoreUObject/Private/UObject/UObjectGlobals.cpp
--- a/Engine/Source/Runtime/CoreUObject/Private/UObject/UObjectGlobals.cpp
+++ b/Engine/Source/Runtime/CoreUObject/Private/UObject/UObjectGlobals.cpp
@@ -832,7 +832,7 @@ bool ParseObject( const TCHAR* Stream, const TCHAR* Match, UClass* Class, UObjec
 		if( !Res )
 		{
 			if (Class->IsChildOf<UClass>())
-			{ 
+			{
 				const FString ObjectName = FPackageName::ObjectPathToObjectName(FString(TempStr));
 				const FName RedirectedObjectName = FLinkerLoad::FindNewNameForClass(*ObjectName, false);
 				if (!RedirectedObjectName.IsNone())
@@ -1140,7 +1140,7 @@ UPackage* LoadPackageInternal(UPackage* InOuter, const FPackagePath& PackagePath
 
 #if WITH_IOSTORE_IN_EDITOR
 		// Use the old loader if an uncooked package exists on disk
-		const bool bDoesUncookedPackageExist = FPackageName::DoesPackageExist(PackagePath) && !DoesPackageExistInIoStore(PackageName);
+		const bool bDoesUncookedPackageExist = FPackageName::DoesPackageExistEx(PackagePath, FPackageName::EPackageLocationFilter::Uncooked) != FPackageName::EPackageLocationFilter::None;
 		if (!bDoesUncookedPackageExist)
 #endif
 		{
@@ -2725,7 +2725,7 @@ FObjectInitializer::FObjectInitializer(UObject* InObj, const FStaticConstructObj
 	}
 
 	Construct_Internal();
-}
+}	
 
 FObjectInitializer::FObjectInitializer(UObject* InObj, UObject* InObjectArchetype, bool bInCopyTransientsFromClassDefaults, bool bInShouldInitializeProps, struct FObjectInstancingGraph* InInstanceGraph)
 	: Obj(InObj)
@@ -3455,7 +3455,7 @@ UObject* StaticConstructObject_Internal(const FStaticConstructObjectParameters&
 	check(Result != nullptr);
 	// Don't call the constructor on recycled subobjects, they haven't been destroyed.
 	if (!bRecycledSubobject)
-	{
+	{		
 		STAT(FScopeCycleCounterUObject ConstructorScope(InClass->GetFName().IsNone() ? nullptr : InClass, GET_STATID(STAT_ConstructObject)));
 		(*InClass->ClassConstructor)(FObjectInitializer(Result, Params));
 	}
diff --git a/Engine/Source/Runtime/CoreUObject/Public/Misc/PackageName.h b/Engine/Source/Runtime/CoreUObject/Public/Misc/PackageName.h
--- a/Engine/Source/Runtime/CoreUObject/Public/Misc/PackageName.h
+++ b/Engine/Source/Runtime/CoreUObject/Public/Misc/PackageName.h
@@ -105,7 +105,7 @@ public:
 	 */
 	static bool TryConvertLongPackageNameToFilename(const FString& InLongPackageName, FString& OutFilename, const FString& InExtension = TEXT(""));
 
-	/**
+	/** 
 	 * Find the MountPoint for a LocalPath, LongPackageName, or ObjectPath and return its elements. Use this function instead of TryConvertFilenameToLongPackageName or
 	 * TryConvertLongPackageNameToFilename if you need to handle InPaths that might be ObjectPaths.
 	 * @param InPath					The LocalPath (with path,name,extension), PackageName, or ObjectPath we want to 
@@ -321,6 +321,26 @@ public:
 	 **/
 	static bool DoesPackageExist(const FPackagePath& PackagePath, FPackagePath* OutPackagePath);
 
+	enum class EPackageLocationFilter : uint8
+	{
+		None = 0,
+		Cooked = 1,
+		Uncooked = 2,
+		Any = 0xFF, // special filter to find if it exists anywhere at all, and won't need to check both. in this case, as soon as one is found, DoesPackageExistEx will return true
+	};
+
+	/**
+	 * Checks if the package exists in IOStore containers, on disk outsode of IOStore, both, or neither
+	 *
+	 * @param PackagePath Package package.
+	 * @param Filter Indication of where it should look for 
+	 * @param Guid If nonnull, and the package is found on disk but does not have this PackageGuid in its FPackageFileSummary::Guid, false is returned
+	 * @param bMatchCaseOnDisk If true, the OutPackagePath is modified to match the capitalization of the discovered file
+	 * @param OutPackagePath If nonnull and the package exists, set to a copy of PackagePath with the HeaderExtension set to the extension that exists on disk (and if bMatchCaseOnDisk is true, capitalization changed to match). If not found, this variable is not written
+	 * @return the set of locations where the package exists (cooked or uncooked, both or neither)
+	 **/
+	static EPackageLocationFilter DoesPackageExistEx(const FPackagePath& PackagePath, EPackageLocationFilter Filterconst, const FGuid* Guid = nullptr, bool bMatchCaseOnDisk = false, FPackagePath* OutPackagePath = nullptr);
+
 	/**
 	 * Attempts to find a package given its short name on disk (very slow).
 	 * 
@@ -638,7 +658,7 @@ public:
 	 */
 	static bool DoesPackageNameContainInvalidCharacters(FStringView InLongPackageName, FText* OutReason);
 	static bool DoesPackageNameContainInvalidCharacters(FStringView InLongPackageName, EErrorCode* OutReason = nullptr);
-
+	
 	/**
 	* Checks if a package can be found using known package extensions (header extensions only; files with the extensions of other segments are not returned).
 	*
diff --git a/Engine/Source/Runtime/CoreUObject/Public/UObject/UObjectGlobals.h b/Engine/Source/Runtime/CoreUObject/Public/UObject/UObjectGlobals.h
--- a/Engine/Source/Runtime/CoreUObject/Public/UObject/UObjectGlobals.h
+++ b/Engine/Source/Runtime/CoreUObject/Public/UObject/UObjectGlobals.h
@@ -512,13 +512,6 @@ COREUOBJECT_API bool IsEventDrivenLoaderEnabledInCookedBuilds();
 */
 COREUOBJECT_API bool IsEventDrivenLoaderEnabled();
 
-#if WITH_IOSTORE_IN_EDITOR
-/**
- * Returns true if the specified package is cooked and stored in I/O store
-*/
-bool DoesPackageExistInIoStore(FName InPackageName);
-#endif
-
 /**
  * Returns the async load percentage for a package in flight with the passed in name or -1 if there isn't one.
  * @warning THIS IS SLOW. MAY BLOCK ASYNC LOADING.
