AsyncLoading: Add missing call to FPlatformMisc::PumpEssentialAppMessages() when processing loaded packages (PostLoad and CompletionCallbacks) on the GT.
Make sure to call OnAsyncLoadingFlushUpdate frequently enough (to e.g. process platform messages) when running without a dedicated async loading thread.

Follow up fixes for cl 14462152: Add a lightweight message pump function that can be used for platforms that need regular updates even during asset loading & processing. Fix some suspend timeouts on engine startup.

diff --git a/Engine/Source/Runtime/CoreUObject/Private/Serialization/AsyncLoading.cpp b/Engine/Source/Runtime/CoreUObject/Private/Serialization/AsyncLoading.cpp
--- a/Engine/Source/Runtime/CoreUObject/Private/Serialization/AsyncLoading.cpp
+++ b/Engine/Source/Runtime/CoreUObject/Private/Serialization/AsyncLoading.cpp
@@ -4027,6 +4027,7 @@ EAsyncPackageState::Type FAsyncLoadingThread::ProcessAsyncLoading(int32& OutPack
 			{
 				// Update heartbeat after 32 events
 				FThreadHeartBeat::Get().HeartBeat();
+				FCoreDelegates::OnAsyncLoadingFlushUpdate.Broadcast();
 			}
 
 			bool bDidSomething = false;
@@ -4356,6 +4357,8 @@ EAsyncPackageState::Type FAsyncLoadingThread::ProcessLoadedPackages(bool bUseTim
 
 	for (int32 PackageIndex = 0; PackageIndex < LoadedPackagesToProcess.Num() && !IsAsyncLoadingSuspendedInternal(); ++PackageIndex)
 	{
+		FPlatformMisc::PumpEssentialAppMessages();
+
 		FAsyncPackage* Package = LoadedPackagesToProcess[PackageIndex];
 		if (Package->GetDependencyRefCount() == 0)
 		{
diff --git a/Engine/Source/Runtime/CoreUObject/Private/Serialization/AsyncLoading2.cpp b/Engine/Source/Runtime/CoreUObject/Private/Serialization/AsyncLoading2.cpp
--- a/Engine/Source/Runtime/CoreUObject/Private/Serialization/AsyncLoading2.cpp
+++ b/Engine/Source/Runtime/CoreUObject/Private/Serialization/AsyncLoading2.cpp
@@ -4042,8 +4042,6 @@ EAsyncPackageState::Type FAsyncLoadingThread2::ProcessAsyncLoadingFromGameThread
 
 	check(IsInGameThread());
 
-	// If we're not multithreaded and flushing async loading, update the thread heartbeat
-	const bool bNeedsHeartbeatTick = !FAsyncLoadingThread2::IsMultithreaded();
 	OutPackagesProcessed = 0;
 
 #if ALT2_VERIFY_RECURSIVE_LOADS 
@@ -4056,10 +4054,12 @@ EAsyncPackageState::Type FAsyncLoadingThread2::ProcessAsyncLoadingFromGameThread
 	{
 		do 
 		{
-			if (bNeedsHeartbeatTick && (++LoopIterations) % 32 == 31)
+			if ((++LoopIterations) % 32 == 31)
 			{
+				// We're not multithreaded and flushing async loading
 				// Update heartbeat after 32 events
 				FThreadHeartBeat::Get().HeartBeat();
+				FCoreDelegates::OnAsyncLoadingFlushUpdate.Broadcast();
 			}
 
 			if (ThreadState.IsTimeLimitExceeded(TEXT("ProcessAsyncLoadingFromGameThread")))
