Make sure an object can't be accessed via weak pointers after it's been FinishDestroyed.

Fix race condition between FWeakObjectPtr and GC AsyncPurge.
When AsyncPurge calls ResetSerialNumberAndFlags() then FWeakObjectPtr::IsValid() could read the old good SerialNumber and the newly cleared Flags making it return true despite that the object is being destroyed.

diff --git a/Engine/Source/Runtime/CoreUObject/Private/UObject/Obj.cpp b/Engine/Source/Runtime/CoreUObject/Private/UObject/Obj.cpp
--- a/Engine/Source/Runtime/CoreUObject/Private/UObject/Obj.cpp
+++ b/Engine/Source/Runtime/CoreUObject/Private/UObject/Obj.cpp
@@ -1029,6 +1029,9 @@ bool UObject::ConditionalFinishDestroy()
 #endif
 		FinishDestroy();
 
+		// Make sure this object can't be accessed via weak pointers after it's been FinishDestroyed
+		GUObjectArray.ResetSerialNumber(this);
+
 		// Make sure this object can't be found through any delete listeners (annotation maps etc) after it's been FinishDestroyed
 		GUObjectArray.RemoveObjectFromDeleteListeners(this);
 
diff --git a/Engine/Source/Runtime/CoreUObject/Private/UObject/UObjectArray.cpp b/Engine/Source/Runtime/CoreUObject/Private/UObject/UObjectArray.cpp
--- a/Engine/Source/Runtime/CoreUObject/Private/UObject/UObjectArray.cpp
+++ b/Engine/Source/Runtime/CoreUObject/Private/UObject/UObjectArray.cpp
@@ -244,6 +244,19 @@ void FUObjectArray::AllocateUObjectIndex(UObjectBase* Object, bool bMergingThrea
 	}
 }
 
+/**
+ * Reset the serial number from the game thread to invalidate all weak object pointers to it
+ *
+ * @param Object to reset
+ */
+void FUObjectArray::ResetSerialNumber(UObjectBase* Object)
+{
+	int32 Index = Object->InternalIndex;
+	FUObjectItem* ObjectItem = IndexToObject(Index);
+	checkSlow(ObjectItem);
+	ObjectItem->SerialNumber = 0;
+}
+
 /**
  * Removes an object from delete listeners
  *
diff --git a/Engine/Source/Runtime/CoreUObject/Public/UObject/UObjectArray.h b/Engine/Source/Runtime/CoreUObject/Public/UObject/UObjectArray.h
--- a/Engine/Source/Runtime/CoreUObject/Public/UObject/UObjectArray.h
+++ b/Engine/Source/Runtime/CoreUObject/Public/UObject/UObjectArray.h
@@ -572,6 +572,15 @@ public:
 **/
 class COREUOBJECT_API FUObjectArray
 {
+	friend class UObject;
+private:
+	/**
+	 * Reset the serial number from the game thread to invalidate all weak object pointers to it
+	 *
+	 * @param Object to reset
+	 */
+	void ResetSerialNumber(UObjectBase* Object);
+
 public:
 
 	enum ESerialNumberConstants
