Fixed non unity compilation error from missing include

diff --git a/Engine/Source/Runtime/CoreUObject/Public/UObject/PropertyIterator.h b/Engine/Source/Runtime/CoreUObject/Public/UObject/PropertyIterator.h
--- a/Engine/Source/Runtime/CoreUObject/Public/UObject/PropertyIterator.h
+++ b/Engine/Source/Runtime/CoreUObject/Public/UObject/PropertyIterator.h
@@ -6,6 +6,7 @@ PropertyIterrator.h: TPropertyIterator implementation.
 #pragma once
 
 #include "UObject/Field.h"
+#include "UObject/UObjectIterator.h"
 
 /**
  * Iterates over all structs and their properties
