Fix FStructUtils::TheSameLayout returning false when passed null structs. This fixes a regression introduced by sparse class data changes (CL# 16048941).

diff --git a/Engine/Source/Runtime/CoreUObject/Private/UObject/Class.cpp b/Engine/Source/Runtime/CoreUObject/Private/UObject/Class.cpp
--- a/Engine/Source/Runtime/CoreUObject/Private/UObject/Class.cpp
+++ b/Engine/Source/Runtime/CoreUObject/Private/UObject/Class.cpp
@@ -6013,8 +6013,9 @@ bool FStructUtils::ArePropertiesTheSame(const FProperty* A, const FProperty* B,
 
 bool FStructUtils::TheSameLayout(const UStruct* StructA, const UStruct* StructB, bool bCheckPropertiesNames)
 {
-	bool bResult = false;
-	if (StructA 
+	bool bResult = (StructA == StructB);
+	if (!bResult
+		&& StructA 
 		&& StructB 
 		&& (StructA->GetPropertiesSize() == StructB->GetPropertiesSize())
 		&& (StructA->GetMinAlignment() == StructB->GetMinAlignment()))
