Temporarily backout TheSameLayout fix while investigating why anim code relied on incorrect behavior

diff --git a/Engine/Source/Runtime/CoreUObject/Private/UObject/Class.cpp b/Engine/Source/Runtime/CoreUObject/Private/UObject/Class.cpp
--- a/Engine/Source/Runtime/CoreUObject/Private/UObject/Class.cpp
+++ b/Engine/Source/Runtime/CoreUObject/Private/UObject/Class.cpp
@@ -6013,9 +6013,8 @@ bool FStructUtils::ArePropertiesTheSame(const FProperty* A, const FProperty* B,
 
 bool FStructUtils::TheSameLayout(const UStruct* StructA, const UStruct* StructB, bool bCheckPropertiesNames)
 {
-	bool bResult = (StructA == StructB);
-	if (!bResult
-		&& StructA 
+	bool bResult = false;
+	if (StructA 
 		&& StructB 
 		&& (StructA->GetPropertiesSize() == StructB->GetPropertiesSize())
 		&& (StructA->GetMinAlignment() == StructB->GetMinAlignment()))
