Do not exit early if no file on the path

diff --git a/Engine/Source/Runtime/CoreUObject/Private/UObject/Obj.cpp b/Engine/Source/Runtime/CoreUObject/Private/UObject/Obj.cpp
--- a/Engine/Source/Runtime/CoreUObject/Private/UObject/Obj.cpp
+++ b/Engine/Source/Runtime/CoreUObject/Private/UObject/Obj.cpp
@@ -2768,7 +2768,7 @@ bool UObject::TryUpdateDefaultConfigFile(const FString& SpecificFileLocation, bo
 {
 	FString ConfigFile = SpecificFileLocation.IsEmpty() ? GetDefaultConfigFilename() : SpecificFileLocation;
 
-	if (FPaths::FileExists(ConfigFile) && !IFileManager::Get().IsReadOnly(*ConfigFile))
+	if (!FPaths::FileExists(ConfigFile) || !IFileManager::Get().IsReadOnly(*ConfigFile))
 	{
 		UpdateSingleSectionOfConfigFile(ConfigFile);
 
diff --git a/Engine/Source/Runtime/CoreUObject/Public/UObject/Object.h b/Engine/Source/Runtime/CoreUObject/Public/UObject/Object.h
--- a/Engine/Source/Runtime/CoreUObject/Public/UObject/Object.h
+++ b/Engine/Source/Runtime/CoreUObject/Public/UObject/Object.h
@@ -1124,7 +1124,7 @@ public:
 	 * Try to Saves just the section(s) for this class into the default ini file for the class (with just the changes from base)
 	 *
 	 *  @param SpecificFileLocation The Ini file or if Empty the default ini file for the call
-	 *  @param bWarnIfFail If unable to update the Ini file due to incorrect path or unable to write will log a warning
+	 *  @param bWarnIfFail If true and unable to update the Ini file due to being read only log a warning
 	 *  @return true if it was able to write false otherwise
 	 */
 	bool TryUpdateDefaultConfigFile(const FString& SpecificFileLocation = "", bool bWarnIfFail = true);
