Fix warning reporting issue about file not existing

diff --git a/Engine/Source/Runtime/CoreUObject/Private/UObject/Obj.cpp b/Engine/Source/Runtime/CoreUObject/Private/UObject/Obj.cpp
--- a/Engine/Source/Runtime/CoreUObject/Private/UObject/Obj.cpp
+++ b/Engine/Source/Runtime/CoreUObject/Private/UObject/Obj.cpp
@@ -2777,7 +2777,7 @@ bool UObject::TryUpdateDefaultConfigFile(const FString& SpecificFileLocation, bo
 
 	if (bWarnIfFail)
 	{
-		UE_LOG(LogObj, Warning, TEXT("Ini File '%s' was not found or is read-only"), *ConfigFile);
+		UE_LOG(LogObj, Warning, TEXT("Ini File '%s' is read-only and cannot be written to"), *ConfigFile);
 	}
 
 	return false;
