LargeMemoryArchive::ReleaseOwnership() sets TotalSize() to 0, so TotalSize must be called before calling ReleaseOwnership. Calling them both as arguments into a function was succeeding on some compilers and failing on others due to compiler optimization differences.
#rnx
#[fyi] Francis.Hurteau

diff --git a/Engine/Source/Developer/IoStoreUtilities/Private/ZenStoreWriter.cpp b/Engine/Source/Developer/IoStoreUtilities/Private/ZenStoreWriter.cpp
--- a/Engine/Source/Developer/IoStoreUtilities/Private/ZenStoreWriter.cpp
+++ b/Engine/Source/Developer/IoStoreUtilities/Private/ZenStoreWriter.cpp
@@ -317,7 +317,8 @@ FZenStoreWriter::~FZenStoreWriter()
 void FZenStoreWriter::WritePackageData(const FPackageInfo& Info, FLargeMemoryWriter& ExportsArchive, const TArray<FFileRegion>& FileRegions)
 {
 	check(Info.ChunkId.IsValid());
-	FIoBuffer PackageData(FIoBuffer::AssumeOwnership, ExportsArchive.ReleaseOwnership(), ExportsArchive.TotalSize());
+	int64 DataSize = ExportsArchive.TotalSize();
+	FIoBuffer PackageData(FIoBuffer::AssumeOwnership, ExportsArchive.ReleaseOwnership(), DataSize);
 
 	TRACE_CPUPROFILER_EVENT_SCOPE(FZenStoreWriter::WritePackageData);
 
diff --git a/Engine/Source/Runtime/CoreUObject/Private/Serialization/PackageWriterToSharedBuffer.cpp b/Engine/Source/Runtime/CoreUObject/Private/Serialization/PackageWriterToSharedBuffer.cpp
--- a/Engine/Source/Runtime/CoreUObject/Private/Serialization/PackageWriterToSharedBuffer.cpp
+++ b/Engine/Source/Runtime/CoreUObject/Private/Serialization/PackageWriterToSharedBuffer.cpp
@@ -31,7 +31,11 @@ void FPackageWriterRecords::WritePackageData(const IPackageWriter::FPackageInfo&
 	FLargeMemoryWriter& ExportsArchive, const TArray<FFileRegion>& FileRegions)
 {
 	ValidatePackageName(Info.PackageName);
-	FSharedBuffer Buffer = FSharedBuffer::TakeOwnership(ExportsArchive.ReleaseOwnership(), ExportsArchive.TotalSize(),
+	int64 DataSize= ExportsArchive.TotalSize();
+	checkf(DataSize > 0, TEXT("IPackageWriter->WritePackageData must not be called with an empty ExportsArchive"));
+	checkf(static_cast<uint64>(DataSize) >= Info.HeaderSize,
+		TEXT("IPackageWriter->WritePackageData must not be called with HeaderSize > ExportsArchive.TotalSize"));
+	FSharedBuffer Buffer = FSharedBuffer::TakeOwnership(ExportsArchive.ReleaseOwnership(), DataSize,
 		FMemory::Free);
 	Package = FPackage{ Info, MoveTemp(Buffer), FileRegions };
 }
diff --git a/Engine/Source/Runtime/CoreUObject/Private/UObject/SavePackage2.cpp b/Engine/Source/Runtime/CoreUObject/Private/UObject/SavePackage2.cpp
--- a/Engine/Source/Runtime/CoreUObject/Private/UObject/SavePackage2.cpp
+++ b/Engine/Source/Runtime/CoreUObject/Private/UObject/SavePackage2.cpp
@@ -1918,8 +1918,9 @@ ESavePackageResult FinalizeFile(FStructuredArchive::FRecord& StructuredArchiveRo
 			}
 
 			// Add the uasset file to the list of output files
+			int64 DataSize = Writer->TotalSize();
 			SaveContext.AdditionalPackageFiles.Emplace(SaveContext.GetFilename(),
-				FLargeMemoryPtr(Writer->ReleaseOwnership()), Linker->FileRegions, Writer->TotalSize());
+				FLargeMemoryPtr(Writer->ReleaseOwnership()), Linker->FileRegions, DataSize);
 
 			for (FSavePackageOutputFile& Entry : SaveContext.AdditionalPackageFiles)
 			{
