Added Units and ForceUnits to ObjectMacros.h for discoverability

diff --git a/Engine/Source/Runtime/CoreUObject/Public/UObject/ObjectMacros.h b/Engine/Source/Runtime/CoreUObject/Public/UObject/ObjectMacros.h
--- a/Engine/Source/Runtime/CoreUObject/Public/UObject/ObjectMacros.h
+++ b/Engine/Source/Runtime/CoreUObject/Public/UObject/ObjectMacros.h
@@ -1330,6 +1330,14 @@ namespace UM
 		/// [PropertyMetadata] Used for float and integer properties.  Specifies the highest that the value slider should represent.
 		UIMax,
 
+		/// [PropertyMetadata] Used on any numeric property to declare the internal units for a property. See UnitConversion.h/cpp for details.
+		///                    If editor preferances allow, properties marked up with Units= will be displayed in the user's locale preference,
+		///                    and/or in the most appropriate unit for the value's magnitude (ie, showing cm as km, or bytes as megabytes etc).
+		Units,
+
+		/// [PropertyMetadata] Used instead of 'Units' metadata to forcibly show a property in a fixed unit, without locale or other conversion
+		ForceUnits,
+
 		/// [PropertyMetadata] Used for SoftObjectPtr/SoftObjectPath properties to specify a reference should not be tracked. This reference will not be automatically cooked or saved into the asset registry for redirector/delete fixup.
 		Untracked,
 
