VirtualizedUntypedBulkData: Fix assertion in BulkDataRegistry: IsMemoryOnlyPayload was returning true even if the payload existed on disk.

diff --git a/Engine/Source/Runtime/CoreUObject/Private/Serialization/VirtualizedBulkData.cpp b/Engine/Source/Runtime/CoreUObject/Private/Serialization/VirtualizedBulkData.cpp
--- a/Engine/Source/Runtime/CoreUObject/Private/Serialization/VirtualizedBulkData.cpp
+++ b/Engine/Source/Runtime/CoreUObject/Private/Serialization/VirtualizedBulkData.cpp
@@ -1115,7 +1115,7 @@ bool FVirtualizedUntypedBulkData::CanUnloadData() const
 
 bool FVirtualizedUntypedBulkData::IsMemoryOnlyPayload() const
 {
-	return !Payload.IsNull() && !CanUnloadData();
+	return !Payload.IsNull() && !IsDataVirtualized() && PackagePath.IsEmpty();
 }
 
 void FVirtualizedUntypedBulkData::Reset()
