Add comment to clarify UAsset vs TopLevelAsset.
#rnx

diff --git a/Engine/Source/Runtime/CoreUObject/Public/AssetRegistry/AssetData.h b/Engine/Source/Runtime/CoreUObject/Public/AssetRegistry/AssetData.h
--- a/Engine/Source/Runtime/CoreUObject/Public/AssetRegistry/AssetData.h
+++ b/Engine/Source/Runtime/CoreUObject/Public/AssetRegistry/AssetData.h
@@ -146,7 +146,10 @@ public:
 		return !ObjectPath.IsNone();
 	}
 
-	/** Returns true if this is the primary asset in a package, true for maps and assets but false for secondary objects like class redirectors */
+	/**
+	 * Returns true if this is the main asset in a package, true for maps and assets but false for secondary objects like class redirectors
+	 * Every UAsset is also a TopLevelAsset.
+	 */
 	bool IsUAsset() const
 	{
 		if (!IsValid())
@@ -162,18 +165,23 @@ public:
 		return DetectIsUAssetByNames(PackageNameStrBuilder, AssetNameStrBuilder);
 	}
 
-	/** Returns true if the given UObject is the primary asset in a package, true for maps and assets but false for secondary objects like class redirectors */
+	/**
+	 * Returns true if the given UObject is the main asset in a package, true for maps and assets but false for secondary objects like class redirectors
+	 * Every UAsset is also a TopLevelAsset.
+	 */
 	COREUOBJECT_API static bool IsUAsset(UObject* Object);
 
 	/**
 	 * Returns true iff the Asset is a TopLevelAsset (not a subobject, its outer is a UPackage).
 	 * Only TopLevelAssets can be PrimaryAssets in the AssetManager.
+	 * A TopLevelAsset is not necessarily the main asset in a package; see IsUAsset.
 	 */
 	COREUOBJECT_API bool IsTopLevelAsset() const;
 	
 	/**
 	 * Returns true iff the given Object, assumed to be an Asset, is a TopLevelAsset (not a subobject, its outer is a UPackage).
 	 * Only TopLevelAssets can be PrimaryAssets in the AssetManager.
+	 * A TopLevelAsset is not necessarily the main asset in a package; see IsUAsset.
 	 */
 	COREUOBJECT_API static bool IsTopLevelAsset(UObject* Object);
 
