Fix for editor assert when expanding asset paths so they exceed NAME_SIZE

diff --git a/Engine/Source/Runtime/CoreUObject/Private/UObject/UObjectGlobals.cpp b/Engine/Source/Runtime/CoreUObject/Private/UObject/UObjectGlobals.cpp
--- a/Engine/Source/Runtime/CoreUObject/Private/UObject/UObjectGlobals.cpp
+++ b/Engine/Source/Runtime/CoreUObject/Private/UObject/UObjectGlobals.cpp
@@ -7,6 +7,7 @@
 #include "UObject/UObjectGlobals.h"
 #include "HAL/PlatformFileManager.h"
 #include "HAL/FileManager.h"
+#include "Misc/AsciiSet.h"
 #include "Misc/Paths.h"
 #include "Misc/ITransaction.h"
 #include "Serialization/ArchiveProxy.h"
@@ -742,22 +743,23 @@ bool ResolveName(UObject*& InPackage, FString& InOutName, bool Create, bool Thro
 	bool bSubobjectPath = false;
 
 	// Handle specified packages.
-	int32 DotIndex = INDEX_NONE;// InOutName.Find(TEXT("."), ESearchCase::CaseSensitive);
-
-	// to make parsing the name easier, replace the subobject delimiter with an extra dot
-	InOutName.ReplaceInline(SUBOBJECT_DELIMITER, TEXT(".."), ESearchCase::CaseSensitive);
-	while ((DotIndex = InOutName.Find(TEXT("."), ESearchCase::CaseSensitive)) != INDEX_NONE)
+	constexpr FAsciiSet Delimiters = FAsciiSet(".") + SUBOBJECT_DELIMITER_CHAR;
+	while (true)
 	{
-		FString PartialName = InOutName.Left(DotIndex);
+		const TCHAR* DelimiterOrEnd = FAsciiSet::FindFirstOrEnd(*InOutName, Delimiters);
 
-		// if the next part of InOutName ends in two dots, it indicates that the next object in the path name
-		// is not a top-level object (i.e. it's a subobject).  e.g. SomePackage.SomeGroup.SomeObject..Subobject
-		if (InOutName.IsValidIndex(DotIndex+1) && InOutName[DotIndex+1] == TEXT('.'))
+		if (*DelimiterOrEnd == '\0')
+		{
+			return true;
+		}
+		else if (*DelimiterOrEnd == SUBOBJECT_DELIMITER_CHAR)
 		{
-			InOutName.RemoveAt(DotIndex, 1, false);
 			bSubobjectPath = true;
 			Create         = false;
 		}
+		
+		int32 DotIndex = DelimiterOrEnd - *InOutName;
+		FString PartialName = InOutName.Left(DotIndex);
 
 		bool bIsScriptPackage = false;
 		if (!bSubobjectPath)
@@ -814,8 +816,6 @@ bool ResolveName(UObject*& InPackage, FString& InOutName, bool Create, bool Thro
 		}
 		InOutName.RemoveAt(0, DotIndex + 1, false);
 	}
-
-	return true;
 }
 
 bool ParseObject( const TCHAR* Stream, const TCHAR* Match, UClass* Class, UObject*& DestRes, UObject* InParent, bool* bInvalidObject )
