SavePackage: Allow save.FixupStandaloneFlags to apply to all assets except external-package assets.
#rnx

diff --git a/Engine/Source/Runtime/CoreUObject/Private/UObject/SavePackage2.cpp b/Engine/Source/Runtime/CoreUObject/Private/UObject/SavePackage2.cpp
--- a/Engine/Source/Runtime/CoreUObject/Private/UObject/SavePackage2.cpp
+++ b/Engine/Source/Runtime/CoreUObject/Private/UObject/SavePackage2.cpp
@@ -133,7 +133,7 @@ ESavePackageResult ValidatePackage(FSaveContext& SaveContext)
 		EObjectFlags TopLevelFlags = SaveContext.GetTopLevelFlags();
 		if (TopLevelFlags != RF_NoFlags && !Asset->HasAnyFlags(TopLevelFlags))
 		{
-			if (SaveContext.IsFixupStandaloneFlags() && FAssetData::IsUAsset(Asset) && EnumHasAnyFlags(TopLevelFlags, RF_Standalone))
+			if (SaveContext.IsFixupStandaloneFlags() && !Asset->GetExternalPackage() && EnumHasAnyFlags(TopLevelFlags, RF_Standalone))
 			{
 				UE_LOG(LogSavePackage, Warning, TEXT("The Asset %s being saved is missing the RF_Standalone flag; adding it."), *Asset->GetPathName());
 				Asset->SetFlags(RF_Standalone);
@@ -146,7 +146,7 @@ ESavePackageResult ValidatePackage(FSaveContext& SaveContext)
 				static_assert(sizeof(TopLevelFlags) <= sizeof(uint32), "Expect EObjectFlags to be uint32");
 				Arguments.Add(TEXT("Flags"), FText::FromString(FString::Printf(TEXT("%x"), (uint32)TopLevelFlags)));
 				FText ErrorText;
-				if (FAssetData::IsUAsset(Asset) && EnumHasAnyFlags(TopLevelFlags, RF_Standalone))
+				if (!Asset->GetExternalPackage() && EnumHasAnyFlags(TopLevelFlags, RF_Standalone))
 				{
 					ErrorText = FText::Format(NSLOCTEXT("SavePackage2", "AssetSaveMissingStandaloneFlag",
 						"The Asset {Name} being saved does not have any of the provided object flags (0x{Flags}); saving the package would cause data loss. "
