Change IPackageResourceManager::OpenAsyncReadPackage to return the Format along with the IAsyncReadFileHandle, so that FAsyncArchive can abort with an error if it is a text archive. This is necessary for cooked editor, which tries to use FAsyncArchive when available, but has to support text-format packages.
#rnx

diff --git a/Engine/Source/Editor/UnrealEd/Private/Cooker/CookPackageData.cpp b/Engine/Source/Editor/UnrealEd/Private/Cooker/CookPackageData.cpp
--- a/Engine/Source/Editor/UnrealEd/Private/Cooker/CookPackageData.cpp
+++ b/Engine/Source/Editor/UnrealEd/Private/Cooker/CookPackageData.cpp
@@ -46,7 +46,7 @@ namespace UE::Cook
 
 	FPackageData::FPackageData(FPackageDatas& PackageDatas, const FName& InPackageName, const FName& InFileName)
 		: GeneratedOwner(nullptr), PackageName(InPackageName), FileName(InFileName), PackageDatas(PackageDatas)
-		, PreloadableFileFormat(EPackageFormat::Binary), Instigator(EInstigator::NotYetRequested), bIsUrgent(0)
+		, Instigator(EInstigator::NotYetRequested), bIsUrgent(0)
 		, bIsVisited(0), bIsPreloadAttempted(0)
 		, bIsPreloaded(0), bHasSaveCache(0), bHasBeginPrepareSaveFailed(0), bCookedPlatformDataStarted(0)
 		, bCookedPlatformDataCalled(0), bCookedPlatformDataComplete(0), bMonitorIsCooked(0)
@@ -815,16 +815,18 @@ namespace UE::Cook
 			PreloadableFile.Get()->InitializeAsync([this]()
 				{
 					TStringBuilder<NAME_SIZE> FileNameString;
-					// Note this is an unsynchronized read of this->GetFilename. It is allowed because GetFilename does
-					// not change until this is destructed, and the destructor of this waits for PreloadableFile to
-					// finish initialization
+					// Note this async callback has an read of this->GetFilename and a write of PreloadableFileOpenResult
+					// outside of a critical section. This read and write is allowed because GetFilename does
+					// not change until this is destructed, and the destructor does not run and other threads do not read
+					// or write PreloadableFileOpenResult until after PreloadableFile.Get() has finished initialization
+					// and this callback is therefore complete.
+					// The code that accomplishes that waiting is in TryPreload (IsInitialized) and ClearPreload (ReleaseCache)
 					this->GetFileName().ToString(FileNameString);
 					FPackagePath PackagePath = FPackagePath::FromLocalPath(FileNameString);
 					FOpenPackageResult Result = IPackageResourceManager::Get().OpenReadPackage(PackagePath);
 					if (Result.Archive)
 					{
-						// PreloadableFileFormat is an atomic, so this write is safe with other reads
-						this->PreloadableFileFormat = Result.Format;
+						this->PreloadableFileOpenResult.CopyMetaData(Result);
 					}
 					return Result.Archive.Release();
 				},
@@ -850,19 +852,19 @@ namespace UE::Cook
 			UE_LOG(LogCook, Warning, TEXT("Failed to find file when preloading %s."), *GetFileName().ToString());
 			SetIsPreloadAttempted(true);
 			PreloadableFile.Reset(*this);
-			PreloadableFileFormat = EPackageFormat::Binary;
+			PreloadableFileOpenResult = FOpenPackageResult();
 			return true;
 		}
 
 		TStringBuilder<NAME_SIZE> FileNameString;
 		GetFileName().ToString(FileNameString);
 		if (!IPackageResourceManager::TryRegisterPreloadableArchive(FPackagePath::FromLocalPath(FileNameString),
-			FilePtr, PreloadableFileFormat))
+			FilePtr, PreloadableFileOpenResult))
 		{
 			UE_LOG(LogCook, Warning, TEXT("Failed to register %s for preload."), *GetFileName().ToString());
 			SetIsPreloadAttempted(true);
 			PreloadableFile.Reset(*this);
-			PreloadableFileFormat = EPackageFormat::Binary;
+			PreloadableFileOpenResult = FOpenPackageResult();
 			return true;
 		}
 
@@ -911,7 +913,7 @@ namespace UE::Cook
 		}
 
 		PreloadableFile.Reset(*this);
-		PreloadableFileFormat = EPackageFormat::Binary;
+		PreloadableFileOpenResult = FOpenPackageResult();
 		SetIsPreloaded(false);
 		SetIsPreloadAttempted(false);
 	}
diff --git a/Engine/Source/Editor/UnrealEd/Private/Cooker/CookPackageData.h b/Engine/Source/Editor/UnrealEd/Private/Cooker/CookPackageData.h
--- a/Engine/Source/Editor/UnrealEd/Private/Cooker/CookPackageData.h
+++ b/Engine/Source/Editor/UnrealEd/Private/Cooker/CookPackageData.h
@@ -543,7 +543,7 @@ namespace UE::Cook
 		FTrackedPreloadableFilePtr PreloadableFile;
 		int32 NumPendingCookedPlatformData = 0;
 		int32 CookedPlatformDataNextIndex = 0;
-		std::atomic<EPackageFormat> PreloadableFileFormat;
+		FOpenPackageResult PreloadableFileOpenResult;
 		FInstigator Instigator;
 
 		uint32 State : int32(EPackageState::BitCount);
diff --git a/Engine/Source/Editor/UnrealEd/Private/EditorDomain/EditorDomain.cpp b/Engine/Source/Editor/UnrealEd/Private/EditorDomain/EditorDomain.cpp
--- a/Engine/Source/Editor/UnrealEd/Private/EditorDomain/EditorDomain.cpp
+++ b/Engine/Source/Editor/UnrealEd/Private/EditorDomain/EditorDomain.cpp
@@ -345,10 +345,11 @@ FOpenPackageResult FEditorDomain::OpenReadPackage(const FPackagePath& PackagePat
 	}
 
 	const EPackageFormat Format = bHasEditorSource ? EPackageFormat::Binary : Result->GetPackageFormat();
-	return FOpenPackageResult{ TUniquePtr<FArchive>(Result), Format };
+	bool bNeedsEngineVersionChecks = bHasEditorSource ? false : (Result->GetPackageSource() != EPackageSource::Editor);
+	return FOpenPackageResult{ TUniquePtr<FArchive>(Result), Format, bNeedsEngineVersionChecks};
 }
 
-IAsyncReadFileHandle* FEditorDomain::OpenAsyncReadPackage(const FPackagePath& PackagePath, EPackageSegment PackageSegment)
+FOpenAsyncPackageResult FEditorDomain::OpenAsyncReadPackage(const FPackagePath& PackagePath, EPackageSegment PackageSegment)
 {
 	using namespace UE::EditorDomain;
 	using namespace UE::DerivedData;
@@ -366,9 +367,18 @@ IAsyncReadFileHandle* FEditorDomain::OpenAsyncReadPackage(const FPackagePath& Pa
 		return Workspace->OpenAsyncReadPackage(PackagePath, PackageSegment);
 	}
 
+	// TODO: Change priority to Normal instead of Blocking once we have removed the GetPackageFormat below
+	// and don't need to block on the result before exiting this function
+	EPriority Priority = EPriority::Blocking;
+	FEditorDomainAsyncReadFileHandle* Result =
+		new FEditorDomainAsyncReadFileHandle(Locks, PackagePath, PackageSource, Priority);
+	const bool bHasEditorSource = (PackageSource->Source == EPackageSource::Editor);
+
+	// Unlock before requesting the package because the completion callback takes the lock.
+	ScopeLock.Unlock();
+
 	// Fetch meta-data only in the initial request
 	ECachePolicy SkipFlags = ECachePolicy::SkipData & ~ECachePolicy::SkipMeta;
-	FEditorDomainAsyncReadFileHandle* Result = new FEditorDomainAsyncReadFileHandle(Locks, PackagePath, PackageSource);
 	RequestEditorDomainPackage(PackagePath, PackageSource->Digest,
 		SkipFlags, Result->GetRequestOwner(),
 		[Result](FCacheGetCompleteParams&& Params)
@@ -377,7 +387,9 @@ IAsyncReadFileHandle* FEditorDomain::OpenAsyncReadPackage(const FPackagePath& Pa
 			Result->OnRecordRequestComplete(MoveTemp(Params));
 		});
 
-	return Result;
+	const EPackageFormat Format = bHasEditorSource ? EPackageFormat::Binary : Result->GetPackageFormat();
+	bool bNeedsEngineVersionChecks = bHasEditorSource ? false : (Result->GetPackageSource() != EPackageSource::Editor);
+	return FOpenAsyncPackageResult{ TUniquePtr<IAsyncReadFileHandle>(Result), Format, bNeedsEngineVersionChecks };
 }
 
 IMappedFileHandle* FEditorDomain::OpenMappedHandleToPackage(const FPackagePath& PackagePath,
@@ -402,7 +414,7 @@ bool FEditorDomain::DoesExternalResourceExist(EPackageExternalResource ResourceT
 	return Workspace->DoesExternalResourceExist(ResourceType, Identifier);
 }
 
-IAsyncReadFileHandle* FEditorDomain::OpenAsyncReadExternalResource(
+FOpenAsyncPackageResult FEditorDomain::OpenAsyncReadExternalResource(
 	EPackageExternalResource ResourceType, FStringView Identifier)
 {
 	return Workspace->OpenAsyncReadExternalResource(ResourceType, Identifier);
@@ -553,4 +565,4 @@ void FEditorDomain::OnAssetUpdatedOnDisk(const FAssetData& AssetData)
 	}
 	FScopeLock ScopeLock(&Locks->Lock);
 	PackageSources.Remove(PackageName);
-}
\ No newline at end of file
+}
diff --git a/Engine/Source/Editor/UnrealEd/Private/EditorDomain/EditorDomain.h b/Engine/Source/Editor/UnrealEd/Private/EditorDomain/EditorDomain.h
--- a/Engine/Source/Editor/UnrealEd/Private/EditorDomain/EditorDomain.h
+++ b/Engine/Source/Editor/UnrealEd/Private/EditorDomain/EditorDomain.h
@@ -45,6 +45,14 @@ DECLARE_LOG_CATEGORY_EXTERN(LogEditorDomain, Log, All);
 class FEditorDomain : public IPackageResourceManager, public FTickableEditorObject
 {
 public:
+	/** Different options for which domain a package comes from. */
+	enum class EPackageSource
+	{
+		Undecided,
+		Workspace,
+		Editor
+	};
+
 	FEditorDomain();
 	virtual ~FEditorDomain();
 
@@ -60,14 +68,14 @@ public:
 		FPackagePath* OutUpdatedPath = nullptr) override;
 	virtual FOpenPackageResult OpenReadPackage(const FPackagePath& PackagePath, EPackageSegment PackageSegment,
 		FPackagePath* OutUpdatedPath = nullptr) override;
-	virtual IAsyncReadFileHandle* OpenAsyncReadPackage(const FPackagePath& PackagePath,
+	virtual FOpenAsyncPackageResult OpenAsyncReadPackage(const FPackagePath& PackagePath,
 		EPackageSegment PackageSegment) override;
 	virtual IMappedFileHandle* OpenMappedHandleToPackage(const FPackagePath& PackagePath,
 		EPackageSegment PackageSegment, FPackagePath* OutUpdatedPath = nullptr) override;
 	virtual bool TryMatchCaseOnDisk(const FPackagePath& PackagePath, FPackagePath* OutNormalizedPath = nullptr) override;
 	virtual TUniquePtr<FArchive> OpenReadExternalResource(EPackageExternalResource ResourceType, FStringView Identifier) override;
 	virtual bool DoesExternalResourceExist(EPackageExternalResource ResourceType, FStringView Identifier) override;
-	virtual IAsyncReadFileHandle* OpenAsyncReadExternalResource(
+	virtual FOpenAsyncPackageResult OpenAsyncReadExternalResource(
 		EPackageExternalResource ResourceType, FStringView Identifier) override;
 	virtual void FindPackagesRecursive(TArray<TPair<FPackagePath, EPackageSegment>>& OutPackages, FStringView PackageMount,
 		FStringView FileMount, FStringView RootRelPath, FStringView BasenameWildcard) override;
@@ -99,13 +107,6 @@ private:
 		FEditorDomain* Owner;
 		FCriticalSection Lock;
 	};
-	/** Different options for which domain a package comes from. */
-	enum class EPackageSource
-	{
-		Undecided,
-		Workspace,
-		Editor
-	};
 	/**
 	 * Data about which domain a package comes from. Multiple queries of the same
 	 * package have to align (for e.g. BulkData offsets) so we have to keep track of this.
diff --git a/Engine/Source/Editor/UnrealEd/Private/EditorDomain/EditorDomainArchive.cpp b/Engine/Source/Editor/UnrealEd/Private/EditorDomain/EditorDomainArchive.cpp
--- a/Engine/Source/Editor/UnrealEd/Private/EditorDomain/EditorDomainArchive.cpp
+++ b/Engine/Source/Editor/UnrealEd/Private/EditorDomain/EditorDomainArchive.cpp
@@ -120,6 +120,7 @@ bool FEditorDomainPackageSegments::Serialize(void* V, int64 Length, bool bIsErro
 			{
 				FMemory::Memcpy(V, static_cast<const uint8*>(MRUSegment->Data.GetData()) + SegmentPos, LengthInSegment);
 				Pos += LengthInSegment;
+				V = static_cast<void*>(static_cast<uint8*>(V) + LengthInSegment);
 			}
 			else
 			{
@@ -675,6 +676,13 @@ EPackageFormat FEditorDomainReadArchive::GetPackageFormat() const
 	return PackageFormat;
 }
 
+FEditorDomain::EPackageSource FEditorDomainReadArchive::GetPackageSource() const
+{
+	WaitForReady();
+	return Segments.GetSource() == FEditorDomainPackageSegments::ESource::Segments ?
+		FEditorDomain::EPackageSource::Editor : FEditorDomain::EPackageSource::Workspace;
+}
+
 /** An IAsyncReadRequest SizeRequest that returns a value known at construction time. */
 class FAsyncSizeRequestConstant : public IAsyncReadRequest
 {
@@ -736,8 +744,9 @@ protected:
 
 
 FEditorDomainAsyncReadFileHandle::FEditorDomainAsyncReadFileHandle(const TRefCountPtr<FEditorDomain::FLocks>& InLocks,
-	const FPackagePath& InPackagePath, const TRefCountPtr<FEditorDomain::FPackageSource>& InPackageSource)
-	: Segments(InLocks, InPackagePath, InPackageSource, UE::DerivedData::EPriority::Normal)
+	const FPackagePath& InPackagePath, const TRefCountPtr<FEditorDomain::FPackageSource>& InPackageSource,
+	UE::DerivedData::EPriority Priority)
+	: Segments(InLocks, InPackagePath, InPackageSource, Priority)
 {
 }
 
@@ -749,16 +758,37 @@ void FEditorDomainAsyncReadFileHandle::OnRecordRequestComplete(UE::DerivedData::
 	);
 }
 
+EPackageFormat FEditorDomainAsyncReadFileHandle::GetPackageFormat() const
+{
+	if (Segments.GetSource() == FEditorDomainPackageSegments::ESource::Uninitialized)
+	{
+		Segments.WaitForReady();
+	}
+	return PackageFormat;
+}
+
+FEditorDomain::EPackageSource FEditorDomainAsyncReadFileHandle::GetPackageSource() const
+{
+	if (Segments.GetSource() == FEditorDomainPackageSegments::ESource::Uninitialized)
+	{
+		Segments.WaitForReady();
+	}
+	return Segments.GetSource() == FEditorDomainPackageSegments::ESource::Segments ?
+		FEditorDomain::EPackageSource::Editor : FEditorDomain::EPackageSource::Workspace;
+}
+
 void FEditorDomainAsyncReadFileHandle::CreateSegmentData(bool bValid)
 {
+	PackageFormat = EPackageFormat::Binary;
 }
 
 bool FEditorDomainAsyncReadFileHandle::TryCreateFallbackData(FEditorDomain& EditorDomain)
 {
 	IPackageResourceManager& Workspace = *EditorDomain.Workspace;
-	IAsyncReadFileHandle* Result = Workspace.OpenAsyncReadPackage(Segments.GetPackagePath(), EPackageSegment::Header);
-	check(Result);
-	InnerArchive.Reset(Result);
+	FOpenAsyncPackageResult Result = Workspace.OpenAsyncReadPackage(Segments.GetPackagePath(), EPackageSegment::Header);
+	check(Result.Handle);
+	InnerArchive = MoveTemp(Result.Handle);
+	PackageFormat = Result.Format;
 	return true;
 }
 
diff --git a/Engine/Source/Editor/UnrealEd/Private/EditorDomain/EditorDomainArchive.h b/Engine/Source/Editor/UnrealEd/Private/EditorDomain/EditorDomainArchive.h
--- a/Engine/Source/Editor/UnrealEd/Private/EditorDomain/EditorDomainArchive.h
+++ b/Engine/Source/Editor/UnrealEd/Private/EditorDomain/EditorDomainArchive.h
@@ -207,6 +207,7 @@ public:
 	void OnRecordRequestComplete(UE::DerivedData::FCacheGetCompleteParams&& Params);
 	/** Get the PackageFormat, which depends on the domain the data is read from. */
 	EPackageFormat GetPackageFormat() const;
+	FEditorDomain::EPackageSource GetPackageSource() const;
 
 	// FArchive interface
 	virtual void Seek(int64 InPos) override;
@@ -218,7 +219,6 @@ public:
 	virtual void Flush() override;
 	virtual void FlushCache() override;
 	virtual bool Precache(int64 PrecacheOffset, int64 PrecacheSize) override;
-	virtual bool NeedsEngineVersionChecks() const override { return false; }
 
 private:
 	void WaitForReady() const;
@@ -238,12 +238,16 @@ class FEditorDomainAsyncReadFileHandle : public IAsyncReadFileHandle
 {
 public:
 	FEditorDomainAsyncReadFileHandle(const TRefCountPtr<FEditorDomain::FLocks>& InLocks,
-		const FPackagePath& InPackagePath, const TRefCountPtr<FEditorDomain::FPackageSource>& InPackageSource);
+		const FPackagePath& InPackagePath, const TRefCountPtr<FEditorDomain::FPackageSource>& InPackageSource,
+		UE::DerivedData::EPriority Priority);
 
 	/** Get the request owner for requests that will feed this archive. */
 	UE::DerivedData::IRequestOwner& GetRequestOwner() { return Segments.GetRequestOwner(); }
 	/** Callback from the request of the Record; set whether we're reading from EditorDomain bytes or WorkspaceDomain archive. */
 	void OnRecordRequestComplete(UE::DerivedData::FCacheGetCompleteParams&& Params);
+	/** Get the PackageFormat, which depends on the domain the data is read from. */
+	EPackageFormat GetPackageFormat() const;
+	FEditorDomain::EPackageSource GetPackageSource() const;
 
 	// IAsyncReadFileHandle interface
 	virtual IAsyncReadRequest* SizeRequest(FAsyncFileCallBack* CompleteCallback = nullptr) override;
@@ -258,4 +262,5 @@ private:
 
 	FEditorDomainPackageSegments Segments;
 	TUniquePtr<IAsyncReadFileHandle> InnerArchive;
+	EPackageFormat PackageFormat = EPackageFormat::Binary;
 };
\ No newline at end of file
diff --git a/Engine/Source/Runtime/Core/Public/Serialization/Archive.h b/Engine/Source/Runtime/Core/Public/Serialization/Archive.h
--- a/Engine/Source/Runtime/Core/Public/Serialization/Archive.h
+++ b/Engine/Source/Runtime/Core/Public/Serialization/Archive.h
@@ -1910,12 +1910,6 @@ public:
 	 */
 	virtual bool IsUsingEventDrivenLoader() const;
 
-	/**
-	 * Used by LinkerLoad to decide whether it needs to check its Loader for FEngineVersion::Current().IsCompatibleWith.
-	 * EditorDomain Loaders are selected based on CustomVersions and Schema and do not need to check minor version changes.
-	 */
-	virtual bool NeedsEngineVersionChecks() const { return true; }
-
 #if USE_STABLE_LOCALIZATION_KEYS
 	using FArchiveState::SetLocalizationNamespace;
 	using FArchiveState::GetLocalizationNamespace;
diff --git a/Engine/Source/Runtime/CoreUObject/Private/Serialization/AsyncLoading.cpp b/Engine/Source/Runtime/CoreUObject/Private/Serialization/AsyncLoading.cpp
--- a/Engine/Source/Runtime/CoreUObject/Private/Serialization/AsyncLoading.cpp
+++ b/Engine/Source/Runtime/CoreUObject/Private/Serialization/AsyncLoading.cpp
@@ -7005,7 +7005,9 @@ FAsyncArchive::FAsyncArchive(const FPackagePath& InPackagePath, FLinkerLoad* InO
 	, HeaderSize(0)
 	, HeaderSizeWhenReadingExportsFromSplitFile(0)
 	, LoadPhase(ELoadPhase::WaitingForSize)
+	, LoadError(ELoadError::Unknown)
 	, bCookedForEDLInEditor(false)
+	, bNeedsEngineVersionChecks(true)
 	, PackagePath(InPackagePath)
 	, OpenTime(FPlatformTime::Seconds())
 	, SummaryReadTime(0.0)
@@ -7014,8 +7016,15 @@ FAsyncArchive::FAsyncArchive(const FPackagePath& InPackagePath, FLinkerLoad* InO
 	, OwnerLinker(InOwner)
 {
 	LogItem(TEXT("Open"));
-	Handle = IPackageResourceManager::Get().OpenAsyncReadPackage(PackagePath, EPackageSegment::Header);
+	FOpenAsyncPackageResult OpenResult = IPackageResourceManager::Get().OpenAsyncReadPackage(PackagePath, EPackageSegment::Header);
+	Handle = OpenResult.Handle.Release();
 	check(Handle); // OpenAsyncReadPackage guarantees a non-null return value; the handle will fail to read later if the path does not exist
+	if (OpenResult.Format != EPackageFormat::Binary)
+	{
+		SetError();
+		LoadError = ELoadError::UnsupportedFormat;
+	}
+	bNeedsEngineVersionChecks = OpenResult.bNeedsEngineVersionChecks;
 
 	ReadCallbackFunction = [this](bool bWasCancelled, IAsyncReadRequest* Request)
 	{
@@ -7054,6 +7063,10 @@ void FAsyncArchive::ReadCallback(bool bWasCancelled, IAsyncReadRequest* Request)
 {
 	if (bWasCancelled || IsError())
 	{
+		if (!IsError())
+		{
+			LoadError = ELoadError::Cancelled;
+		}
 		SetError();
 		return; // we don't do much with this, the code on the other thread knows how to deal with my request
 	}
@@ -7064,6 +7077,7 @@ void FAsyncArchive::ReadCallback(bool bWasCancelled, IAsyncReadRequest* Request)
 		if (FileSize < 32)
 		{
 			SetError();
+			LoadError = FileSize > 0 ? ELoadError::CorruptData : ELoadError::FileDoesNotExist;
 		}
 		else
 		{
@@ -7101,6 +7115,7 @@ void FAsyncArchive::ReadCallback(bool bWasCancelled, IAsyncReadRequest* Request)
 		if (!Mem)
 		{
 			SetError();
+			LoadError = ELoadError::CorruptData;
 			FPlatformMisc::MemoryBarrier();
 			LoadPhase = ELoadPhase::WaitingForHeader;
 		}
@@ -7112,6 +7127,7 @@ void FAsyncArchive::ReadCallback(bool bWasCancelled, IAsyncReadRequest* Request)
 			if (Ar.IsError() || Sum.TotalHeaderSize > FileSize || Sum.IsFileVersionTooOld())
 			{
 				SetError();
+				LoadError = ELoadError::CorruptData;
 			}
 			else
 			{
@@ -7320,6 +7336,7 @@ void FAsyncArchive::CompleteRead()
 		if (!Mem)
 		{
 			SetError();
+			LoadError = ELoadError::CorruptData;
 		}
 		else
 		{
@@ -7507,6 +7524,7 @@ bool FAsyncArchive::PrecacheInternal(int64 RequestOffset, int64 RequestSize, boo
 	if (ReadRequestSize <= 0)
 	{
 		SetError();
+		LoadError = ELoadError::CorruptData;
 		return true;
 	}
 	double StartTime = FPlatformTime::Seconds();
@@ -7566,7 +7584,8 @@ void FAsyncArchive::FirstExportStarting()
 
 		HeaderSizeWhenReadingExportsFromSplitFile = HeaderSize;
 
-		Handle = IPackageResourceManager::Get().OpenAsyncReadPackage(PackagePath, EPackageSegment::Exports);
+		FOpenAsyncPackageResult OpenResult = IPackageResourceManager::Get().OpenAsyncReadPackage(PackagePath, EPackageSegment::Exports);
+		Handle = OpenResult.Handle.Release();
 		check(Handle); // OpenAsyncReadPackage guarantees a non-null return value; the handle will fail to read later if the path does not exist
 
 		check(!SizeRequestPtr);
@@ -7588,7 +7607,8 @@ IAsyncReadRequest* FAsyncArchive::MakeEventDrivenPrecacheRequest(int64 Offset, i
 		IAsyncReadFileHandle* NewHandle;
 		{
 			double StartTime = FPlatformTime::Seconds();
-			NewHandle = IPackageResourceManager::Get().OpenAsyncReadPackage(PackagePath, EPackageSegment::Exports);
+			FOpenAsyncPackageResult OpenResult = IPackageResourceManager::Get().OpenAsyncReadPackage(PackagePath, EPackageSegment::Exports);
+			NewHandle = OpenResult.Handle.Release();
 			check(NewHandle); // OpenAsyncReadPackage guarantees a non-null return value; the handle will fail to read later if the path does not exist
 			LogItem(TEXT("Open UExp"), Offset - HeaderSizeWhenReadingExportsFromSplitFile, BytesToRead, StartTime);
 		}
diff --git a/Engine/Source/Runtime/CoreUObject/Private/Serialization/AsyncLoadingPrivate.h b/Engine/Source/Runtime/CoreUObject/Private/Serialization/AsyncLoadingPrivate.h
--- a/Engine/Source/Runtime/CoreUObject/Private/Serialization/AsyncLoadingPrivate.h
+++ b/Engine/Source/Runtime/CoreUObject/Private/Serialization/AsyncLoadingPrivate.h
@@ -19,6 +19,14 @@ public:
 		WaitingForFirstExport,
 		ProcessingExports,
 	};
+	enum class ELoadError : uint8
+	{
+		Unknown,
+		UnsupportedFormat,
+		FileDoesNotExist,
+		CorruptData,
+		Cancelled,
+	};
 
 	FAsyncArchive(const FPackagePath& InPackagePath, FLinkerLoad* InOwner, TFunction<void()>&& InSummaryReadyCallback);
 	virtual ~FAsyncArchive ();
@@ -59,6 +67,16 @@ public:
 		return bCookedForEDLInEditor;
 	}
 
+	ELoadError GetLoadError() const
+	{
+		return LoadError;
+	}
+	bool NeedsEngineVersionChecks() const
+	{
+		return bNeedsEngineVersionChecks;
+	}
+
+
 private:
 #if DEVIRTUALIZE_FLinkerLoad_Serialize
 	/**
@@ -112,9 +130,12 @@ private:
 	int64 HeaderSizeWhenReadingExportsFromSplitFile;
 
 	ELoadPhase LoadPhase;
+	ELoadError LoadError;
 
 	/** If true, this package is a cooked EDL package loaded in uncooked builds */
 	bool bCookedForEDLInEditor;
+	/** True if the linker should do version and corruption checks on bytes of this archive. */
+	bool bNeedsEngineVersionChecks;
 
 	FAsyncFileCallBack ReadCallbackFunction;
 	/** Cached PackagePath for debugging.												*/
diff --git a/Engine/Source/Runtime/CoreUObject/Private/Serialization/BulkData.cpp b/Engine/Source/Runtime/CoreUObject/Private/Serialization/BulkData.cpp
--- a/Engine/Source/Runtime/CoreUObject/Private/Serialization/BulkData.cpp
+++ b/Engine/Source/Runtime/CoreUObject/Private/Serialization/BulkData.cpp
@@ -1883,15 +1883,17 @@ IAsyncReadFileHandle* FUntypedBulkData::OpenAsyncReadHandle() const
 	}
 #endif //WITH_IOSTORE_IN_EDITOR
 	
+	FOpenAsyncPackageResult OpenResult;
 	if (IsInExternalResource())
 	{
-		return IPackageResourceManager::Get().OpenAsyncReadExternalResource(EPackageExternalResource::WorkspaceDomainFile,
+		OpenResult = IPackageResourceManager::Get().OpenAsyncReadExternalResource(EPackageExternalResource::WorkspaceDomainFile,
 			GetPackagePath().GetPackageName());
 	}
 	else
 	{
-		return IPackageResourceManager::Get().OpenAsyncReadPackage(GetPackagePath(), GetPackageSegment());
+		OpenResult = IPackageResourceManager::Get().OpenAsyncReadPackage(GetPackagePath(), GetPackageSegment());
 	}
+	return OpenResult.Handle.Release();
 }
 
 IBulkDataIORequest* FUntypedBulkData::CreateStreamingRequest(EAsyncIOPriorityAndFlags Priority, FBulkDataIORequestCallBack* CompleteCallback, uint8* UserSuppliedMemory) const
@@ -1924,16 +1926,17 @@ IBulkDataIORequest* FUntypedBulkData::CreateStreamingRequest(int64 OffsetInBulkD
 	UE_CLOG(IsStoredCompressedOnDisk(), LogSerialization, Fatal, TEXT("Package level compression is no longer supported (%s)."), *PackagePath.GetDebugName(PackageSegment));
 	UE_CLOG(GetBulkDataSize() <= 0, LogSerialization, Error, TEXT("(%s) has invalid bulk data size."), *PackagePath.GetDebugName(PackageSegment));
 
-	IAsyncReadFileHandle* IORequestHandle;
+	FOpenAsyncPackageResult OpenResult;
 	if (IsInExternalResource())
 	{
-		IORequestHandle = IPackageResourceManager::Get().OpenAsyncReadExternalResource(
+		OpenResult = IPackageResourceManager::Get().OpenAsyncReadExternalResource(
 			EPackageExternalResource::WorkspaceDomainFile, PackagePath.GetPackageName());
 	}
 	else
 	{
-		IORequestHandle = IPackageResourceManager::Get().OpenAsyncReadPackage(PackagePath, PackageSegment);
+		OpenResult = IPackageResourceManager::Get().OpenAsyncReadPackage(PackagePath, PackageSegment);
 	}
+	IAsyncReadFileHandle* IORequestHandle = OpenResult.Handle.Release();
 	check(IORequestHandle); // this generally cannot fail because it is async
 
 	if (IORequestHandle == nullptr)
@@ -1989,7 +1992,8 @@ IBulkDataIORequest* FUntypedBulkData::CreateStreamingRequestForRange(const FPack
 
 	// TODO: The caller is assuming that their specified PackageSegments are the way in which bulkdata is stored for their package
 	// To allow more flexible bulkdata storage, we will need to eliminate this interface and have them provide an array of bulkdatas.
-	IAsyncReadFileHandle* IORequestHandle = IPackageResourceManager::Get().OpenAsyncReadPackage(PackagePath, PackageSegment);
+	FOpenAsyncPackageResult OpenResult = IPackageResourceManager::Get().OpenAsyncReadPackage(PackagePath, PackageSegment);
+	IAsyncReadFileHandle* IORequestHandle = OpenResult.Handle.Release();
 	check(IORequestHandle); // this generally cannot fail because it is async
 
 	if (IORequestHandle == nullptr)
diff --git a/Engine/Source/Runtime/CoreUObject/Private/Serialization/BulkData2.cpp b/Engine/Source/Runtime/CoreUObject/Private/Serialization/BulkData2.cpp
--- a/Engine/Source/Runtime/CoreUObject/Private/Serialization/BulkData2.cpp
+++ b/Engine/Source/Runtime/CoreUObject/Private/Serialization/BulkData2.cpp
@@ -1314,7 +1314,8 @@ IAsyncReadFileHandle* FBulkDataBase::OpenAsyncReadHandle() const
 	}
 	else
 	{
-		return IPackageResourceManager::Get().OpenAsyncReadPackage(GetPackagePath(), GetPackageSegmentFromFlags());
+		FOpenAsyncPackageResult OpenResult = IPackageResourceManager::Get().OpenAsyncReadPackage(GetPackagePath(), GetPackageSegmentFromFlags());
+		return OpenResult.Handle.Release();
 	}
 }
 
@@ -1360,7 +1361,8 @@ IBulkDataIORequest* FBulkDataBase::CreateStreamingRequest(int64 OffsetInBulkData
 		UE_CLOG(IsStoredCompressedOnDisk(), LogSerialization, Fatal, TEXT("Package level compression is no longer supported (%s)."), *PackagePath.GetDebugName(PackageSegment));
 		UE_CLOG(BulkDataSize <= 0, LogSerialization, Error, TEXT("(%s) has invalid bulk data size."), *PackagePath.GetDebugName(PackageSegment));
 
-		IAsyncReadFileHandle* IORequestHandle = IPackageResourceManager::Get().OpenAsyncReadPackage(PackagePath, PackageSegment);
+		FOpenAsyncPackageResult OpenResult = IPackageResourceManager::Get().OpenAsyncReadPackage(PackagePath, PackageSegment);
+		IAsyncReadFileHandle* IORequestHandle = OpenResult.Handle.Release();
 		checkf(IORequestHandle, TEXT("OpenAsyncRead failed")); // An assert as there shouldn't be a way for this to fail
 
 		if (IORequestHandle == nullptr)
diff --git a/Engine/Source/Runtime/CoreUObject/Private/UObject/LinkerLoad.cpp b/Engine/Source/Runtime/CoreUObject/Private/UObject/LinkerLoad.cpp
--- a/Engine/Source/Runtime/CoreUObject/Private/UObject/LinkerLoad.cpp
+++ b/Engine/Source/Runtime/CoreUObject/Private/UObject/LinkerLoad.cpp
@@ -454,7 +454,7 @@ FLinkerLoad* FLinkerLoad::CreateLinker(FUObjectSerializeContext* LoadContext, UP
 		{
 			// The linker can't have an associated loader here if we have a loader override
 			check(!Linker->Loader);
-			Linker->SetLoader(InLoader);
+			Linker->SetLoader(InLoader, true /* bInLoaderNeedsEngineVersionChecks */);
 			// Set the basic archive flags on the linker
 			Linker->ResetStatusInfo();
 		}
@@ -477,9 +477,10 @@ void FLinkerLoad::SetPackagePath(const FPackagePath& InPackagePath)
 	PRAGMA_ENABLE_DEPRECATION_WARNINGS
 }
 
-void FLinkerLoad::SetLoader(FArchive* InLoader)
+void FLinkerLoad::SetLoader(FArchive* InLoader, bool bInLoaderNeedsEngineVersionChecks)
 {
 	Loader = InLoader;
+	bLoaderNeedsEngineVersionChecks = bInLoaderNeedsEngineVersionChecks;
 
 	check(StructuredArchive == nullptr);
 	check(!StructuredArchiveRootRecord.IsSet());
@@ -874,6 +875,7 @@ FLinkerLoad::FLinkerLoad(UPackage* InParent, const FPackagePath& InPackagePath,
 , bTimeLimitExceeded(false)
 , bUseTimeLimit(false)
 , bUseFullTimeLimit(false)
+, bLoaderNeedsEngineVersionChecks(true)
 , IsTimeLimitExceededCallCount(0)
 , TimeLimit(0.0f)
 , TickStartTime(0.0)
@@ -1024,29 +1026,32 @@ FLinkerLoad::ELinkerStatus FLinkerLoad::CreateLoader(
 		// If want to be able to load cooked data in the editor we need to use FAsyncArchive which supports EDL cooked packages,
 		// otherwise the generic file reader is faster in the editor so use that
 		bool bCanUseAsyncLoader = FPlatformProperties::RequiresCookedData() || GAllowCookedDataInEditorBuilds;
-#if WITH_TEXT_ARCHIVE_SUPPORT
-		bCanUseAsyncLoader &= !FPackageName::IsTextPackageExtension(PackagePath.GetHeaderExtension());
-#endif //if WITH_TEXT_ARCHIVE_SUPPORT
 
 		if (bCanUseAsyncLoader)
 		{
-			Loader = new FAsyncArchive(GetPackagePath(), this,
-				GEventDrivenLoaderEnabled ? Forward<TFunction<void()>>(InSummaryReadyCallback) : TFunction<void()>([]() {})
-			);
-			if (Loader->IsError())
+			FAsyncArchive* AsyncArchive = new FAsyncArchive(GetPackagePath(), this,
+				GEventDrivenLoaderEnabled ? Forward<TFunction<void()>>(InSummaryReadyCallback) : TFunction<void()>([]() {}));
+			Loader = AsyncArchive; // We're only allowed to delete any FAsyncArchive with this->DestroyLoader
+			bLoaderNeedsEngineVersionChecks = AsyncArchive->NeedsEngineVersionChecks();
+			if (AsyncArchive->IsError())
 			{
+				bool bRetryWithNormalArchive = AsyncArchive->GetLoadError() == FAsyncArchive::ELoadError::UnsupportedFormat;
 				DestroyLoader();
-				UE_LOG(LogLinker, Warning, TEXT("Error opening file '%s'."), *GetDebugName());
-				return LINKER_Failed;
+				bCanUseAsyncLoader = false;
+				if (!bRetryWithNormalArchive)
+				{
+					UE_LOG(LogLinker, Warning, TEXT("Error opening file '%s'."), *GetDebugName());
+					return LINKER_Failed;
+				}
 			}
 		}
-		else
+		if (!Loader)
 		{
 			FOpenPackageResult OpenResult;
 #if WITH_EDITOR
 			if (FLinkerLoad::GetPreloadingEnabled() && FLinkerLoad::TryGetPreloadedLoader(GetPackagePath(), OpenResult))
 			{
-				// Loader set by TryGetPreloadedLoader
+				// OpenResult set by TryGetPreloadedLoader
 			}
 			else
 #endif
@@ -1054,6 +1059,7 @@ FLinkerLoad::ELinkerStatus FLinkerLoad::CreateLoader(
 				OpenResult = IPackageResourceManager::Get().OpenReadPackage(GetPackagePath());
 			}
 			Loader = OpenResult.Archive.Release();
+			bLoaderNeedsEngineVersionChecks = OpenResult.bNeedsEngineVersionChecks;
 			if (!Loader || Loader->IsError())
 			{
 				if (Loader)
@@ -1125,7 +1131,7 @@ FLinkerLoad::ELinkerStatus FLinkerLoad::CreateLoader(
 			bIsAsyncLoader = bCanUseAsyncLoader;
 		}
 
-		SetLoader(Loader);
+		SetLoader(Loader, bLoaderNeedsEngineVersionChecks);
 
 		check(Loader);
 		check(!Loader->IsError());
@@ -1226,9 +1232,8 @@ FLinkerLoad::ELinkerStatus FLinkerLoad::SerializePackageFileSummaryInternal()
 		}
 	}
 
-	bool bLoaderVersionCheck = Loader->NeedsEngineVersionChecks();
 	// Don't load packages that are only compatible with an engine version newer than the current one.
-	if (bLoaderVersionCheck && GEnforcePackageCompatibleVersionCheck && !FEngineVersion::Current().IsCompatibleWith(Summary.CompatibleWithEngineVersion))
+	if (bLoaderNeedsEngineVersionChecks && GEnforcePackageCompatibleVersionCheck && !FEngineVersion::Current().IsCompatibleWith(Summary.CompatibleWithEngineVersion))
 	{
 		UE_LOG(LogLinker, Warning, TEXT("Asset '%s' has been saved with a newer engine and can't be loaded. CurrentEngineVersion: %s (Licensee=%d). AssetEngineVersion: %s (Licensee=%d)"),
 			*GetDebugName(),
@@ -1244,7 +1249,7 @@ FLinkerLoad::ELinkerStatus FLinkerLoad::SerializePackageFileSummaryInternal()
 	SetUseUnversionedPropertySerialization(bUseUnversionedProperties);
 	Loader->SetUseUnversionedPropertySerialization(bUseUnversionedProperties);
 
-	if (bLoaderVersionCheck && !FPlatformProperties::RequiresCookedData()
+	if (bLoaderNeedsEngineVersionChecks && !FPlatformProperties::RequiresCookedData()
 		&& !Summary.SavedByEngineVersion.HasChangelist() && FEngineVersion::Current().HasChangelist())
 	{
 		// This warning can be disabled in ini with [Core.System] ZeroEngineVersionWarning=False
@@ -1264,7 +1269,7 @@ FLinkerLoad::ELinkerStatus FLinkerLoad::SerializePackageFileSummaryInternal()
 	}
 
 	// Don't load packages that were saved with package version newer than the current one.
-	if (bLoaderVersionCheck && ((Summary.IsFileVersionTooNew()) || (Summary.GetFileVersionLicenseeUE() > GPackageFileLicenseeUEVersion)))
+	if (bLoaderNeedsEngineVersionChecks && ((Summary.IsFileVersionTooNew()) || (Summary.GetFileVersionLicenseeUE() > GPackageFileLicenseeUEVersion)))
 	{
 		UE_LOG(LogLinker, Warning, TEXT("Unable to load package (%s) PackageVersion %i, MaxExpected %i : LicenseePackageVersion %i, MaxExpected %i."),
 			*GetDebugName(), Summary.GetFileVersionUE().ToValue(), GPackageFileUEVersion.ToValue(), Summary.GetFileVersionLicenseeUE(), GPackageFileLicenseeUEVersion);
@@ -1303,7 +1308,7 @@ FLinkerLoad::ELinkerStatus FLinkerLoad::SerializePackageFileSummaryInternal()
 	if (!FPlatformProperties::RequiresCookedData() &&
 		// We can't check the post tag if the file is an EDL cooked package
 		!((Summary.GetPackageFlags() & PKG_FilterEditorOnly) && Summary.PreloadDependencyCount > 0 && Summary.PreloadDependencyOffset > 0)
-		&& !IsTextFormat() && bLoaderVersionCheck)
+		&& !IsTextFormat() && bLoaderNeedsEngineVersionChecks)
 	{
 		// check if this package version stored the 4-byte magic post tag
 		// get the offset of the post tag
diff --git a/Engine/Source/Runtime/CoreUObject/Private/UObject/PackageResourceManager.cpp b/Engine/Source/Runtime/CoreUObject/Private/UObject/PackageResourceManager.cpp
--- a/Engine/Source/Runtime/CoreUObject/Private/UObject/PackageResourceManager.cpp
+++ b/Engine/Source/Runtime/CoreUObject/Private/UObject/PackageResourceManager.cpp
@@ -2,6 +2,7 @@
 
 #include "UObject/PackageResourceManager.h"
 
+#include "Async/AsyncFileHandle.h"
 #include "Misc/PackageSegment.h"
 #include "Misc/PreloadableFile.h"
 #include "Misc/ScopeLock.h"
@@ -64,7 +65,7 @@ FOpenPackageResult IPackageResourceManager::OpenReadPackage(const FPackagePath&
 	return OpenReadPackage(PackagePath, EPackageSegment::Header, OutUpdatedPath);
 }
 
-IAsyncReadFileHandle* IPackageResourceManager::OpenAsyncReadPackage(const FPackagePath& PackagePath)
+FOpenAsyncPackageResult IPackageResourceManager::OpenAsyncReadPackage(const FPackagePath& PackagePath)
 {
 	return OpenAsyncReadPackage(PackagePath, EPackageSegment::Header);
 }
@@ -154,11 +155,11 @@ void IPackageResourceManager::IteratePackagesStatInLocalOnlyDirectory(FStringVie
 }
 
 #if WITH_EDITOR
-TMap<FName, TPair<TSharedPtr<FPreloadableArchive>, EPackageFormat>> IPackageResourceManager::PreloadedPaths;
+TMap<FName, TPair<TSharedPtr<FPreloadableArchive>, FOpenPackageResult>> IPackageResourceManager::PreloadedPaths;
 FCriticalSection IPackageResourceManager::PreloadedPathsLock;
 
 bool IPackageResourceManager::TryRegisterPreloadableArchive(const FPackagePath& PackagePath,
-	const TSharedPtr<FPreloadableArchive>& PreloadableArchive, EPackageFormat PackageFormat)
+	const TSharedPtr<FPreloadableArchive>& PreloadableArchive, const FOpenPackageResult& OpenResult)
 {
 	const FName PackageName = PackagePath.GetPackageFName();
 	if (PackageName.IsNone())
@@ -167,14 +168,14 @@ bool IPackageResourceManager::TryRegisterPreloadableArchive(const FPackagePath&
 	}
 	check(PreloadableArchive.IsValid());
 	FScopeLock ScopeLock(&PreloadedPathsLock);
-	TPair<TSharedPtr<FPreloadableArchive>, EPackageFormat>& ExistingData = PreloadedPaths.FindOrAdd(PackageName);
+	TPair<TSharedPtr<FPreloadableArchive>, FOpenPackageResult>& ExistingData = PreloadedPaths.FindOrAdd(PackageName);
 	TSharedPtr<FPreloadableArchive>& ExistingArchive = ExistingData.Get<0>();
-	EPackageFormat& ExistingFormat = ExistingData.Get<1>();
+	FOpenPackageResult& ExistingResult = ExistingData.Get<1>();
 	if (ExistingArchive)
 	{
 		if (ExistingArchive.Get() == PreloadableArchive.Get())
 		{
-			check(ExistingFormat == PackageFormat);
+			check(ExistingResult.Format == OpenResult.Format && ExistingResult.bNeedsEngineVersionChecks == OpenResult.bNeedsEngineVersionChecks);
 			return true;
 		}
 		else
@@ -185,7 +186,7 @@ bool IPackageResourceManager::TryRegisterPreloadableArchive(const FPackagePath&
 	else
 	{
 		ExistingArchive = PreloadableArchive;
-		ExistingFormat = PackageFormat;
+		ExistingResult.CopyMetaData(OpenResult);
 		return true;
 	}
 }
@@ -193,6 +194,7 @@ bool IPackageResourceManager::TryRegisterPreloadableArchive(const FPackagePath&
 bool IPackageResourceManager::TryTakePreloadableArchive(const FPackagePath& PackagePath, FOpenPackageResult& OutResult)
 {
 	OutResult.Format = EPackageFormat::Binary;
+	OutResult.bNeedsEngineVersionChecks = true;
 	OutResult.Archive = nullptr;
 	const FName PackageName = PackagePath.GetPackageFName();
 	if (PackageName.IsNone())
@@ -205,20 +207,20 @@ bool IPackageResourceManager::TryTakePreloadableArchive(const FPackagePath& Pack
 		return false;
 	}
 
-	TPair<TSharedPtr<FPreloadableArchive>, EPackageFormat> ExistingData;
+	TPair<TSharedPtr<FPreloadableArchive>, FOpenPackageResult> ExistingData;
 	if (!PreloadedPaths.RemoveAndCopyValue(PackageName, ExistingData))
 	{
 		return false;
 	}
 	TSharedPtr<FPreloadableArchive>& PreloadableArchive = ExistingData.Get<0>();
-	EPackageFormat PackageFormat = ExistingData.Get<1>();
+	FOpenPackageResult& OpenResult = ExistingData.Get<1>();
 	if (!PreloadableArchive || !PreloadableArchive->IsInitialized())
 	{
 		// Someone has called Close on the Archive already.
 		return false;
 	}
 
-	OutResult.Format = PackageFormat;
+	OutResult.CopyMetaData(OpenResult);
 	OutResult.Archive = TUniquePtr<FArchive>(PreloadableArchive->DetachLowerLevel());
 	// If DetachedLowerLevel returns non-null, the PreloadableArchive is in PreloadHandle mode;
 	// it is not preloading bytes, but instead is only providing a pre-opened (and possibly primed) sync handle
@@ -241,4 +243,19 @@ bool IPackageResourceManager::UnRegisterPreloadableArchive(const FPackagePath& P
 	int32 NumRemoved = PreloadedPaths.Remove(PackageName);
 	return NumRemoved > 0;
 }
-#endif
\ No newline at end of file
+#endif
+
+EPackageFormat ExtensionToPackageFormat(EPackageExtension Extension)
+{
+	switch (Extension)
+	{
+	case EPackageExtension::TextAsset: return EPackageFormat::Text;
+	case EPackageExtension::TextMap: return EPackageFormat::Text;
+	default: return EPackageFormat::Binary;
+	}
+}
+
+FOpenAsyncPackageResult::~FOpenAsyncPackageResult()
+{
+	// Defined in CPP so we can forward declare IAsyncReadFileHandle
+}
diff --git a/Engine/Source/Runtime/CoreUObject/Private/UObject/PackageResourceManagerFile.cpp b/Engine/Source/Runtime/CoreUObject/Private/UObject/PackageResourceManagerFile.cpp
--- a/Engine/Source/Runtime/CoreUObject/Private/UObject/PackageResourceManagerFile.cpp
+++ b/Engine/Source/Runtime/CoreUObject/Private/UObject/PackageResourceManagerFile.cpp
@@ -39,7 +39,7 @@ public:
 		FPackagePath* OutUpdatedPath = nullptr) override;
 	virtual FOpenPackageResult OpenReadPackage(const FPackagePath& PackagePath, EPackageSegment PackageSegment,
 		FPackagePath* OutUpdatedPath = nullptr) override;
-	virtual IAsyncReadFileHandle* OpenAsyncReadPackage(const FPackagePath& PackagePath,
+	virtual FOpenAsyncPackageResult OpenAsyncReadPackage(const FPackagePath& PackagePath,
 		EPackageSegment PackageSegment) override;
 	virtual IMappedFileHandle* OpenMappedHandleToPackage(const FPackagePath& PackagePath, EPackageSegment PackageSegment,
 		FPackagePath* OutUpdatedPath = nullptr) override;
@@ -47,7 +47,7 @@ public:
 
 	virtual TUniquePtr<FArchive> OpenReadExternalResource(EPackageExternalResource ResourceType, FStringView Identifier) override;
 	virtual bool DoesExternalResourceExist(EPackageExternalResource ResourceType, FStringView Identifier) override;
-	virtual IAsyncReadFileHandle* OpenAsyncReadExternalResource(
+	virtual FOpenAsyncPackageResult OpenAsyncReadExternalResource(
 		EPackageExternalResource ResourceType, FStringView Identifier) override;
 
 	virtual void FindPackagesRecursive(TArray<TPair<FPackagePath, EPackageSegment>>& OutPackages,
@@ -194,7 +194,7 @@ int64 FPackageResourceManagerFile::FileSize(const FPackagePath& PackagePath, EPa
 FOpenPackageResult FPackageResourceManagerFile::OpenReadPackage(const FPackagePath& PackagePath,
 	EPackageSegment PackageSegment, FPackagePath* OutUpdatedPath)
 {
-	FOpenPackageResult Result{ nullptr, EPackageFormat::Binary };
+	FOpenPackageResult Result{ nullptr, EPackageFormat::Binary, true /* bNeedsEngineVersionChecks */};
 
 	IFileManager* FileManager = &IFileManager::Get();
 	IteratePossibleFiles(PackagePath, PackageSegment, OutUpdatedPath,
@@ -226,10 +226,10 @@ FOpenPackageResult FPackageResourceManagerFile::OpenReadPackage(const FPackagePa
 	return Result;
 }
 
-IAsyncReadFileHandle* FPackageResourceManagerFile::OpenAsyncReadPackage(const FPackagePath& PackagePath,
+FOpenAsyncPackageResult FPackageResourceManagerFile::OpenAsyncReadPackage(const FPackagePath& PackagePath,
 	EPackageSegment PackageSegment)
 {
-	IAsyncReadFileHandle* Result = nullptr;
+	FOpenAsyncPackageResult Result { nullptr, EPackageFormat::Binary, true /* bNeedsEngineVersionChecks */ };
 
 	EPackageExtension Extension = EPackageExtension::Unspecified;
 	TConstArrayView<EPackageExtension> Extensions = PackagePath.GetPossibleExtensions(PackageSegment);
@@ -255,13 +255,14 @@ IAsyncReadFileHandle* FPackageResourceManagerFile::OpenAsyncReadPackage(const FP
 		{
 			FilePath << ((Extension != EPackageExtension::Custom) ?
 				FStringView(LexToString(Extension)) : PackagePath.GetExtensionString(EPackageSegment::Header));
-			Result = FPlatformFileManager::Get().GetPlatformFile().OpenAsyncRead(FilePath.ToString());
-			check(Result != nullptr); // OpenAsyncRead guarantees a non-null return value
+			Result.Handle.Reset(FPlatformFileManager::Get().GetPlatformFile().OpenAsyncRead(FilePath.ToString()));
+			check(Result.Handle); // OpenAsyncRead guarantees a non-null return value
+			Result.Format = ExtensionToPackageFormat(Extension);
 		}
 	}
-	if (Result == nullptr)
+	if (!Result.Handle)
 	{
-		Result = new FAsyncReadFileHandleNull();
+		Result.Handle.Reset(new FAsyncReadFileHandleNull());
 	}
 	return Result;
 }
@@ -358,7 +359,7 @@ bool FPackageResourceManagerFile::DoesExternalResourceExist(EPackageExternalReso
 	}
 }
 
-IAsyncReadFileHandle* FPackageResourceManagerFile::OpenAsyncReadExternalResource(EPackageExternalResource ResourceType,
+FOpenAsyncPackageResult FPackageResourceManagerFile::OpenAsyncReadExternalResource(EPackageExternalResource ResourceType,
 	FStringView Identifier)
 {
 	switch (ResourceType)
@@ -368,13 +369,13 @@ IAsyncReadFileHandle* FPackageResourceManagerFile::OpenAsyncReadExternalResource
 		FPackagePath PackagePath;
 		if (!FPackagePath::TryFromPackageName(Identifier, PackagePath))
 		{
-			return new FAsyncReadFileHandleNull();
+			return FOpenAsyncPackageResult{ TUniquePtr<IAsyncReadFileHandle>(new FAsyncReadFileHandleNull()), EPackageFormat::Binary };
 		}
 		return OpenAsyncReadPackage(PackagePath, EPackageSegment::Header);
 	}
 	default:
 		checkNoEntry();
-		return new FAsyncReadFileHandleNull();
+		return FOpenAsyncPackageResult{ TUniquePtr<IAsyncReadFileHandle>(new FAsyncReadFileHandleNull()), EPackageFormat::Binary };
 	}
 }
 
diff --git a/Engine/Source/Runtime/CoreUObject/Private/UObject/UObjectGlobals.cpp b/Engine/Source/Runtime/CoreUObject/Private/UObject/UObjectGlobals.cpp
--- a/Engine/Source/Runtime/CoreUObject/Private/UObject/UObjectGlobals.cpp
+++ b/Engine/Source/Runtime/CoreUObject/Private/UObject/UObjectGlobals.cpp
@@ -1094,7 +1094,7 @@ public:
 		FOpenPackageResult OtherFile = IPackageResourceManager::Get().OpenReadPackage(DiffPackagePath);
 		checkf(!OtherFile.Archive.IsValid() || OtherFile.Format == EPackageFormat::Binary, TEXT("Text format is not yet supported with DiffPackage"));
 		FDiffFileArchive* DiffArchive = new FDiffFileArchive(GetLoader(), OtherFile.Archive.Release());
-		SetLoader(DiffArchive);
+		SetLoader(DiffArchive, true /* bInLoaderNeedsEngineVersionChecks */);
 	}
 };
 
diff --git a/Engine/Source/Runtime/CoreUObject/Public/UObject/LinkerLoad.h b/Engine/Source/Runtime/CoreUObject/Public/UObject/LinkerLoad.h
--- a/Engine/Source/Runtime/CoreUObject/Public/UObject/LinkerLoad.h
+++ b/Engine/Source/Runtime/CoreUObject/Public/UObject/LinkerLoad.h
@@ -195,7 +195,7 @@ private:
 
 protected:
 
-	void SetLoader(FArchive* InLoader);
+	void SetLoader(FArchive* InLoader, bool bInLoaderNeedsEngineVersionChecks);
 	FArchive* GetLoader() const { return Loader; }
 
 public:
@@ -362,6 +362,8 @@ private:
 	bool					bUseTimeLimit:1;
 	/** Whether to use the full time limit, even if we're blocked on I/O													*/
 	bool					bUseFullTimeLimit:1;
+	/** Whether the loader needs version and correctness checks (see OpenReadPackage)										*/
+	bool					bLoaderNeedsEngineVersionChecks : 1;
 	/** Call count of IsTimeLimitExceeded.																					*/
 	int32					IsTimeLimitExceededCallCount;
 	/** Current time limit to use if bUseTimeLimit is true.																	*/
diff --git a/Engine/Source/Runtime/CoreUObject/Public/UObject/PackageResourceManager.h b/Engine/Source/Runtime/CoreUObject/Public/UObject/PackageResourceManager.h
--- a/Engine/Source/Runtime/CoreUObject/Public/UObject/PackageResourceManager.h
+++ b/Engine/Source/Runtime/CoreUObject/Public/UObject/PackageResourceManager.h
@@ -32,11 +32,12 @@ DECLARE_DELEGATE_RetVal(IPackageResourceManager*, FSetPackageResourceManager);
  * Format for a package payload reported by the PackageResourceManager.
  * The appropriate ArchiveInputFormatter must be used for each format.
  */
-enum class EPackageFormat
+enum class EPackageFormat : uint8
 {
 	Binary, // Standard UnrealFormat, every field marshalled to bytes with Serialize(void*,int)
 	Text, // Text format, use FJsonArchiveInputFormatter
 };
+COREUOBJECT_API EPackageFormat ExtensionToPackageFormat(EPackageExtension Extension);
 
 struct FOpenPackageResult
 {
@@ -46,7 +47,53 @@ struct FOpenPackageResult
 	 * Format of the archive, binary or text.
 	 * Currently only header segments can have EPackageFormat::Text, all other segments have EPackageFormat::Binary
 	 */
-	EPackageFormat Format;
+	EPackageFormat Format = EPackageFormat::Binary;
+	/**
+	 * True if the package is of unknown version and needs to check for version and corruption.
+	 * False if the package was loaded from a repository specifically for the current binary's versions
+	 * and has already been checked for corruption.
+	 */
+	bool bNeedsEngineVersionChecks = true;
+
+	FOpenPackageResult() = default;
+	FOpenPackageResult(const FOpenPackageResult&) = delete;
+	FOpenPackageResult(FOpenPackageResult&&) = default;
+	FOpenPackageResult& operator=(const FOpenPackageResult&) = delete;
+	FOpenPackageResult& operator=(FOpenPackageResult&&) = default;
+	void CopyMetaData(const FOpenPackageResult& Other)
+	{
+		Format = Other.Format;
+		bNeedsEngineVersionChecks = Other.bNeedsEngineVersionChecks;
+	}
+};
+
+struct FOpenAsyncPackageResult
+{
+	/** AsyncReadFileHandle for the requested Segment bytes, in canceled state if it does not exist. */
+	TUniquePtr<IAsyncReadFileHandle> Handle;
+	/**
+	 * Format of the archive, binary or text.
+	 * Currently only header segments can have EPackageFormat::Text, all other segments have EPackageFormat::Binary
+	 */
+	EPackageFormat Format = EPackageFormat::Binary;
+	/**
+	 * True if the package is of unknown version and needs to check for version and corruption.
+	 * False if the package was loaded from a repository specifically for the current binary's versions
+	 * and has already been checked for corruption.
+	 */
+	bool bNeedsEngineVersionChecks = true;
+
+	FOpenAsyncPackageResult() = default;
+	FOpenAsyncPackageResult(const FOpenAsyncPackageResult&) = delete;
+	FOpenAsyncPackageResult(FOpenAsyncPackageResult&&) = default;
+	FOpenAsyncPackageResult& operator=(const FOpenAsyncPackageResult&) = delete;
+	FOpenAsyncPackageResult& operator=(FOpenAsyncPackageResult&&) = default;
+	COREUOBJECT_API ~FOpenAsyncPackageResult();
+	void CopyMetaData(const FOpenPackageResult& Other)
+	{
+		Format = Other.Format;
+		bNeedsEngineVersionChecks = Other.bNeedsEngineVersionChecks;
+	}
 };
 
 enum class EPackageExternalResource
@@ -132,10 +179,11 @@ public:
 	 * This call will always return a non-null handle, even if the package does not exist
 	 *
 	 * @param PackagePath The package to look for
-	 * @return An IAsyncReadFileHandle that will read from the package if it exists, or will be in the
-	 *         canceled state if the package does not exist
+	 * @return An FOpenAsyncPackageResult, with Handle that will read from the package if it exists,
+	 *         or will be in the canceled state if the package does not exist, and with other data
+	 *         describing the returned archive (see FOpenAsyncPackageResult)
 	 */
-	COREUOBJECT_API IAsyncReadFileHandle* OpenAsyncReadPackage(const FPackagePath& PackagePath);
+	COREUOBJECT_API FOpenAsyncPackageResult OpenAsyncReadPackage(const FPackagePath& PackagePath);
 
 	/**
 	 * Open an IMappedFileHandle to the package, if the PackageResourceManager supports it
@@ -181,10 +229,11 @@ public:
 	 *
 	 * @param ResourceType Id for the method used to map the identifier to an archive.
 	 * @param Identifier Id for which resource to return within the ResourceType's domain.
-	 * @return An IAsyncReadFileHandle that will read from the package if it exists, or will be in the
-	 *         canceled state if the package does not exist
+	 * @return An FOpenAsyncPackageResult, with Handle that will read from the package if it exists,
+	 *         or will be in the canceled state if the package does not exist, and with other data
+	 *         describing the returned archive (see FOpenAsyncPackageResult)
 	 */
-	virtual IAsyncReadFileHandle* OpenAsyncReadExternalResource(
+	virtual FOpenAsyncPackageResult OpenAsyncReadExternalResource(
 		EPackageExternalResource ResourceType, FStringView Identifier) = 0;
 
 	/**
@@ -298,7 +347,7 @@ public:
 		FPackagePath* OutUpdatedPath = nullptr) = 0;
 
 	/** OpenAsyncReadPackage that takes a PackageSegment */
-	virtual IAsyncReadFileHandle* OpenAsyncReadPackage(const FPackagePath& PackagePath, EPackageSegment PackageSegment) = 0;
+	virtual FOpenAsyncPackageResult OpenAsyncReadPackage(const FPackagePath& PackagePath, EPackageSegment PackageSegment) = 0;
 
 	/* OpenMappedHandleToPackage that takes a PackageSegment */
 	virtual IMappedFileHandle* OpenMappedHandleToPackage(const FPackagePath& PackagePath, EPackageSegment PackageSegment,
@@ -346,7 +395,7 @@ public:
 	 *         not be registered or a different archive was already registered
 	 */
 	COREUOBJECT_API static bool TryRegisterPreloadableArchive(const FPackagePath& PackagePath,
-		const TSharedPtr<FPreloadableArchive>& PreloadableArchive, EPackageFormat PackageFormat);
+		const TSharedPtr<FPreloadableArchive>& PreloadableArchive, const FOpenPackageResult& PackageFormat);
 
 	/**
 	 * Look up an FPreloadableFile instance registered for the given PackagePath, and return an FOpenPackageResult from it
@@ -367,7 +416,7 @@ public:
 	COREUOBJECT_API static bool UnRegisterPreloadableArchive(const FPackagePath& PackagePath);
 
 private:
-	static TMap<FName, TPair<TSharedPtr<FPreloadableArchive>, EPackageFormat>> PreloadedPaths;
+	static TMap<FName, TPair<TSharedPtr<FPreloadableArchive>, FOpenPackageResult>> PreloadedPaths;
 	static FCriticalSection PreloadedPathsLock;
 #endif
 };
