Compiling out gc.CalculateTokenStreamSize in shipping config

diff --git a/Engine/Source/Runtime/CoreUObject/Private/UObject/GarbageCollection.cpp b/Engine/Source/Runtime/CoreUObject/Private/UObject/GarbageCollection.cpp
--- a/Engine/Source/Runtime/CoreUObject/Private/UObject/GarbageCollection.cpp
+++ b/Engine/Source/Runtime/CoreUObject/Private/UObject/GarbageCollection.cpp
@@ -202,6 +202,7 @@ static FAutoConsoleVariableRef CGarbageReferenceTrackingEnabled(
 	ECVF_Default
 );
 
+#if !UE_BUILD_SHIPPING
 static FAutoConsoleCommand CmdCalculateTokenStreamSize(
 	TEXT("gc.CalculateTokenStreamSize"),
 	TEXT(""),
@@ -219,6 +220,7 @@ static FAutoConsoleCommand CmdCalculateTokenStreamSize(
 			UE_LOG(LogGarbage, Display, TEXT("Memory allocated for GC Tokens: %lldb, Debug Info: %lldb, Total: %lldb (%d classes)"), TokenSize, DebugSize, (TokenSize + DebugSize), NumClasses);
 		})
 );
+#endif // !UE_BUILD_SHIPPING
 
 #if PERF_DETAILED_PER_CLASS_GC_STATS
 /** Map from a UClass' FName to the number of objects that were purged during the last purge phase of this class.	*/
