- modified P_GET_STRUCT() and P_GET_STRUCT_REF() macros so that they can pass on types that have commas in them such as templates with multiple parameters
- enhanced PREPROCESSOR_COMMA_SEPARATED() so that it will accept 1+ arguments and not just 2+ arguments

diff --git a/Engine/Source/Runtime/Core/Public/HAL/PreprocessorHelpers.h b/Engine/Source/Runtime/Core/Public/HAL/PreprocessorHelpers.h
--- a/Engine/Source/Runtime/Core/Public/HAL/PreprocessorHelpers.h
+++ b/Engine/Source/Runtime/Core/Public/HAL/PreprocessorHelpers.h
@@ -19,8 +19,8 @@
 #define PREPROCESSOR_IF_INNER_1(x, y) x
 #define PREPROCESSOR_IF_INNER_0(x, y) y
 
-// Expands to the parameter list of the macro - used for when you need to pass a comma-separated identifier to another macro as a single parameter
-#define PREPROCESSOR_COMMA_SEPARATED(first, second, ...) first, second, ##__VA_ARGS__
+// Expands to the parameter list of the macro - used to pass a *potentially* comma-separated identifier to another macro as a single parameter
+#define PREPROCESSOR_COMMA_SEPARATED(first, ...) first, ##__VA_ARGS__
 
 // Expands to nothing - used as a placeholder
 #define PREPROCESSOR_NOTHING
diff --git a/Engine/Source/Runtime/CoreUObject/Public/UObject/ScriptMacros.h b/Engine/Source/Runtime/CoreUObject/Public/UObject/ScriptMacros.h
--- a/Engine/Source/Runtime/CoreUObject/Public/UObject/ScriptMacros.h
+++ b/Engine/Source/Runtime/CoreUObject/Public/UObject/ScriptMacros.h
@@ -62,8 +62,8 @@ enum {MAX_VARIABLE_SIZE = 0x0FFF };
 #define P_GET_UBOOL64(ParamName)					uint64 ParamName=0;                             Stack.StepCompiledIn<FBoolProperty>(&ParamName); ParamName = ParamName ? 1 : 0; // translate the bitfield into a bool type for non-intel platforms
 #define P_GET_UBOOL_REF(ParamName)					PARAM_PASSED_BY_REF_ZEROED(ParamName, FBoolProperty, bool)
 
-#define P_GET_STRUCT(StructType,ParamName)			PARAM_PASSED_BY_VAL(ParamName, FStructProperty, StructType)
-#define P_GET_STRUCT_REF(StructType,ParamName)		PARAM_PASSED_BY_REF(ParamName, FStructProperty, StructType)
+#define P_GET_STRUCT(StructType,ParamName)			PARAM_PASSED_BY_VAL(ParamName, FStructProperty, PREPROCESSOR_COMMA_SEPARATED(StructType))
+#define P_GET_STRUCT_REF(StructType,ParamName)		PARAM_PASSED_BY_REF(ParamName, FStructProperty, PREPROCESSOR_COMMA_SEPARATED(StructType))
 
 #define P_GET_OBJECT(ObjectType,ParamName)			PARAM_PASSED_BY_VAL_ZEROED(ParamName, FObjectPropertyBase, ObjectType*)
 #define P_GET_OBJECT_REF(ObjectType,ParamName)		PARAM_PASSED_BY_REF_ZEROED(ParamName, FObjectPropertyBase, ObjectType*)
