GarbageCollection compile fix for error C4459: declaration of 'MaxOffset' hides global declaration

#rnx

diff --git a/Engine/Source/Runtime/CoreUObject/Public/UObject/GarbageCollection.h b/Engine/Source/Runtime/CoreUObject/Public/UObject/GarbageCollection.h
--- a/Engine/Source/Runtime/CoreUObject/Public/UObject/GarbageCollection.h
+++ b/Engine/Source/Runtime/CoreUObject/Public/UObject/GarbageCollection.h
@@ -103,8 +103,7 @@ struct FGCReferenceInfo
 		, Offset(InOffset)
 	{
 		checkf(InReferenceType != GCRT_None && InReferenceType <= 0x1F, TEXT("Invalid GC Token Reference Type (%d)"), (uint32)InReferenceType);
-		const uint32 MaxOffset = 0x7FFFF;
-		checkf((InOffset & ~MaxOffset) == 0, TEXT("Invalid GC Token Offset (%d), max is %d"), InOffset, MaxOffset);
+		checkf((InOffset & ~0x7FFFF) == 0, TEXT("Invalid GC Token Offset (%d), max is %d"), InOffset, 0x7FFFF);
 	}
 	/**
 	 * Constructor
@@ -685,4 +684,4 @@ COREUOBJECT_API double GetLastGCTime();
 FORCEINLINE bool IsGarbageCollecting()
 {
 	return GIsGarbageCollecting;
-}
\ No newline at end of file
+}
