Fix 2.5MB of untagged memory when initializing token streams

diff --git a/Engine/Source/Runtime/CoreUObject/Private/UObject/GarbageCollection.cpp b/Engine/Source/Runtime/CoreUObject/Private/UObject/GarbageCollection.cpp
--- a/Engine/Source/Runtime/CoreUObject/Private/UObject/GarbageCollection.cpp
+++ b/Engine/Source/Runtime/CoreUObject/Private/UObject/GarbageCollection.cpp
@@ -2999,6 +2999,8 @@ struct FScopeLockIfNotNative
 
 void UClass::AssembleReferenceTokenStream(bool bForce)
 {
+	LLM_SCOPE(ELLMTag::GC);
+
 	// Lock for non-native classes
 	FScopeLockIfNotNative ReferenceTokenStreamLock(ReferenceTokenStreamCritical, !(ClassFlags & CLASS_Native));
 
