Reduce reallocs due to auto shrinkage of RemoveSwap causing bouncing between 0 element and 1 element very often
  -  Gets rid of 572K reallocs from process start to end of skydiving for a replayrun

diff --git a/Engine/Source/Runtime/CoreUObject/Private/UObject/GarbageCollection.cpp b/Engine/Source/Runtime/CoreUObject/Private/UObject/GarbageCollection.cpp
--- a/Engine/Source/Runtime/CoreUObject/Private/UObject/GarbageCollection.cpp
+++ b/Engine/Source/Runtime/CoreUObject/Private/UObject/GarbageCollection.cpp
@@ -2597,12 +2597,12 @@ bool FStructProperty::ContainsObjectReference(TArray<const FStructProperty*>& En
 			{
 				if (Property->ContainsObjectReference(EncounteredStructProps, InReferenceType))
 				{
-					EncounteredStructProps.RemoveSingleSwap(this);
+					EncounteredStructProps.RemoveSingleSwap(this, false /*bAllowShrinking*/);
 					return true;
 				}
 				Property = Property->PropertyLinkNext;
 			}
-			EncounteredStructProps.RemoveSingleSwap(this);
+			EncounteredStructProps.RemoveSingleSwap(this, false /*bAllowShrinking*/);
 		}
 		return false;
 	}
@@ -2880,7 +2880,7 @@ void FStructProperty::EmitReferenceInfo(UClass& OwnerClass, int32 BaseOffset, TA
 		{
 			bHasPropertiesWithObjectReferences = Property->ContainsObjectReference(EncounteredStructProps, EPropertyObjectReferenceType::Strong | EPropertyObjectReferenceType::Weak);
 		}
-		EncounteredStructProps.RemoveSingleSwap(this);
+		EncounteredStructProps.RemoveSingleSwap(this, false /* bAllowShrinking */);
 	}
 	// If the struct has UObject properties (and only if) emit tokens for them
 	if (bHasPropertiesWithObjectReferences)
