EDLCookChecker: suppress missing export warnings for packages that were iteratively skipped and thus do have the exports but did not declare them.

diff --git a/Engine/Source/Editor/UnrealEd/Private/Commandlets/CookCommandlet.cpp b/Engine/Source/Editor/UnrealEd/Private/Commandlets/CookCommandlet.cpp
--- a/Engine/Source/Editor/UnrealEd/Private/Commandlets/CookCommandlet.cpp
+++ b/Engine/Source/Editor/UnrealEd/Private/Commandlets/CookCommandlet.cpp
@@ -45,6 +45,7 @@
 #include "UObject/Class.h"
 #include "UObject/MetaData.h"
 #include "UObject/Package.h"
+#include "UObject/SavePackage.h"
 #include "UObject/UObjectIterator.h"
 
 DEFINE_LOG_CATEGORY_STATIC(LogCookCommandlet, Log, All);
@@ -1180,10 +1181,10 @@ bool UCookCommandlet::CookByTheBook( const TArray<ITargetPlatform*>& Platforms)
 	}
 #endif
 
-	if (!bIterativeCooking && StartupOptions.DLCName.IsEmpty())
+	if (StartupOptions.DLCName.IsEmpty())
 	{
 		bool bFullReferencesExpected = !(CookOptions & ECookByTheBookOptions::SkipHardReferences);
-		VerifyEDLCookInfo(bFullReferencesExpected);
+		UE::SavePackageUtilities::VerifyEDLCookInfo(bFullReferencesExpected);
 	}
 
 	return true;
diff --git a/Engine/Source/Editor/UnrealEd/Private/CookOnTheFlyServer.cpp b/Engine/Source/Editor/UnrealEd/Private/CookOnTheFlyServer.cpp
--- a/Engine/Source/Editor/UnrealEd/Private/CookOnTheFlyServer.cpp
+++ b/Engine/Source/Editor/UnrealEd/Private/CookOnTheFlyServer.cpp
@@ -5877,6 +5877,9 @@ void UCookOnTheFlyServer::PopulateCookedPackages(TArrayView<const ITargetPlatfor
 				{
 					PackageData->SetPlatformCooked(TargetPlatform, true /* bSucceeded */);
 				}
+
+				// Declare the package to the EDLCookInfo verification so we don't warn about missing exports from it
+				UE::SavePackageUtilities::EDLCookInfoAddIterativelySkippedPackage(IdenticalPackage);
 			}
 			PackageWriter.MarkPackagesUpToDate(Difference.IdenticalCookedPackages.Array());
 			for (FName UncookedPackage : Difference.IdenticalUncookedPackages)
@@ -7960,7 +7963,7 @@ void UCookOnTheFlyServer::StartCookByTheBook( const FCookByTheBookStartupOptions
 
 	if (CurrentCookMode == ECookMode::CookByTheBook && !IsCookFlagSet(ECookInitializationFlags::Iterative))
 	{
-		StartSavingEDLCookInfoForVerification();
+		UE::SavePackageUtilities::StartSavingEDLCookInfoForVerification();
 	}
 
 	{
diff --git a/Engine/Source/Editor/UnrealEd/Private/Cooker/CookRequestCluster.cpp b/Engine/Source/Editor/UnrealEd/Private/Cooker/CookRequestCluster.cpp
--- a/Engine/Source/Editor/UnrealEd/Private/Cooker/CookRequestCluster.cpp
+++ b/Engine/Source/Editor/UnrealEd/Private/Cooker/CookRequestCluster.cpp
@@ -23,6 +23,7 @@
 #include "Interfaces/ITargetPlatform.h"
 #include "Misc/RedirectCollector.h"
 #include "Misc/StringBuilder.h"
+#include "UObject/SavePackage.h"
 
 namespace UE::Cook
 {
@@ -748,6 +749,8 @@ void FRequestCluster::FGraphSearch::VisitVertex(const FVertexData& VertexData)
 					bFoundCachedTargetDomain = true;
 					PackageData->SetPlatformCooked(TargetPlatform, true);
 					PackageWriter->MarkPackagesUpToDate({ PackageName });
+					// Declare the package to the EDLCookInfo verification so we don't warn about missing exports from it
+					UE::SavePackageUtilities::EDLCookInfoAddIterativelySkippedPackage(PackageName);
 				}
 				HardDependencies.Append(PlatformAttachments.BuildDependencies);
 				if (Cluster.bAllowSoftDependencies && Cluster.bPreexploreDependenciesEnabled)
diff --git a/Engine/Source/Runtime/CoreUObject/Private/UObject/SavePackage/SavePackageUtilities.cpp b/Engine/Source/Runtime/CoreUObject/Private/UObject/SavePackage/SavePackageUtilities.cpp
--- a/Engine/Source/Runtime/CoreUObject/Private/UObject/SavePackage/SavePackageUtilities.cpp
+++ b/Engine/Source/Runtime/CoreUObject/Private/UObject/SavePackage/SavePackageUtilities.cpp
@@ -1327,6 +1327,16 @@ void FEDLCookChecker::FEDLNodeData::AppendPathName(const FEDLCookChecker& Owner,
 	Name.AppendString(Result);
 }
 
+FName FEDLCookChecker::FEDLNodeData::GetPackageName(const FEDLCookChecker& Owner) const
+{
+	if (ParentID != NodeIDInvalid)
+	{
+		// @todo ExternalPackages: We need to store ExternalPackage pointers on the Node and return that
+		return Owner.Nodes[ParentID].GetPackageName(Owner);
+	}
+	return Name;
+}
+
 void FEDLCookChecker::FEDLNodeData::Merge(FEDLCookChecker::FEDLNodeData&& Other)
 {
 	check(ObjectEvent == Other.ObjectEvent);
@@ -1393,7 +1403,7 @@ void FEDLCookChecker::AddExport(UObject* Export)
 		Nodes[SerializeID].bIsExport = true;
 		FEDLNodeID CreateID = FindOrAddNode(FEDLNodeHash(Export, EObjectEvent::Create));
 		Nodes[CreateID].bIsExport = true;
-		AddDependency(SerializeID, CreateID); // every export must be created before it can be serialize...these arcs are implicit and not listed in any table.
+		AddDependency(SerializeID, CreateID); // every export must be created before it can be serialized...these arcs are implicit and not listed in any table.
 	}
 }
 
@@ -1407,6 +1417,15 @@ void FEDLCookChecker::AddArc(UObject* DepObject, bool bDepIsSerialize, UObject*
 	}
 }
 
+void FEDLCookChecker::AddPackageWithUnknownExports(FName LongPackageName)
+{
+	if (bIsActive)
+	{
+		PackagesWithUnknownExports.Add(LongPackageName);
+	}
+}
+
+
 void FEDLCookChecker::AddDependency(FEDLNodeID SourceID, FEDLNodeID TargetID)
 {
 	NodePrereqs.Add(SourceID, TargetID);
@@ -1547,6 +1566,20 @@ void FEDLCookChecker::Merge(FEDLCookChecker&& Other)
 		Other.NodePrereqs.Empty();
 		Other.Nodes.Empty();
 	}
+
+	if (PackagesWithUnknownExports.Num() == 0)
+	{
+		Swap(PackagesWithUnknownExports, Other.PackagesWithUnknownExports);
+	}
+	else
+	{
+		PackagesWithUnknownExports.Reserve(Other.PackagesWithUnknownExports.Num());
+		for (FName PackageName : Other.PackagesWithUnknownExports)
+		{
+			PackagesWithUnknownExports.Add(PackageName);
+		}
+		Other.PackagesWithUnknownExports.Empty();
+	}
 }
 
 void FEDLCookChecker::Verify(bool bFullReferencesExpected)
@@ -1579,6 +1612,14 @@ void FEDLCookChecker::Verify(bool bFullReferencesExpected)
 			{
 				if (NodeData.bIsExport)
 				{
+					// The node is an export; imports of it are valid
+					continue;
+				}
+
+				if (Accumulator.PackagesWithUnknownExports.Contains(NodeData.GetPackageName(Accumulator)))
+				{
+					// The node is an object in a package that exists, but for which we do not know the exports
+					// because e.g. it was iteratively skipped in the current cook. Suppress warnings about it
 					continue;
 				}
 
@@ -1618,6 +1659,9 @@ void FEDLCookChecker::Verify(bool bFullReferencesExpected)
 FCriticalSection FEDLCookChecker::CookCheckerInstanceCritical;
 TArray<FEDLCookChecker*> FEDLCookChecker::CookCheckerInstances;
 
+namespace UE::SavePackageUtilities
+{
+
 void StartSavingEDLCookInfoForVerification()
 {
 	FEDLCookChecker::StartSavingEDLCookInfoForVerification();
@@ -1628,6 +1672,13 @@ void VerifyEDLCookInfo(bool bFullReferencesExpected)
 	FEDLCookChecker::Verify(bFullReferencesExpected);
 }
 
+void EDLCookInfoAddIterativelySkippedPackage(FName LongPackageName)
+{
+	FEDLCookChecker::Get().AddPackageWithUnknownExports(LongPackageName);
+}
+
+}
+
 FScopedSavingFlag::FScopedSavingFlag(bool InSavingConcurrent, UPackage* InSavedPackage)
 	: bSavingConcurrent(InSavingConcurrent)
 	, SavedPackage(InSavedPackage)
diff --git a/Engine/Source/Runtime/CoreUObject/Private/UObject/SavePackage/SavePackageUtilities.h b/Engine/Source/Runtime/CoreUObject/Private/UObject/SavePackage/SavePackageUtilities.h
--- a/Engine/Source/Runtime/CoreUObject/Private/UObject/SavePackage/SavePackageUtilities.h
+++ b/Engine/Source/Runtime/CoreUObject/Private/UObject/SavePackage/SavePackageUtilities.h
@@ -180,6 +180,7 @@ struct FEDLCookChecker : public TThreadSingleton<FEDLCookChecker>
 	void AddImport(UObject* Import, UPackage* ImportingPackage);
 	void AddExport(UObject* Export);
 	void AddArc(UObject* DepObject, bool bDepIsSerialize, UObject* Export, bool bExportIsSerialize);
+	void AddPackageWithUnknownExports(FName LongPackageName);
 
 	static void StartSavingEDLCookInfoForVerification();
 	static void Verify(bool bFullReferencesExpected);
@@ -268,6 +269,7 @@ private:
 
 		FString ToString(const FEDLCookChecker& Owner) const;
 		void AppendPathName(const FEDLCookChecker& Owner, FStringBuilderBase& Result) const;
+		FName GetPackageName(const FEDLCookChecker& Owner) const;
 		void Merge(FEDLNodeData&& Other);
 	};
 
@@ -295,6 +297,11 @@ private:
 	TMap<FEDLNodeHash, FEDLNodeID> NodeHashToNodeID;
 	/** The graph of dependencies between nodes. */
 	TMultiMap<FEDLNodeID, FEDLNodeID> NodePrereqs;
+	/**
+	 * Packages that were cooked iteratively and therefore have an unknown set of exports.
+	 * We suppress warnings for exports missing from these packages.
+	 */
+	TSet<FName> PackagesWithUnknownExports;
 	/** True if the EDLCookChecker should be active; it is turned off if the runtime will not be using EDL. */
 	bool bIsActive;
 
diff --git a/Engine/Source/Runtime/CoreUObject/Public/UObject/Package.h b/Engine/Source/Runtime/CoreUObject/Public/UObject/Package.h
--- a/Engine/Source/Runtime/CoreUObject/Public/UObject/Package.h
+++ b/Engine/Source/Runtime/CoreUObject/Public/UObject/Package.h
@@ -103,9 +103,6 @@ struct FSavePackageResultStruct
 	}
 };
 
-COREUOBJECT_API void StartSavingEDLCookInfoForVerification();
-COREUOBJECT_API void VerifyEDLCookInfo(bool bFullReferencesExpected=true);
-
 /**
 * A package.
 */
diff --git a/Engine/Source/Runtime/CoreUObject/Public/UObject/SavePackage.h b/Engine/Source/Runtime/CoreUObject/Public/UObject/SavePackage.h
--- a/Engine/Source/Runtime/CoreUObject/Public/UObject/SavePackage.h
+++ b/Engine/Source/Runtime/CoreUObject/Public/UObject/SavePackage.h
@@ -132,9 +132,7 @@ public:
 	const bool bForceLegacyOffsets = false;
 };
 
-namespace UE
-{
-namespace SavePackageUtilities
+namespace UE::SavePackageUtilities
 {
 	/**
 	 * Return whether the given save parameters indicate the LoadedPath of the package being saved should be updated.
@@ -166,7 +164,10 @@ namespace SavePackageUtilities
 
 	COREUOBJECT_API void ResetCookStats();
 	COREUOBJECT_API int32 GetNumPackagesSaved();
-}
+
+	COREUOBJECT_API void StartSavingEDLCookInfoForVerification();
+	COREUOBJECT_API void VerifyEDLCookInfo(bool bFullReferencesExpected = true);
+	COREUOBJECT_API void EDLCookInfoAddIterativelySkippedPackage(FName LongPackageName);
 }
 
 COREUOBJECT_API DECLARE_LOG_CATEGORY_EXTERN(LogSavePackage, Log, All);
