Fix bad merge.

diff --git a/Engine/Source/Runtime/CoreUObject/Private/UObject/LazyObjectPtr.cpp b/Engine/Source/Runtime/CoreUObject/Private/UObject/LazyObjectPtr.cpp
--- a/Engine/Source/Runtime/CoreUObject/Private/UObject/LazyObjectPtr.cpp
+++ b/Engine/Source/Runtime/CoreUObject/Private/UObject/LazyObjectPtr.cpp
@@ -120,7 +120,7 @@ void FLazyObjectPtr::PossiblySerializeObjectGuid(UObject *Object, FStructuredArc
 	else if (UnderlyingArchive.IsLoading())
 	{
 		TOptional<FStructuredArchiveSlot> GuidSlot = Record.TryEnterField(SA_FIELD_NAME(TEXT("Guid")), false);
-		if (GuidSlot.IsSet())
+		if (GuidSlot.IsSet() && !UnderlyingArchive.IsError())
 		{
 			FUniqueObjectGuid Guid;
 			GuidSlot.GetValue() << Guid;
