Track when we're using an immediately flushed async load from normal synchronous package loading. Use that information to disable a warning condition around postloads doing synchronous loads and blocking async loading.

diff --git a/Engine/Source/Runtime/CoreUObject/Private/UObject/UObjectGlobals.cpp b/Engine/Source/Runtime/CoreUObject/Private/UObject/UObjectGlobals.cpp
--- a/Engine/Source/Runtime/CoreUObject/Private/UObject/UObjectGlobals.cpp
+++ b/Engine/Source/Runtime/CoreUObject/Private/UObject/UObjectGlobals.cpp
@@ -941,12 +941,15 @@ UObject* StaticLoadObjectInternal(UClass* ObjectClass, UObject* InOuter, const T
 	return Result;
 }
 
+// Track how many nested loads we're doing by triggering an async load and immediately flushing that request.
+static int32 GSyncLoadUsingAsyncLoaderCount = 0;
+
 UObject* StaticLoadObject(UClass* ObjectClass, UObject* InOuter, const TCHAR* InName, const TCHAR* Filename, uint32 LoadFlags, UPackageMap* Sandbox, bool bAllowObjectReconciliation, const FLinkerInstancingContext* InstancingContext)
 {
 	FUObjectThreadContext& ThreadContext = FUObjectThreadContext::Get();
-	if (ThreadContext.IsRoutingPostLoad && IsInAsyncLoadingThread())
+	if ((GSyncLoadUsingAsyncLoaderCount == 0) && ThreadContext.IsRoutingPostLoad && IsInAsyncLoadingThread())
 	{
-		UE_LOG(LogUObjectGlobals, Warning, TEXT("Calling StaticLoadObject(\"%s\", \"%s\", \"%s\") during PostLoad of %s may result in hitches during streaming."), 
+		UE_LOG(LogUObjectGlobals, Warning, TEXT("Calling StaticLoadObject(\"%s\", \"%s\", \"%s\") during PostLoad of %s is illegal and will crash in a cooked runtime"), 
 			*GetFullNameSafe(ObjectClass),
 			*GetFullNameSafe(InOuter),
 			InName,
@@ -1178,12 +1181,14 @@ UPackage* LoadPackageInternal(UPackage* InOuter, const FPackagePath& PackagePath
 				FCoreDelegates::OnSyncLoadPackage.Broadcast(PackageName.ToString());
 			}
 
+			GSyncLoadUsingAsyncLoaderCount++;
 			int32 RequestID = LoadPackageAsync(PackagePath, PackageName);
 
 			if (RequestID != INDEX_NONE)
 			{
 				FlushAsyncLoading(RequestID);
 			}
+			GSyncLoadUsingAsyncLoaderCount--;
 
 			return (InOuter ? InOuter : FindObjectFast<UPackage>(nullptr, PackageName));
 		}
