Saves 1 minute of GC DeleteLinkers during cook due to linear search and array shrinking

diff --git a/Engine/Source/Runtime/CoreUObject/Private/UObject/LinkerManager.h b/Engine/Source/Runtime/CoreUObject/Private/UObject/LinkerManager.h
--- a/Engine/Source/Runtime/CoreUObject/Private/UObject/LinkerManager.h
+++ b/Engine/Source/Runtime/CoreUObject/Private/UObject/LinkerManager.h
@@ -150,7 +150,7 @@ private:
 #if !UE_BUILD_SHIPPING && !UE_BUILD_TEST
 	/** List of all the existing linker loaders **/
 	FCriticalSection LiveLinkersCritical;
-	TArray<FLinkerLoad*> LiveLinkers;
+	TSet<FLinkerLoad*> LiveLinkers;
 #endif
 	
 	/** List of loaders that have new imports **/
