Reworked the virtualization system so that we only virtualize a package when it is submitted to source control, which makes the overall system much less dangerous.

#rnx

### Virtualization
- The idea is to move from our current set up, where the virtualization happens when a package is saved and the payloads are saved to a local virtualization storage system with the final push to persistent storage occurring when the package is submitted to a system where we save payloads into the package file, but move them to persistent virtualization storage when the package file is submitted.
- The main advantage is that if someone submits a package file to source control without virtualizing it, we don't have to worry that others might not be able to load the package, the worst case scenario is that data sizes get bigger until the package is virtualized again.
- This is only the first pass, in the future we can do further optimizations, like not storing payloads locally if we know that they are already in the persistent storage system etc.
- In order to keep the virtualization process simple we want to be able to do it without needing to re-save the package, so to this end we will storage the payloads at the end of the package file as before, but instead of storing them inside of the package file format we will use a new container, the FPackageTrailer that is appended to the package file instead.
- In theory this will make future sidecar work easier as the trailer can just be moved to the sidecar file as as long as we know where to find the trailer things will just work (tm)
- For now VBD will continue to serialize it's offset/size to disk and use those values to read the payload directly when not virtualizing. Ideally we'd use the trailer but this will help reduce the risk.

### Current Issues
- The system does not work with the editor domain and has minor issues with text based assets
- The trailer system is disabled by default via the config system [Core]UsePackageTrailer=False and can then be enabled on select test projects until fully functional.

### PackageTrailer
- The trailer is split into two parts FPackageTrailerBuilder and FPackageTrailer
- FPackageTrailerBuilder is used when building the trailer during package save and FPackageTrailer is the structure we load and actually use.
- When saving the trailer we try to avoid using containers and such so that we can try to keep the specification for the file format very clear so that licensees can build their own tools for the virtualization process if they so desire.
- the header file contains a breakdown of the new format.

### VirtualizedBulkData
- FPayloadToc is now only used by the sidecar experimental feature. In a future update the sidecar will be changed to basically be the package trailer but stored in it's own file.
- FindPayloadsInPackageFile has been moved to PackageTrailer.h as it has a closer association with that code.
- Added a new method, ::LoadFromPackageTrailer, to load payloads from the package trailer without needing to know where the payload is on disk. Using this is opt in with the cvar "Serialization.LoadFromTrailer" and is only provided for debugging purposes. Although once this system has matured it will likely become the preferred way to access payloads.
- We no longer allow virtualization of bulkdata payloads on save, but I am not sure if we might want to allow this as an option in the future. So for now the code branch is disabled by the global constant bAllowVirtualizationOnSave.
- Added utility functions ::GetLinkerLoadFromOwner and ::GetTrailerFromOwner for easy access to the LinkerLoad/Trailer from the owning object.
- Added a utility ::UpdateArchiveData for seeking back to a known position in an archive, writing over it, then seeking back, which is fairly common in the package saving code paths. This might be worth moving to Archive.h for general use.
- ::LoadFromPackageFile has been refactored to prefer returning error values early over nested ifs.
- We no longer serialize the EFlags::IsVirtualized flag to disk. Instead it is applied when the package is loaded if we detect that a payload is virtualized and then removed before saved to disk. This allows the virtualization process to occur without resaving the package.
- We currently support the old serialize to disk path which is enabled when the FPackageTrailer system is disabled and the newer system when it is enabled.
- We do however continue to serialize the offset in the package file to disk when saved. Strictly speaking this is not needed as we can look this up from the package trailer when the package is loaded but continuing to serialize it out to disk helps keep the old and newer data format paths working together without adding too much special case code.
-- Once the FPackageTrailer feature is no longer optional we can consider changing this.

### SavePackage/SavePackage2
- Removed the functions that would create the FPayloadToc
- Added new functions to help create the trailer.
- Moved the end of package tag out of the if to it's own scope (still won't run if a package writer exists), the reasoning is to make it clearer where the package format ends and where we can start writing the payload trailer.
- We have checks to make sure that the trailer does not try to write to disk for both text based assets and if the editor domain is enabled. Support for these features will be added later, but since the package trailer system is disabled by default this shouldn't cause problems.

### LinkerLoad
- After we parse the FPackageFileSummary we now also parse the header of the package trailer. In the workspace domain this will incur additional seek costs (which may or not actually seek and invalidate the internal file cache depending on how large the package's exports are) as the trailer is found after the package
-- In the future, the package trailer header can be stored right after the FPackageFileSummary in the editor domain which will remove this overhead.
- By parsing the header of the trailer at this point, virtualized bulkdata objects in the package can use the look up info to determine if their payloads are virtualized or not, and if not then where they reside inside the package file.

### LinkerSave
- We no longer collect a list of all virtualized bulkdata in a package while saving as we no longer generate the FPayloadToc so BulkDataInPackage has been removed.
- We now gather data to be appended to the package trailer via TrailerBuilder.

### DumpPayloadToc
- This command now dumps info based on the package trailer if one can be found.

diff --git a/Engine/Config/BaseEngine.ini b/Engine/Config/BaseEngine.ini
--- a/Engine/Config/BaseEngine.ini
+++ b/Engine/Config/BaseEngine.ini
@@ -1093,6 +1093,7 @@ GameThreadHeartBeatHitchDuration=-1.0
 GameThreadHeartBeatStackWalk=false
 AllowBulkDataInIoStore=True
 IgnoreInlineBulkDataReloadEnsures=False
+UsePackageTrailer=False
 HangDuration=0.0
 AssetLogShowsDiskPath=True
 AssetLogShowsAbsolutePath=True
diff --git a/Engine/Source/Developer/Virtualization/Private/DumpPayloadTocCommand.cpp b/Engine/Source/Developer/Virtualization/Private/DumpPayloadTocCommand.cpp
--- a/Engine/Source/Developer/Virtualization/Private/DumpPayloadTocCommand.cpp
+++ b/Engine/Source/Developer/Virtualization/Private/DumpPayloadTocCommand.cpp
@@ -3,12 +3,10 @@
 #include "HAL/IConsoleManager.h"
 #include "Logging/LogMacros.h"
 #include "Misc/PackagePath.h"
-#include "Serialization/VirtualizedBulkData.h"
+#include "UObject/PackageTrailer.h"
 
 namespace UE
 {
-namespace Virtualization
-{
 
 #if WITH_EDITORONLY_DATA
 
@@ -33,16 +31,16 @@ void DumpPayloadToc(const TArray<FString>& Args)
 
 		if (FPackagePath::TryFromMountedName(Arg, Path))
 		{
-			TArray<FPayloadId> LocalPayloadIds;
-			TArray<FPayloadId> VirtualizedPayloadIds;
+			TArray<Virtualization::FPayloadId> LocalPayloadIds;
+			TArray<Virtualization::FPayloadId> VirtualizedPayloadIds;
 
-			if (!UE::Virtualization::FindPayloadsInPackageFile(Path, UE::Virtualization::EPayloadType::Local, LocalPayloadIds))
+			if (!UE::FindPayloadsInPackageFile(Path, UE::EPayloadFilter::Local, LocalPayloadIds))
 			{
 				UE_LOG(LogVirtualization, Error, TEXT("Failed to find local payload information from package: '%s'"), *Path.GetDebugName());
 				continue;
 			}
 
-			if (!UE::Virtualization::FindPayloadsInPackageFile(Path, UE::Virtualization::EPayloadType::Virtualized, VirtualizedPayloadIds))
+			if (!UE::FindPayloadsInPackageFile(Path, UE::EPayloadFilter::Virtualized, VirtualizedPayloadIds))
 			{
 				UE_LOG(LogVirtualization, Error, TEXT("Failed to find virtualized payload information from package: '%s'"), *Path.GetDebugName());
 				continue;
@@ -86,5 +84,4 @@ static FAutoConsoleCommand CCmdDumpPayloadToc = FAutoConsoleCommand(
 
 #endif //WITH_EDITORONLY_DATA
 
-} // namespace Virtualization
 } // namespace UE
\ No newline at end of file
diff --git a/Engine/Source/Developer/Virtualization/Private/PackageSubmissionChecks.cpp b/Engine/Source/Developer/Virtualization/Private/PackageSubmissionChecks.cpp
--- a/Engine/Source/Developer/Virtualization/Private/PackageSubmissionChecks.cpp
+++ b/Engine/Source/Developer/Virtualization/Private/PackageSubmissionChecks.cpp
@@ -3,16 +3,25 @@
 #include "PackageSubmissionChecks.h"
 
 #include "Containers/UnrealString.h"
+#include "HAL/FileManager.h"
+#include "HAL/PlatformFileManager.h"
 #include "Internationalization/Internationalization.h"
 #include "Misc/PackageName.h"
+#include "Misc/Paths.h"
 #include "Serialization/VirtualizedBulkData.h"
+#include "UObject/PackageResourceManager.h"
+#include "UObject/PackageTrailer.h"
 #include "Virtualization/VirtualizationSystem.h"
 
 #define LOCTEXT_NAMESPACE "Virtualization"
 
+#define UE_USE_LEGACY_SUBMIT 0
+#define UE_VALIDATE_TRUNCATED_PACKAGE 1
+
 namespace UE::Virtualization
 {
 
+#if UE_USE_LEGACY_SUBMIT
 void OnPrePackageSubmission(const TArray<FString>& FilesToSubmit, TArray<FText>& DescriptionTags, TArray<FText>& Errors)
 {
 	TRACE_CPUPROFILER_EVENT_SCOPE(UE::Virtualization::OnPrePackageSubmission);
@@ -36,7 +45,7 @@ void OnPrePackageSubmission(const TArray<FString>& FilesToSubmit, TArray<FText>&
 		if (FPackageName::IsPackageExtension(PackagePath.GetHeaderExtension()) || FPackageName::IsTextPackageExtension(PackagePath.GetHeaderExtension()))
 		{
 			TArray<FPayloadId> PayloadsInPackage;
-			if (FindPayloadsInPackageFile(PackagePath, EPayloadType::Virtualized, PayloadsInPackage))
+			if (FindPayloadsInPackageFile(PackagePath, EPayloadFilter::Virtualized, PayloadsInPackage))
 			{
 				for (const FPayloadId& PayloadId : PayloadsInPackage)
 				{
@@ -82,6 +91,267 @@ void OnPrePackageSubmission(const TArray<FString>& FilesToSubmit, TArray<FText>&
 		DescriptionTags.Add(Tag);	
 	}
 }
+#else
+
+/** 
+ * Creates a copy of the given package but the copy will not include the FPackageTrailer.
+ * 
+ * @param[in]	PackagePath	The path of the package to copy
+ * @param[in]	CopyPath	The path where the copy should be created
+ * @param[in]	Trailer		The trailer found in 'PackagePath' that is already loaded
+ * @param[out]	Errors		Errors created by the function will be added here
+ * 
+ * @return Returns true if the package was copied correctly, false otherwise. Note even when returning false a file might have been created at 'CopyPath'
+ */
+bool TryCopyPackageWithoutTrailer(const FPackagePath PackagePath, const FString& CopyPath, const FPackageTrailer& Trailer, TArray<FText>& Errors)
+{
+	// TODO: Consider adding a custom copy routine to only copy the data we want, rather than copying the full file then truncating
+
+	const FString PackageFilePath = PackagePath.GetLocalFullPath();
+
+	if (IFileManager::Get().Copy(*CopyPath, *PackageFilePath) != ECopyResult::COPY_OK)
+	{
+		FText Message = FText::Format(	LOCTEXT("Virtualization_CopyFailed", "Unable to copy package file '{0}' for virtualization"),
+										FText::FromString(PackagePath.GetDebugName()));
+		Errors.Add(Message);
+		return false;
+	}
+
+	const int64 PackageSizeWithoutTrailer = IFileManager::Get().FileSize(*PackageFilePath) - Trailer.GetTrailerLength();
+
+	TUniquePtr<IFileHandle> TempFileHandle(FPlatformFileManager::Get().GetPlatformFile().OpenWrite(*CopyPath, true));
+	if (!TempFileHandle.IsValid())
+	{
+		FText Message = FText::Format(	LOCTEXT("Virtualization_TruncOpenFailed", "Failed to open package file for truncation'{0}' when virtualizing"),
+										FText::FromString(CopyPath));
+		Errors.Add(Message);
+		return false;
+	}
+
+	if (!TempFileHandle->Truncate(PackageSizeWithoutTrailer))
+	{
+		FText Message = FText::Format(	LOCTEXT("Virtualization_TruncFailed", "Failed to truncate '{0}' when virtualizing"),
+										FText::FromString(CopyPath));
+		Errors.Add(Message);
+		return false;
+	}
+
+	return true;
+}
+
+/** 
+ * Check that the given package ends with PACKAGE_FILE_TAG. Intended to be used to make sure that
+ * we have truncated a package correctly when removing the trailers. 
+ */
+bool ValidatePackage(const FString& PackagePath, TArray<FText>& Errors)
+{
+	TUniquePtr<IFileHandle> TempFileHandle(FPlatformFileManager::Get().GetPlatformFile().OpenRead(*PackagePath));
+	if (!TempFileHandle.IsValid())
+	{
+		FText ErrorMsg = FText::Format(	LOCTEXT("Virtualization_OpenValidationFailed", "Unable to open '{0}' so that it can be validated"), 
+										FText::FromString(PackagePath));
+		Errors.Add(ErrorMsg);
+		return false;
+	}
+
+	TempFileHandle->SeekFromEnd(-4);
+
+	uint32 PackageTag = INDEX_NONE;
+	if (!TempFileHandle->Read((uint8*)&PackageTag, 4) || PackageTag != PACKAGE_FILE_TAG)
+	{
+		FText ErrorMsg = FText::Format(	LOCTEXT("Virtualization_ValidationFailed", "The package '{0}' does not end with a valid tag, the file is considered corrupt"),
+										FText::FromString(PackagePath));
+		Errors.Add(ErrorMsg);
+		return false;
+	}
+
+
+	return true;
+}
+
+void OnPrePackageSubmission(const TArray<FString>& FilesToSubmit, TArray<FText>& DescriptionTags, TArray<FText>& Errors)
+{
+	TRACE_CPUPROFILER_EVENT_SCOPE(UE::Virtualization::OnPrePackageSubmission);
+
+	IVirtualizationSystem& System = IVirtualizationSystem::Get();
+
+	// TODO: We could check to see if the package is virtualized even if it is disabled for the project
+	// as a safety feature?
+	if (!System.IsEnabled())
+	{
+		return;
+	}
+
+	// Can't virtualize if the payload trailer system is disabled
+	if (!FPackageTrailer::IsEnabled())
+	{
+		return;
+	}
+
+	// Other systems may have added errors to this array, we need to check so later we can determine if this function added any additional errors.
+	const int32 NumErrors = Errors.Num();
+
+	TArray<TPair<FPackagePath,FString>> PackagesToReplace;
+
+	for (const FString& AbsoluteFilePath : FilesToSubmit)
+	{
+		if (NumErrors != Errors.Num())
+		{
+			break; // Give up if previous packages have had errors
+		}
+
+		const FPackagePath PackagePath = FPackagePath::FromLocalPath(AbsoluteFilePath);
+
+		if (FPackageName::IsPackageExtension(PackagePath.GetHeaderExtension()) || FPackageName::IsTextPackageExtension(PackagePath.GetHeaderExtension()))
+		{
+			FPackageTrailer Trailer;
+			if (FPackageTrailer::TryLoadFromPackage(PackagePath, Trailer))
+			{
+				const FString PackageFilePath = PackagePath.GetLocalFullPath();
+				const FString BaseName = FPaths::GetBaseFilename(PackagePath.GetPackageName());
+				const FString TempFilePath = FPaths::CreateTempFilename(*FPaths::ProjectSavedDir(), *BaseName.Left(32));
+
+				// Create copy of package minus the trailer the trailer
+				if (!TryCopyPackageWithoutTrailer(PackagePath, TempFilePath, Trailer, Errors))
+				{
+					continue;
+				}			
+
+#if UE_VALIDATE_TRUNCATED_PACKAGE
+				// Validate we didn't break the package
+				if (!ValidatePackage(TempFilePath, Errors))
+				{
+					continue;
+				}
+#endif //UE_VALIDATE_TRUNCATED_PACKAGE
+
+				// Extract each payload and push
+				{
+					TUniquePtr<FArchive> PackageAr = IPackageResourceManager::Get().OpenReadExternalResource(EPackageExternalResource::WorkspaceDomainFile, PackagePath.GetPackageName());
+
+					if (!PackageAr.IsValid())
+					{
+						FText Message = FText::Format(	LOCTEXT("Virtualization_PkgOpenm", "Failed to open the package '{1}' for reading"),
+														FText::FromString(PackagePath.GetDebugName()));
+						Errors.Add(Message);
+						break;
+					}
+
+					bool bUpdatedTrailer = false;
+
+					for (const FPayloadId& PayloadId : Trailer.GetPayloads(EPayloadFilter::Local))
+					{
+						if (PayloadId.IsValid())
+						{
+							FCompressedBuffer Payload = Trailer.LoadPayload(PayloadId, *PackageAr);
+
+							if (PayloadId != FIoHash(Payload.GetRawHash()))
+							{
+								FText Message = FText::Format(LOCTEXT("Virtualization_WrongPayload", "Package {0} loaded an incorrect payload from the trailer. Expected '{1}' Loaded  '{2}'"),
+									FText::FromString(PackagePath.GetDebugName()),
+									FText::FromString(PayloadId.ToString()),
+									FText::FromString(LexToString(Payload.GetRawHash())));
+								Errors.Add(Message);
+								break;
+							}
+
+							if (Payload)
+							{
+								if (!System.PushData(PayloadId, Payload, EStorageType::Persistent, PackagePath))
+								{
+									FText Message = FText::Format(LOCTEXT("Virtualization_PushFailure", "Failed to push payload '{0}' from the package '{1}'"),
+										FText::FromString(PayloadId.ToString()),
+										FText::FromString(PackagePath.GetDebugName()));
+									Errors.Add(Message);
+									break;
+								}
+							}
+							else
+							{
+								FText Message = FText::Format(LOCTEXT("Virtualization_MissingPayload", "Unable to find the payload '{0}' in the local storage of package '{1}'"),
+									FText::FromString(PayloadId.ToString()),
+									FText::FromString(PackagePath.GetDebugName()));
+								Errors.Add(Message);
+								break;
+							}
+
+							if (Trailer.UpdatePayloadAsVirtualized(PayloadId))
+							{
+								bUpdatedTrailer = true;
+							}
+							else
+							{
+								FText Message = FText::Format(LOCTEXT("Virtualization_UpdateStatusFailed", "Unable to update the status for the payload '{0}' in the package '{1}'"),
+									FText::FromString(PayloadId.ToString()),
+									FText::FromString(PackagePath.GetDebugName()));
+								Errors.Add(Message);
+								break;
+							}		
+						}
+					}
+
+					if (bUpdatedTrailer)
+					{
+						TUniquePtr<FArchive> CopyAr(IFileManager::Get().CreateFileWriter(*TempFilePath, EFileWrite::FILEWRITE_Append));
+						if (!CopyAr.IsValid())
+						{
+							FText Message = FText::Format(	LOCTEXT("Virtualization_TrailerAppendOpen", "Unable to open '{0}' to append the trailer'"),
+															FText::FromString(TempFilePath));
+							Errors.Add(Message);
+							break;
+						}
+
+						FPackageTrailerBuilder TrailerBuilder = FPackageTrailerBuilder::Create(Trailer, *PackageAr);
+						if (!TrailerBuilder.BuildAndAppendTrailer(nullptr, *CopyAr))
+						{
+							FText Message = FText::Format(	LOCTEXT("Virtualization_TrailerAppend", "Failed to append the trailer to '{0}'"),
+															FText::FromString(TempFilePath));
+							Errors.Add(Message);
+							break;
+						}
+
+						// Now that we have successfully created a new version of the package with an updated trailer 
+						// we need to mark that it should replace the original package.
+						PackagesToReplace.Emplace(PackagePath, TempFilePath);
+					}
+				}
+			}
+		}
+	}
+
+	if (NumErrors == Errors.Num())
+	{
+		// TODO: Consider using the SavePackage model (move the original, then replace, so we can restore all of the original packages if needed)
+		// having said that, once a package is in PackagesToReplace it should still be safe to submit so maybe we don't need this level of protection?
+
+		// Since we had no errors we can now replace all of the packages that were virtualized data with the 
+		// virtualized replacement file.
+		for(const TPair<FPackagePath,FString>&  Iterator : PackagesToReplace)
+		{
+			const FString OriginalPackagePath = Iterator.Key.GetLocalFullPath();
+			const FString& NewPackagePath = Iterator.Value;
+
+			if (!IFileManager::Get().Move(*OriginalPackagePath, *NewPackagePath))
+			{
+				FText Message = FText::Format(	LOCTEXT("Virtualization_MoveFailed", "Unable to replace the package '{0}' with the virtualized version"),
+												FText::FromString(Iterator.Key.GetDebugName()));
+				Errors.Add(Message);
+				continue;
+			}
+		}
+	}
+
+	// If we had no new errors add the validation tag to indicate that the packages are safe for submission. 
+	// TODO: Currently this is a hard coded guid, it should be replaced by something unique to the FilesToSubmit collection that could be 
+	// validated by the source control server. This functionality is currently pending.
+	if (NumErrors == Errors.Num())
+	{
+		FText Tag = FText::FromString(TEXT("#virtualization c9c748e7-d00b-45af-bb61-9054abba503a"));
+		DescriptionTags.Add(Tag);
+	}
+}
+#endif 
 } // namespace UE::Virtualization
 
 #undef LOCTEXT_NAMESPACE
+
diff --git a/Engine/Source/Runtime/CoreUObject/Private/Serialization/VirtualizedBulkData.cpp b/Engine/Source/Runtime/CoreUObject/Private/Serialization/VirtualizedBulkData.cpp
--- a/Engine/Source/Runtime/CoreUObject/Private/Serialization/VirtualizedBulkData.cpp
+++ b/Engine/Source/Runtime/CoreUObject/Private/Serialization/VirtualizedBulkData.cpp
@@ -45,6 +45,15 @@ static TAutoConsoleVariable<bool> CVarShouldLoadFromSidecar(
 	false,
 	TEXT("When true FVirtualizedUntypedBulkData will load from the sidecar file"));
 
+/** 
+ * Prefer loading from the package trailer (load the trailer, parse the look up, then load the payload) over 
+ * using the in built OffsetInFile member to load from the package file directly.
+ */
+static TAutoConsoleVariable<bool> CVarShouldLoadFromTrailer(
+	TEXT("Serialization.LoadFromTrailer"),
+	false,
+	TEXT("When true FVirtualizedUntypedBulkData will load payloads via the package trailer rather than the package itself"));
+
 static TAutoConsoleVariable<bool> CVarShouldValidatePayload(
 	TEXT("Serialization.ValidatePayloads"),
 	false,
@@ -55,6 +64,9 @@ static TAutoConsoleVariable<bool> CVarShouldAllowSidecarSyncing(
 	false,
 	TEXT("When true FVirtualizedUntypedBulkData will attempt to sync it's .upayload file via sourcecontrol if the first attempt to load from it fails"));
 
+/** Might expose this as an option but for now set to false only */
+static constexpr bool bAllowVirtualizationOnSave = false;
+
 /** Wrapper around the config file option [Core.System.Experimental]EnablePackageSidecarSaving */
 bool ShouldSaveToPackageSidecar()
 {
@@ -111,6 +123,38 @@ bool IsValid(const FVirtualizedUntypedBulkData& BulkData, const FCompressedBuffe
 	return true;
 }
 
+/** Utility for finding the FLinkerLoad associated with a given UObject */
+FLinkerLoad* GetLinkerLoadFromOwner(UObject* Owner)
+{
+	if (Owner != nullptr)
+	{
+		UPackage* Package = Owner->GetOutermost();
+		checkf(Package != nullptr, TEXT("Owner was not a valid UPackage!"));
+
+		FLinkerLoad* Linker = FLinkerLoad::FindExistingLinkerForPackage(Package);
+		checkf(Linker != nullptr, TEXT("UPackage did not have a valid FLinkerLoad!"));
+
+		return Linker;
+	}
+	else
+	{
+		return nullptr;
+	}
+}
+
+FPackageTrailer* GetTrailerFromOwner(UObject* Owner)
+{
+	FLinkerLoad* Linker = GetLinkerLoadFromOwner(Owner);
+	if (Linker != nullptr)
+	{
+		return Linker->GetPackageTrailer();
+	}
+	else
+	{
+		return nullptr;
+	}
+}
+
 /** Utility for accessing IVirtualizationSourceControlUtilities from the modular feature system. */
 Experimental::IVirtualizationSourceControlUtilities* GetSourceControlInterface()
 {
@@ -418,6 +462,18 @@ void FVirtualizedUntypedBulkData::CreateLegacyUniqueIdentifier(UObject* Owner)
 	}
 }
 
+/** Utility for updating an existing entry in an Archive before returning the archive to it's original seek position */
+template<typename DataType>
+void UpdateArchiveData(FArchive& Ar, int64 DataPosition, DataType& Data)
+{
+	int64 OriginalPosition = Ar.Tell();
+
+	Ar.Seek(DataPosition);
+	Ar << Data;
+
+	Ar.Seek(OriginalPosition);
+}
+
 void FVirtualizedUntypedBulkData::Serialize(FArchive& Ar, UObject* Owner, bool bAllowRegister)
 {
 	TRACE_CPUPROFILER_EVENT_SCOPE(FVirtualizedUntypedBulkData::Serialize);
@@ -492,13 +548,16 @@ void FVirtualizedUntypedBulkData::Serialize(FArchive& Ar, UObject* Owner, bool b
 				UpdateKeyIfNeeded();
 			}
 
-			const bool bCanAttemptVirtualization = LinkerSave != nullptr;
-			if (bCanAttemptVirtualization)
+			if (bAllowVirtualizationOnSave)
 			{
-				FPackagePath LinkerPackagePath;
-				FPackagePath::TryFromPackageName(LinkerSave->LinkerRoot->GetName(), LinkerPackagePath);
+				const bool bCanAttemptVirtualization = LinkerSave != nullptr;
+				if (bCanAttemptVirtualization)
+				{
+					FPackagePath LinkerPackagePath;
+					FPackagePath::TryFromPackageName(LinkerSave->LinkerRoot->GetName(), LinkerPackagePath);
 
-				PushData(LinkerPackagePath); // Note this can change various members if we are going from non-virtualized to virtualized
+					PushData(LinkerPackagePath); // Note this can change various members if we are going from non-virtualized to virtualized
+				}
 			}
 		}
 		else
@@ -554,7 +613,7 @@ void FVirtualizedUntypedBulkData::Serialize(FArchive& Ar, UObject* Owner, bool b
 			}
 			else
 			{
-				bWriteOutPayload = !IsDataVirtualized();
+				bWriteOutPayload = !bAllowVirtualizationOnSave || !IsDataVirtualized();
 			}
 
 			if (bWriteOutPayload)
@@ -596,86 +655,143 @@ void FVirtualizedUntypedBulkData::Serialize(FArchive& Ar, UObject* Owner, bool b
 				const int64 OffsetPos = Ar.Tell();
 				Ar << PlaceholderValue; // OffsetInFile
 
-				// The lambda is mutable so that PayloadToSerialize is not const (due to FArchive api not
-				// accepting const values)
-				auto SerializePayload = [this, OffsetPos, PayloadToSerialize, UpdatedFlags, Owner](FLinkerSave* LinkerSave, FArchive& ExportsArchive, FArchive& DataArchive, int64 DataStartOffset) mutable
-				{
-					checkf(ExportsArchive.IsCooking() == false, TEXT("FVirtualizedUntypedBulkData::Serialize should not be called during a cook"));
+				// If we have a valid linker then we will defer serialization of the payload so that it will
+				// be placed at the end of the output file so we don't have to seek past the payload on load.
+				// If we do not have a linker OR the linker is in text format then we should just serialize
+				// the payload directly to the archive.
+				if (LinkerSave != nullptr && !LinkerSave->IsTextFormat())
+				{	
+					if (FPackageTrailer::IsEnabled() == false) // Legacy path, will save the payload data to the package
+					{
+						// The lambda is mutable so that PayloadToSerialize is not const (due to FArchive api not accepting const values)
+						auto SerializePayload = [this, OffsetPos, PayloadToSerialize, UpdatedFlags, Owner](FLinkerSave& LinkerSave, FArchive& ExportsArchive, FArchive& DataArchive, int64 DataStartOffset) mutable
+						{
+							checkf(ExportsArchive.IsCooking() == false, TEXT("FVirtualizedUntypedBulkData::Serialize should not be called during a cook"));
 
-					SerializeData(DataArchive, PayloadToSerialize, UpdatedFlags);
+							SerializeData(DataArchive, PayloadToSerialize, UpdatedFlags);
 
-					// Record the current archive offset (probably EOF but we cannot be sure)
-					const int64 ReturnPos = ExportsArchive.Tell();
+							// Record the current archive offset (probably EOF but we cannot be sure)
+							const int64 ReturnPos = ExportsArchive.Tell();
 
-					// Update the offset/size entries that we set up during ::Serialize
-					ExportsArchive.Seek(OffsetPos);
-					ExportsArchive << DataStartOffset;
-					// Restore the archive's offset
-					ExportsArchive.Seek(ReturnPos);
-
-					// If we are saving the package to disk (we have access to FLinkerSave and its filepath is valid) 
-					// then we should register a callback to be received once the package has actually been saved to 
-					// disk so that we can update the object's members to be redirected to the saved file.
-					if (LinkerSave != nullptr && !LinkerSave->GetFilename().IsEmpty())
-					{	
-						// At some point saving to the sidecar file will be mutually exclusive with saving to the asset file, at that point
-						// we can split these code paths entirely for clarity. (might need to update ::BuildFlagsForSerialization at that point too!)
-						if (ShouldSaveToPackageSidecar())
-						{
-							FLinkerSave::FSidecarStorageInfo& SidecarData = LinkerSave->SidecarDataToAppend.AddZeroed_GetRef();
-							SidecarData.Identifier = PayloadContentId;
-							SidecarData.Payload = PayloadToSerialize;
-						}
+							// Update the offset/size entries that we set up during ::Serialize
+							ExportsArchive.Seek(OffsetPos);
+							ExportsArchive << DataStartOffset;
+							// Restore the archive's offset
+							ExportsArchive.Seek(ReturnPos);
 
-						auto OnSavePackage = [this, DataStartOffset, UpdatedFlags, Owner](const FPackagePath& InPackagePath, FObjectPostSaveContext ObjectSaveContext)
-						{
-							if (!ObjectSaveContext.IsUpdatingLoadedPath())
+							// If we are saving the package to disk (we have access to FLinkerSave and its filepath is valid) 
+							// then we should register a callback to be received once the package has actually been saved to 
+							// disk so that we can update the object's members to be redirected to the saved file.
+							if (!LinkerSave.GetFilename().IsEmpty())
 							{
-								return;
+								// At some point saving to the sidecar file will be mutually exclusive with saving to the asset file, at that point
+								// we can split these code paths entirely for clarity. (might need to update ::BuildFlagsForSerialization at that point too!)
+								if (ShouldSaveToPackageSidecar())
+								{
+									FLinkerSave::FSidecarStorageInfo& SidecarData = LinkerSave.SidecarDataToAppend.AddZeroed_GetRef();
+									SidecarData.Identifier = PayloadContentId;
+									SidecarData.Payload = PayloadToSerialize;
+								}
+
+								auto OnSavePackage = [this, DataStartOffset, UpdatedFlags, Owner](const FPackagePath& InPackagePath, FObjectPostSaveContext ObjectSaveContext)
+								{
+									if (!ObjectSaveContext.IsUpdatingLoadedPath())
+									{
+										return;
+									}
+
+									this->PackagePath = InPackagePath;
+									check(!this->PackagePath.IsEmpty()); // LinkerSave guarantees a valid PackagePath if we're updating loaded path
+									this->OffsetInFile = DataStartOffset;
+									this->Flags = UpdatedFlags;
+
+									if (CanUnloadData())
+									{
+										this->CompressionSettings.Reset();
+										this->Payload.Reset();
+									}
+
+									// Update our information in the registry
+									// TODO: Pass Owner into Register once the AssetRegistry has been fixed to use the updated PackageGuid from the save
+									Register(nullptr);
+								};
+
+								LinkerSave.PostSaveCallbacks.Add(MoveTemp(OnSavePackage));
 							}
+						};
+
+						auto AdditionalDataCallback = [SerializePayload = MoveTemp(SerializePayload)](FLinkerSave& ExportsArchive, FArchive& DataArchive, int64 DataStartOffset) mutable
+						{
+							SerializePayload(ExportsArchive, ExportsArchive, DataArchive, DataStartOffset);
+						};
 
-							this->PackagePath = InPackagePath;
-							check(!this->PackagePath.IsEmpty()); // LinkerSave guarantees a valid PackagePath if we're updating loaded path
-							this->OffsetInFile = DataStartOffset;
-							this->Flags = UpdatedFlags;
+						LinkerSave->AdditionalDataToAppend.Add(MoveTemp(AdditionalDataCallback));	// -V595 PVS believes that LinkerSave can potentially be nullptr at 
+																									// this point however we test LinkerSave != nullptr to enter this branch.
 
-							if (CanUnloadData())
+					}
+					else // New path that will save the payload to the package trailer
+					{
+						auto OnPayloadWritten = [this, OffsetPos, UpdatedFlags](FLinkerSave& LinkerSave) mutable
+						{
+							checkf(LinkerSave.IsCooking() == false, TEXT("FVirtualizedUntypedBulkData::Serialize should not be called during a cook"));
+
+							int64 PayloadOffset = LinkerSave.PackageTrailerBuilder.FindPayloadOffset(PayloadContentId);
+
+							UpdateArchiveData(LinkerSave, OffsetPos, PayloadOffset);
+
+							// If we are saving the package to disk (we have access to FLinkerSave and its filepath is valid) 
+							// then we should register a callback to be received once the package has actually been saved to 
+							// disk so that we can update the object's members to be redirected to the saved file.
+							if (!LinkerSave.GetFilename().IsEmpty())
 							{
-								this->CompressionSettings.Reset();
-								this->Payload.Reset();
+								auto OnSavePackage = [this, PayloadOffset, UpdatedFlags](const FPackagePath& InPackagePath, FObjectPostSaveContext ObjectSaveContext)
+								{
+									if (!ObjectSaveContext.IsUpdatingLoadedPath())
+									{
+										return;
+									}
+
+									this->PackagePath = InPackagePath;
+									check(!this->PackagePath.IsEmpty()); // LinkerSave guarantees a valid PackagePath if we're updating loaded path
+									this->OffsetInFile = PayloadOffset;
+									this->Flags = UpdatedFlags;
+
+									if (CanUnloadData())
+									{
+										this->CompressionSettings.Reset();
+										this->Payload.Reset();
+									}
+
+									// Update our information in the registry
+									// TODO: Pass Owner into Register once the AssetRegistry has been fixed to use the updated PackageGuid from the save
+									Register(nullptr);
+								};
+
+								LinkerSave.PostSaveCallbacks.Add(MoveTemp(OnSavePackage));
 							}
-
-							// Update our information in the registry
-							// TODO: Pass Owner into Register once the AssetRegistry has been fixed to use the updated PackageGuid from the save
-							Register(nullptr);
 						};
 
-						LinkerSave->PostSaveCallbacks.Add(MoveTemp(OnSavePackage));
-					}
-				};
-
-				// If we have a valid linker then we will defer serialization of the payload so that it will
-				// be placed at the end of the output file so we don't have to seek past the payload on load.
-				// If we do not have a linker OR the linker is in text format then we should just serialize
-				// the payload directly to the archive.
-				if(LinkerSave != nullptr && !LinkerSave->IsTextFormat())
-				{
-					auto AdditionalDataCallback = [SerializePayload = MoveTemp(SerializePayload)](FLinkerSave& ExportsArchive, FArchive& DataArchive, int64 DataStartOffset) mutable
-					{
-						SerializePayload(&ExportsArchive, ExportsArchive, DataArchive, DataStartOffset);
-					};
-					LinkerSave->AdditionalDataToAppend.Add(MoveTemp(AdditionalDataCallback));	// -V595 PVS believes that LinkerSave can potentially be nullptr at 
-																								// this point however we test LinkerSave != nullptr to enter this branch.
+						LinkerSave->PackageTrailerBuilder.AddPayload(PayloadContentId, MoveTemp(PayloadToSerialize), MoveTemp(OnPayloadWritten));
+					}	
 				}
 				else
 				{
-					SerializePayload(nullptr, Ar, Ar, Ar.Tell());
-				}
-			}
+					checkf(Ar.IsCooking() == false, TEXT("FVirtualizedUntypedBulkData::Serialize should not be called during a cook"));
+					
+					int64 DataStartOffset = Ar.Tell();
 
-			if (LinkerSave != nullptr)
-			{
-				LinkerSave->BulkDataInPackage.Add(this);
+					SerializeData(Ar, PayloadToSerialize, UpdatedFlags);
+					
+					// Record the current archive offset to return to (probably EOF but we cannot be sure)
+					const int64 ReturnPos = Ar.Tell();
+
+					// Update the offset entry that we set up earlier
+					Ar.Seek(OffsetPos);
+					Ar << DataStartOffset;
+
+					// Restore the archive's offset
+					Ar.Seek(ReturnPos);
+				}
 			}
 
 			if (CanUnloadData())
@@ -686,7 +802,19 @@ void FVirtualizedUntypedBulkData::Serialize(FArchive& Ar, UObject* Owner, bool b
 		}
 		else if (Ar.IsLoading())
 		{
-			if (IsReferencingOldBulkData())
+			FPackageTrailer* Trailer = GetTrailerFromOwner(Owner);
+
+			// TODO: Error on not EPayloadStatus::NotFound
+			if (Trailer != nullptr && Trailer->FindPayloadStatus(PayloadContentId) == EPayloadStatus::StoredVirtualized)
+			{
+				check(!IsReferencingOldBulkData()); // This should not be possible
+
+				EnumAddFlags(Flags, EFlags::IsVirtualized);
+
+				Ar << OffsetInFile;
+				OffsetInFile = INDEX_NONE;
+			}
+			else if (IsReferencingOldBulkData())
 			{
 				FString PackageNameStr;
 				Ar << PackageNameStr;
@@ -696,6 +824,9 @@ void FVirtualizedUntypedBulkData::Serialize(FArchive& Ar, UObject* Owner, bool b
 			}
 			else if (IsDataVirtualized())
 			{
+				// TODO: Remove once safe to do so
+				// Legacy path for packages that were saved as virtualized
+			
 				// We aren't going to use these members so reset them
 				OffsetInFile = INDEX_NONE;
 
@@ -749,6 +880,7 @@ void FVirtualizedUntypedBulkData::Serialize(FArchive& Ar, UObject* Owner, bool b
 					}
 				}
 			}
+
 			if (bAllowRegister)
 			{
 				Register(Owner);
@@ -844,7 +976,14 @@ FCompressedBuffer FVirtualizedUntypedBulkData::LoadFromDisk() const
 	}
 	else
 	{
-		return LoadFromPackageFile();	 
+		if (CVarShouldLoadFromTrailer.GetValueOnAnyThread())
+		{
+			return LoadFromPackageTrailer();
+		}
+		else
+		{
+			return LoadFromPackageFile();
+		}
 	}
 }
 
@@ -854,8 +993,6 @@ FCompressedBuffer FVirtualizedUntypedBulkData::LoadFromPackageFile() const
 
 	UE_LOG(LogVirtualization, Verbose, TEXT("Attempting to load payload from the package file: '%s'"), *PackagePath.GetLocalFullPath(PackageSegment));
 
-	FCompressedBuffer PayloadFromDisk;
-
 	// Open a reader to the file
 	TUniquePtr<FArchive> BulkArchive;
 	if (!IsReferencingOldBulkData() || PackageSegment != EPackageSegment::Header)
@@ -875,21 +1012,68 @@ FCompressedBuffer FVirtualizedUntypedBulkData::LoadFromPackageFile() const
 		BulkArchive = IPackageResourceManager::Get().OpenReadExternalResource(EPackageExternalResource::WorkspaceDomainFile,
 			PackagePath.GetPackageName());
 	}
-	if (BulkArchive.IsValid())
+
+	if (!BulkArchive.IsValid())
 	{
-		checkf(OffsetInFile != INDEX_NONE, TEXT("Attempting to load '%s' from disk with an invalid OffsetInFile!"), *PackagePath.GetDebugNameWithExtension(PackageSegment));
-		// Move the correct location of the data in the file
-		BulkArchive->Seek(OffsetInFile);
+		LogPackageOpenFailureMessage(PackagePath, PackageSegment);
+		return FCompressedBuffer();
+	}
+
+	checkf(OffsetInFile != INDEX_NONE, TEXT("Attempting to load '%s' from disk with an invalid OffsetInFile!"), *PackagePath.GetDebugNameWithExtension(PackageSegment));
+	// Move the correct location of the data in the file
+	BulkArchive->Seek(OffsetInFile);
+
+	// Now we can actually serialize it
+	FCompressedBuffer PayloadFromDisk;
+	SerializeData(*BulkArchive, PayloadFromDisk, Flags);
+
+	return PayloadFromDisk;
+}
+
+FCompressedBuffer FVirtualizedUntypedBulkData::LoadFromPackageTrailer() const
+{
+	TRACE_CPUPROFILER_EVENT_SCOPE(FVirtualizedUntypedBulkData::LoadFromPackageTrailer);
+
+	UE_LOG(LogVirtualization, Verbose, TEXT("Attempting to load payload from the package trailer: '%s'"), *PackagePath.GetLocalFullPath(PackageSegment));
 
-		// Now we can actually serialize it
-		SerializeData(*BulkArchive, PayloadFromDisk, Flags);
+	// Open a reader to the file
+	TUniquePtr<FArchive> BulkArchive;
+	if (!IsReferencingOldBulkData() || PackageSegment != EPackageSegment::Header)
+	{
+		FOpenPackageResult Result = IPackageResourceManager::Get().OpenReadPackage(PackagePath, PackageSegment);
+		if (Result.Format == EPackageFormat::Binary)
+		{
+			BulkArchive = MoveTemp(Result.Archive);
+		}
 	}
 	else
+	{
+		// *this may have been loaded from the EditorDomain, but saved with a reference to the bulk data in the
+		// Workspace Domain file. This was only possible if PackageSegment == Header; we checked that when serializing to the EditorDomain
+		// In that case, we need to use OpenReadExternalResource to access the Workspace Domain file
+		// In the cases where *this was loaded from the WorkspaceDomain, OpenReadExternalResource and OpenReadPackage are identical.
+		BulkArchive = IPackageResourceManager::Get().OpenReadExternalResource(EPackageExternalResource::WorkspaceDomainFile,
+			PackagePath.GetPackageName());
+	}
+
+	if (!BulkArchive.IsValid())
 	{
 		LogPackageOpenFailureMessage(PackagePath, PackageSegment);
+		return FCompressedBuffer();
 	}
 
-	return PayloadFromDisk;
+	BulkArchive->Seek(BulkArchive->TotalSize());
+
+	FPackageTrailer Trailer;
+	
+	if (Trailer.TryLoadBackwards(*BulkArchive))
+	{
+		return Trailer.LoadPayload(PayloadContentId, *BulkArchive);
+	}
+	else
+	{
+		return FCompressedBuffer();
+	}	
 }
 
 FCompressedBuffer FVirtualizedUntypedBulkData::LoadFromSidecarFileInternal(ErrorVerbosity Verbosity) const
@@ -1091,14 +1275,11 @@ FCompressedBuffer FVirtualizedUntypedBulkData::PullData() const
 FPackagePath FVirtualizedUntypedBulkData::GetPackagePathFromOwner(UObject* Owner, EPackageSegment& OutPackageSegment) const
 {
 	OutPackageSegment = EPackageSegment::Header;
-	if (Owner != nullptr)
-	{
-		UPackage* Package = Owner->GetOutermost();
-		checkf(Package != nullptr, TEXT("Owner was not a valid UPackage!"));
 
-		FLinkerLoad* Linker = FLinkerLoad::FindExistingLinkerForPackage(Package);
-		checkf(Linker != nullptr, TEXT("UPackage did not have a valid FLinkerLoad!"));
+	FLinkerLoad* Linker = GetLinkerLoadFromOwner(Owner);
 
+	if (Linker != nullptr)
+	{
 		return Linker->GetPackagePath();
 	}
 	else
@@ -1457,6 +1638,12 @@ FVirtualizedUntypedBulkData::EFlags FVirtualizedUntypedBulkData::BuildFlagsForSe
 			EnumRemoveFlags(UpdatedFlags, EFlags::ReferencesLegacyFile | EFlags::LegacyFileIsCompressed | EFlags::LegacyKeyWasGuidDerived);
 		}
 
+		// If we are re-hydrating packages on save then we need to remove the virtualization flag
+		if (bAllowVirtualizationOnSave == false && LinkerSave != nullptr)
+		{
+			EnumRemoveFlags(UpdatedFlags, EFlags::IsVirtualized);
+		}
+
 		return UpdatedFlags;
 	}
 	else
@@ -1573,73 +1760,6 @@ void operator<<(FStructuredArchive::FSlot Slot, FPayloadToc& TableOfContents)
 	Record << SA_VALUE(TEXT("Entries"), TableOfContents.Contents);
 }
 
-bool FindPayloadsInPackageFile(const FPackagePath& PackagePath, EPayloadType Filter, TArray<FPayloadId>& OutPayloadIds)
-{
-	if (FPackageName::IsTextPackageExtension(PackagePath.GetHeaderExtension()))
-	{
-		UE_LOG(LogVirtualization, Warning, TEXT("Attempting to call 'FindPayloadsInPackageFile' on a text based asset '%s' this is not currently supported"), *PackagePath.GetDebugName());
-		return false;
-	}
-
-	TUniquePtr<FArchive> Ar = IPackageResourceManager::Get().OpenReadExternalResource(EPackageExternalResource::WorkspaceDomainFile, PackagePath.GetPackageName());
-
-	if (Ar.IsValid())
-	{
-		FPackageFileSummary PackageFileSummary;
-		*Ar << PackageFileSummary;
-
-		if (Ar->IsError())
-		{
-			UE_LOG(LogVirtualization, Warning, TEXT("Failed to parse the FPackageFileSummary for '%s'"), *PackagePath.GetDebugName());
-			return false;
-		}
-
-		if (PackageFileSummary.Tag == PACKAGE_FILE_TAG)
-		{
-			if (PackageFileSummary.PayloadTocOffset != INDEX_NONE)
-			{
-				Ar->Seek(PackageFileSummary.PayloadTocOffset);
-
-				FPayloadToc PayloadToC;
-				*Ar << PayloadToC;
-
-				for (const UE::Virtualization::FTocEntry& Entry : PayloadToC.GetContents())
-				{
-					switch (Filter)
-					{
-					case EPayloadType::Local:
-						if (Entry.OffsetInFile != INDEX_NONE)
-						{
-							OutPayloadIds.Add(Entry.Identifier);
-						}
-						break;
-					case EPayloadType::Virtualized:
-						if (Entry.OffsetInFile == INDEX_NONE)
-						{
-							OutPayloadIds.Add(Entry.Identifier);
-						}
-						break;
-					case EPayloadType::All:
-						OutPayloadIds.Add(Entry.Identifier);
-						break;
-					}
-				}
-			}
-		}
-		else
-		{
-			UE_LOG(LogVirtualization, Warning, TEXT("Failed to find PACKAGE_FILE_TAG in '%s', not a valid package!"), *PackagePath.GetDebugName());
-		}
-
-		return true;
-	}
-	else
-	{
-		UE_LOG(LogVirtualization, Warning, TEXT("Unable to open '%s' for reading"), *PackagePath.GetDebugName());
-		return false;
-	}
-}
-
 } // namespace UE::Virtualization
 
 #undef VBD_ALLOW_LINKERLOADER_ATTACHMENT
diff --git a/Engine/Source/Runtime/CoreUObject/Private/UObject/LinkerLoad.cpp b/Engine/Source/Runtime/CoreUObject/Private/UObject/LinkerLoad.cpp
--- a/Engine/Source/Runtime/CoreUObject/Private/UObject/LinkerLoad.cpp
+++ b/Engine/Source/Runtime/CoreUObject/Private/UObject/LinkerLoad.cpp
@@ -14,6 +14,7 @@
 #include "UObject/ObjectRedirector.h"
 #include "UObject/Package.h"
 #include "UObject/PackageResourceManager.h"
+#include "UObject/PackageTrailer.h"
 #include "UObject/UObjectHash.h"
 #include "Misc/PackageName.h"
 #include "Blueprint/BlueprintSupport.h"
@@ -720,6 +721,13 @@ FLinkerLoad::ELinkerStatus FLinkerLoad::Tick( float InTimeLimit, bool bInUseTime
 				Status = SerializePackageFileSummary();
 			}
 
+			// Serialize the header for the package trailer
+			if (Status == LINKER_Loaded)
+			{
+				SCOPED_LOADTIMER(LinkerLoad_SerializePackageTrailer);
+				Status = SerializePackageTrailer();
+			}
+
 			// Serialize the name map and register the names.
 			if( Status == LINKER_Loaded )
 			{
@@ -857,6 +865,7 @@ FLinkerLoad::FLinkerLoad(UPackage* InParent, const FPackagePath& InPackagePath,
 , PackagePath(InPackagePath)
 , Loader(nullptr)
 , InstancingContext(MoveTemp(InInstancingContext))
+, PackageTrailer(nullptr)
 , AsyncRoot(nullptr)
 , GatherableTextDataMapIndex(0)
 , ImportMapIndex(0)
@@ -864,6 +873,7 @@ FLinkerLoad::FLinkerLoad(UPackage* InParent, const FPackagePath& InPackagePath,
 , DependsMapIndex(0)
 , ExportHashIndex(0)
 , bHasSerializedPackageFileSummary(false)
+, bHasSerializedPackageTrailer(false)
 , bHasReconstructedImportAndExportMap(false)
 , bHasSerializedPreloadDependencies(false)
 , bHasFixedUpImportMap(false)
@@ -1502,6 +1512,34 @@ FLinkerLoad::ELinkerStatus FLinkerLoad::UpdateFromPackageFileSummary()
 	return LINKER_Loaded;
 }
 
+FLinkerLoad::ELinkerStatus FLinkerLoad::SerializePackageTrailer()
+{
+	if (bHasSerializedPackageTrailer)
+	{
+		return LINKER_Loaded;
+	}
+
+	check(PackageTrailer == nullptr);
+
+	if (UE::FPackageTrailer::IsEnabled() && Summary.PayloadTocOffset != INDEX_NONE)
+	{
+		int64 CurPos = Tell();
+		Seek(Summary.PayloadTocOffset);
+		
+		PackageTrailer = MakeUnique<UE::FPackageTrailer>();
+		if (!PackageTrailer->TryLoad(*this))
+		{
+			PackageTrailer.Reset();
+		}
+
+		Seek(CurPos);
+	}
+
+	bHasSerializedPackageTrailer = true;
+
+	return LINKER_Loaded;
+}
+
 /**
  * Serializes the name table.
  */
diff --git a/Engine/Source/Runtime/CoreUObject/Private/UObject/PackageTrailer.cpp b/Engine/Source/Runtime/CoreUObject/Private/UObject/PackageTrailer.cpp
new file mode 100644
--- /dev/null
+++ b/Engine/Source/Runtime/CoreUObject/Private/UObject/PackageTrailer.cpp
@@ -0,0 +1,461 @@
+// Copyright Epic Games, Inc. All Rights Reserved.
+
+#include "UObject/PackageTrailer.h"
+
+#include "CoreGlobals.h"
+#include "Misc/ConfigCacheIni.h"
+#include "Misc/PackagePath.h"
+#include "Misc/PackageSegment.h"
+#include "UObject/PackageResourceManager.h"
+
+namespace UE
+{
+
+/** The version number for the FPackageTrailer format */
+enum class EPackageTrailerVersion : uint32
+{
+	// The original trailer format when it was first added
+	INITIAL = 0,
+
+	// -----<new versions can be added before this line>-------------------------------------------------
+	// - this needs to be the last line (see note below)
+	AUTOMATIC_VERSION_PLUS_ONE,
+	AUTOMATIC_VERSION = AUTOMATIC_VERSION_PLUS_ONE - 1
+};
+
+// These asserts are here to make sure that any changes to the size of disk constants are intentional.
+// If the change was intentional then just update the assert.
+static_assert(FPackageTrailer::FHeader::StaticHeaderSizeOnDisk == 29, "FPackageTrailer::FHeader size has been changed, if this was intentional then update this assert");
+static_assert(Private::FLookupTableEntry::SizeOnDisk == 44, "FLookupTableEntry size has been changed, if this was intentional then update this assert");
+static_assert(FPackageTrailer::FFooter::SizeOnDisk == 20, "FPackageTrailer::FFooter size has been changed, if this was intentional then update this assert");
+
+/** Utility for recording failed package open reasons */
+void LogPackageOpenFailureMessage(const FPackagePath& PackagePath)
+{
+	// TODO: Check the various error paths here again!
+	const uint32 SystemError = FPlatformMisc::GetLastError();
+	// If we have a system error we can give a more informative error message but don't output it if the error is zero as 
+	// this can lead to very confusing error messages.
+	if (SystemError != 0)
+	{
+		TCHAR SystemErrorMsg[2048] = { 0 };
+		FPlatformMisc::GetSystemErrorMessage(SystemErrorMsg, sizeof(SystemErrorMsg), SystemError);
+		UE_LOG(LogVirtualization, Error, TEXT("Could not open the file '%s' for reading due to system error: '%s' (%d))"), *PackagePath.GetDebugName(), SystemErrorMsg, SystemError);
+	}
+	else
+	{
+		UE_LOG(LogVirtualization, Error, TEXT("Could not open (%s) to read FVirtualizedUntypedBulkData with an unknown error"), *PackagePath.GetDebugName());
+	}
+}
+
+namespace Private
+{
+
+FLookupTableEntry::FLookupTableEntry(const Virtualization::FPayloadId& InIdentifier, uint64 InRawSize)
+	: Identifier(InIdentifier)
+	, RawSize(InRawSize)
+{
+
+}
+
+FArchive& operator<<(FArchive& Ar, FLookupTableEntry& Entry)
+{
+	Ar << Entry.Identifier;
+	Ar << Entry.OffsetInFile;
+	Ar << Entry.CompressedSize;
+	Ar << Entry.RawSize;
+
+	return Ar;
+}
+
+} // namespace Private
+
+FPackageTrailerBuilder FPackageTrailerBuilder::Create(const FPackageTrailer& Trailer, FArchive& Ar)
+{
+	FPackageTrailerBuilder Builder;
+
+	for (const Private::FLookupTableEntry& Entry : Trailer.Header.PayloadLookupTable)
+	{
+		if (Entry.IsVirtualized())
+		{
+			Builder.VirtualizedEntries.Add(Entry.Identifier, VirtualizedEntry(Entry.CompressedSize, Entry.RawSize));
+		}
+		else
+		{
+			FCompressedBuffer Payload = Trailer.LoadPayload(Entry.Identifier, Ar);
+			Builder.LocalEntries.Add(Entry.Identifier, LocalEntry(MoveTemp(Payload)));
+		}
+	}
+
+	return Builder;
+}
+
+int64 FPackageTrailerBuilder::FindPayloadOffset(const Virtualization::FPayloadId& Identifier) const
+{
+	check(TrailerPositionInFile != INDEX_NONE);
+	check(PayloadPosInFile != INDEX_NONE);
+	check(PayloadLookupTable.Num() == LocalEntries.Num() + VirtualizedEntries.Num());
+
+	for (const Private::FLookupTableEntry& Entry : PayloadLookupTable)
+	{
+		if (Entry.Identifier == Identifier)
+		{
+			return PayloadPosInFile + Entry.OffsetInFile;
+		}
+	}
+
+	return INDEX_NONE;
+}
+
+bool FPackageTrailerBuilder::BuildAndAppendTrailer(FLinkerSave* Linker, FArchive& DataArchive)
+{
+	checkf(TrailerPositionInFile == INDEX_NONE, TEXT("Attempting to build the same FPackageTrailer multiple times"));
+
+	// Note that we do not serialize containers directly as we want a file format that is 
+	// 100% under our control. This will allow people to create external scripts that can
+	// parse and manipulate the trailer without needing to worry that we might change how
+	// our containers serialize. 
+	TrailerPositionInFile = DataArchive.Tell();
+
+	uint64 HeaderTag = FPackageTrailer::FHeader::HeaderTag;
+	DataArchive << HeaderTag;
+
+	EPackageTrailerVersion Version = EPackageTrailerVersion::AUTOMATIC_VERSION;
+	DataArchive << Version;
+
+	const uint32 DynamicHeaderSizeOnDisk = ((LocalEntries.Num() + VirtualizedEntries.Num()) * Private::FLookupTableEntry::SizeOnDisk); // Add the length of the lookup table
+	
+	uint32 HeaderSizeOnDisk = FPackageTrailer::FHeader::StaticHeaderSizeOnDisk + DynamicHeaderSizeOnDisk;
+	DataArchive << HeaderSizeOnDisk;
+
+	int64 PayloadsDataLength = 0;
+	PayloadLookupTable.Reserve(LocalEntries.Num() + VirtualizedEntries.Num());
+
+	for (const TPair<Virtualization::FPayloadId, LocalEntry>& It : LocalEntries)
+	{
+		Private::FLookupTableEntry& Entry = PayloadLookupTable.AddDefaulted_GetRef();
+		Entry.Identifier = It.Key;
+		Entry.OffsetInFile = PayloadsDataLength;
+		Entry.CompressedSize = It.Value.Payload.GetCompressedSize();
+		Entry.RawSize = It.Value.Payload.GetRawSize();
+
+		PayloadsDataLength += It.Value.Payload.GetCompressedSize();
+	}
+
+	for (const TPair<Virtualization::FPayloadId, VirtualizedEntry>& It : VirtualizedEntries)
+	{
+		Private::FLookupTableEntry& Entry = PayloadLookupTable.AddDefaulted_GetRef();
+		Entry.Identifier = It.Key;
+		Entry.OffsetInFile = INDEX_NONE;
+		Entry.CompressedSize = It.Value.CompressedSize;
+		Entry.RawSize = It.Value.RawSize;
+	}
+
+	DataArchive << PayloadsDataLength;
+
+	// Currently we only support relative access!
+	// Referenced access will come with the editor domain support.
+	EPayloadAccessMode AccessMode = EPayloadAccessMode::Relative;
+	DataArchive << AccessMode;
+
+	int32 NumPayloads = PayloadLookupTable.Num();
+	DataArchive << NumPayloads;
+
+	for (Private::FLookupTableEntry& Entry : PayloadLookupTable)
+	{
+		DataArchive << Entry;
+	}
+
+	checkf((TrailerPositionInFile + HeaderSizeOnDisk) == DataArchive.Tell(), TEXT("Header length was calculated as %d bytes but we wrote %" INT64_FMT " bytes!"), HeaderSizeOnDisk, DataArchive.Tell() - TrailerPositionInFile);
+
+	PayloadPosInFile = DataArchive.Tell();
+
+	for (TPair<Virtualization::FPayloadId, LocalEntry>& It : LocalEntries)
+	{
+		DataArchive << It.Value.Payload;
+	}
+
+	checkf((PayloadPosInFile + PayloadsDataLength) == DataArchive.Tell(), TEXT("Total payload length was calculated as %" INT64_FMT " bytes but we wrote %" INT64_FMT " bytes!"), PayloadsDataLength, DataArchive.Tell() - PayloadPosInFile);
+
+	uint64 FooterTag = FPackageTrailer::FFooter::FooterTag;
+	DataArchive << FooterTag;
+
+	int64 TrailerLength = HeaderSizeOnDisk + PayloadsDataLength + FPackageTrailer::FFooter::SizeOnDisk;
+	DataArchive << TrailerLength;
+	
+	uint32 PackageTag = PACKAGE_FILE_TAG;
+	DataArchive << PackageTag;
+
+	checkf((TrailerPositionInFile + TrailerLength) == DataArchive.Tell(), TEXT("Trailer length was calculated as %" INT64_FMT " bytes but we wrote %" INT64_FMT " bytes!"), TrailerLength, DataArchive.Tell() - TrailerPositionInFile);
+
+	for (const AdditionalDataCallback& Callback : Callbacks)
+	{
+		Callback(*Linker);
+	}
+
+	return !DataArchive.IsError();
+}
+
+bool FPackageTrailerBuilder::IsEmpty() const
+{
+	return LocalEntries.IsEmpty() && VirtualizedEntries.IsEmpty();
+}
+
+void FPackageTrailerBuilder::AddPayload(const Virtualization::FPayloadId& Identifier, FCompressedBuffer Payload, AdditionalDataCallback&& Callback)
+{
+	checkf(TrailerPositionInFile == INDEX_NONE, TEXT("Attempting to add payloads after the trailer has been built"));
+
+	Callbacks.Emplace(MoveTemp(Callback));
+	LocalEntries.Add(Identifier, LocalEntry(MoveTemp(Payload)));
+}
+
+bool FPackageTrailer::IsEnabled()
+{
+	static struct FUsePackageTrailer
+	{
+		bool bEnabled = true;
+
+		FUsePackageTrailer()
+		{		
+			GConfig->GetBool(TEXT("Core.System"), TEXT("UsePackageTrailer"), bEnabled, GEngineIni);
+			UE_LOG(LogSerialization, Log, TEXT("UsePackageTrailer: '%s'"), bEnabled ? TEXT("true") : TEXT("false"));
+
+			// Check to make sure that the editor domain is not also enabled and assert if it is.
+			// Currently the package trailer system will not work correctly with the editor domain and as
+			// it is an opt in feature in development we should just prevent people running with both
+			// options enabled.
+			// We check the config values directly to avoid needing to introduce dependencies between modules.
+			if (bEnabled)
+			{
+				FConfigFile PlatformEngineIni;
+				FConfigCacheIni::LoadLocalIniFile(PlatformEngineIni, TEXT("Editor"), true);
+
+				bool bEditorDomainEnabled = false;
+				if (GConfig->GetBool(TEXT("CookSettings"), TEXT("EditorDomainEnabled"), bEditorDomainEnabled, GEditorIni))
+				{
+					checkf(!bEditorDomainEnabled, TEXT("The package trailer system does not yet work with the editor domain!"));
+				}
+
+				if (GConfig->GetBool(TEXT("EditorDomain"), TEXT("EditorDomainEnabled"), bEditorDomainEnabled, GEditorIni))
+				{
+					checkf(!bEditorDomainEnabled, TEXT("The package trailer system does not yet work with the editor domain!"));
+				}
+			}
+		}
+	} UsePackageTrailer;
+
+	return UsePackageTrailer.bEnabled;
+}
+
+bool FPackageTrailer::TryLoadFromPackage(const FPackagePath& PackagePath, FPackageTrailer& OutTrailer)
+{
+	//TODO: Need to do this in a way that supports text based assets
+	TUniquePtr<FArchive> PackageAr = IPackageResourceManager::Get().OpenReadExternalResource(EPackageExternalResource::WorkspaceDomainFile, PackagePath.GetPackageName());
+
+	if (PackageAr.IsValid())
+	{
+		PackageAr->Seek(PackageAr->TotalSize());
+		return OutTrailer.TryLoadBackwards(*PackageAr);
+	}
+	else
+	{
+		LogPackageOpenFailureMessage(PackagePath);
+		return false;
+	}	
+}
+
+bool FPackageTrailer::TryLoad(FArchive& Ar)
+{
+	check(Ar.IsLoading());
+
+	TrailerPositionInFile = Ar.Tell();
+
+	Ar << Header.Tag;
+
+	// Make sure that we are parsing a valid FPackageTrailer
+	if (Header.Tag != FPackageTrailer::FHeader::HeaderTag)
+	{
+		return false;
+	}
+
+	Ar << Header.Version;
+
+	Ar << Header.HeaderLength;
+	Ar << Header.PayloadsDataLength;
+	Ar << Header.AccessMode;
+
+	int32 NumPayloads = 0;
+	Ar << NumPayloads;
+
+	Header.PayloadLookupTable.Reserve(NumPayloads);
+
+	for (int32 Index = 0; Index < NumPayloads; ++Index)
+	{
+		Private::FLookupTableEntry& Entry = Header.PayloadLookupTable.AddDefaulted_GetRef();
+		Ar << Entry;
+	}
+
+	return !Ar.IsError();
+}
+
+bool FPackageTrailer::TryLoadBackwards(FArchive& Ar)
+{
+	check(Ar.IsLoading());
+
+	Ar.Seek(Ar.Tell() - FFooter::SizeOnDisk);
+
+	FFooter Footer;
+
+	Ar << Footer.Tag;
+	Ar << Footer.TrailerLength;
+	Ar << Footer.PackageTag;
+
+	// First check the package tag as this indicates if the file is corrupted or not
+	if (Footer.PackageTag != PACKAGE_FILE_TAG)
+	{
+		return false;
+	}
+
+	// Now check the footer tag which will indicate if this is actually a FPackageTrailer that we are parsing
+	if (Footer.Tag != FFooter::FooterTag)
+	{
+		return false;
+	}
+
+	Ar.Seek(Ar.Tell() - Footer.TrailerLength);
+
+	return TryLoad(Ar);
+}
+
+FCompressedBuffer FPackageTrailer::LoadPayload(Virtualization::FPayloadId Id, FArchive& Ar) const
+{
+	const Private::FLookupTableEntry* Entry = Header.PayloadLookupTable.FindByPredicate([&Id](const Private::FLookupTableEntry& Entry)->bool
+		{
+			return Entry.Identifier == Id;
+		});
+
+	// TODO: Should we check if payload is correct here or let caller do that?
+
+	if (Entry == nullptr)
+	{
+		return FCompressedBuffer();
+	}
+
+	const int64 OffsetInFile = TrailerPositionInFile + + Header.HeaderLength + Entry->OffsetInFile;
+	Ar.Seek(OffsetInFile);
+
+	FCompressedBuffer Payload;
+	Ar << Payload;
+
+	return Payload;
+}
+
+bool FPackageTrailer::UpdatePayloadAsVirtualized(Virtualization::FPayloadId Identifier)
+{
+	Private::FLookupTableEntry* Entry = Header.PayloadLookupTable.FindByPredicate([&Identifier](const Private::FLookupTableEntry& Entry)->bool
+		{
+			return Entry.Identifier == Identifier;
+		});
+
+	if (Entry != nullptr)
+	{
+		Entry->OffsetInFile = INDEX_NONE;
+		return true;
+	}
+	else
+	{
+		return false;
+	}
+}
+
+EPayloadStatus FPackageTrailer::FindPayloadStatus(Virtualization::FPayloadId Id) const
+{
+	const Private::FLookupTableEntry* Entry = Header.PayloadLookupTable.FindByPredicate([&Id](const Private::FLookupTableEntry& Entry)->bool
+		{
+			return Entry.Identifier == Id;
+		});
+
+	if (Entry == nullptr)
+	{
+		return EPayloadStatus::NotFound;
+	}
+
+	return Entry->IsVirtualized() ? EPayloadStatus::StoredVirtualized : EPayloadStatus::StoredLocally;
+}
+
+int64 FPackageTrailer::GetTrailerLength() const
+{
+	return Header.HeaderLength + Header.PayloadsDataLength + FFooter::SizeOnDisk;
+}
+
+TArray<Virtualization::FPayloadId> FPackageTrailer::GetPayloads(EPayloadFilter Type) const
+{
+	TArray<Virtualization::FPayloadId> Identifiers;
+	Identifiers.Reserve(Header.PayloadLookupTable.Num());
+
+	for (const Private::FLookupTableEntry& Entry : Header.PayloadLookupTable)
+	{
+		switch (Type)
+		{
+		case EPayloadFilter::All:
+			Identifiers.Add(Entry.Identifier);
+			break;
+
+		case EPayloadFilter::Local:
+			if (!Entry.IsVirtualized())
+			{
+				Identifiers.Add(Entry.Identifier);
+			}
+			break;
+
+		case EPayloadFilter::Virtualized:
+			if (Entry.IsVirtualized())
+			{
+				Identifiers.Add(Entry.Identifier);
+			}
+			break;
+
+		default:
+			checkNoEntry();
+		}	
+	}
+
+	return Identifiers;
+}
+
+bool FindPayloadsInPackageFile(const FPackagePath& PackagePath, EPayloadFilter Filter, TArray<Virtualization::FPayloadId>& OutPayloadIds)
+{
+	if (FPackageName::IsTextPackageExtension(PackagePath.GetHeaderExtension()))
+	{
+		UE_LOG(LogVirtualization, Warning, TEXT("Attempting to call 'FindPayloadsInPackageFile' on a text based asset '%s' this is not currently supported"), *PackagePath.GetDebugName());
+		return false;
+	}
+
+	TUniquePtr<FArchive> Ar = IPackageResourceManager::Get().OpenReadExternalResource(EPackageExternalResource::WorkspaceDomainFile, PackagePath.GetPackageName());
+
+	if (Ar.IsValid())
+	{
+		Ar->Seek(Ar->TotalSize());
+
+		FPackageTrailer Trailer;
+		
+		if (Trailer.TryLoadBackwards(*Ar))
+		{
+			OutPayloadIds = Trailer.GetPayloads(Filter);
+			return true;
+		}
+		else
+		{
+			UE_LOG(LogVirtualization, Warning, TEXT("Failed to parse the FPackageTrailer for '%s'"), *PackagePath.GetDebugName());
+			return false;
+		}	
+	}
+	else
+	{
+		UE_LOG(LogVirtualization, Warning, TEXT("Unable to open '%s' for reading"), *PackagePath.GetDebugName());
+		return false;
+	}
+}
+
+} //namespace UE
+
diff --git a/Engine/Source/Runtime/CoreUObject/Private/UObject/SavePackage.cpp b/Engine/Source/Runtime/CoreUObject/Private/UObject/SavePackage.cpp
--- a/Engine/Source/Runtime/CoreUObject/Private/UObject/SavePackage.cpp
+++ b/Engine/Source/Runtime/CoreUObject/Private/UObject/SavePackage.cpp
@@ -1656,74 +1656,24 @@ struct FPackageExportTagger
 	}
 };
 
-/**
- * Writes out the initial state of the payload table of contents.
- *
- * However because we store the table of contents before the payloads in the package
- * format the data will not be correct as we cannot know the correct payload offsets.
- * The correct data will be written out by a subsequent call to 'UpdatePayloadToC'.
- *
- * Because of this we only write out the table of contents if we are currently writing
- * to a binary format package. This will reserve the space in the file that we need
- * to later write the correct data too.
- * If the package is in text format then we do not need to write anything as we can
- * simple write out the correct data later when we know it.
- */
-static ESavePackageResult WritePayloadToC(FLinkerSave* Linker, bool bTextFormat)
+[[nodiscard]] ESavePackageResult BuildAndWriteTrailer(FLinkerSave* Linker, FStructuredArchive::FRecord& StructuredArchiveRoot, bool bTextFormat)
 {
-	if (!bTextFormat)
-	{
-		Linker->Summary.PayloadTocOffset = Linker->Tell();
-
-		UE::Virtualization::FPayloadToc PayloadTableOfContents;
-		for (const UE::Virtualization::FVirtualizedUntypedBulkData* BulkData : Linker->BulkDataInPackage)
-		{
-			if (BulkData != nullptr)
-			{
-				PayloadTableOfContents.AddEntry(*BulkData);
-			}
-		}
-
-		*Linker << PayloadTableOfContents;
-	}
-
-	return ESavePackageResult::Success;
-}
-
-/**
- * Either update the existing table of contents (if writing out in binary format) or write the table
- * of contents (if writing out in text format).
- *
- * @see WritePayloadToC
- */
-ESavePackageResult UpdatePayloadToC(FLinkerSave* Linker,  FStructuredArchive::FRecord& StructuredArchiveRoot, bool bTextFormat, int64 OffsetAfterPayloadToc)
-{
-	UE::Virtualization::FPayloadToc PayloadTableOfContents;
-	for (const UE::Virtualization::FVirtualizedUntypedBulkData* BulkData : Linker->BulkDataInPackage)
-	{
-		if (BulkData != nullptr)
-		{
-			PayloadTableOfContents.AddEntry(*BulkData);
-		}
-	}
+	check(Linker != nullptr);
 
 	if (bTextFormat)
 	{
-		FStructuredArchive::FStream Stream = StructuredArchiveRoot.EnterStream(SA_FIELD_NAME(TEXT("PayloadTableOfContents")));
-		Stream.EnterElement() << PayloadTableOfContents;
+		// We shouldn't have any data in the trailer as text based assets should store 
+		// all payloads inline at the moment.
+		checkf(Linker->PackageTrailerBuilder.IsEmpty(), TEXT("Attempting to build a package trailer for text based asset, this is not supported!"));
+		Linker->Summary.PayloadTocOffset = INDEX_NONE;
 	}
 	else
 	{
-		check(Linker->Summary.PayloadTocOffset != INDEX_NONE);
-
-		int64 CurrentPos = Linker->Tell();
-		Linker->Seek(Linker->Summary.PayloadTocOffset);
-
-		*Linker << PayloadTableOfContents;
-
-		checkf(OffsetAfterPayloadToc == Linker->Tell(), TEXT("Mismatch between ::WritePayloadToC and ::UpdatePayloadToC, the package will be corrupt"));
-
-		Linker->Seek(CurrentPos);
+		Linker->Summary.PayloadTocOffset = Linker->Tell();
+		if (!Linker->PackageTrailerBuilder.BuildAndAppendTrailer(Linker, *Linker))
+		{
+			return ESavePackageResult::Error;
+		}
 	}
 
 	return ESavePackageResult::Success;
@@ -2025,7 +1975,7 @@ FSavePackageResultStruct UPackage::Save(UPackage* InOuter, UObject* InAsset, con
 
 		FText StatusMessage = FText::Format(NSLOCTEXT("Core", "SavingFile", "Saving file: {CleanFilename}..."), Args);
 
-		const int32 TotalSaveSteps = 35;
+		const int32 TotalSaveSteps = 34;
 		FScopedSlowTask SlowTask(TotalSaveSteps, StatusMessage, bSlowTask);
 		SlowTask.MakeDialogDelayed(3.0f, SaveFlags & SAVE_FromAutosave ? true : false);
 
@@ -3824,15 +3774,6 @@ FSavePackageResultStruct UPackage::Save(UPackage* InOuter, UObject* InAsset, con
 					return ESavePackageResult::Canceled;
 				}
 
-				SlowTask.EnterProgressFrame();
-				ESavePackageResult WritePayloadTocResult = WritePayloadToC(Linker.Get(), bTextFormat);
-				if (WritePayloadTocResult != ESavePackageResult::Success)
-				{
-					return WritePayloadTocResult;
-				}
-
-				const int64 OffsetAfterPayloadToc = Linker->Tell();
-
 				SlowTask.EnterProgressFrame(1, NSLOCTEXT("Core", "SerializingBulkData", "Serializing bulk data"));
 
 				FSavePackageOutputFileArray AdditionalOutputFiles;
@@ -3908,13 +3849,18 @@ FSavePackageResultStruct UPackage::Save(UPackage* InOuter, UObject* InAsset, con
 						uint32 Tag = PACKAGE_FILE_TAG;
 						StructuredArchiveRoot.GetUnderlyingArchive() << Tag;
 					}
-				}
 
-				SlowTask.EnterProgressFrame();
-				ESavePackageResult UpdatePayloadTocResult = UpdatePayloadToC(Linker.Get(), StructuredArchiveRoot, bTextFormat, OffsetAfterPayloadToc);
-				if (UpdatePayloadTocResult != ESavePackageResult::Success)
-				{
-					return UpdatePayloadTocResult;
+					// Now that the package is written out we can write the package trailer that is appended
+					// to the file. This should be the last thing written to the file!
+					SlowTask.EnterProgressFrame();
+					if (UE::FPackageTrailer::IsEnabled())
+					{
+						ESavePackageResult TrailerResult = BuildAndWriteTrailer(Linker.Get(), StructuredArchiveRoot, bTextFormat);
+						if (TrailerResult != ESavePackageResult::Success)
+						{
+							return TrailerResult;
+						}
+					}
 				}
 
 				int64 OffsetBeforeUpdates = Linker->Tell();
diff --git a/Engine/Source/Runtime/CoreUObject/Private/UObject/SavePackage2.cpp b/Engine/Source/Runtime/CoreUObject/Private/UObject/SavePackage2.cpp
--- a/Engine/Source/Runtime/CoreUObject/Private/UObject/SavePackage2.cpp
+++ b/Engine/Source/Runtime/CoreUObject/Private/UObject/SavePackage2.cpp
@@ -1650,76 +1650,22 @@ ESavePackageResult WriteExports(FStructuredArchive::FRecord& StructuredArchiveRo
 	return Linker->IsError() ? ESavePackageResult::Error : ReturnSuccessOrCancel();
 }
 
-/** 
- * Writes out the initial state of the payload table of contents.
- * 
- * However because we store the table of contents before the payloads in the package
- * format the data will not be correct as we cannot know the correct payload offsets.
- * The correct data will be written out by a subsequent call to 'UpdatePayloadToC'.
- * 
- * Because of this we only write out the table of contents if we are currently writing
- * to a binary format package. This will reserve the space in the file that we need 
- * to later write the correct data too.
- * If the package is in text format then we do not need to write anything as we can
- * simple write out the correct data later when we know it.
- */
-ESavePackageResult WritePayloadToC(FSaveContext& SaveContext)
+[[nodiscard]] ESavePackageResult BuildAndWriteTrailer(FStructuredArchive::FRecord& StructuredArchiveRoot, FSaveContext& SaveContext)
 {
-	if (!SaveContext.IsTextFormat())
-	{
-		SaveContext.Linker->Summary.PayloadTocOffset = SaveContext.Linker->Tell();
-
-		UE::Virtualization::FPayloadToc PayloadTableOfContents;
-		for (const UE::Virtualization::FVirtualizedUntypedBulkData* BulkData : SaveContext.Linker->BulkDataInPackage)
-		{
-			if (BulkData != nullptr)
-			{
-				PayloadTableOfContents.AddEntry(*BulkData);
-			}
-		}
-
-		*SaveContext.Linker << PayloadTableOfContents;
-
-		SaveContext.OffsetAfterPayloadToc = SaveContext.Linker->Tell();
-	}
-	
-	return ESavePackageResult::Success;
-}
-
-/** 
- * Either update the existing table of contents (if writing out in binary format) or write the table
- * of contents (if writing out in text format).
- * 
- * @see WritePayloadToC
- */
-ESavePackageResult UpdatePayloadToC(FStructuredArchive::FRecord& StructuredArchiveRoot, FSaveContext& SaveContext)
-{
-	UE::Virtualization::FPayloadToc PayloadTableOfContents;
-	for (const UE::Virtualization::FVirtualizedUntypedBulkData* BulkData : SaveContext.Linker->BulkDataInPackage)
-	{
-		if (BulkData != nullptr)
-		{
-			PayloadTableOfContents.AddEntry(*BulkData);
-		}
-	}
-
 	if (SaveContext.IsTextFormat())
 	{
-		FStructuredArchive::FStream Stream = StructuredArchiveRoot.EnterStream(SA_FIELD_NAME(TEXT("PayloadTableOfContents")));
-		Stream.EnterElement() << PayloadTableOfContents;
+		// We shouldn't have any data in the trailer as text based assets should store 
+		// all payloads inline at the moment.
+		checkf(SaveContext.Linker->PackageTrailerBuilder.IsEmpty(), TEXT("Attempting to build a package trailer for text based asset '%s', this is not supported!"), *SaveContext.GetPackage()->GetName());
+		SaveContext.Linker->Summary.PayloadTocOffset = INDEX_NONE;
 	}
 	else
 	{
-		check(SaveContext.Linker->Summary.PayloadTocOffset != INDEX_NONE);
-
-		int64 CurrentPos = SaveContext.Linker->Tell();
-		SaveContext.Linker->Seek(SaveContext.Linker->Summary.PayloadTocOffset);
-
-		*SaveContext.Linker << PayloadTableOfContents;
-
-		checkf(SaveContext.OffsetAfterPayloadToc == SaveContext.Linker->Tell(), TEXT("Mismatch between ::WritePayloadToC and ::UpdatePayloadToC, the package will be corrupt"));
-
-		SaveContext.Linker->Seek(CurrentPos);
+		SaveContext.Linker->Summary.PayloadTocOffset = SaveContext.Linker->Tell();
+		if (!SaveContext.Linker->PackageTrailerBuilder.BuildAndAppendTrailer(SaveContext.Linker.Get(), *SaveContext.Linker.Get()))
+		{
+			return ESavePackageResult::Error;
+		}
 	}
 
 	return ESavePackageResult::Success;
@@ -2052,7 +1998,7 @@ ESavePackageResult InnerSave(FSaveContext& SaveContext)
 	SaveContext.SetEDLCookChecker(&FEDLCookChecker::Get());
 
 	// Create slow task dialog if needed
-	const int32 TotalSaveSteps = 14;
+	const int32 TotalSaveSteps = 13;
 	FScopedSlowTask SlowTask(TotalSaveSteps, FText(), SaveContext.IsUsingSlowTask());
 	SlowTask.MakeDialogDelayed(3.0f, SaveContext.IsFromAutoSave());
 
@@ -2144,13 +2090,6 @@ ESavePackageResult InnerSave(FSaveContext& SaveContext)
 		}
 	}
 
-	SlowTask.EnterProgressFrame();
-	SaveContext.Result = WritePayloadToC(SaveContext);
-	if (SaveContext.Result != ESavePackageResult::Success)
-	{
-		return SaveContext.Result;
-	}
-
 	IPackageWriter* PackageWriter = SaveContext.GetPackageWriter();
 	if (PackageWriter)
 	{
@@ -2170,19 +2109,25 @@ ESavePackageResult InnerSave(FSaveContext& SaveContext)
 		{
 			return SaveContext.Result;
 		}
+	}
 
-		if (!SaveContext.IsTextFormat())
-		{
-			uint32 Tag = PACKAGE_FILE_TAG;
-			StructuredArchiveRoot.GetUnderlyingArchive() << Tag;
-		}
+	// Write out a tag to the end of the package
+	if (PackageWriter == nullptr && !SaveContext.IsTextFormat())
+	{
+		uint32 Tag = PACKAGE_FILE_TAG;
+		StructuredArchiveRoot.GetUnderlyingArchive() << Tag;
 	}
 
+	// Now that the package is written out we can write the package trailer that is appended
+	// to the file. This should be the last thing written to the file!
 	SlowTask.EnterProgressFrame();
-	SaveContext.Result = UpdatePayloadToC(StructuredArchiveRoot, SaveContext);
-	if (SaveContext.Result != ESavePackageResult::Success)
+	if (UE::FPackageTrailer::IsEnabled() && PackageWriter == nullptr)
 	{
-		return SaveContext.Result;
+		SaveContext.Result = BuildAndWriteTrailer(StructuredArchiveRoot, SaveContext);
+		if (SaveContext.Result != ESavePackageResult::Success)
+		{
+			return SaveContext.Result;
+		}	
 	}
 
 	int64 ExportsSize = SaveContext.Linker->Tell();
diff --git a/Engine/Source/Runtime/CoreUObject/Public/Serialization/VirtualizedBulkData.h b/Engine/Source/Runtime/CoreUObject/Public/Serialization/VirtualizedBulkData.h
--- a/Engine/Source/Runtime/CoreUObject/Public/Serialization/VirtualizedBulkData.h
+++ b/Engine/Source/Runtime/CoreUObject/Public/Serialization/VirtualizedBulkData.h
@@ -343,6 +343,7 @@ private:
 
 	FCompressedBuffer LoadFromDisk() const;
 	FCompressedBuffer LoadFromPackageFile() const;
+	FCompressedBuffer LoadFromPackageTrailer() const;
 	FCompressedBuffer LoadFromSidecarFile() const;
 	FCompressedBuffer LoadFromSidecarFileInternal(ErrorVerbosity Verbosity) const;
 
@@ -480,7 +481,10 @@ struct COREUOBJECT_API FTocEntry
 	int64 UncompressedSize = INDEX_NONE;
 };
 
-/** A table of contents showing the location of all virtualized bulkdata payloads in a file. */
+/** 
+ * A table of contents showing the location of all virtualized bulkdata payloads in a file. 
+ * NOTE: This is currently only used by the sidecar functionality.
+ */
 class COREUOBJECT_API FPayloadToc
 {
 public:
@@ -507,30 +511,6 @@ private:
 	TArray<UE::Virtualization::FTocEntry> Contents;
 };
 
-/** Used to filter a request to a specific type of payload */
-enum class EPayloadType
-{
-	/** All payload types. */
-	All,
-	/** All payloads stored locally on disk. */
-	Local,
-	/** All payloads stored in a virtualized backend. */
-	Virtualized
-};
-
-/**
- * Used to find the identifiers of the payload in a given package. Note that this will return the payloads included in the
- * package on disk and will not take into account any edits to the package if they are in memory and unsaved.
- *
- * @param PackagePath	The package to look in.
- * @param Filter		What sort of payloads should be returned. @see EPayloadType
- * @param OutPayloadIds	This array will be filled with the FPayloadId values found in the package that passed the filter.
- *						Note that existing content in the array will be preserved. It is up to the caller to empty it.
- *
- * @return 				True if the package was parsed successfully (although it might not have contained any payloads) and false if opening or parsing the package file failed.
- */
-COREUOBJECT_API bool FindPayloadsInPackageFile(const FPackagePath& PackagePath, EPayloadType Filter, TArray<FPayloadId>& OutPayloadIds);
-
 } // namespace UE::Virtualization
 
 //#endif //WITH_EDITORONLY_DATA
diff --git a/Engine/Source/Runtime/CoreUObject/Public/UObject/LinkerLoad.h b/Engine/Source/Runtime/CoreUObject/Public/UObject/LinkerLoad.h
--- a/Engine/Source/Runtime/CoreUObject/Public/UObject/LinkerLoad.h
+++ b/Engine/Source/Runtime/CoreUObject/Public/UObject/LinkerLoad.h
@@ -17,6 +17,12 @@ class IPakFile;
 class ULinkerPlaceholderExportObject;
 struct FScopedSlowTask;
 struct FUntypedBulkData;
+
+namespace UE
+{
+class FPackageTrailer;
+}
+
 namespace UE::Virtualization
 {
 class FVirtualizedUntypedBulkData;
@@ -161,6 +167,11 @@ public:
 		return InstancingContext;
 	}
 
+	UE::FPackageTrailer* GetPackageTrailer() const
+	{
+		return PackageTrailer.Get();
+	}
+
 private:
 
 	/** True if the linker is currently deleting loader */
@@ -185,6 +196,9 @@ private:
 	/** The linker instancing context. */
 	FLinkerInstancingContext InstancingContext;
 
+	/** The trailer for the package */
+	TUniquePtr<UE::FPackageTrailer> PackageTrailer;
+
 	// Helper function to access the InstancingContext IsInstanced, 
 	// returns false if WITH_EDITOR isn't defined.
 	bool IsContextInstanced() const;
@@ -340,6 +354,8 @@ private:
 
 	/** Whether we already serialized the package file summary.																*/
 	bool					bHasSerializedPackageFileSummary:1;
+	/** Whether we already serialized the package trailer.																	*/
+	bool					bHasSerializedPackageTrailer : 1;
 	/** Whether we have already reconstructed the import/export tables for a text asset */
 	bool					bHasReconstructedImportAndExportMap:1;
 	/** Whether we already serialized preload dependencies.																	*/
@@ -1010,6 +1026,11 @@ private:
 	 * Updates the linker, loader and root package with data from the package file summary.
 	 * */
 	ELinkerStatus UpdateFromPackageFileSummary();
+	
+	/** 
+	 * Serializes the header for the package trailer, allowing us to identify the payloads of the package 
+	 */
+	ELinkerStatus SerializePackageTrailer();
 
 	/**
 	 * Serializes the name map.
diff --git a/Engine/Source/Runtime/CoreUObject/Public/UObject/LinkerSave.h b/Engine/Source/Runtime/CoreUObject/Public/UObject/LinkerSave.h
--- a/Engine/Source/Runtime/CoreUObject/Public/UObject/LinkerSave.h
+++ b/Engine/Source/Runtime/CoreUObject/Public/UObject/LinkerSave.h
@@ -9,6 +9,7 @@
 #include "Templates/RefCounting.h"
 #include "UObject/Linker.h"
 #include "UObject/ObjectResource.h"
+#include "UObject/PackageTrailer.h"
 #include "UObject/UObjectThreadContext.h"
 #include "Virtualization/PayloadId.h"
 
@@ -113,9 +114,9 @@ public:
 
 	/** Used by FVirtualizedUntypedBulkData to add payloads to be added to the payload sidecar file (currently an experimental feature) */
 	TArray<FSidecarStorageInfo> SidecarDataToAppend;
-
-	/** A list of all virtualized bulkdata objects that serialized while saving the package. Note that we do not expect nullptr values to be added! */
-	TArray<UE::Virtualization::FVirtualizedUntypedBulkData*> BulkDataInPackage;
+	
+	/** Gathers all payloads while save the package, so that they can be stored ina  single data structure @see FPackageTrailer */
+	UE::FPackageTrailerBuilder PackageTrailerBuilder;
 
 	/** 
 	 * Array of callbacks that will be invoked when the package has successfully saved to disk.
diff --git a/Engine/Source/Runtime/CoreUObject/Public/UObject/PackageTrailer.h b/Engine/Source/Runtime/CoreUObject/Public/UObject/PackageTrailer.h
new file mode 100644
--- /dev/null
+++ b/Engine/Source/Runtime/CoreUObject/Public/UObject/PackageTrailer.h
@@ -0,0 +1,361 @@
+// Copyright Epic Games, Inc. All Rights Reserved.
+
+#pragma once
+
+#include "Compression/CompressedBuffer.h"
+#include "Containers/Map.h"
+#include "Virtualization/PayloadId.h"
+
+class FArchive;
+class FLinkerSave;
+class FPackagePath;
+
+namespace UE
+{
+
+/** Trailer Format
+ * The FPackageTrailer is a container that will commonly be appended to the end of a package file. The main purpose of the trailer is to store
+ * the bulkdata payloads contained by the package until they are virtualized or moved to an additional storage location.
+ * 
+ * By storing the payloads in a data format adjacent to the rest of the package we can perform the virtualization process without needing to 
+ * re-save the package itself which in turn should allow for external tools to be able to perform the virtualization process themselves
+ * rather than needing to force it through engine code.
+ * 
+ * The package trailer is intended to an easy format for external code/script to be able to manipulate. To make things clearer we do not 
+ * serialize containers directly but write out each data structure one at a time so that it should be easy to see how to manipulate the file.
+ * 
+ * The file is split into three parts:
+ * 
+ * [Header]
+ * The header contains the useful info about the trailer and the payloads in general. @See UE::FLookupTableEntry for details about 
+ * the look up table's data.
+ * 
+ * [Payload Data]
+ * If the trailer is in the workspace domain package then we will store all non-virtualized payloads here. If the trailer is in the editor 
+ * domain then there will be no payload data section and the header will be referencing the trailer in the workspace domain instead.
+ * 
+ * [Footer]
+ * The footer allows for us to load the trailer in reverse and replicates the end of package file tag (PACKAGE_FILE_TAG), it should only be
+ * used for finding the start of the trailer or validation.
+ * 
+ * CurrentVersion UE::EPackageTrailerVersion::INITIAL
+ * ______________________________________________________________________________________________________________________________________________
+ * | [Header]																																	|
+ * | Tag				| uint64			| Should match FHeader::HeaderTag, used to identify that the data being read is an FPackageTrailer	|
+ * | Version			| uint32			| Version number of the format@see UE::EPackageTrailerVersion										|
+ * | HeaderLength		| uint32			| The total size of the header on disk in bytes.													|
+ * | PayloadsDataLength	| uint64			| The total size of the payload data on disk in bytes												|
+ * | NumPayloads		| int32				| The number of payloads in LookupTableArray														|
+ * | LookupTableArray	| FLookupTableEntry | An array of FLookupTableEntry @see UE::Private::FLookupTableEntry									|
+ * |____________________________________________________________________________________________________________________________________________|
+ * | [Payload Data]																																|
+ * | Array				| FCompressedBuffer | A binary blob containing all of the payloads. Individual payloads can be found via				|
+ * |										 the LookupTableArray found in the header.															|
+ * |____________________________________________________________________________________________________________________________________________|
+ * [Footer]
+ * | Tag				| uint64			| Should match FFooter::FooterTag, used to identify that the data being read is an FPackageTrailer	|
+ * | TrailerLength		| uint64			| The total size of the trailer on disk in bytes. Can be used to find the start of the trailer when	|
+ * |										  reading backwards.																				|
+ * | PackageTag			| uint32			| The end of package tag, PACKAGE_FILE_TAG. This is used to validate that a package file on disk is	|
+ * |										  not corrupt. By ending the trailer with this tag we allow that validation code to work.			|
+ * |____________________________________________________________________________________________________________________________________________|
+ */
+
+ /** Used to filter requests to a specific type of payload */
+enum class EPayloadFilter
+{
+	/** All payload types. */
+	All,
+	/** All payloads stored locally in the package file. */
+	Local,
+	/** All payloads stored in a virtualized backend. */
+	Virtualized
+};
+
+/** Used to show the status of a payload */
+enum class EPayloadStatus
+{
+	/** The payload is not registered in the package trailer */
+	NotFound = 0,
+	/** The payload is stored locally on disk */
+	StoredLocally,
+	/** The payload is virtualized and needs to be accessed via the IVirtualizationSystem */
+	StoredVirtualized,
+};
+
+/** Lists the various methods of payload access that the trailer supports */
+enum class EPayloadAccessMode : uint8
+{
+	/** The payloads are stored in the Payload Data segment of the trailer and the offsets in FLookupTableEntry will be relative to the start of this segment */
+	Relative = 0,
+	/** The payloads are stored in the trailer of another file (most likely the workspace domain package file) and the offsets in FLookupTableEntry are absolute offsets in that external file */
+	Referenced
+};
+
+namespace Private
+{
+
+struct FLookupTableEntry
+{
+	/** Size of the entry when serialized to disk in bytes */
+	static constexpr uint32 SizeOnDisk = 44;	// Identifier		| 20 bytes
+												// OffsetInFile		| 8 bytes
+												// CompressedSize	| 8 bytes
+												// RawSize			| 8 bytes
+
+	FLookupTableEntry() = default;
+	FLookupTableEntry(const Virtualization::FPayloadId& InIdentifier, uint64 InRawSize);
+
+	friend FArchive& operator<<(FArchive& Ar, FLookupTableEntry& Entry);
+
+	[[nodiscard]] bool IsVirtualized() const
+	{
+		return OffsetInFile == INDEX_NONE;
+	}
+
+	/** Identifier for the payload */
+	Virtualization::FPayloadId Identifier;
+	/** The offset into the file where we can find the payload, note that a virtualized payload will have an offset of INDEX_NONE */
+	int64 OffsetInFile = INDEX_NONE;
+	/** The size of the payload when compressed. This will be the same value as RawSize if the payload is not compressed */
+	uint64 CompressedSize = INDEX_NONE;
+	/** The size of the payload when uncompressed. */
+	uint64 RawSize = INDEX_NONE;
+};
+
+} // namespace Private
+
+/** 
+ * This class is used to build a FPackageTrailer and write it disk.
+ * 
+ * While saving a package, payloads should be added to a FPackageTrailer via ::AddPayload then once
+ * the package has been saved to disk ::BuildAndAppendTrailer should be called. 
+ */
+class COREUOBJECT_API FPackageTrailerBuilder
+{
+public:
+	using AdditionalDataCallback = TFunction<void(FLinkerSave& LinkerSave)>;
+
+	[[nodiscard]] static FPackageTrailerBuilder Create(const class FPackageTrailer& Trailer, FArchive& Ar);
+
+	// Methods that can be called while building the trailer
+
+	/*
+	 * Adds a payload to the builder to be written to the trailer. Duplicate payloads will be discarded and only a 
+	 * single instance stored in the trailer.
+	 * 
+	 * @param Identifier	The identifier of the payload
+	 * @param Payload		The payload data
+	 * @param Callback		This callback will be invoked once the FPackageTrailer has been built and appended to disk.
+	 */
+	void AddPayload(const Virtualization::FPayloadId& Identifier, FCompressedBuffer Payload, AdditionalDataCallback&& Callback);
+	
+	/*
+	 * @param ExportsArchive	The linker associated with the package being written to disk.
+	 * @param DataArchive		The archive where the package data has been written to. This is where the FPackageTrailer will be written to
+	 * @return True if the builder was created and appended successfully and false if any error was encountered
+	 */
+	[[nodiscard]] bool BuildAndAppendTrailer(FLinkerSave* Linker, FArchive& DataArchive);
+	
+	/** Returns if the builder has any payload entries or not */
+	[[nodiscard]] bool IsEmpty() const;
+
+	// Methods that can be called after building the trailer and appending it to the package file
+	[[nodiscard]] int64 FindPayloadOffset(const Virtualization::FPayloadId& Identifier) const;
+
+private:
+	
+	/** All of the data required to add a payload that is stored on disk */
+	struct LocalEntry
+	{
+		LocalEntry() = default;
+		LocalEntry(FCompressedBuffer&& InPayload)
+			: Payload(InPayload)
+		{
+
+		}
+		~LocalEntry() = default;
+
+		FCompressedBuffer Payload;
+	};
+
+	/** All of the data required to add a payload that is virtualized */
+	struct VirtualizedEntry
+	{
+		VirtualizedEntry() = default;
+		VirtualizedEntry(int64 InCompressedSize, int64 InRawSize)
+			: CompressedSize(InCompressedSize)
+			, RawSize(InRawSize)
+		{
+
+		}
+		~VirtualizedEntry() = default;
+
+		int64 CompressedSize = INDEX_NONE;
+		int64 RawSize = INDEX_NONE;
+	};
+
+	// Members used when building the trailer
+
+	/** Payloads that will be stored locally when the trailer is written to disk */
+	TMap<Virtualization::FPayloadId, LocalEntry> LocalEntries;
+	/** Payloads that are already virtualized and so will not be written to disk */
+	TMap<Virtualization::FPayloadId, VirtualizedEntry> VirtualizedEntries;
+
+	/** Callbacks to invoke once the trailer has been written to the end of a package */
+	TArray<AdditionalDataCallback> Callbacks;
+
+	// Members that are only valid one after building the trailer and appending it to the package file
+
+	/** Where in the package file the trailer is located */
+	int64 TrailerPositionInFile = INDEX_NONE;
+
+	/** Where in the package file that the payload data is located */
+	int64 PayloadPosInFile = INDEX_NONE;
+
+	/** The same look up table that the trailer would have */
+	TArray<Private::FLookupTableEntry> PayloadLookupTable;
+};
+
+/** 
+ *
+ * The package trailer should only ever stored the payloads in the workspace domain. If the package trailer is in the editor
+ * domain then it's values should be valid, but when loading non-virtualized payloads they need to come from the workspace 
+ * domain package.
+ */
+class COREUOBJECT_API FPackageTrailer
+{
+public:
+	/** 
+	 * Returns if the feature is enabled or disabled.
+	 * 
+	 * Note that this is for development purposes only and should ship as always enabled!
+	 */
+	[[nodiscard]] static bool IsEnabled();
+
+	/** Try to load a trailer from a given package path. Note that it will always try to load the trailer from the workspace domain */
+	[[nodiscard]] static bool TryLoadFromPackage(const FPackagePath& PackagePath, FPackageTrailer& OutTrailer);
+
+	FPackageTrailer() = default;
+	~FPackageTrailer() = default;
+
+	/** 
+	 * Serializes the trailer from the given archive assuming that the seek position of the archive is already at the correct position
+	 * for the trailer.
+	 * 
+	 * @param Ar	The archive to load the trailer from
+	 * @return		True if a valid trailer was found and was able to be loaded, otherwise false.
+	 */
+	[[nodiscard]] bool TryLoad(FArchive& Ar);
+
+	/** 
+	 * Serializes the trailer from the given archive BUT assumes that the seek position of the archive is at the end of the trailer
+	 * and so will attempt to read the footer first and use that to find the start of the trailer in order to read the header.
+	 * 
+	 * @param Ar	The archive to load the trailer from
+	 * @return		True if a valid trailer was found and was able to be loaded, otherwise false.
+	 */
+	[[nodiscard]] bool TryLoadBackwards(FArchive& Ar);
+
+	/** 
+	 * Loads the a payload from the provided archive 
+	 * 
+	 * @param Id The payload to load
+	 * @param Ar The archive to load the payload from, it is assumed that this archive is the package file in the workspace domain
+	 * 
+	 * @return	The payload in the form of a FCompressedBuffer. If the payload does not exist in the trailer or in the archive then
+	 *			the FCompressedBuffer will be null.
+	 */
+	[[nodiscard]] FCompressedBuffer LoadPayload(Virtualization::FPayloadId Id, FArchive& Ar) const;
+
+	/** 
+	 * Calling this indicates that the payload has been virtualized and will no longer be stored on disk. 
+	 * 
+	 * @param Identifier The payload that has been virtualized
+	 * @return True if the payload was in the trailer, otherwise false
+	 */
+	[[nodiscard]] bool UpdatePayloadAsVirtualized(Virtualization::FPayloadId Identifier);
+
+	/** Attempt to find the status of the given payload. @See EPayloadStatus */
+	[[nodiscard]] EPayloadStatus FindPayloadStatus(Virtualization::FPayloadId Id) const;
+
+	/** Returns the total size of the of the trailer on disk in bytes */
+	[[nodiscard]] int64 GetTrailerLength() const;
+
+	/** Returns an array of the payloads that match the given type. @See EPayloadType */
+	[[nodiscard]] TArray<Virtualization::FPayloadId> GetPayloads(EPayloadFilter Type) const;
+
+	struct FHeader
+	{
+		/** Unique value used to identify the header */
+		static constexpr uint64 HeaderTag = 0xD1C43B2E80A5F697;
+
+		/** 
+		 * Size of the static header data when serialized to disk in bytes. Note that we still need to 
+		 * add the size of the data in PayloadLookupTable to get the final header size on disk 
+		 */
+		static constexpr uint32 StaticHeaderSizeOnDisk = 29;	// HeaderTag			| 8 bytes
+																// Version				| 4 bytes
+																// HeaderLength			| 4 bytes
+																// PayloadsDataLength	| 8 bytes
+																// AccessMode			| 1 byte
+																// NumPayloads			| 4 bytes
+
+		/** Expected tag at the start of the header */
+		uint64 Tag = 0;
+		/** Version of the header */
+		int32 Version = INDEX_NONE;
+		/** Total length of the header on disk in bytes */
+		uint32 HeaderLength = 0;
+		/** Total length of the payloads on disk in bytes */
+		uint64 PayloadsDataLength = 0;
+		/** What sort of access to the payloads does the trailer have */
+		EPayloadAccessMode AccessMode = EPayloadAccessMode::Relative;
+		/** Lookup table for the payloads on disk */
+		TArray<Private::FLookupTableEntry> PayloadLookupTable;
+	};
+
+	struct FFooter
+	{
+		/** Unique value used to identify the footer */
+		static constexpr uint64 FooterTag = 0x29BFCA045138DE76;
+
+		/** Size of the footer when serialized to disk in bytes */
+		static constexpr uint32 SizeOnDisk = 20;	// Tag				| 8 bytes
+													// TrailerLength	| 8 bytes
+													// PackageTag		| 4 bytes
+
+		/** Expected tag at the start of the footer */
+		uint64 Tag = 0;
+		/** Total length of the trailer on disk in bytes */
+		uint64 TrailerLength = 0;	
+		/** End the trailer with PACKAGE_FILE_TAG, which we expect all package files to end with */
+		uint32 PackageTag = 0;
+	};
+
+private:
+	friend class FPackageTrailerBuilder;
+
+	/** Where in the workspace domain package file the trailer is located */
+	int64 TrailerPositionInFile = INDEX_NONE;
+
+	/** 
+	 * The header of the trailer. Since this contains the lookup table for payloads we keep this in memory once the trailer
+	 * has been loaded. There is no need to keep the footer in memory */
+	FHeader Header;
+};
+
+/**
+ * Used to find the identifiers of the payload in a given package. Note that this will return the payloads included in the
+ * package on disk and will not take into account any edits to the package if they are in memory and unsaved.
+ *
+ * @param PackagePath	The package to look in.
+ * @param Filter		What sort of payloads should be returned. @see EPayloadType
+ * @param OutPayloadIds	This array will be filled with the FPayloadId values found in the package that passed the filter.
+ *						Note that existing content in the array will be preserved. It is up to the caller to empty it.
+ *
+ * @return 				True if the package was parsed successfully (although it might not have contained any payloads) and false if opening or parsing the package file failed.
+ */
+[[nodiscard]] COREUOBJECT_API bool FindPayloadsInPackageFile(const FPackagePath& PackagePath, EPayloadFilter Filter, TArray<Virtualization::FPayloadId>& OutPayloadIds);
+
+} //namespace UE
\ No newline at end of file
