Fix a crash when cooking in numerous projects

#rnx

- Removed an assert that isn't strictly needed and is firing now that GetLinkerLoadFromOwner is being called in other places.

diff --git a/Engine/Source/Runtime/CoreUObject/Private/Serialization/VirtualizedBulkData.cpp b/Engine/Source/Runtime/CoreUObject/Private/Serialization/VirtualizedBulkData.cpp
--- a/Engine/Source/Runtime/CoreUObject/Private/Serialization/VirtualizedBulkData.cpp
+++ b/Engine/Source/Runtime/CoreUObject/Private/Serialization/VirtualizedBulkData.cpp
@@ -131,10 +131,7 @@ FLinkerLoad* GetLinkerLoadFromOwner(UObject* Owner)
 		UPackage* Package = Owner->GetOutermost();
 		checkf(Package != nullptr, TEXT("Owner was not a valid UPackage!"));
 
-		FLinkerLoad* Linker = FLinkerLoad::FindExistingLinkerForPackage(Package);
-		checkf(Linker != nullptr, TEXT("UPackage did not have a valid FLinkerLoad!"));
-
-		return Linker;
+		return FLinkerLoad::FindExistingLinkerForPackage(Package);
 	}
 	else
 	{
