Improve UObject constructor comment to not confuse users by saying it is deprecated

Fix deprecation version

diff --git a/Engine/Source/Runtime/CoreUObject/Public/UObject/Object.h b/Engine/Source/Runtime/CoreUObject/Public/UObject/Object.h
--- a/Engine/Source/Runtime/CoreUObject/Public/UObject/Object.h
+++ b/Engine/Source/Runtime/CoreUObject/Public/UObject/Object.h
@@ -69,7 +69,11 @@ class COREUOBJECT_API UObject : public UObjectBaseUtility
 	/** Default constructor */
 	UObject();
 
-	/** Deprecated constructor, ObjectInitializer is no longer needed but is supported for older classes. */
+	/** 
+	 *  Constructor that takes an ObjectInitializer. 
+	 *  Typically not needed, but can be useful for class hierarchies that support
+	 *  optional subobjects or subobject class overriding
+	 */
 	UObject(const FObjectInitializer& ObjectInitializer);
 
 	/** DO NOT USE. This constructor is for internal usage only for statically-created objects. */
diff --git a/Engine/Source/Runtime/CoreUObject/Public/UObject/UObjectGlobals.h b/Engine/Source/Runtime/CoreUObject/Public/UObject/UObjectGlobals.h
--- a/Engine/Source/Runtime/CoreUObject/Public/UObject/UObjectGlobals.h
+++ b/Engine/Source/Runtime/CoreUObject/Public/UObject/UObjectGlobals.h
@@ -825,7 +825,7 @@ public:
 	 */
 	FObjectInitializer(UObject* InObj, UObject* InObjectArchetype, EObjectInitializerOptions InOptions, struct FObjectInstancingGraph* InInstanceGraph = nullptr);
 
-	UE_DEPRECATED(5.01, "Use version that takes EObjectInitializerOptions")
+	UE_DEPRECATED(5.0, "Use version that takes EObjectInitializerOptions")
 	FObjectInitializer(UObject* InObj, UObject* InObjectArchetype, bool bInCopyTransientsFromClassDefaults, bool bInShouldInitializeProps, struct FObjectInstancingGraph* InInstanceGraph = nullptr)
 		: FObjectInitializer(InObj, InObjectArchetype, 
 			(bInCopyTransientsFromClassDefaults ? EObjectInitializerOptions::CopyTransientsFromClassDefaults : EObjectInitializerOptions::None) |
