Added missing Import/ExportTextItem and ARO support to FStructVariant

#rnx

diff --git a/Engine/Source/Runtime/CoreUObject/Private/UObject/StructVariant.cpp b/Engine/Source/Runtime/CoreUObject/Private/UObject/StructVariant.cpp
--- a/Engine/Source/Runtime/CoreUObject/Private/UObject/StructVariant.cpp
+++ b/Engine/Source/Runtime/CoreUObject/Private/UObject/StructVariant.cpp
@@ -2,6 +2,7 @@
 
 
 #include "UObject/StructVariant.h"
+#include "UObject/UnrealType.h"
 #include "UObject/PropertyPortFlags.h"
 #include "UObject/UObjectHash.h"
 #include "UObject/Package.h"
@@ -91,6 +92,66 @@ bool FStructVariant::Identical(const FStructVariant* InOther, uint32 PortFlags)
 	return true;
 }
 
+bool FStructVariant::ExportTextItem(FString& ValueStr, const FStructVariant& DefaultValue, UObject* Parent, int32 PortFlags, UObject* ExportRootScope) const
+{
+	if (const UScriptStruct* StructTypePtr = GetStructType())
+	{
+		ValueStr += StructTypePtr->GetPathName();
+		StructTypePtr->ExportText(ValueStr, StructInstance, StructTypePtr == DefaultValue.GetStructType() ? DefaultValue.StructInstance : nullptr, Parent, PortFlags, ExportRootScope);
+	}
+	else
+	{
+		ValueStr += TEXT("None");
+	}
+	return true;
+}
+
+bool FStructVariant::ImportTextItem(const TCHAR*& Buffer, int32 PortFlags, UObject* Parent, FOutputDevice* ErrorText)
+{
+	FNameBuilder StructPathName;
+	if (const TCHAR* Result = FPropertyHelpers::ReadToken(Buffer, StructPathName, /*bDottedNames*/true))
+	{
+		Buffer = Result;
+	}
+	else
+	{
+		return false;
+	}
+
+	if (StructPathName.Len() == 0 || FCString::Stricmp(StructPathName.ToString(), TEXT("None")) == 0)
+	{
+		SetStructType(nullptr);
+	}
+	else
+	{
+		UScriptStruct* StructTypePtr = LoadObject<UScriptStruct>(nullptr, StructPathName.ToString());
+		if (!StructTypePtr)
+		{
+			return false;
+		}
+
+		SetStructType(StructTypePtr);
+		if (const TCHAR* Result = StructTypePtr->ImportText(Buffer, StructInstance, Parent, PortFlags, ErrorText, [StructTypePtr]() { return StructTypePtr->GetName(); }))
+		{
+			Buffer = Result;
+		}
+		else
+		{
+			return false;
+		}
+	}
+
+	return true;
+}
+
+void FStructVariant::AddStructReferencedObjects(FReferenceCollector& Collector)
+{
+	if (const UScriptStruct* StructTypePtr = GetStructType())
+	{
+		Collector.AddReferencedObjects(StructTypePtr, StructInstance);
+	}
+}
+
 const UScriptStruct* FStructVariant::GetStructType() const
 {
 	return StructType.Get();
diff --git a/Engine/Source/Runtime/CoreUObject/Public/UObject/StructVariant.h b/Engine/Source/Runtime/CoreUObject/Public/UObject/StructVariant.h
--- a/Engine/Source/Runtime/CoreUObject/Public/UObject/StructVariant.h
+++ b/Engine/Source/Runtime/CoreUObject/Public/UObject/StructVariant.h
@@ -41,6 +41,17 @@ public:
 	bool operator!=(const FStructVariant& InOther) const;
 	bool Identical(const FStructVariant* InOther, uint32 PortFlags) const;
 
+	/**
+	 * Import/ExportText.
+	 */
+	bool ExportTextItem(FString& ValueStr, const FStructVariant& DefaultValue, UObject* Parent, int32 PortFlags, UObject* ExportRootScope) const;
+	bool ImportTextItem(const TCHAR*& Buffer, int32 PortFlags, UObject* Parent, FOutputDevice* ErrorText);
+
+	/**
+	 * Reference collection.
+	 */
+	void AddStructReferencedObjects(class FReferenceCollector& Collector);
+
 	/**
 	 * Get the struct type of this variant.
 	 */
@@ -160,6 +171,9 @@ struct TStructOpsTypeTraits<FStructVariant> : public TStructOpsTypeTraitsBase2<F
 	enum
 	{
 		WithIdentical = true,
+		WithExportTextItem = true,
+		WithImportTextItem = true,
+		WithAddStructReferencedObjects = true,
 		WithStructuredSerializer = true,
 		WithGetPreloadDependencies = true,
 	};
