CIS Fix: Flip order of the if statement to satisfy MSVC code analysis

#rnx

- The analysis warning complains that as bAllowVirtualizationOnSave is always false, was & insteaded instead of &&.
- However bAllowVirtualizationOnSave only exists as a development tool, to allow people to enable the old functionality more easily, so we expect bAllowVirtualizationOnSave to only be evaluated as false. It is currently a constant and not a define as we * might * expose this as a config option at some point.
- Flipping the order of evaluation is enough to satify the warning.

diff --git a/Engine/Source/Runtime/CoreUObject/Private/Serialization/VirtualizedBulkData.cpp b/Engine/Source/Runtime/CoreUObject/Private/Serialization/VirtualizedBulkData.cpp
--- a/Engine/Source/Runtime/CoreUObject/Private/Serialization/VirtualizedBulkData.cpp
+++ b/Engine/Source/Runtime/CoreUObject/Private/Serialization/VirtualizedBulkData.cpp
@@ -1636,7 +1636,7 @@ FVirtualizedUntypedBulkData::EFlags FVirtualizedUntypedBulkData::BuildFlagsForSe
 		}
 
 		// If we are re-hydrating packages on save then we need to remove the virtualization flag
-		if (bAllowVirtualizationOnSave == false && LinkerSave != nullptr)
+		if (LinkerSave != nullptr && bAllowVirtualizationOnSave == false)
 		{
 			EnumRemoveFlags(UpdatedFlags, EFlags::IsVirtualized);
 		}
