Avoid gathering very expensive asset registry tags only to extract package name and class
  - Saves more than 5 minutes of GetAssetRegistryTags during cook on many different object types, especially expensive on Niagara/Blueprints...
  - Also toggle off preload during cooking

diff --git a/Engine/Source/Runtime/CoreUObject/Private/AssetRegistry/AssetData.cpp b/Engine/Source/Runtime/CoreUObject/Private/AssetRegistry/AssetData.cpp
--- a/Engine/Source/Runtime/CoreUObject/Private/AssetRegistry/AssetData.cpp
+++ b/Engine/Source/Runtime/CoreUObject/Private/AssetRegistry/AssetData.cpp
@@ -85,12 +85,12 @@ FAssetData::FAssetData(const FString& InLongPackageName, const FString& InObject
 	AssetName = FName(*InObjectPath.Mid(CharPos + 1));
 }
 
-FAssetData::FAssetData(const UObject* InAsset, bool bAllowBlueprintClass)
+FAssetData::FAssetData(const UObject* InAsset, FAssetData::ECreationFlags InCreationFlags)
 {
 	if (InAsset != nullptr)
 	{
 		const UClass* InClass = Cast<UClass>(InAsset);
-		if (InClass && InClass->ClassGeneratedBy && !bAllowBlueprintClass)
+		if (InClass && InClass->ClassGeneratedBy && !EnumHasAnyFlags(InCreationFlags, FAssetData::ECreationFlags::AllowBlueprintClass))
 		{
 			// For Blueprints, the AssetData refers to the UBlueprint and not the UBlueprintGeneratedClass
 			InAsset = InClass->ClassGeneratedBy;
@@ -104,7 +104,10 @@ FAssetData::FAssetData(const UObject* InAsset, bool bAllowBlueprintClass)
 		AssetClass = InAsset->GetClass()->GetFName();
 		ObjectPath = FName(*InAsset->GetPathName());
 
-		InAsset->GetAssetRegistryTags(*this);
+		if (!EnumHasAnyFlags(InCreationFlags, FAssetData::ECreationFlags::SkipAssetRegistryTagsGathering))
+		{
+			InAsset->GetAssetRegistryTags(*this);
+		}
 
 		ChunkIDs = Package->GetChunkIDs();
 		PackageFlags = Package->GetPackageFlags();
diff --git a/Engine/Source/Runtime/CoreUObject/Public/AssetRegistry/AssetData.h b/Engine/Source/Runtime/CoreUObject/Public/AssetRegistry/AssetData.h
--- a/Engine/Source/Runtime/CoreUObject/Public/AssetRegistry/AssetData.h
+++ b/Engine/Source/Runtime/CoreUObject/Public/AssetRegistry/AssetData.h
@@ -81,6 +81,14 @@ public:
 		return TEXT("CL_");
 	}
 
+	enum class ECreationFlags
+	{
+		None = 0,
+		SkipAssetRegistryTagsGathering = 1 << 0, // Do not perform expensive step of gathering asset registry tags at construction.
+		AllowBlueprintClass            = 1 << 1, // Unless specified, the default when trying to create one for a blueprint class 
+												 // will create one for the UBlueprint instead, but this can be overridden
+	};
+
 public:
 	/** The object path for the asset in the form PackageName.ObjectName, or PackageName.ObjectName:SubObjectName. */
 	FName ObjectPath;
@@ -117,7 +125,13 @@ public:
 	COREUOBJECT_API FAssetData(const FString& InLongPackageName, const FString& InObjectPath, FName InAssetClass, FAssetDataTagMap InTags = FAssetDataTagMap(), TArrayView<const int32> InChunkIDs = TArrayView<const int32>(), uint32 InPackageFlags = 0);
 
 	/** Constructor taking a UObject. By default trying to create one for a blueprint class will create one for the UBlueprint instead, but this can be overridden */
-	COREUOBJECT_API FAssetData(const UObject* InAsset, bool bAllowBlueprintClass = false);
+	COREUOBJECT_API FAssetData(const UObject* InAsset, FAssetData::ECreationFlags InCreationFlags = ECreationFlags::None);
+
+	/** Constructor taking a UObject. By default trying to create one for a blueprint class will create one for the UBlueprint instead, but this can be overridden */
+	inline FAssetData(const UObject* InAsset, bool bAllowBlueprintClass)
+		: FAssetData(InAsset, bAllowBlueprintClass ? ECreationFlags::AllowBlueprintClass : ECreationFlags::None)
+	{
+	}
 
 	/** FAssetDatas are equal if their object paths match */
 	bool operator==(const FAssetData& Other) const
@@ -500,6 +514,8 @@ private:
 	}
 };
 
+ENUM_CLASS_FLAGS(FAssetData::ECreationFlags);
+
 FORCEINLINE uint32 GetTypeHash(const FAssetData& AssetData)
 {
 	return GetTypeHash(AssetData.ObjectPath);
