Fixing references nested in struct properties not properly reporting the class the struct property was declared in when dumping reference chains.

Example: ACustomCharacter::BasedMovement.MovementBase becomes ACharacter::BasedMovement.MovementBase

diff --git a/Engine/Source/Runtime/CoreUObject/Private/UObject/Property.cpp b/Engine/Source/Runtime/CoreUObject/Private/UObject/Property.cpp
--- a/Engine/Source/Runtime/CoreUObject/Private/UObject/Property.cpp
+++ b/Engine/Source/Runtime/CoreUObject/Private/UObject/Property.cpp
@@ -593,6 +593,7 @@ FString FGCStackSizeHelper::GetPropertyPath() const
 {
 	FString Result;
 	const FProperty* PreviousProperty = nullptr;
+	const TCHAR DelimiterChar = '.';
 
 	for (int32 PropertyIndex = 0; PropertyIndex < PropertyStack.Num(); ++PropertyIndex)
 	{
@@ -606,7 +607,7 @@ FString FGCStackSizeHelper::GetPropertyPath() const
 				// but we do want to keep TMapName.TMapName_Key
 				continue;
 			}
-			Result += '.';
+			Result += DelimiterChar;
 		}
 		Result += Property->GetName();
 		PreviousProperty = Property;
@@ -614,6 +615,72 @@ FString FGCStackSizeHelper::GetPropertyPath() const
 	return Result;
 }
 
+bool FGCStackSizeHelper::ConvertPathToProperties(UClass* ObjectClass, const FName& InPropertyPath, TArray<FProperty*>& OutProperties)
+{
+	const TCHAR DelimiterChar = '.';
+	FString PropertyNameOrPath = InPropertyPath.ToString();
+	int32 DelimiterIndex = -1;
+	bool bFullPathConstructed = true;
+
+	if (!PropertyNameOrPath.FindChar(DelimiterChar, DelimiterIndex))
+	{
+		// 99% of the time we're be dealing with just a single property
+		FProperty* FoundProperty = ObjectClass->FindPropertyByName(*PropertyNameOrPath);
+		if (FoundProperty)
+		{
+			OutProperties.Add(FoundProperty);
+		}
+		else
+		{
+			bFullPathConstructed = false;
+		}
+	}
+	else
+	{
+		// Try and find the first property as we can't start processing the rest of the path without it
+		FString PropertyName = PropertyNameOrPath.Left(DelimiterIndex);
+		FProperty* FoundProperty = ObjectClass->FindPropertyByName(*PropertyName);
+		if (FoundProperty)
+		{
+			OutProperties.Add(FoundProperty);
+
+			int32 StartIndex = DelimiterIndex + 1;
+			const TCHAR DelimiterStr[] = { DelimiterChar, '\0' };
+			do
+			{
+				// Determine the next property name
+				DelimiterIndex = PropertyNameOrPath.Find(DelimiterStr, ESearchCase::CaseSensitive, ESearchDir::FromStart, StartIndex);
+				PropertyName = PropertyNameOrPath.Mid(StartIndex, DelimiterIndex >= 0 ? (DelimiterIndex - StartIndex) : (PropertyNameOrPath.Len() - StartIndex));
+
+				if (FStructProperty* StructProp = CastField<FStructProperty>(FoundProperty))
+				{
+					// If the previous property was a struct property, the next one belongs to the struct the previous property represented
+					FoundProperty = StructProp->Struct->FindPropertyByName(*PropertyName);
+				}
+				else
+				{
+					// In all other case (though in reality it should only be a TMap) find the inner property
+					FoundProperty = CastField<FProperty>(FoundProperty->GetInnerFieldByName(*PropertyName));
+				}
+
+				if (FoundProperty)
+				{
+					OutProperties.Add(FoundProperty);
+				}
+				else
+				{
+					bFullPathConstructed = false;
+				}
+			} while (DelimiterIndex >= 0 && bFullPathConstructed);
+		}
+		else
+		{
+			bFullPathConstructed = false;
+		}
+	}
+	return bFullPathConstructed;
+}
+
 /*-----------------------------------------------------------------------------
 	FProperty implementation.
 -----------------------------------------------------------------------------*/
diff --git a/Engine/Source/Runtime/CoreUObject/Private/UObject/ReferenceChainSearch.cpp b/Engine/Source/Runtime/CoreUObject/Private/UObject/ReferenceChainSearch.cpp
--- a/Engine/Source/Runtime/CoreUObject/Private/UObject/ReferenceChainSearch.cpp
+++ b/Engine/Source/Runtime/CoreUObject/Private/UObject/ReferenceChainSearch.cpp
@@ -336,14 +336,18 @@ void FReferenceChainSearch::DumpChain(FReferenceChainSearch::FReferenceChain* Ch
 			if (ReferenceInfo.Type == EReferenceType::Property)
 			{
 				FString ReferencingPropertyName;
-				FProperty* ReferencingProperty = ReferencerObject->GetClass()->FindPropertyByName(ReferenceInfo.ReferencerName);
-				if (ReferencingProperty)
+				TArray<FProperty*> ReferencingProperties;
+
+				if (FGCStackSizeHelper::ConvertPathToProperties(ReferencerObject->GetClass(), ReferenceInfo.ReferencerName, ReferencingProperties))
 				{
+					FProperty* InnermostProperty = ReferencingProperties.Last();
+					FProperty* OutermostProperty = ReferencingProperties[0];
+
 					ReferencingPropertyName = FString::Printf(TEXT("%s %s%s::%s"),
-						*ReferencingProperty->GetCPPType(),
-						ReferencingProperty->GetOwnerClass()->GetPrefixCPP(),
-						*ReferencingProperty->GetOwnerClass()->GetName(),
-						*ReferencingProperty->GetName());
+						*InnermostProperty->GetCPPType(),
+						OutermostProperty->GetOwnerClass()->GetPrefixCPP(),
+						*OutermostProperty->GetOwnerClass()->GetName(),
+						*ReferenceInfo.ReferencerName.ToString());
 				}
 				else
 				{
diff --git a/Engine/Source/Runtime/CoreUObject/Public/UObject/UnrealType.h b/Engine/Source/Runtime/CoreUObject/Public/UObject/UnrealType.h
--- a/Engine/Source/Runtime/CoreUObject/Public/UObject/UnrealType.h
+++ b/Engine/Source/Runtime/CoreUObject/Public/UObject/UnrealType.h
@@ -137,6 +137,15 @@ public:
 	/** Converts the interal property stack to a string representing the current property path (Member.StructMember.InnerStructMember) */
 	FString GetPropertyPath() const;
 
+	/** 
+	 * Converts a property path constructed with GetPropertyPath() to an array of properties (from the outermost to the innermost) *
+	 * @param ObjectClass Class that defines the outermost property
+	 * @param InPropertyPath Property path
+	 * @param OutProperties resulting array
+	 * @returns true if the conversion was successful, false otherwise
+	 */
+	static bool ConvertPathToProperties(UClass* ObjectClass, const FName& InPropertyPath, TArray<FProperty*>& OutProperties);
+
 	int32 GetMaxStackSize() const
 	{
 		return MaxStackSize;
