Add Save Package harvesting context:
Add 3 contexts when harvesting references during package saving: Game (cooking), Editor and the new Optional context
This will be able to be used for saving multiple package output using the new optional context but also fix problem when using `CanSkipEditorReferencedPackagesWhenCooking` and properly purge editor only internal references from the cooked context.
Add a new save flag to save the new optional context

diff --git a/Engine/Source/Runtime/CoreUObject/Private/UObject/SavePackage/PackageHarvester.cpp b/Engine/Source/Runtime/CoreUObject/Private/UObject/SavePackage/PackageHarvester.cpp
--- a/Engine/Source/Runtime/CoreUObject/Private/UObject/SavePackage/PackageHarvester.cpp
+++ b/Engine/Source/Runtime/CoreUObject/Private/UObject/SavePackage/PackageHarvester.cpp
@@ -63,7 +63,7 @@ EObjectMark GenerateMarksForObject(const UObject* InObject, const ITargetPlatfor
 	return Marks;
 }
 
-bool ConditionallyExcludeObjectForTarget(FSaveContext& SaveContext, UObject* Obj)
+bool ConditionallyExcludeObjectForTarget(FSaveContext& SaveContext, UObject* Obj, ESaveRealm HarvestingContext)
 {
 	if (!Obj || Obj->GetOutermost()->GetFName() == GLongCoreUObjectPackageName)
 	{
@@ -71,12 +71,14 @@ bool ConditionallyExcludeObjectForTarget(FSaveContext& SaveContext, UObject* Obj
 		return false;
 	}
 
+	//@note: currently always query excluded content against the default/main save harvesting context
+	// this is because generated excluded marks are not aware of the different harvesting context yet
 	bool bExcluded = false;
 	if (SaveContext.IsExcluded(Obj))
 	{
 		return true;
 	}
-	else if (!SaveContext.IsIncluded(Obj))
+	else if (!SaveContext.IsIncluded(Obj, HarvestingContext))
 	{
 		const EObjectMark ExcludedObjectMarks = SaveContext.GetExcludedObjectMarks();
 		const ITargetPlatform* TargetPlatform = SaveContext.GetTargetPlatform();
@@ -86,12 +88,12 @@ bool ConditionallyExcludeObjectForTarget(FSaveContext& SaveContext, UObject* Obj
 			UObject* ObjOuter = Obj->GetOuter();
 			UClass* ObjClass = Obj->GetClass();
 
-			if (ConditionallyExcludeObjectForTarget(SaveContext, ObjClass))
+			if (ConditionallyExcludeObjectForTarget(SaveContext, ObjClass, HarvestingContext))
 			{
 				// If the object class is excluded, the object must be excluded too
 				bExcluded = true;
 			}
-			else if (ConditionallyExcludeObjectForTarget(SaveContext, ObjOuter))
+			else if (ConditionallyExcludeObjectForTarget(SaveContext, ObjOuter, HarvestingContext))
 			{
 				// If the object outer is excluded, the object must be excluded too
 				bExcluded = true;
@@ -102,7 +104,7 @@ bool ConditionallyExcludeObjectForTarget(FSaveContext& SaveContext, UObject* Obj
 			if (ThisStruct && ThisStruct->GetSuperStruct())
 			{
 				UObject* SuperStruct = ThisStruct->GetSuperStruct();
-				if (ConditionallyExcludeObjectForTarget(SaveContext, SuperStruct))
+				if (ConditionallyExcludeObjectForTarget(SaveContext, SuperStruct, HarvestingContext))
 				{
 					bExcluded = true;
 				}
@@ -112,7 +114,7 @@ bool ConditionallyExcludeObjectForTarget(FSaveContext& SaveContext, UObject* Obj
 			UObject* Archetype = Obj->GetArchetype();
 			if (Archetype)
 			{
-				if (ConditionallyExcludeObjectForTarget(SaveContext, Archetype))
+				if (ConditionallyExcludeObjectForTarget(SaveContext, Archetype, HarvestingContext))
 				{
 					bExcluded = true;
 				}
@@ -149,8 +151,25 @@ bool DoesObjectNeedLoadForEditorGame(UObject* InObject)
 	return bNeedsLoadForEditorGame;
 }
 
+FPackageHarvester::FExportScope::FExportScope(FPackageHarvester& InHarvester, const FExportWithContext& InToProcess, bool bIsEditorOnlyObject)
+	: Harvester(InHarvester)
+	, PreviousContext(InHarvester.CurrentExportHarvestingRealm)
+{
+	check(Harvester.CurrentExportDependencies.CurrentExport == nullptr);
+	Harvester.CurrentExportDependencies = { InToProcess.Export };
+	Harvester.CurrentExportHarvestingRealm = InToProcess.HarvestedFromRealm;
+	Harvester.bIsEditorOnlyExportOnStack = bIsEditorOnlyObject;
+}
+
+FPackageHarvester::FExportScope::~FExportScope()
+{
+	Harvester.AppendCurrentExportDependencies();
+	Harvester.CurrentExportHarvestingRealm = PreviousContext;
+}
+
 FPackageHarvester::FPackageHarvester(FSaveContext& InContext)
 	: SaveContext(InContext)
+	, CurrentExportHarvestingRealm(ESaveRealm::None)
 	, bIsEditorOnlyExportOnStack(false)
 {
 	this->SetIsSaving(true);
@@ -162,29 +181,43 @@ FPackageHarvester::FPackageHarvester(FSaveContext& InContext)
 	this->SetFilterEditorOnly(SaveContext.IsFilterEditorOnly());
 	this->SetCookingTarget(SaveContext.GetTargetPlatform());
 	this->SetSerializeContext(SaveContext.GetSerializeContext());
+	this->SetUseUnversionedPropertySerialization(SaveContext.IsSaveUnversionedProperties());
 }
 
-UObject* FPackageHarvester::PopExportToProcess()
+FPackageHarvester::FExportWithContext FPackageHarvester::PopExportToProcess()
 {
-	UObject* Export = nullptr;
-	ExportsToProcess.Dequeue(Export);
-	return Export;
+	FExportWithContext ExportToProcess;
+	ExportsToProcess.Dequeue(ExportToProcess);
+	return ExportToProcess;
 }
 
-void FPackageHarvester::ProcessExport(UObject* InObject)
-{
-	check(SaveContext.IsExport(InObject));
-	bool bReferencerIsEditorOnly = IsEditorOnlyObject(InObject, true /* bCheckRecursive */, true /* bCheckMarks */) && !InObject->HasNonEditorOnlyReferences();
-	FExportScope HarvesterScope(*this, InObject, bReferencerIsEditorOnly);
+void FPackageHarvester::ProcessExport(const FExportWithContext& InProcessContext)
+{	
+	UObject* Export = InProcessContext.Export;
+
+	// No need to check marks since we do not set them on objects anymore
+	bool bReferencerIsEditorOnly = IsEditorOnlyObject(Export, true /* bCheckRecursive */, false /* bCheckMarks */) && !Export->HasNonEditorOnlyReferences();
+	FExportScope HarvesterScope(*this, InProcessContext, bReferencerIsEditorOnly);
+	
+	// The export scope set the current harvesting context
+	check(SaveContext.GetHarvestedRealm(CurrentExportHarvestingRealm).IsExport(Export));
 
 	// Harvest its class 
-	UClass* Class = InObject->GetClass();
+	UClass* Class = Export->GetClass();
 	*this << Class;
 
 	// Harvest the export outer
-	if (UObject* Outer = InObject->GetOuter())
+	if (UObject* Outer = Export->GetOuter())
 	{
-		if (!Outer->IsInPackage(SaveContext.GetPackage()))
+		auto ShouldHarvestOuterAsDependencies = [this](UObject* InObject, UObject* InOuter)
+		{
+			// Harvest the outer as dependencies if the outer is not in the package or if the outer is a ref from optional to non optional object in an optional context
+			return !InOuter->IsInPackage(SaveContext.GetPackage()) ||
+				(CurrentExportHarvestingRealm == ESaveRealm::Optional && 
+					InObject->GetClass()->HasAnyClassFlags(CLASS_Optional) && !InOuter->GetClass()->HasAnyClassFlags(CLASS_Optional));
+		};
+
+		if (ShouldHarvestOuterAsDependencies(Export, Outer))
 		{
 			*this << Outer;
 		}
@@ -197,7 +230,7 @@ void FPackageHarvester::ProcessExport(UObject* InObject)
 	}
 
 	// Harvest its template, if any
-	UObject* Template = InObject->GetArchetype();
+	UObject* Template = Export->GetArchetype();
 	if (Template
 		 && (Template != Class->GetDefaultObject() || SaveContext.IsCooking())
 		)
@@ -206,9 +239,9 @@ void FPackageHarvester::ProcessExport(UObject* InObject)
 	}
 
 	// Serialize the object or CDO
-	if (InObject->HasAnyFlags(RF_ClassDefaultObject))
+	if (Export->HasAnyFlags(RF_ClassDefaultObject))
 	{
-		Class->SerializeDefaultObject(InObject, *this);
+		Class->SerializeDefaultObject(Export, *this);
 		//@ todo FH: I don't think recursing into the template subobject is necessary, serializing it should catch the necessary sub objects
 		// GetCDOSubobjects??
 	}
@@ -217,7 +250,7 @@ void FPackageHarvester::ProcessExport(UObject* InObject)
 		// @todo FH: always serialize???
 		// In the CDO case the above would serialize most of the references, including transient properties
 		// but we still want to serialize the object using the normal path to collect all custom versions it might be using.
-		InObject->Serialize(*this);
+		Export->Serialize(*this);
 	}
 
 	// Gather object preload dependencies
@@ -228,7 +261,7 @@ void FPackageHarvester::ProcessExport(UObject* InObject)
 			// We want to tag these as imports, but not as dependencies, here since they are handled separately to the the DependsMap as SerializationBeforeSerializationDependencies instead of CreateBeforeSerializationDependencies 
 			FIgnoreDependenciesScope IgnoreDependencies(*this);
 
-			InObject->GetPreloadDependencies(Deps);
+			Export->GetPreloadDependencies(Deps);
 			for (UObject* Dep : Deps)
 			{
 				// We assume nothing in coreuobject ever loads assets in a constructor
@@ -242,7 +275,7 @@ void FPackageHarvester::ProcessExport(UObject* InObject)
 		if (SaveContext.IsProcessingPrestreamingRequests())
 		{
 			Deps.Reset();
-			InObject->GetPrestreamPackages(Deps);
+			Export->GetPrestreamPackages(Deps);
 			for (UObject* Dep : Deps)
 			{
 				if (Dep)
@@ -262,23 +295,28 @@ void FPackageHarvester::TryHarvestExport(UObject* InObject)
 {
 	// Those should have been already validated
 	check(InObject && InObject->IsInPackage(SaveContext.GetPackage()));
-	if (!SaveContext.IsExport(InObject))
+
+	// Get the realm in which we should harvest this export
+	EIllegalRefReason Reason = EIllegalRefReason::None;
+	ESaveRealm HarvestContext = GetObjectHarvestingRealm(InObject, Reason);
+	if (!SaveContext.GetHarvestedRealm(HarvestContext).IsExport(InObject))
 	{
 		SaveContext.MarkUnsaveable(InObject);
 		bool bExcluded = false;
 		if (!InObject->HasAnyFlags(RF_Transient))
 		{
-			bExcluded = ConditionallyExcludeObjectForTarget(SaveContext, InObject);
+			bExcluded = ConditionallyExcludeObjectForTarget(SaveContext, InObject, HarvestContext);
 		}
 		if (!InObject->HasAnyFlags(RF_Transient) && !bExcluded)
 		{
 			// It passed filtering so mark as export
-			SaveContext.AddExport(InObject, !DoesObjectNeedLoadForEditorGame(InObject));
-
-			// Harvest the export name
-			HarvestPackageHeaderName(InObject->GetFName());
+			HarvestExport(InObject, HarvestContext);
+		}
 
-			ExportsToProcess.Enqueue(InObject);
+		// If we have a illegal ref reason, record it
+		if (Reason != EIllegalRefReason::None)
+		{
+			SaveContext.RecordIllegalReference(CurrentExportDependencies.CurrentExport, InObject, EIllegalRefReason::ReferenceToOptional);
 		}
 	}
 }
@@ -288,7 +326,7 @@ void FPackageHarvester::TryHarvestImport(UObject* InObject)
 	// Those should have been already validated
 	check(InObject);
 	check(!InObject->IsInPackage(SaveContext.GetPackage()));
-	
+
 	auto IsObjNative = [](UObject* InObj)
 	{
 		bool bIsNative = InObj->IsNative();
@@ -301,16 +339,12 @@ void FPackageHarvester::TryHarvestImport(UObject* InObject)
 		return bIsNative;
 	};
 
-	bool bExcluded = ConditionallyExcludeObjectForTarget(SaveContext, InObject);
+	bool bExcluded = ConditionallyExcludeObjectForTarget(SaveContext, InObject, CurrentExportHarvestingRealm);
 	bool bExcludePackageFromCook = InObject && FCoreUObjectDelegates::ShouldCookPackageForPlatform.IsBound() ? !FCoreUObjectDelegates::ShouldCookPackageForPlatform.Execute(InObject->GetOutermost(), CookingTarget()) : false;
 	if (!bExcludePackageFromCook && !bExcluded && !SaveContext.IsUnsaveable(InObject))
 	{
 		bool bIsNative = IsObjNative(InObject);
-		bool bIsEditorOnly = false;
-#if WITH_EDITORONLY_DATA
-		bIsEditorOnly = bIsEditorOnlyExportOnStack || IsEditorOnlyPropertyOnTheStack();
-#endif
-		SaveContext.AddImport(InObject, bIsEditorOnly);
+		HarvestImport(InObject);
 
 		UObject* ObjOuter = InObject->GetOuter();
 		UClass* ObjClass = InObject->GetClass();
@@ -373,6 +407,14 @@ void FPackageHarvester::TryHarvestImport(UObject* InObject)
 			HarvestPackageHeaderName(ObjClass->GetOuter()->GetFName());
 		}
 	}
+
+	// Check for illegal reference
+	EIllegalRefReason Reason = EIllegalRefReason::None;
+	GetObjectHarvestingRealm(InObject, Reason);
+	if (Reason != EIllegalRefReason::None)
+	{
+		SaveContext.RecordIllegalReference(CurrentExportDependencies.CurrentExport, InObject, Reason);
+	}
 }
 
 FString FPackageHarvester::GetArchiveName() const
@@ -432,7 +474,7 @@ FArchive& FPackageHarvester::operator<<(UObject*& Obj)
 		return bIsNative;
 	};
 
-	if (SaveContext.IsIncluded(Obj))
+	if (SaveContext.IsIncluded(Obj, CurrentExportHarvestingRealm))
 	{
 		HarvestDependency(Obj, IsObjNative(Obj));
 	}
@@ -482,12 +524,12 @@ FArchive& FPackageHarvester::operator<<(FSoftObjectPath& Value)
 			FString Path = Value.ToString();
 			FName PackageName = FName(*FPackageName::ObjectPathToPackageName(Path));
 			HarvestPackageHeaderName(PackageName);
-			SaveContext.SoftPackageReferenceList.AddUnique(PackageName);
+			SaveContext.GetHarvestedRealm(CurrentExportHarvestingRealm).GetSoftPackageReferenceList().Add(PackageName);
 #if WITH_EDITORONLY_DATA
 			if (CollectType != ESoftObjectPathCollectType::EditorOnlyCollect && !bIsEditorOnlyExportOnStack)
 #endif
 			{
-				SaveContext.SoftPackagesUsedInGame.Add(PackageName);
+				SaveContext.GetHarvestedRealm(ESaveRealm::Game).GetSoftPackageReferenceList().Add(PackageName);
 			}
 		}
 	}
@@ -522,17 +564,17 @@ void FPackageHarvester::HarvestDependency(UObject* InObj, bool bIsNative)
 
 bool FPackageHarvester::CurrentExportHasDependency(UObject* InObj) const
 {
-	return SaveContext.ExportObjectDependencies.Contains(InObj) || SaveContext.ExportNativeObjectDependencies.Contains(InObj);
+	return SaveContext.GetHarvestedRealm(CurrentExportHarvestingRealm).GetObjectDependencies().Contains(InObj) || SaveContext.GetHarvestedRealm(CurrentExportHarvestingRealm).GetNativeObjectDependencies().Contains(InObj);
 }
 
 void FPackageHarvester::HarvestExportDataName(FName Name)
 {
-	SaveContext.NamesReferencedFromExportData.Add(Name.GetDisplayIndex());
+	SaveContext.GetHarvestedRealm(CurrentExportHarvestingRealm).GetNamesReferencedFromExportData().Add(Name.GetDisplayIndex());
 }
 
 void FPackageHarvester::HarvestPackageHeaderName(FName Name)
 {
-	SaveContext.NamesReferencedFromPackageHeader.Add(Name.GetDisplayIndex());
+	SaveContext.GetHarvestedRealm(CurrentExportHarvestingRealm).GetNamesReferencedFromPackageHeader().Add(Name.GetDisplayIndex());
 }
 
 void FPackageHarvester::HarvestSearchableName(UObject* TypeObject, FName Name)
@@ -544,13 +586,70 @@ void FPackageHarvester::HarvestSearchableName(UObject* TypeObject, FName Name)
 	}
 
 	HarvestPackageHeaderName(Name);
-	SaveContext.SearchableNamesObjectMap.FindOrAdd(TypeObject).AddUnique(Name);
+	SaveContext.GetHarvestedRealm(CurrentExportHarvestingRealm).GetSearchableNamesObjectMap().FindOrAdd(TypeObject).AddUnique(Name);
+}
+
+ESaveRealm FPackageHarvester::GetObjectHarvestingRealm(UObject* InObject, EIllegalRefReason& OutReason) const
+{
+	OutReason = EIllegalRefReason::None;
+	switch (CurrentExportHarvestingRealm)
+	{
+	// We are harvesting InObject from the root (i.e. asset or top level flag)
+	case ESaveRealm::None:
+		// if the object is optional and we are cooking, harvest in the Optional context
+		if (InObject->GetClass()->HasAnyClassFlags(CLASS_Optional) && SaveContext.IsCooking())
+		{
+			return ESaveRealm::Optional;
+		}
+
+		// Otherwise, just return the current default harvesting context (i.e. Game while cooking, Editor otherwise)
+		return SaveContext.CurrentHarvestingRealm;
+
+	// We are harvesting InObject from an optional object
+	case ESaveRealm::Optional:
+		// whatever the type of object we are harvesting, harvest it in the Optional context
+		return ESaveRealm::Optional;
+
+	// We are harvesting InObject from a game or editor only object
+	default:
+		//@todo FH: check CanSkipEditorReferencedPackagesWhenCooking to propagate the editor context when trimming editor reference and potentially skip processing
+
+		// if we are harvesting an optional object while in a different context, record an illegal reference to display on validation
+		if (InObject->GetClass()->HasAnyClassFlags(CLASS_Optional) && SaveContext.IsCooking())
+		{
+			OutReason = EIllegalRefReason::ReferenceToOptional;
+			return ESaveRealm::Optional;
+		}
+		// Otherwise propagate the current context
+		return CurrentExportHarvestingRealm;
+	}
+}
+
+void FPackageHarvester::HarvestExport(UObject* InObject, ESaveRealm InContext)
+{
+	SaveContext.GetHarvestedRealm(InContext).AddExport(InObject, !DoesObjectNeedLoadForEditorGame(InObject));
+	SaveContext.GetHarvestedRealm(InContext).GetNamesReferencedFromPackageHeader().Add(InObject->GetFName().GetDisplayIndex());
+	ExportsToProcess.Enqueue({ InObject, InContext });
+}
+
+void FPackageHarvester::HarvestImport(UObject* InObject)
+{
+	bool bIsEditorOnly = false;
+#if WITH_EDITORONLY_DATA
+	bIsEditorOnly = bIsEditorOnlyExportOnStack || IsEditorOnlyPropertyOnTheStack() || CurrentExportHarvestingRealm == ESaveRealm::Optional;
+#endif
+	SaveContext.GetHarvestedRealm(CurrentExportHarvestingRealm).AddImport(InObject);
+	// No matter the current context, if the import is not editor only also add it to the game context, this is later used in asset registry saving
+	if (!bIsEditorOnly)
+	{
+		SaveContext.GetHarvestedRealm(ESaveRealm::Game).AddImport(InObject);
+	}
 }
 
 void FPackageHarvester::AppendCurrentExportDependencies()
 {
 	check(CurrentExportDependencies.CurrentExport);
-	SaveContext.ExportObjectDependencies.Add(CurrentExportDependencies.CurrentExport, MoveTemp(CurrentExportDependencies.ObjectReferences));
-	SaveContext.ExportNativeObjectDependencies.Add(CurrentExportDependencies.CurrentExport, MoveTemp(CurrentExportDependencies.NativeObjectReferences));
+	SaveContext.GetHarvestedRealm(CurrentExportHarvestingRealm).GetObjectDependencies().Add(CurrentExportDependencies.CurrentExport, MoveTemp(CurrentExportDependencies.ObjectReferences));
+	SaveContext.GetHarvestedRealm(CurrentExportHarvestingRealm).GetNativeObjectDependencies().Add(CurrentExportDependencies.CurrentExport, MoveTemp(CurrentExportDependencies.NativeObjectReferences));
 	CurrentExportDependencies.CurrentExport = nullptr;
 }
\ No newline at end of file
diff --git a/Engine/Source/Runtime/CoreUObject/Private/UObject/SavePackage/PackageHarvester.h b/Engine/Source/Runtime/CoreUObject/Private/UObject/SavePackage/PackageHarvester.h
--- a/Engine/Source/Runtime/CoreUObject/Private/UObject/SavePackage/PackageHarvester.h
+++ b/Engine/Source/Runtime/CoreUObject/Private/UObject/SavePackage/PackageHarvester.h
@@ -7,6 +7,8 @@
 #include "Serialization/ArchiveUObject.h"
 #include "UObject/NameTypes.h"
 
+enum class ESaveRealm : uint32;
+enum class EIllegalRefReason : uint8;
 class FSaveContext;
 class UObject;
 struct FWeakObjectPtr;
@@ -20,24 +22,27 @@ struct FSoftObjectPath;
 class FPackageHarvester : public FArchiveUObject
 {
 public:
-	class FExportScope
+	struct FExportWithContext
 	{
-	public:
-		FExportScope(FPackageHarvester& InHarvester, UObject* InExport, bool bIsEditorOnlyObject)
-			: Harvester(InHarvester)
-		{
-			check(Harvester.CurrentExportDependencies.CurrentExport == nullptr);
-			Harvester.CurrentExportDependencies = { InExport };
-			Harvester.bIsEditorOnlyExportOnStack = bIsEditorOnlyObject;
-		}
+		/* Export to process */
+		UObject* Export = nullptr;
+		/* The harvesting context from which this export was harvested. */
+		ESaveRealm HarvestedFromRealm;
 
-		~FExportScope()
+		operator bool() const
 		{
-			Harvester.AppendCurrentExportDependencies();
-			Harvester.bIsEditorOnlyExportOnStack = false;
+			return Export != nullptr;
 		}
+	};
+
+	class FExportScope
+	{
+	public:
+		FExportScope(FPackageHarvester& InHarvester, const FExportWithContext& InToProcess, bool bIsEditorOnlyObject);
+		~FExportScope();
 	private:
 		FPackageHarvester& Harvester;
+		ESaveRealm PreviousContext;
 	};
 
 	class FIgnoreDependenciesScope
@@ -62,9 +67,9 @@ public:
 public:
 	FPackageHarvester(FSaveContext& InContext);
 
-	UObject* PopExportToProcess();
+	FExportWithContext PopExportToProcess();
 
-	void ProcessExport(UObject* InObject);
+	void ProcessExport(const FExportWithContext& InProcessContext);
 	void TryHarvestExport(UObject* InObject);
 	void TryHarvestImport(UObject* InObject);
 
@@ -85,6 +90,10 @@ public:
 	virtual FArchive& operator<<(FName& Name) override;
 
 private:
+	ESaveRealm GetObjectHarvestingRealm(UObject* InObject, EIllegalRefReason& OutReason) const;
+
+	void HarvestExport(UObject* InObject, ESaveRealm InContext);
+	void HarvestImport(UObject* InObject);
 
 	void AppendCurrentExportDependencies();
 
@@ -97,7 +106,10 @@ private:
 	};
 
 	FSaveContext& SaveContext;
-	TQueue<UObject*> ExportsToProcess;
+
+	TQueue<FExportWithContext> ExportsToProcess;
 	FExportDependencies CurrentExportDependencies;
+	ESaveRealm CurrentExportHarvestingRealm;
+	//@todo: bIsEditorOnlyExportOnStack can be probably be folded in CurrentExportHarvestingContext
 	bool bIsEditorOnlyExportOnStack;
 };
\ No newline at end of file
diff --git a/Engine/Source/Runtime/CoreUObject/Private/UObject/SavePackage/SaveContext.cpp b/Engine/Source/Runtime/CoreUObject/Private/UObject/SavePackage/SaveContext.cpp
--- a/Engine/Source/Runtime/CoreUObject/Private/UObject/SavePackage/SaveContext.cpp
+++ b/Engine/Source/Runtime/CoreUObject/Private/UObject/SavePackage/SaveContext.cpp
@@ -4,6 +4,26 @@
 
 #include "Serialization/PackageWriter.h"
 
+
+
+TArray<ESaveRealm> FSaveContext::GetHarvestedRealmsToSave()
+{
+	TArray<ESaveRealm> HarvestedContextsToSave;
+	if (IsCooking())
+	{
+		HarvestedContextsToSave.Add(ESaveRealm::Game);
+		if (IsSaveOptional())
+		{
+			HarvestedContextsToSave.Add(ESaveRealm::Optional);
+		}
+	}
+	else
+	{
+		HarvestedContextsToSave.Add(ESaveRealm::Editor);
+	}
+	return HarvestedContextsToSave;
+}
+
 void FSaveContext::MarkUnsaveable(UObject* InObject)
 {
 	if (IsUnsaveable(InObject))
diff --git a/Engine/Source/Runtime/CoreUObject/Private/UObject/SavePackage/SaveContext.h b/Engine/Source/Runtime/CoreUObject/Private/UObject/SavePackage/SaveContext.h
--- a/Engine/Source/Runtime/CoreUObject/Private/UObject/SavePackage/SaveContext.h
+++ b/Engine/Source/Runtime/CoreUObject/Private/UObject/SavePackage/SaveContext.h
@@ -32,15 +32,22 @@
 {
 	UObject* Obj;
 	uint32 bNotAlwaysLoadedForEditorGame : 1;
+	/**
+	 * Indicate that this export should have a public hash even if it isn't marked as RF_Public
+	 * This will artificially mark the object RF_Public in the linker tables so the iostore generates the public hash 
+	 */
+	uint32 bGeneratePublicHash : 1;
 
 	FTaggedExport()
 		: Obj(nullptr)
 		, bNotAlwaysLoadedForEditorGame(false)
+		, bGeneratePublicHash(false)
 	{}
 
 	FTaggedExport(UObject* InObj, bool bInNotAlwaysLoadedForEditorGame = true)
 		: Obj(InObj)
 		, bNotAlwaysLoadedForEditorGame(bInNotAlwaysLoadedForEditorGame)
+		, bGeneratePublicHash(false)
 	{}
 
 	inline bool operator == (const FTaggedExport& Other) const
@@ -54,12 +61,209 @@
 	return GetTypeHash(Export.Obj);
 }
 
+/** 
+ * Available save realm during save package harvesting 
+ * A realm is the set of objects gathered and referenced for a particular domain/context
+ */
+enum class ESaveRealm : uint32
+{
+	Game		= 0,
+	Optional,
+	Editor,
+	RealmCount,
+	None		= RealmCount
+};
 
+/** Reason for harvested illegal reference */
+enum class EIllegalRefReason : uint8
+{
+	None = 0,
+	ReferenceToOptional,
+	ReferenceFromOptionalToMissingGameExport,
+};
+
+/** Small struct to store illegal references harvested during save */
+struct FIllegalReference
+{
+	UObject* From = nullptr;
+	UObject* To = nullptr;
+	EIllegalRefReason Reason;
+};
+
+/** Hold the harvested exports and imports for a realm */
+struct FHarvestedRealm
+{
+	void AddImport(UObject* InObject)
+	{
+		Imports.Add(InObject);
+	}
+
+	void AddExport(UObject* InObj, bool bNotAlwaysLoadedForEditorGame)
+	{
+		Exports.Add(FTaggedExport(InObj, bNotAlwaysLoadedForEditorGame));
+	}
+
+	void AddExcluded(UObject* InObject)
+	{
+		Excluded.Add(InObject);
+	}
+
+	bool IsImport(UObject* InObject) const
+	{
+		return Imports.Contains(InObject);
+	}
+
+	bool IsExport(UObject* InObject) const
+	{
+		return Exports.Contains(InObject);
+	}
+
+	bool IsIncluded(UObject* InObject) const
+	{
+		return IsImport(InObject) || IsExport(InObject);
+	}
+
+	bool IsExcluded(UObject* InObject) const
+	{
+		return Excluded.Contains(InObject);
+	}
+
+	TSet<FTaggedExport>& GetExports()
+	{
+		return Exports;
+	}
+
+	const TSet<UObject*>& GetImports() const
+	{
+		return Imports;
+	}
+
+	const TSet<FName>& GetSoftPackageReferenceList() const
+	{
+		return SoftPackageReferenceList;
+	}
+
+	TSet<FName>& GetSoftPackageReferenceList()
+	{
+		return SoftPackageReferenceList;
+	}
+
+	const TMap<UObject*, TArray<FName>>& GetSearchableNamesObjectMap() const
+	{
+		return SearchableNamesObjectMap;
+	}
+
+	TMap<UObject*, TArray<FName>>& GetSearchableNamesObjectMap()
+	{
+		return SearchableNamesObjectMap;
+	}
+
+	const TSet<FNameEntryId>& GetNamesReferencedFromExportData() const
+	{
+		return NamesReferencedFromExportData;
+	}
+
+	TSet<FNameEntryId>& GetNamesReferencedFromExportData()
+	{
+		return NamesReferencedFromExportData;
+	}
+
+	const TSet<FNameEntryId>& GetNamesReferencedFromPackageHeader() const
+	{
+		return NamesReferencedFromPackageHeader;
+	}
+
+	TSet<FNameEntryId>& GetNamesReferencedFromPackageHeader()
+	{
+		return NamesReferencedFromPackageHeader;
+	}
+
+	const TMap<UObject*, TSet<UObject*>>& GetObjectDependencies() const
+	{
+		return ExportObjectDependencies;
+	}
+
+	TMap<UObject*, TSet<UObject*>>& GetObjectDependencies()
+	{
+		return ExportObjectDependencies;
+	}
+
+	const TMap<UObject*, TSet<UObject*>>& GetNativeObjectDependencies() const
+	{
+		return ExportNativeObjectDependencies;
+	}
+
+	TMap<UObject*, TSet<UObject*>>& GetNativeObjectDependencies()
+	{
+		return ExportNativeObjectDependencies;
+	}
+
+	bool NameExists(FNameEntryId ComparisonId) const
+	{
+		for (FNameEntryId DisplayId : NamesReferencedFromExportData)
+		{
+			if (FName::GetComparisonIdFromDisplayId(DisplayId) == ComparisonId)
+			{
+				return true;
+			}
+		}
+		for (FNameEntryId DisplayId : NamesReferencedFromPackageHeader)
+		{
+			if (FName::GetComparisonIdFromDisplayId(DisplayId) == ComparisonId)
+			{
+				return true;
+			}
+		}
+		return false;
+	}
+
+private:
+	// Set of objects excluded (import or exports) through marks or otherwise (i.e. transient flags, etc)
+	TSet<UObject*> Excluded;
+	// Set of objects marked as export
+	TSet<FTaggedExport> Exports;
+	// Set of objects marked as import
+	TSet<UObject*> Imports;
+	// Set of names referenced from export serialization
+	TSet<FNameEntryId> NamesReferencedFromExportData;
+	// Set of names referenced from the package header (import and export table object names etc)
+	TSet<FNameEntryId> NamesReferencedFromPackageHeader;
+	// List of soft package reference found
+	TSet<FName> SoftPackageReferenceList;
+	// Map of objects to their list of searchable names
+	TMap<UObject*, TArray<FName>> SearchableNamesObjectMap;
+	// Map of objects to their dependencies
+	TMap<UObject*, TSet<UObject*>> ExportObjectDependencies;
+	// Map of objects to their native dependencies
+	TMap<UObject*, TSet<UObject*>> ExportNativeObjectDependencies;
+};
+
+
 /**
  * Helper class that encapsulate the full necessary context and intermediate result to save a package
  */
 class FSaveContext
 {
 public:
+	struct FSetSaveRealmToSaveScope
+	{
+		FSetSaveRealmToSaveScope(FSaveContext& InContext, ESaveRealm InHarvestingRealm)
+			: Context(InContext)
+			, PreviousHarvestingRealm(InContext.CurrentHarvestingRealm)
+		{
+			Context.CurrentHarvestingRealm = InHarvestingRealm;
+		}
+
+		~FSetSaveRealmToSaveScope()
+		{
+			Context.CurrentHarvestingRealm = PreviousHarvestingRealm;
+		}
+
+	private:
+		FSaveContext& Context;
+		ESaveRealm PreviousHarvestingRealm;
+	};
+
+public:
 	FSaveContext(UPackage* InPackage, UObject* InAsset, const TCHAR* InFilename, const FSavePackageArgs& InSaveArgs, FUObjectSerializeContext* InSerializeContext = nullptr)
 		: Package(InPackage)
@@ -109,7 +313,10 @@
 		}
 
 		ObjectSaveContext.Set(InPackage, GetTargetPlatform(), TargetPackagePath, SaveArgs.SaveFlags);
-	}
 
+		// Setup the harvesting flags and generate the context for harvesting the package
+		CreateHarvestingRealms();
+	} 
+
 	~FSaveContext()
 	{
@@ -265,6 +472,11 @@
 		return !!(SaveArgs.SaveFlags & SAVE_Unversioned_Properties) && bCanUseUnversionedPropertySerialization;
 	}
 
+	bool IsSaveOptional() const
+	{
+		return !!(SaveArgs.SaveFlags & SAVE_Optional);
+	}
+
 	bool IsComputeHash() const
 	{
 		return !!(SaveArgs.SaveFlags & SAVE_ComputeHash);
@@ -345,112 +557,135 @@
 		bGenerateFileStub = true;
 	}
 
-	void AddImport(UObject* InObject, bool bIsEditorOnlyImport = false)
+	ESaveRealm GetCurrentHarvestingRealm() const
 	{
-		Imports.Add(InObject);
-		if (!bIsEditorOnlyImport)
-		{
-			ImportsUsedInGame.Add(InObject);
-		}
+		return CurrentHarvestingRealm;
 	}
 
-	void AddExport(UObject* InObj, bool bNotAlwaysLoadedForEditorGame)
-	{
-		Exports.Add(FTaggedExport(InObj, bNotAlwaysLoadedForEditorGame));
-	}
+	TArray<ESaveRealm> GetHarvestedRealmsToSave();
 
-	void AddExcluded(UObject* InObject)
-	{
-		Excluded.Add(InObject);
-	}
-
 	void MarkUnsaveable(UObject* InObject);
 
 	bool IsUnsaveable(UObject* InObject) const;
 
+	void RecordIllegalReference(UObject* InFrom, UObject* InTo, EIllegalRefReason InReason)
+	{
+		HarvestedIllegalReferences.Add({ InFrom, InTo, InReason });
+	}
+
+	const TArray<FIllegalReference>& GetIllegalReferences() const
+	{
+		return HarvestedIllegalReferences;
+	}
+	
+	void AddImport(UObject* InObject)
+	{
+		GetHarvestedRealm().AddImport(InObject);
+	}
+
+	void AddExport(UObject* InObj, bool bNotAlwaysLoadedForEditorGame)
+	{
+		GetHarvestedRealm().AddExport(InObj, bNotAlwaysLoadedForEditorGame);
+	}
+
+	void AddExcluded(UObject* InObject)
+	{
+		GetHarvestedRealm().AddExcluded(InObject);
+	}
+
 	bool IsImport(UObject* InObject) const
 	{
-		return Imports.Contains(InObject);
+		return GetHarvestedRealm().IsImport(InObject);
 	}
 
 	bool IsExport(UObject* InObject) const
 	{
-		return Exports.Contains(InObject);
+		return GetHarvestedRealm().IsExport(InObject);
 	}
 
-	bool IsIncluded(UObject* InObject) const
+	bool IsIncluded(UObject* InObject, ESaveRealm InContext = ESaveRealm::None) const
 	{
-		return IsImport(InObject) || IsExport(InObject);
+		return GetHarvestedRealm(InContext).IsIncluded(InObject);
 	}
 
 	bool IsExcluded(UObject* InObject) const
 	{
-		return Excluded.Contains(InObject);
+		return GetHarvestedRealm().IsExcluded(InObject);
 	}
 
 	TSet<FTaggedExport>& GetExports()
 	{
-		return Exports;
+		return GetHarvestedRealm().GetExports();
 	}
 
 	const TSet<UObject*>& GetImports() const
 	{
-		return Imports;
+		return GetHarvestedRealm().GetImports();
 	}
 
 	const TSet<UObject*>& GetImportsUsedInGame() const
 	{
-		return ImportsUsedInGame;
+		return GetHarvestedRealm(ESaveRealm::Game).GetImports();
 	}
 
-	const TArray<FName>& GetSoftPackageReferenceList() const
+	const TSet<FName>& GetSoftPackageReferenceList() const
 	{
-		return SoftPackageReferenceList;
+		return GetHarvestedRealm().GetSoftPackageReferenceList();
 	}
 
-	TArray<FName>& GetSoftPackageReferenceList()
+	TSet<FName>& GetSoftPackageReferenceList()
 	{
-		return SoftPackageReferenceList;
+		return GetHarvestedRealm().GetSoftPackageReferenceList();
 	}
 
 	const TSet<FName>& GetSoftPackagesUsedInGame() const
 	{
-		return SoftPackagesUsedInGame;
+		return GetHarvestedRealm(ESaveRealm::Game).GetSoftPackageReferenceList();
 	}
 
+	TSet<FName>& GetSoftPackagesUsedInGame()
+	{
+		return GetHarvestedRealm(ESaveRealm::Game).GetSoftPackageReferenceList();
+	}
+
 	const TMap<UObject*, TArray<FName>>& GetSearchableNamesObjectMap() const
 	{
-		return SearchableNamesObjectMap;
+		return GetHarvestedRealm().GetSearchableNamesObjectMap();
 	}
 
 	TMap<UObject*, TArray<FName>>& GetSearchableNamesObjectMap()
 	{
-		return SearchableNamesObjectMap;
+		return GetHarvestedRealm().GetSearchableNamesObjectMap();
 	}
 
 	const TSet<FNameEntryId>& GetNamesReferencedFromExportData() const
 	{
-		return NamesReferencedFromExportData;
+		return GetHarvestedRealm().GetNamesReferencedFromExportData();
 	}
 
 	const TSet<FNameEntryId>& GetNamesReferencedFromPackageHeader() const
 	{
-		return NamesReferencedFromPackageHeader;
+		return GetHarvestedRealm().GetNamesReferencedFromPackageHeader();
 	}
 
-	const FCustomVersionContainer& GetCustomVersions() const
-	{
-		return CustomVersions;
-	}
-
 	const TMap<UObject*, TSet<UObject*>>& GetObjectDependencies() const
 	{
-		return ExportObjectDependencies;
+		return GetHarvestedRealm().GetObjectDependencies();
 	}
 
 	const TMap<UObject*, TSet<UObject*>>& GetNativeObjectDependencies() const
 	{
-		return ExportNativeObjectDependencies;
+		return GetHarvestedRealm().GetNativeObjectDependencies();
 	}
 
+	bool NameExists(FNameEntryId ComparisonId) const
+	{
+		return GetHarvestedRealm().NameExists(ComparisonId);
+	}
+
+	const FCustomVersionContainer& GetCustomVersions() const
+	{
+		return CustomVersions;
+	}
+
 	const TSet<UPackage*>& GetPrestreamPackages() const
@@ -473,25 +708,6 @@
 		PrestreamPackages.Add(InPackage);
 	}
 
-	bool NameExists(FNameEntryId ComparisonId) const
-	{
-		for (FNameEntryId DisplayId : NamesReferencedFromExportData)
-		{
-			if (FName::GetComparisonIdFromDisplayId(DisplayId) == ComparisonId)
-			{
-				return true;
-			}
-		}
-		for (FNameEntryId DisplayId : NamesReferencedFromPackageHeader)
-		{
-			if (FName::GetComparisonIdFromDisplayId(DisplayId) == ComparisonId)
-			{
-				return true;
-			}
-		}
-		return false;
-	}
-
 	void SetCustomVersions(FCustomVersionContainer InCustomVersions)
 	{
 		CustomVersions = MoveTemp(InCustomVersions);
@@ -550,6 +766,15 @@
 		return SaveArgs.SavePackageContext ? SaveArgs.SavePackageContext->GetValidator() : nullptr;
 	}
 
+	const FHarvestedRealm& GetHarvestedRealm(ESaveRealm Realm = ESaveRealm::None) const
+	{
+		return HarvestedRealms[(uint32)(Realm == ESaveRealm::None ? CurrentHarvestingRealm : Realm)];
+	}
+	FHarvestedRealm& GetHarvestedRealm(ESaveRealm Realm = ESaveRealm::None)
+	{
+		return HarvestedRealms[(uint32)(Realm == ESaveRealm::None ? CurrentHarvestingRealm : Realm)];
+	}
+
 public:
 	ESavePackageResult Result;
 
@@ -563,6 +788,7 @@
 
 	EPropertyLocalizationGathererResultFlags GatherableTextResultFlags = EPropertyLocalizationGathererResultFlags::Empty;
 
+	int64 PackageHeaderAndExportSize = 0;
 	int64 TotalPackageSizeUncompressed = 0;
 	int32 OffsetAfterPackageFileSummary = 0;
 	int32 OffsetAfterImportMap = 0;
@@ -573,6 +799,17 @@
 	TArray<FLargeMemoryWriter, TInlineAllocator<4>> AdditionalFilesFromExports;
 	FSavePackageOutputFileArray AdditionalPackageFiles;
 private:
+
+	// Create the needed harvesting context depending on the save context options
+	void CreateHarvestingRealms()
+	{
+		// Create the different harvesting realms
+		HarvestedRealms.AddDefaulted((uint32)ESaveRealm::RealmCount);
+	
+		// if cooking the default harvesting context is Game, otherwise it's the editor context
+		CurrentHarvestingRealm = IsCooking() ? ESaveRealm::Game : ESaveRealm::Editor;
+	}
+		
 	friend class FPackageHarvester;
 
 	// Args
@@ -602,32 +839,18 @@
 
 	// Matching any mark in ExcludedObjectMarks indicates that an object should be excluded from being either an import or an export for this save
 	const EObjectMark ExcludedObjectMarks;
-	// Set of objects excluded (import or exports) through through marks or otherwise (i.e. transient flags, etc)
-	TSet<UObject*> Excluded;
 
-	// Set of objects marked as export
-	TSet<FTaggedExport> Exports;
-	// Set of objects marked as import
-	TSet<UObject*> Imports;
-	// Subset of this->Imports which are referenced from not-editoronly properties
-	TSet<UObject*> ImportsUsedInGame;
-	// Set of names referenced from export serialization
-	TSet<FNameEntryId> NamesReferencedFromExportData;
-	// Set of names referenced from the package header (import and export table object names etc)
-	TSet<FNameEntryId> NamesReferencedFromPackageHeader;
-	// List of soft package reference found
-	TArray<FName> SoftPackageReferenceList;
-	// Subset of this->SoftPackageReferenceList which are referenced from not-editoronly properties
-	TSet<FName> SoftPackagesUsedInGame;
-
-	// Map of objects to their list of searchable names
-	TMap<UObject*, TArray<FName>> SearchableNamesObjectMap;
-	// Map of objects to their dependencies
-	TMap<UObject*, TSet<UObject*>> ExportObjectDependencies;
-	// Map of objects to their native dependencies
-	TMap<UObject*, TSet<UObject*>> ExportNativeObjectDependencies;
-	// Set of harvested prestream packages
-	TSet<UPackage*> PrestreamPackages;
 	// Harvested custom versions
 	FCustomVersionContainer CustomVersions;
+
+	// The current default harvesting context being queried by the save context
+	ESaveRealm CurrentHarvestingRealm = ESaveRealm::None;
+
+	// Set of harvested content split per harvesting context
+	TArray<FHarvestedRealm> HarvestedRealms;
+
+	TArray<FIllegalReference> HarvestedIllegalReferences;
+
+	// Set of harvested prestream packages, should be deprecated
+	TSet<UPackage*> PrestreamPackages;
 };
diff --git a/Engine/Source/Runtime/CoreUObject/Private/UObject/SavePackage2.cpp b/Engine/Source/Runtime/CoreUObject/Private/UObject/SavePackage2.cpp
--- a/Engine/Source/Runtime/CoreUObject/Private/UObject/SavePackage2.cpp
+++ b/Engine/Source/Runtime/CoreUObject/Private/UObject/SavePackage2.cpp
@@ -306,7 +306,6 @@ ESavePackageResult HarvestPackage(FSaveContext& SaveContext)
 	SCOPED_SAVETIMER(UPackage_Save_HarvestPackage);
 
 	FPackageHarvester Harvester(SaveContext);
-	Harvester.SetUseUnversionedPropertySerialization(SaveContext.IsSaveUnversionedProperties());
 	EObjectFlags TopLevelFlags = SaveContext.GetTopLevelFlags();
 	UObject* Asset = SaveContext.GetAsset();
 
@@ -314,9 +313,9 @@ ESavePackageResult HarvestPackage(FSaveContext& SaveContext)
 	if (TopLevelFlags == RF_NoFlags)
 	{
 		Harvester.TryHarvestExport(Asset);
-		while (UObject* Export = Harvester.PopExportToProcess())
+		while (FPackageHarvester::FExportWithContext ExportContext = Harvester.PopExportToProcess())
 		{
-			Harvester.ProcessExport(Export);
+			Harvester.ProcessExport(ExportContext);
 		}
 	}
 	// Otherwise process all objects which have the relevant flags
@@ -333,12 +332,43 @@ ESavePackageResult HarvestPackage(FSaveContext& SaveContext)
 				}
 				return true;
 			}, true/*bIncludeNestedObjects */, RF_Transient);
-		while (UObject* Export = Harvester.PopExportToProcess())
+		while (FPackageHarvester::FExportWithContext ExportContext = Harvester.PopExportToProcess())
 		{
-			Harvester.ProcessExport(Export);
+			Harvester.ProcessExport(ExportContext);
 		}
 	}
 
+	// If we have a valid optional context and we are saving it,
+	// transform any harvested non optional export into imports
+	if (SaveContext.IsSaveOptional() &&
+		SaveContext.IsCooking() &&
+		SaveContext.GetHarvestedRealm(ESaveRealm::Optional).GetExports().Num())
+	{
+		FHarvestedRealm& OptionalContext = SaveContext.GetHarvestedRealm(ESaveRealm::Optional);
+		for (auto It = OptionalContext.GetExports().CreateIterator(); It; ++It)
+		{
+			if (!It->Obj->GetClass()->HasAnyClassFlags(CLASS_Optional))
+			{
+				// Make sure the export is found in the game context as well
+				if (FTaggedExport* GameExport = SaveContext.GetHarvestedRealm(ESaveRealm::Game).GetExports().Find(It->Obj))
+				{
+					// Flag the export in the game context to generate it's public hash
+					GameExport->bGeneratePublicHash = true;
+					// Transform the export as an import
+					OptionalContext.AddImport(It->Obj);
+				}
+				// if not found in the game context, record an illegal reference
+				else
+				{
+					SaveContext.RecordIllegalReference(nullptr, It->Obj, EIllegalRefReason::ReferenceFromOptionalToMissingGameExport);
+				}
+				It.RemoveCurrent();
+			}
+		}
+		// Also add the current package itself as an import, if the optional context isn't empty
+		OptionalContext.AddImport(SaveContext.GetPackage());
+	}
+
 	// Trim PrestreamPackage list 
 	TSet<UPackage*>& PrestreamPackages = SaveContext.GetPrestreamPackages();
 	TSet<UPackage*> KeptPrestreamPackages;
@@ -379,10 +409,60 @@ ESavePackageResult HarvestPackage(FSaveContext& SaveContext)
 
 static FNameEntryId NAME_UniqueObjectNameForCookingComparisonIndex = FName("UniqueObjectNameForCooking").GetComparisonIndex();
 
+ESavePackageResult ValidateRealms(FSaveContext& SaveContext)
+{
+	SCOPED_SAVETIMER(UPackage_Save_ValidateRealms);
+	if (SaveContext.GetIllegalReferences().Num())
+	{
+		for (const FIllegalReference& Reference : SaveContext.GetIllegalReferences())
+		{
+			FString ErrorMessage;
+			switch (Reference.Reason)
+			{
+			case EIllegalRefReason::ReferenceToOptional:
+				ErrorMessage = FString::Printf(TEXT("Can't save %s: Non-optional object (%s) has a reference to optional object (%s). Only optional objects can refer to other optional objects."),
+					SaveContext.GetFilename(),
+					Reference.From ? *Reference.From->GetPathName() : TEXT("Unknown"),
+					Reference.To ? *Reference.To->GetPathName() : TEXT("Unknown"));
+				break;
+			case EIllegalRefReason::ReferenceFromOptionalToMissingGameExport:
+				ErrorMessage = FString::Printf(TEXT("Can't save %s: Optional object (%s) has a reference to cooked object (%s) which is missing. Non optional objects referenced by optional objects needs to be present in cooked data."),
+					SaveContext.GetFilename(),
+					Reference.From ? *Reference.From->GetPathName() : TEXT("Unknown"),
+					Reference.To ? *Reference.To->GetPathName() : TEXT("Unknown"));
+				break;
+			default:
+				ErrorMessage = FString::Printf(TEXT("Can't save %s: Unknown Illegal reference from object (%s) to object (%s)"),
+					SaveContext.GetFilename(),
+					Reference.From ? *Reference.From->GetPathName() : TEXT("Unknown"),
+					Reference.To ? *Reference.To->GetPathName() : TEXT("Unknown"));
+			}
+
+			if (SaveContext.IsGenerateSaveError())
+			{
+				SaveContext.GetError()->Logf(ELogVerbosity::Warning, TEXT("%s"), *ErrorMessage);
+			}
+			else
+			{
+				UE_LOG(LogSavePackage, Error, TEXT("%s"), *ErrorMessage);
+			}
+
+		}
+		return ESavePackageResult::Error;
+	}
+	return ReturnSuccessOrCancel();
+}
+
 ESavePackageResult ValidateExports(FSaveContext& SaveContext)
 {
 	SCOPED_SAVETIMER(UPackage_Save_ValidateExports);
 
+	//@todo: add validation on optional context
+	if (SaveContext.GetCurrentHarvestingRealm() == ESaveRealm::Optional)
+	{
+		return ReturnSuccessOrCancel();
+	}
+
 	// Check if we gathered any exports
 	if (SaveContext.GetExports().Num() == 0)
 	{
@@ -652,7 +732,11 @@ ESavePackageResult ValidateImports(FSaveContext& SaveContext)
 		}
 
 		// if an import outer is an export and that import doesn't have a specific package set then, there's an error
-		const bool bWrongImport = Import->GetOuter() && Import->GetOuter()->IsInPackage(SaveContext.GetPackage()) && Import->GetExternalPackage() == nullptr;
+		const bool bWrongImport = Import->GetOuter() 
+			&& Import->GetOuter()->IsInPackage(SaveContext.GetPackage()) 
+			&& Import->GetExternalPackage() == nullptr
+			// The optional context will have import that are actually in the same package, similar to external package
+			&& SaveContext.GetCurrentHarvestingRealm() != ESaveRealm::Optional;
 		if (bWrongImport)
 		{
 			if (!Import->HasAllFlags(RF_Transient) || !Import->IsNative())
@@ -671,7 +755,8 @@ ESavePackageResult ValidateImports(FSaveContext& SaveContext)
 		check(!bWrongImport || Import->HasAllFlags(RF_Transient) || Import->IsNative());
 
 		// if this import shares a outer with top level object of this package then the reference is acceptable
-		if (!SaveContext.IsCooking() && (IsInAnyTopLevelObject(Import) || AnyTopLevelObjectIsIn(Import) || AnyTopLevelObjectHasSameOutermostObject(Import)))
+		if ((!SaveContext.IsCooking() || SaveContext.GetCurrentHarvestingRealm() == ESaveRealm::Optional) &&
+			(IsInAnyTopLevelObject(Import) || AnyTopLevelObjectIsIn(Import) || AnyTopLevelObjectHasSameOutermostObject(Import)))
 		{
 			continue;
 		}
@@ -703,7 +788,8 @@ ESavePackageResult ValidateImports(FSaveContext& SaveContext)
 	}
 
 	// Cooking checks
-	if (SaveContext.IsCooking())
+	// Currently do not use the edl checker for the optional context
+	if (SaveContext.IsCooking() && SaveContext.GetCurrentHarvestingRealm() != ESaveRealm::Optional)
 	{
 		// Now that imports are validated add them to the cook checker if available
 		if (FEDLCookChecker* EDLCookChecker = SaveContext.GetEDLCookChecker())
@@ -936,6 +1022,14 @@ ESavePackageResult BuildLinker(FSaveContext& SaveContext)
 		for (const FTaggedExport& TaggedExport : SaveContext.GetExports())
 		{
 			FObjectExport& Export = Linker->ExportMap.Add_GetRef(FObjectExport(TaggedExport.Obj, TaggedExport.bNotAlwaysLoadedForEditorGame));
+			// @note: Artificially flag export as public if needed
+			// This is temporary for testing the optional realm cooked content with the iostore
+			// It should be changed to a separate member on `FObjectExport`
+			if (TaggedExport.bGeneratePublicHash)
+			{
+				Export.ObjectFlags |= RF_Public;
+			}
+
 			if (UPackage* Package = Cast<UPackage>(Export.Object))
 			{
 				Export.PackageFlags = Package->GetPackageFlags();
@@ -1007,7 +1101,7 @@ ESavePackageResult BuildLinker(FSaveContext& SaveContext)
 
 	// Build Searchable Name Map
 	{
-		Linker->SoftPackageReferenceList = SaveContext.GetSoftPackageReferenceList();
+		Linker->SoftPackageReferenceList = SaveContext.GetSoftPackageReferenceList().Array();
 
 		// Convert the searchable names map from UObject to packageindex
 		for (TPair<UObject*, TArray<FName>>& SearchableNamePair : SaveContext.GetSearchableNamesObjectMap())
@@ -1298,8 +1392,8 @@ void SavePreloadDependencies(FStructuredArchive::FRecord& StructuredArchiveRoot,
 				IncludeIndexAsDependency(CreateBeforeCreateDependencies, Export.OuterIndex);
 				IncludeIndexAsDependency(CreateBeforeCreateDependencies, Export.SuperIndex);
 			}
-			FEDLCookChecker* EDLCookChecker = SaveContext.GetEDLCookChecker();
-			check(EDLCookChecker);
+			// Currently do not validate the optional context with the edl checker
+			FEDLCookChecker* EDLCookChecker = SaveContext.GetCurrentHarvestingRealm() != ESaveRealm::Optional ? SaveContext.GetEDLCookChecker() : nullptr;
 			auto AddArcForDepChecking = [&Linker, &Export, EDLCookChecker](bool bExportIsSerialize, FPackageIndex Dep, bool bDepIsSerialize)
 			{
 				check(Export.Object);
@@ -1308,7 +1402,10 @@ void SavePreloadDependencies(FStructuredArchive::FRecord& StructuredArchiveRoot,
 				check(DepObject);
 
 				Linker->DepListForErrorChecking.Add(Dep);
-				EDLCookChecker->AddArc(DepObject, bDepIsSerialize, Export.Object, bExportIsSerialize);
+				if (EDLCookChecker)
+				{
+					EDLCookChecker->AddArc(DepObject, bDepIsSerialize, Export.Object, bExportIsSerialize);
+				}
 			};
 
 			for (FPackageIndex Index : SerializationBeforeSerializationDependencies)
@@ -1987,32 +2084,14 @@ void PostSavePackage(FSaveContext& SaveContext)
 #endif
 }
 
-
-/**
- * InnerSave is the portion of Save that can be safely run concurrently
- */
-ESavePackageResult InnerSave(FSaveContext& SaveContext)
+ESavePackageResult SaveHarvestedRealms(FSaveContext& SaveContext, ESaveRealm HarvestingContextToSave)
 {
-	TRefCountPtr<FUObjectSerializeContext> SerializeContext(FUObjectThreadContext::Get().GetSerializeContext());
-	SaveContext.SetSerializeContext(SerializeContext);
-	SaveContext.SetEDLCookChecker(&FEDLCookChecker::Get());
+	// Set the current harvested context to save
+	FSaveContext::FSetSaveRealmToSaveScope Scope(SaveContext, HarvestingContextToSave);
 
 	// Create slow task dialog if needed
-	const int32 TotalSaveSteps = 13;
+	const int32 TotalSaveSteps = 11;
 	FScopedSlowTask SlowTask(TotalSaveSteps, FText(), SaveContext.IsUsingSlowTask());
-	SlowTask.MakeDialogDelayed(3.0f, SaveContext.IsFromAutoSave());
-
-	// Harvest Package
-	SlowTask.EnterProgressFrame();
-	SaveContext.Result = HarvestPackage(SaveContext);
-	if (SaveContext.Result != ESavePackageResult::Success)
-	{
-		return SaveContext.Result;
-	}
-
-	// @todo: Need to adjust GIsSavingPackage to properly prevent generating reference once, package harvesting is done
-	// GIsSavingPackage is too harsh however, since it should be scope only to the current package
-	FScopedSavingFlag IsSavingFlag(SaveContext.IsConcurrent(), SaveContext.GetPackage());
 
 	// Validate Exports
 	SlowTask.EnterProgressFrame();
@@ -2135,6 +2214,11 @@ ESavePackageResult InnerSave(FSaveContext& SaveContext)
 	{
 		PackageWriter->AddToExportsSize(ExportsSize);
 	}
+	// Store the package header and export size of the non optional realm
+	if (SaveContext.GetCurrentHarvestingRealm() != ESaveRealm::Optional)
+	{
+		SaveContext.PackageHeaderAndExportSize = ExportsSize;
+	}
 	SaveContext.TotalPackageSizeUncompressed += ExportsSize;
 	for (const FSavePackageOutputFile& File : SaveContext.AdditionalPackageFiles)
 	{
@@ -2158,9 +2242,6 @@ ESavePackageResult InnerSave(FSaveContext& SaveContext)
 	}
 
 	//COOK_STAT(FSavePackageStats::MBWritten += ((double)SaveContext.TotalPackageSizeUncompressed) / 1024.0 / 1024.0);
-
-	// Mark Exports & Package RF_Loaded
-	SlowTask.EnterProgressFrame();
 	{
 		SCOPED_SAVETIMER(UPackage_Save_MarkExportLoaded);
 		FLinkerSave* Linker = SaveContext.Linker.Get();
@@ -2168,7 +2249,7 @@ ESavePackageResult InnerSave(FSaveContext& SaveContext)
 		// This is to ensue that newly created packages are properly marked as loaded (since they now exist on disk and 
 		// in memory in the exact same state).
 
- 		// Nobody should be touching those objects beside us while we are saving them here as this can potentially be executed from another thread
+		// Nobody should be touching those objects beside us while we are saving them here as this can potentially be executed from another thread
 		for (auto& Export : Linker->ExportMap)
 		{
 			if (Export.Object)
@@ -2176,11 +2257,70 @@ ESavePackageResult InnerSave(FSaveContext& SaveContext)
 				Export.Object->SetFlags(RF_WasLoaded | RF_LoadCompleted);
 			}
 		}
-		if (Linker->LinkerRoot)
+	}
+	return SaveContext.Result;
+}
+
+/**
+ * InnerSave is the portion of Save that can be safely run concurrently
+ */
+ESavePackageResult InnerSave(FSaveContext& SaveContext)
+{
+	TRefCountPtr<FUObjectSerializeContext> SerializeContext(FUObjectThreadContext::Get().GetSerializeContext());
+	SaveContext.SetSerializeContext(SerializeContext);
+	SaveContext.SetEDLCookChecker(&FEDLCookChecker::Get());
+
+	// Create slow task dialog if needed
+	const int32 TotalSaveSteps = 4;
+	FScopedSlowTask SlowTask(TotalSaveSteps, FText(), SaveContext.IsUsingSlowTask());
+	SlowTask.MakeDialogDelayed(3.0f, SaveContext.IsFromAutoSave());
+
+	// Harvest Package
+	SlowTask.EnterProgressFrame();
+	SaveContext.Result = HarvestPackage(SaveContext);
+	if (SaveContext.Result != ESavePackageResult::Success)
+	{
+		return SaveContext.Result;
+	}
+
+	SlowTask.EnterProgressFrame();
+	SaveContext.Result = ValidateRealms(SaveContext);
+	if (SaveContext.Result != ESavePackageResult::Success)
+	{
+		return SaveContext.Result;
+	}
+
+	// @todo: Need to adjust GIsSavingPackage to properly prevent generating reference once, package harvesting is done
+	// GIsSavingPackage is too harsh however, since it should be scope only to the current package
+	FScopedSavingFlag IsSavingFlag(SaveContext.IsConcurrent(), SaveContext.GetPackage());
+
+	// Split the save context into its harvested contexts,
+	// This essentially mean that a package can produce multiple package output. 
+	// This is different then the multiple file outputs a package can already produce
+	// since each harvested context will produce those multiple file outputs i.e:
+	// Input package -> Main cooked package -> .uasset
+	//										-> .uexp
+	//										-> .ubulk
+	//										-> etc
+	//					Sub cooked package	-> .o.uasset
+	//										-> .o.uexp
+	//										-> .o.ubulk
+	//										-> etc
+	SlowTask.EnterProgressFrame();
+	for (ESaveRealm HarvestingContext : SaveContext.GetHarvestedRealmsToSave())
+	{
+		SaveContext.Result = SaveHarvestedRealms(SaveContext, HarvestingContext);
+		if (SaveContext.Result != ESavePackageResult::Success)
 		{
-			// And finally set the flag on the package itself.
-			Linker->LinkerRoot->SetFlags(RF_WasLoaded | RF_LoadCompleted);
+			return SaveContext.Result;
 		}
+	}
+
+	SlowTask.EnterProgressFrame();
+	{
+		// Mark the Package RF_Loaded after its been serialized. 
+		// This was already done for each object in the package in SaveHarvestedRealms
+		SaveContext.GetPackage()->SetFlags(RF_WasLoaded | RF_LoadCompleted);
 
 		// Clear dirty flag if desired
 		if (!SaveContext.IsKeepDirty())
@@ -2189,7 +2329,10 @@ ESavePackageResult InnerSave(FSaveContext& SaveContext)
 		}
 
 		// Update package FileSize value
-		SaveContext.GetPackage()->FileSize = ExportsSize;
+		if (SaveContext.IsUpdatingLoadedPath())
+		{
+			SaveContext.GetPackage()->FileSize = SaveContext.PackageHeaderAndExportSize;
+		}
 	}
 	return SaveContext.Result;
 }
@@ -2288,10 +2431,10 @@ FSavePackageResultStruct UPackage::Save2(UPackage* InPackage, UObject* InAsset,
 	// PreSave Asset
 	SlowTask.EnterProgressFrame();
 	PreSavePackage(SaveContext);
-	if (InAsset && !SaveContext.IsConcurrent())
+	if (SaveContext.GetAsset() && !SaveContext.IsConcurrent())
 	{
 		FObjectSaveContextData& ObjectSaveContext = SaveContext.GetObjectSaveContext();
-		UE::SavePackageUtilities::CallPreSaveRoot(InAsset, ObjectSaveContext);
+		UE::SavePackageUtilities::CallPreSaveRoot(SaveContext.GetAsset(), ObjectSaveContext);
 		SaveContext.SetPreSaveCleanup(ObjectSaveContext.bCleanupRequired);
 	}
 
@@ -2322,9 +2465,9 @@ FSavePackageResultStruct UPackage::Save2(UPackage* InPackage, UObject* InAsset,
 
 	// PostSave Asset
 	SlowTask.EnterProgressFrame();
-	if (InAsset && !SaveContext.IsConcurrent())
+	if (SaveContext.GetAsset() && !SaveContext.IsConcurrent())
 	{
-		UE::SavePackageUtilities::CallPostSaveRoot(InAsset, SaveContext.GetObjectSaveContext(), SaveContext.GetPreSaveCleanup());
+		UE::SavePackageUtilities::CallPostSaveRoot(SaveContext.GetAsset(), SaveContext.GetObjectSaveContext(), SaveContext.GetPreSaveCleanup());
 		SaveContext.SetPreSaveCleanup(false);
 	}
 
diff --git a/Engine/Source/Runtime/CoreUObject/Public/UObject/ObjectMacros.h b/Engine/Source/Runtime/CoreUObject/Public/UObject/ObjectMacros.h
--- a/Engine/Source/Runtime/CoreUObject/Public/UObject/ObjectMacros.h
+++ b/Engine/Source/Runtime/CoreUObject/Public/UObject/ObjectMacros.h
@@ -82,23 +82,24 @@ enum ELoadFlags
 /** Flags for saving objects/packages, passed into UPackage::SavePackage() as a uint32 */
 enum ESaveFlags
 {
-	SAVE_None			= 0x00000000,	///< No flags
-	SAVE_NoError		= 0x00000001,	///< Don't generate errors on save
-	SAVE_FromAutosave	= 0x00000002,   ///< Used to indicate this save was initiated automatically
-	SAVE_KeepDirty		= 0x00000004,	///< Do not clear the dirty flag when saving
-	SAVE_KeepGUID		= 0x00000008,	///< Keep the same guid, used to save cooked packages
-	SAVE_Async			= 0x00000010,	///< Save to a memory writer, then actually write to disk async
-	SAVE_Unversioned_Native	= 0x00000020,	///< Save all versions as zero. Upon load this is changed to the current version. This is only reasonable to use with full cooked builds for distribution.
-	SAVE_CutdownPackage	= 0x00000040,	///< Saving cutdown packages in a temp location WITHOUT renaming the package.
+	SAVE_None						= 0x00000000,	///< No flags
+	SAVE_NoError					= 0x00000001,	///< Don't generate errors on save
+	SAVE_FromAutosave				= 0x00000002,   ///< Used to indicate this save was initiated automatically
+	SAVE_KeepDirty					= 0x00000004,	///< Do not clear the dirty flag when saving
+	SAVE_KeepGUID					= 0x00000008,	///< Keep the same guid, used to save cooked packages
+	SAVE_Async						= 0x00000010,	///< Save to a memory writer, then actually write to disk async
+	SAVE_Unversioned_Native			= 0x00000020,	///< Save all versions as zero. Upon load this is changed to the current version. This is only reasonable to use with full cooked builds for distribution.
+	SAVE_CutdownPackage				= 0x00000040,	///< Saving cutdown packages in a temp location WITHOUT renaming the package.
 	SAVE_KeepEditorOnlyCookedPackages = 0x00000080,  ///< Keep packages which are marked as editor only even though we are cooking
-	SAVE_Concurrent		= 0x00000100,	///< We are save packages in multiple threads at once and should not call non-threadsafe functions or rely on globals. GIsSavingPackage should be set and PreSave/Postsave functions should be called before/after the entire concurrent save.
-	SAVE_DiffOnly		UE_DEPRECATED(5.0, "Diffing is now done using FDiffPackageWriter.") = 0x00000200,
-	SAVE_DiffCallstack  UE_DEPRECATED(5.0, "Diffing is now done using FDiffPackageWriter.") = 0x00000400,
-	SAVE_ComputeHash    = 0x00000800,	///< Compute the MD5 hash of the cooked data
-	SAVE_CompareLinker	= 0x00001000,	///< Return the linker save to compare against another
-	SAVE_BulkDataByReference = 0x00002000, ///< When saving to a different file than the package's LoadedPath, point bulkdata in the new file to be loaded from the original file.
-	SAVE_Unversioned_Properties = 0x00004000, /// Properties are saved without property name information, and are saved/loaded in the order of the current binary.
-											      ///This is only reasonable to use with full cooked builds for distribution, or with a domain that selects the payload matching a given version externally.
+	SAVE_Concurrent					= 0x00000100,	///< We are save packages in multiple threads at once and should not call non-threadsafe functions or rely on globals. GIsSavingPackage should be set and PreSave/Postsave functions should be called before/after the entire concurrent save.
+	SAVE_DiffOnly UE_DEPRECATED(5.0, "Diffing is now done using FDiffPackageWriter.") = 0x00000200,	
+	SAVE_DiffCallstack UE_DEPRECATED(5.0, "Diffing is now done using FDiffPackageWriter.") = 0x00000400,
+	SAVE_ComputeHash				= 0x00000800,	///< Compute the MD5 hash of the cooked data
+	SAVE_CompareLinker				= 0x00001000,	///< Return the linker save to compare against another
+	SAVE_BulkDataByReference		= 0x00002000,	///< When saving to a different file than the package's LoadedPath, point bulkdata in the new file to be loaded from the original file.
+	SAVE_Unversioned_Properties		= 0x00004000,	///< Properties are saved without property name information, and are saved/loaded in the order of the current binary.
+													///This is only reasonable to use with full cooked builds for distribution, or with a domain that selects the payload matching a given version externally.
+	SAVE_Optional					= 0x00008000,	///< Indicate that we to save optional exports. This flag is only valid while cooking. Optional exports are filtered if not specified during cooking. 
 	SAVE_Unversioned = SAVE_Unversioned_Native | SAVE_Unversioned_Properties,
 };
 
@@ -259,7 +260,7 @@ ENUM_CLASS_FLAGS(EClassFlags);
 #define CLASS_RecompilerClear ((EClassFlags)(CLASS_Inherit | CLASS_Abstract | CLASS_NoExport | CLASS_Native | CLASS_Intrinsic | CLASS_TokenStreamAssembled))
 
 /** These flags will be cleared by the compiler when the class is parsed during script compilation */
-#define CLASS_ShouldNeverBeLoaded ((EClassFlags)(CLASS_Native | CLASS_Intrinsic | CLASS_TokenStreamAssembled))
+#define CLASS_ShouldNeverBeLoaded ((EClassFlags)(CLASS_Native | CLASS_Optional | CLASS_Intrinsic | CLASS_TokenStreamAssembled))
 
 /** These flags will be inherited from the base class only for non-intrinsic classes */
 #define CLASS_ScriptInherit ((EClassFlags)(CLASS_Inherit | CLASS_EditInlineNew | CLASS_CollapseCategories))
@@ -272,6 +273,7 @@ ENUM_CLASS_FLAGS(EClassFlags);
 	CLASS_ProjectUserConfig | \
 	CLASS_Config | \
 	CLASS_Transient | \
+	CLASS_Optional | \
 	CLASS_Native | \
 	CLASS_NotPlaceable | \
 	CLASS_PerObjectConfig | \
@@ -486,20 +488,20 @@ enum EObjectFlags
 {
 	// Do not add new flags unless they truly belong here. There are alternatives.
 	// if you change any the bit of any of the RF_Load flags, then you will need legacy serialization
-	RF_NoFlags						= 0x00000000,	///< No flags, used to avoid a cast
+	RF_NoFlags					= 0x00000000,	///< No flags, used to avoid a cast
 
 	// This first group of flags mostly has to do with what kind of object it is. Other than transient, these are the persistent object flags.
 	// The garbage collector also tends to look at these.
 	RF_Public					=0x00000001,	///< Object is visible outside its package.
 	RF_Standalone				=0x00000002,	///< Keep object around for editing even if unreferenced.
-	RF_MarkAsNative					=0x00000004,	///< Object (UField) will be marked as native on construction (DO NOT USE THIS FLAG in HasAnyFlags() etc)
+	RF_MarkAsNative				=0x00000004,	///< Object (UField) will be marked as native on construction (DO NOT USE THIS FLAG in HasAnyFlags() etc)
 	RF_Transactional			=0x00000008,	///< Object is transactional.
 	RF_ClassDefaultObject		=0x00000010,	///< This object is its class's default object
 	RF_ArchetypeObject			=0x00000020,	///< This object is a template for another object - treat like a class default object
 	RF_Transient				=0x00000040,	///< Don't save object.
 
 	// This group of flags is primarily concerned with garbage collection.
-	RF_MarkAsRootSet					=0x00000080,	///< Object will be marked as root set on construction and not be garbage collected, even if unreferenced (DO NOT USE THIS FLAG in HasAnyFlags() etc)
+	RF_MarkAsRootSet			=0x00000080,	///< Object will be marked as root set on construction and not be garbage collected, even if unreferenced (DO NOT USE THIS FLAG in HasAnyFlags() etc)
 	RF_TagGarbageTemp			=0x00000100,	///< This is a temp user flag for various utilities that need to use the garbage collector. The garbage collector itself does not interpret it.
 
 	// The group of flags tracks the stages of the lifetime of a uobject
