SavePackage: turn off the check for missing toplevelflag when cooking. It is only a concern when saving editor packages, and it can have false positives when cooking.
#rnx

diff --git a/Engine/Source/Runtime/CoreUObject/Private/UObject/SavePackage2.cpp b/Engine/Source/Runtime/CoreUObject/Private/UObject/SavePackage2.cpp
--- a/Engine/Source/Runtime/CoreUObject/Private/UObject/SavePackage2.cpp
+++ b/Engine/Source/Runtime/CoreUObject/Private/UObject/SavePackage2.cpp
@@ -129,9 +129,10 @@ ESavePackageResult ValidatePackage(FSaveContext& SaveContext)
 			return ESavePackageResult::Error;
 		}
 
-		// If An asset is provided, validate it has the requested TopLevelFlags.
+		// If An asset is provided, validate it has the requested TopLevelFlags. This is necessary to prevent dataloss, but only
+		// when saving packages to the WorkspaceDomain
 		EObjectFlags TopLevelFlags = SaveContext.GetTopLevelFlags();
-		if (TopLevelFlags != RF_NoFlags && !Asset->HasAnyFlags(TopLevelFlags))
+		if (!SaveContext.IsCooking() && TopLevelFlags != RF_NoFlags && !Asset->HasAnyFlags(TopLevelFlags))
 		{
 			if (SaveContext.IsFixupStandaloneFlags() && !Asset->GetExternalPackage() && EnumHasAnyFlags(TopLevelFlags, RF_Standalone))
 			{
@@ -321,9 +322,6 @@ ESavePackageResult HarvestPackage(FSaveContext& SaveContext)
 	// Otherwise process all objects which have the relevant flags
 	else
 	{
-		// Validate that if an asset is provided it has the appropriate top level flags
-		checkf(!Asset || Asset->HasAnyFlags(TopLevelFlags), TEXT("The asset to save %s in package %s does not contain any of the provided object flags.")
-			TEXT("This should have been checked already in ValidatePackage"), *Asset->GetName(), *SaveContext.GetPackage()->GetName());
 		ForEachObjectWithPackage(SaveContext.GetPackage(), [&Harvester, TopLevelFlags](UObject* InObject)
 			{
 				if (InObject->HasAnyFlags(TopLevelFlags))
@@ -470,6 +468,28 @@ ESavePackageResult ValidateExports(FSaveContext& SaveContext)
 		return SaveContext.IsCooking() ? ESavePackageResult::ContainsEditorOnlyData : ESavePackageResult::Error;
 	}
 
+	// Validate that if an asset was provided it had the proper flags to be present in the exports.
+	UObject* Asset = SaveContext.GetAsset();
+	if (Asset)
+	{
+		if (!SaveContext.GetExports().Contains(FTaggedExport{ Asset }) &&
+			SaveContext.GetTopLevelFlags() != RF_NoFlags && !Asset->HasAnyFlags(SaveContext.GetTopLevelFlags()))
+		{
+			FString ErrorMessage = FString::Printf(
+				TEXT("The asset to save %s in package %s does not contain any of the provided object flags."),
+				*Asset->GetName(), *SaveContext.GetPackage()->GetName());
+			if (SaveContext.IsGenerateSaveError())
+			{
+				SaveContext.GetError()->Logf(ELogVerbosity::Warning, TEXT("%s"), *ErrorMessage);
+			}
+			else
+			{
+				UE_LOG(LogSavePackage, Error, TEXT("%s"), *ErrorMessage);
+			}
+			return ESavePackageResult::Error;
+		}
+	}
+
 #if WITH_EDITOR
 	if (GOutputCookingWarnings)
 	{
