Prevent invalid FPayloadIds (empty payloads) from being added to the package trailer.

#rnx

- Having invalid payloads in the trailer's lookup table and writing empty payload data to disk was a waste of time and meant we had to add more error handling code elsewhere.
- Now when a VBD attempts to add one to the trailer it will be rejected, although the provided callback will still be invoked so that the VBD object will still mark itself as being serialized to disk and update it's flags etc.
- Asserts have been added when writing the FPackageTrailer to disk or creating a new FPackageTrailerBuilder from an existing FPackageTrailer if invalid FPayloadIds are somehow found. This is not an expected condition and the asserts are there to guard against future code changes that might break the contract.
- Now when virtualizing the package trailer during submission checks we can assume that all FPayloadIds reference actual data.

diff --git a/Engine/Source/Developer/Virtualization/Private/PackageSubmissionChecks.cpp b/Engine/Source/Developer/Virtualization/Private/PackageSubmissionChecks.cpp
--- a/Engine/Source/Developer/Virtualization/Private/PackageSubmissionChecks.cpp
+++ b/Engine/Source/Developer/Virtualization/Private/PackageSubmissionChecks.cpp
@@ -213,9 +213,7 @@ void OnPrePackageSubmission(const TArray<FString>& FilesToSubmit, TArray<FText>&
 
 		for (int32 Index = 0; Index < PackageInfo.LocalPayloads.Num(); ++Index)
 		{
-			// TODO: Remove FPayloadStatus::Invalid when trailers no longer support invalid payloads ids
-			if (PayloadStatuses[PackageInfo.PayloadIndex + Index] == FPayloadStatus::FoundAll ||
-				PayloadStatuses[PackageInfo.PayloadIndex + Index] == FPayloadStatus::Invalid)
+			if (PayloadStatuses[PackageInfo.PayloadIndex + Index] == FPayloadStatus::FoundAll)
 			{
 				if (PackageInfo.Trailer.UpdatePayloadAsVirtualized(PackageInfo.LocalPayloads[Index]))
 				{
@@ -265,7 +263,7 @@ void OnPrePackageSubmission(const TArray<FString>& FilesToSubmit, TArray<FText>&
 
 		for (const FPayloadId& PayloadId : PackageInfo.LocalPayloads)
 		{
-			check(PayloadId.IsValid());
+			checkf(PayloadId.IsValid(), TEXT("PackageTrailer for package '%s' should not contain invalid FPayloadIds"), *PackageInfo.Path.GetDebugName());
 			
 			FCompressedBuffer Payload = PackageInfo.Trailer.LoadPayload(PayloadId, *PackageAr);
 
@@ -364,7 +362,7 @@ void OnPrePackageSubmission(const TArray<FString>& FilesToSubmit, TArray<FText>&
 			return;
 		}
 
-		FPackageTrailerBuilder TrailerBuilder = FPackageTrailerBuilder::Create(PackageInfo.Trailer, *PackageAr);
+		FPackageTrailerBuilder TrailerBuilder = FPackageTrailerBuilder::Create(PackageInfo.Trailer, *PackageAr, PackagePath);
 		if (!TrailerBuilder.BuildAndAppendTrailer(nullptr, *CopyAr))
 		{
 			FText Message = FText::Format(	LOCTEXT("Virtualization_TrailerAppend", "Failed to append the trailer to '{0}'"),
diff --git a/Engine/Source/Runtime/CoreUObject/Private/Serialization/VirtualizedBulkData.cpp b/Engine/Source/Runtime/CoreUObject/Private/Serialization/VirtualizedBulkData.cpp
--- a/Engine/Source/Runtime/CoreUObject/Private/Serialization/VirtualizedBulkData.cpp
+++ b/Engine/Source/Runtime/CoreUObject/Private/Serialization/VirtualizedBulkData.cpp
@@ -801,7 +801,6 @@ void FVirtualizedUntypedBulkData::Serialize(FArchive& Ar, UObject* Owner, bool b
 		{
 			FPackageTrailer* Trailer = GetTrailerFromOwner(Owner);
 
-			// TODO: Error on not EPayloadStatus::NotFound
 			if (Trailer != nullptr && Trailer->FindPayloadStatus(PayloadContentId) == EPayloadStatus::StoredVirtualized)
 			{
 				check(!IsReferencingOldBulkData()); // This should not be possible
diff --git a/Engine/Source/Runtime/CoreUObject/Private/UObject/LinkerSave.cpp b/Engine/Source/Runtime/CoreUObject/Private/UObject/LinkerSave.cpp
--- a/Engine/Source/Runtime/CoreUObject/Private/UObject/LinkerSave.cpp
+++ b/Engine/Source/Runtime/CoreUObject/Private/UObject/LinkerSave.cpp
@@ -25,6 +25,7 @@ TMap<FString, TArray<uint8> > FLinkerSave::PackagesToScriptSHAMap;
 FLinkerSave::FLinkerSave(UPackage* InParent, const TCHAR* InFilename, bool bForceByteSwapping, bool bInSaveUnversioned)
 :	FLinker(ELinkerType::Save, InParent)
 ,	Saver(nullptr)
+,	PackageTrailerBuilder(LinkerRoot)
 {
 	SetFilename(InFilename);
 
@@ -79,6 +80,7 @@ FLinkerSave::FLinkerSave(UPackage* InParent, const TCHAR* InFilename, bool bForc
 FLinkerSave::FLinkerSave(UPackage* InParent, FArchive *InSaver, bool bForceByteSwapping, bool bInSaveUnversioned)
 : FLinker(ELinkerType::Save, InParent)
 , Saver(nullptr)
+, PackageTrailerBuilder(LinkerRoot)
 {
 	SetFilename(TEXT("$$Memory$$"));
 	if (FPlatformProperties::HasEditorOnlyData())
@@ -132,6 +134,7 @@ FLinkerSave::FLinkerSave(UPackage* InParent, FArchive *InSaver, bool bForceByteS
 FLinkerSave::FLinkerSave(UPackage* InParent, bool bForceByteSwapping, bool bInSaveUnversioned )
 :	FLinker(ELinkerType::Save, InParent)
 ,	Saver(nullptr)
+,	PackageTrailerBuilder(LinkerRoot)
 {
 	SetFilename(TEXT("$$Memory$$"));
 	if (FPlatformProperties::HasEditorOnlyData())
diff --git a/Engine/Source/Runtime/CoreUObject/Private/UObject/PackageTrailer.cpp b/Engine/Source/Runtime/CoreUObject/Private/UObject/PackageTrailer.cpp
--- a/Engine/Source/Runtime/CoreUObject/Private/UObject/PackageTrailer.cpp
+++ b/Engine/Source/Runtime/CoreUObject/Private/UObject/PackageTrailer.cpp
@@ -7,6 +7,7 @@
 #include "Misc/ConfigCacheIni.h"
 #include "Misc/PackagePath.h"
 #include "Misc/PackageSegment.h"
+#include "UObject/Package.h"
 #include "UObject/PackageResourceManager.h"
 
 namespace UE
@@ -71,12 +72,14 @@ FArchive& operator<<(FArchive& Ar, FLookupTableEntry& Entry)
 
 } // namespace Private
 
-FPackageTrailerBuilder FPackageTrailerBuilder::Create(const FPackageTrailer& Trailer, FArchive& Ar)
+FPackageTrailerBuilder FPackageTrailerBuilder::Create(const FPackageTrailer& Trailer, FArchive& Ar, const FPackagePath& PackagePath)
 {
-	FPackageTrailerBuilder Builder;
+	FPackageTrailerBuilder Builder(PackagePath.GetPackageFName());
 
 	for (const Private::FLookupTableEntry& Entry : Trailer.Header.PayloadLookupTable)
 	{
+		checkf(Entry.Identifier.IsValid(), TEXT("PackageTrailer for package should not contain invalid FPayloadIds. Package '%s'"), *PackagePath.GetPackageName());
+
 		if (Entry.IsVirtualized())
 		{
 			Builder.VirtualizedEntries.Add(Entry.Identifier, VirtualizedEntry(Entry.CompressedSize, Entry.RawSize));
@@ -91,26 +94,34 @@ FPackageTrailerBuilder FPackageTrailerBuilder::Create(const FPackageTrailer& Tra
 	return Builder;
 }
 
-int64 FPackageTrailerBuilder::FindPayloadOffset(const Virtualization::FPayloadId& Identifier) const
+FPackageTrailerBuilder::FPackageTrailerBuilder(UPackage* Package)
 {
-	check(TrailerPositionInFile != INDEX_NONE);
-	check(PayloadPosInFile != INDEX_NONE);
-	check(PayloadLookupTable.Num() == LocalEntries.Num() + VirtualizedEntries.Num());
-
-	for (const Private::FLookupTableEntry& Entry : PayloadLookupTable)
+	if (Package != nullptr)
 	{
-		if (Entry.Identifier == Identifier)
-		{
-			return PayloadPosInFile + Entry.OffsetInFile;
-		}
+		PackageName = Package->GetFName();
 	}
+}
 
-	return INDEX_NONE;
+FPackageTrailerBuilder::FPackageTrailerBuilder(FName InPackageName)
+	: PackageName(InPackageName)
+{
+}
+
+void FPackageTrailerBuilder::AddPayload(const Virtualization::FPayloadId& Identifier, FCompressedBuffer Payload, AdditionalDataCallback&& Callback)
+{
+	checkf(TrailerPositionInFile == INDEX_NONE, TEXT("Attempting to add payloads after the trailer has been built. Package '%s'"), *PackageName.ToString());
+
+	Callbacks.Emplace(MoveTemp(Callback));
+
+	if (Identifier.IsValid())
+	{
+		LocalEntries.Add(Identifier, LocalEntry(MoveTemp(Payload)));
+	}
 }
 
 bool FPackageTrailerBuilder::BuildAndAppendTrailer(FLinkerSave* Linker, FArchive& DataArchive)
 {
-	checkf(TrailerPositionInFile == INDEX_NONE, TEXT("Attempting to build the same FPackageTrailer multiple times"));
+	checkf(TrailerPositionInFile == INDEX_NONE, TEXT("Attempting to build the same FPackageTrailer multiple times. Package '%s'"), *PackageName.ToString());
 
 	// Note that we do not serialize containers directly as we want a file format that is 
 	// 100% under our control. This will allow people to create external scripts that can
@@ -134,6 +145,8 @@ bool FPackageTrailerBuilder::BuildAndAppendTrailer(FLinkerSave* Linker, FArchive
 
 	for (const TPair<Virtualization::FPayloadId, LocalEntry>& It : LocalEntries)
 	{
+		checkf(It.Key.IsValid(), TEXT("PackageTrailer should not contain invalid FPayloadIds. Package '%s'"), *PackageName.ToString());
+
 		Private::FLookupTableEntry& Entry = PayloadLookupTable.AddDefaulted_GetRef();
 		Entry.Identifier = It.Key;
 		Entry.OffsetInFile = PayloadsDataLength;
@@ -145,6 +158,8 @@ bool FPackageTrailerBuilder::BuildAndAppendTrailer(FLinkerSave* Linker, FArchive
 
 	for (const TPair<Virtualization::FPayloadId, VirtualizedEntry>& It : VirtualizedEntries)
 	{
+		checkf(It.Key.IsValid(), TEXT("PackageTrailer should not contain invalid FPayloadIds. Package '%s'"), *PackageName.ToString());
+
 		Private::FLookupTableEntry& Entry = PayloadLookupTable.AddDefaulted_GetRef();
 		Entry.Identifier = It.Key;
 		Entry.OffsetInFile = INDEX_NONE;
@@ -202,12 +217,24 @@ bool FPackageTrailerBuilder::IsEmpty() const
 	return LocalEntries.IsEmpty() && VirtualizedEntries.IsEmpty();
 }
 
-void FPackageTrailerBuilder::AddPayload(const Virtualization::FPayloadId& Identifier, FCompressedBuffer Payload, AdditionalDataCallback&& Callback)
+int64 FPackageTrailerBuilder::FindPayloadOffset(const Virtualization::FPayloadId& Identifier) const
 {
-	checkf(TrailerPositionInFile == INDEX_NONE, TEXT("Attempting to add payloads after the trailer has been built"));
+	check(TrailerPositionInFile != INDEX_NONE);
+	check(PayloadPosInFile != INDEX_NONE);
+	check(PayloadLookupTable.Num() == LocalEntries.Num() + VirtualizedEntries.Num());
 
-	Callbacks.Emplace(MoveTemp(Callback));
-	LocalEntries.Add(Identifier, LocalEntry(MoveTemp(Payload)));
+	if (Identifier.IsValid())
+	{
+		for (const Private::FLookupTableEntry& Entry : PayloadLookupTable)
+		{
+			if (Entry.Identifier == Identifier)
+			{
+				return PayloadPosInFile + Entry.OffsetInFile;
+			}
+		}
+	}
+
+	return INDEX_NONE;
 }
 
 bool FPackageTrailer::IsEnabled()
diff --git a/Engine/Source/Runtime/CoreUObject/Public/UObject/PackageTrailer.h b/Engine/Source/Runtime/CoreUObject/Public/UObject/PackageTrailer.h
--- a/Engine/Source/Runtime/CoreUObject/Public/UObject/PackageTrailer.h
+++ b/Engine/Source/Runtime/CoreUObject/Public/UObject/PackageTrailer.h
@@ -4,6 +4,7 @@
 
 #include "Compression/CompressedBuffer.h"
 #include "Containers/Map.h"
+#include "UObject/NameTypes.h"
 #include "Virtualization/PayloadId.h"
 
 class FArchive;
@@ -136,7 +137,17 @@ class COREUOBJECT_API FPackageTrailerBuilder
 public:
 	using AdditionalDataCallback = TFunction<void(FLinkerSave& LinkerSave)>;
 
-	[[nodiscard]] static FPackageTrailerBuilder Create(const class FPackageTrailer& Trailer, FArchive& Ar);
+	/**
+	 * @param Trailer
+	 * @param Ar
+	 * @param PackagePath
+	 */
+	[[nodiscard]] static FPackageTrailerBuilder Create(const class FPackageTrailer& Trailer, FArchive& Ar, const FPackagePath& PackagePath);
+
+	FPackageTrailerBuilder() = delete;
+	FPackageTrailerBuilder(UPackage* Package);
+	FPackageTrailerBuilder(FName InPackageName);
+	~FPackageTrailerBuilder() = default;
 
 	// Methods that can be called while building the trailer
 
@@ -197,6 +208,9 @@ private:
 
 	// Members used when building the trailer
 
+	/** Name of the package the trailer is being built for, used to give meaningful error messages */
+	FName PackageName;
+
 	/** Payloads that will be stored locally when the trailer is written to disk */
 	TMap<Virtualization::FPayloadId, LocalEntry> LocalEntries;
 	/** Payloads that are already virtualized and so will not be written to disk */
