Add an experimental feature to content virtualization to allow some bulkdata objects to opt out of being virtualized.

#rnx

- This is an experimental feature being used during testing and development, it is not intended to be used by a shippable project.
-- The code is all wrapped by a single define, UE_ENABLE_VIRTUALIZATION_TOGGLE, making it easy to disable and remove
-- The actual method called on bulkdata, ::SetVirtualizationOptOut is being submitted already deprecated to prevent it's accidental use.
- The use case is to allow projects to opt out of virtualizing static meshes as they are more fragile (i.e. hard crashes) when their payloads cannot be accessed.
- Even when ::SetVirtualizationOptOut is called the feature will not work unless  [Core.System.Experimental]AllowVirtualizationOptOut=True is set in the config file. This allows projects to opt into this.
- When a bulkdata is opting out, we simply send the payload to be stored at the end of the package like we used to, rather than send it to the package trailer. This way we do not need to store any additional state in the trailer to track if a payload is allowed to be virtualized or not. Since this feature is throw away I did not want to make any data format changes.

diff --git a/Engine/Config/BaseEngine.ini b/Engine/Config/BaseEngine.ini
--- a/Engine/Config/BaseEngine.ini
+++ b/Engine/Config/BaseEngine.ini
@@ -1119,6 +1119,7 @@ DefaultThreadStackSize=1048576
 
 [Core.System.Experimental]
 EnablePackageSidecarSaving=False
+AllowVirtualizationOptOut =False
 
 ; Options for the content virtualization system
 [Core.ContentVirtualization]
diff --git a/Engine/Source/Runtime/CoreUObject/Private/Serialization/VirtualizedBulkData.cpp b/Engine/Source/Runtime/CoreUObject/Private/Serialization/VirtualizedBulkData.cpp
--- a/Engine/Source/Runtime/CoreUObject/Private/Serialization/VirtualizedBulkData.cpp
+++ b/Engine/Source/Runtime/CoreUObject/Private/Serialization/VirtualizedBulkData.cpp
@@ -83,6 +83,23 @@ bool ShouldSaveToPackageSidecar()
 	return ConfigSetting.bEnabled;
 }
 
+#if UE_ENABLE_VIRTUALIZATION_TOGGLE
+bool ShouldAllowVirtualizationOptOut()
+{
+	static struct FAllowVirtualizationOptOut
+	{
+		bool bEnabled = true;
+
+		FAllowVirtualizationOptOut()
+		{
+			GConfig->GetBool(TEXT("Core.System.Experimental"), TEXT("AllowVirtualizationOptOut"), bEnabled, GEngineIni);
+		}
+	} AllowVirtualizationOptOut;
+
+	return AllowVirtualizationOptOut.bEnabled;
+}
+#endif // UE_ENABLE_VIRTUALIZATION_TOGGLE
+
 /** Utility for logging extended error messages when we fail to open a package for reading */
 void LogPackageOpenFailureMessage(const FPackagePath& PackagePath, EPackageSegment PackageSegment)
 {
@@ -658,117 +675,19 @@ void FVirtualizedUntypedBulkData::Serialize(FArchive& Ar, UObject* Owner, bool b
 				// the payload directly to the archive.
 				if (LinkerSave != nullptr && !LinkerSave->IsTextFormat())
 				{	
-					if (FPackageTrailer::IsEnabled() == false) // Legacy path, will save the payload data to the package
+#if UE_ENABLE_VIRTUALIZATION_TOGGLE
+					if (FPackageTrailer::IsEnabled() == false || bSkipVirtualization == true)
+#else
+					if (FPackageTrailer::IsEnabled() == false)
+#endif //UE_ENABLE_VIRTUALIZATION_TOGGLE
 					{
-						// The lambda is mutable so that PayloadToSerialize is not const (due to FArchive api not accepting const values)
-						auto SerializePayload = [this, OffsetPos, PayloadToSerialize, UpdatedFlags, Owner](FLinkerSave& LinkerSave, FArchive& ExportsArchive, FArchive& DataArchive, int64 DataStartOffset) mutable
-						{
-							checkf(ExportsArchive.IsCooking() == false, TEXT("FVirtualizedUntypedBulkData::Serialize should not be called during a cook"));
-
-							SerializeData(DataArchive, PayloadToSerialize, UpdatedFlags);
-
-							// Record the current archive offset (probably EOF but we cannot be sure)
-							const int64 ReturnPos = ExportsArchive.Tell();
-
-							// Update the offset/size entries that we set up during ::Serialize
-							ExportsArchive.Seek(OffsetPos);
-							ExportsArchive << DataStartOffset;
-							// Restore the archive's offset
-							ExportsArchive.Seek(ReturnPos);
-
-							// If we are saving the package to disk (we have access to FLinkerSave and its filepath is valid) 
-							// then we should register a callback to be received once the package has actually been saved to 
-							// disk so that we can update the object's members to be redirected to the saved file.
-							if (!LinkerSave.GetFilename().IsEmpty())
-							{
-								// At some point saving to the sidecar file will be mutually exclusive with saving to the asset file, at that point
-								// we can split these code paths entirely for clarity. (might need to update ::BuildFlagsForSerialization at that point too!)
-								if (ShouldSaveToPackageSidecar())
-								{
-									FLinkerSave::FSidecarStorageInfo& SidecarData = LinkerSave.SidecarDataToAppend.AddZeroed_GetRef();
-									SidecarData.Identifier = PayloadContentId;
-									SidecarData.Payload = PayloadToSerialize;
-								}
-
-								auto OnSavePackage = [this, DataStartOffset, UpdatedFlags, Owner](const FPackagePath& InPackagePath, FObjectPostSaveContext ObjectSaveContext)
-								{
-									if (!ObjectSaveContext.IsUpdatingLoadedPath())
-									{
-										return;
-									}
-
-									this->PackagePath = InPackagePath;
-									check(!this->PackagePath.IsEmpty()); // LinkerSave guarantees a valid PackagePath if we're updating loaded path
-									this->OffsetInFile = DataStartOffset;
-									this->Flags = UpdatedFlags;
-
-									if (CanUnloadData())
-									{
-										this->CompressionSettings.Reset();
-										this->Payload.Reset();
-									}
-
-									// Update our information in the registry
-									// TODO: Pass Owner into Register once the AssetRegistry has been fixed to use the updated PackageGuid from the save
-									Register(nullptr);
-								};
-
-								LinkerSave.PostSaveCallbacks.Add(MoveTemp(OnSavePackage));
-							}
-						};
-
-						auto AdditionalDataCallback = [SerializePayload = MoveTemp(SerializePayload)](FLinkerSave& ExportsArchive, FArchive& DataArchive, int64 DataStartOffset) mutable
-						{
-							SerializePayload(ExportsArchive, ExportsArchive, DataArchive, DataStartOffset);
-						};
-
-						LinkerSave->AdditionalDataToAppend.Add(MoveTemp(AdditionalDataCallback));	// -V595 PVS believes that LinkerSave can potentially be nullptr at 
-																									// this point however we test LinkerSave != nullptr to enter this branch.
-
+						// Legacy path, will save the payload data to the package
+						SerializeToLegacyPath(*LinkerSave, OffsetPos, PayloadToSerialize, UpdatedFlags, Owner);
 					}
-					else // New path that will save the payload to the package trailer
+					else 
 					{
-						auto OnPayloadWritten = [this, OffsetPos, UpdatedFlags](FLinkerSave& LinkerSave) mutable
-						{
-							checkf(LinkerSave.IsCooking() == false, TEXT("FVirtualizedUntypedBulkData::Serialize should not be called during a cook"));
-
-							int64 PayloadOffset = LinkerSave.PackageTrailerBuilder.FindPayloadOffset(PayloadContentId);
-
-							UpdateArchiveData(LinkerSave, OffsetPos, PayloadOffset);
-
-							// If we are saving the package to disk (we have access to FLinkerSave and its filepath is valid) 
-							// then we should register a callback to be received once the package has actually been saved to 
-							// disk so that we can update the object's members to be redirected to the saved file.
-							if (!LinkerSave.GetFilename().IsEmpty())
-							{
-								auto OnSavePackage = [this, PayloadOffset, UpdatedFlags](const FPackagePath& InPackagePath, FObjectPostSaveContext ObjectSaveContext)
-								{
-									if (!ObjectSaveContext.IsUpdatingLoadedPath())
-									{
-										return;
-									}
-
-									this->PackagePath = InPackagePath;
-									check(!this->PackagePath.IsEmpty()); // LinkerSave guarantees a valid PackagePath if we're updating loaded path
-									this->OffsetInFile = PayloadOffset;
-									this->Flags = UpdatedFlags;
-
-									if (CanUnloadData())
-									{
-										this->CompressionSettings.Reset();
-										this->Payload.Reset();
-									}
-
-									// Update our information in the registry
-									// TODO: Pass Owner into Register once the AssetRegistry has been fixed to use the updated PackageGuid from the save
-									Register(nullptr);
-								};
-
-								LinkerSave.PostSaveCallbacks.Add(MoveTemp(OnSavePackage));
-							}
-						};
-
-						LinkerSave->PackageTrailerBuilder.AddPayload(PayloadContentId, MoveTemp(PayloadToSerialize), MoveTemp(OnPayloadWritten));
+						// New path that will save the payload to the package trailer
+						SerializeToPackageTrailer(*LinkerSave, OffsetPos, PayloadToSerialize, UpdatedFlags, Owner);
 					}	
 				}
 				else
@@ -1357,6 +1276,119 @@ FGuid FVirtualizedUntypedBulkData::GetIdentifier() const
 	return BulkDataId;
 }
 
+void FVirtualizedUntypedBulkData::SerializeToLegacyPath(FLinkerSave& LinkerSave, int64 OffsetPos, FCompressedBuffer PayloadToSerialize, EFlags UpdatedFlags, UObject* Owner)
+{
+	// The lambda is mutable so that PayloadToSerialize is not const (due to FArchive api not accepting const values)
+	auto SerializePayload = [this, OffsetPos, PayloadToSerialize, UpdatedFlags, Owner](FLinkerSave& LinkerSave, FArchive& ExportsArchive, FArchive& DataArchive, int64 DataStartOffset) mutable
+	{
+		checkf(ExportsArchive.IsCooking() == false, TEXT("FVirtualizedUntypedBulkData::Serialize should not be called during a cook"));
+
+		SerializeData(DataArchive, PayloadToSerialize, UpdatedFlags);
+
+		// Record the current archive offset (probably EOF but we cannot be sure)
+		const int64 ReturnPos = ExportsArchive.Tell();
+
+		// Update the offset/size entries that we set up during ::Serialize
+		ExportsArchive.Seek(OffsetPos);
+		ExportsArchive << DataStartOffset;
+		// Restore the archive's offset
+		ExportsArchive.Seek(ReturnPos);
+
+		// If we are saving the package to disk (we have access to FLinkerSave and its filepath is valid) 
+		// then we should register a callback to be received once the package has actually been saved to 
+		// disk so that we can update the object's members to be redirected to the saved file.
+		if (!LinkerSave.GetFilename().IsEmpty())
+		{
+			// At some point saving to the sidecar file will be mutually exclusive with saving to the asset file, at that point
+			// we can split these code paths entirely for clarity. (might need to update ::BuildFlagsForSerialization at that point too!)
+			if (ShouldSaveToPackageSidecar())
+			{
+				FLinkerSave::FSidecarStorageInfo& SidecarData = LinkerSave.SidecarDataToAppend.AddZeroed_GetRef();
+				SidecarData.Identifier = PayloadContentId;
+				SidecarData.Payload = PayloadToSerialize;
+			}
+
+			auto OnSavePackage = [this, DataStartOffset, UpdatedFlags, Owner](const FPackagePath& InPackagePath, FObjectPostSaveContext ObjectSaveContext)
+			{
+				if (!ObjectSaveContext.IsUpdatingLoadedPath())
+				{
+					return;
+				}
+
+				this->PackagePath = InPackagePath;
+				check(!this->PackagePath.IsEmpty()); // LinkerSave guarantees a valid PackagePath if we're updating loaded path
+				this->OffsetInFile = DataStartOffset;
+				this->Flags = UpdatedFlags;
+
+				if (CanUnloadData())
+				{
+					this->CompressionSettings.Reset();
+					this->Payload.Reset();
+				}
+
+				// Update our information in the registry
+				// TODO: Pass Owner into Register once the AssetRegistry has been fixed to use the updated PackageGuid from the save
+				Register(nullptr);
+			};
+
+			LinkerSave.PostSaveCallbacks.Add(MoveTemp(OnSavePackage));
+		}
+	};
+
+	auto AdditionalDataCallback = [SerializePayload = MoveTemp(SerializePayload)](FLinkerSave& ExportsArchive, FArchive& DataArchive, int64 DataStartOffset) mutable
+	{
+		SerializePayload(ExportsArchive, ExportsArchive, DataArchive, DataStartOffset);
+	};
+
+	LinkerSave.AdditionalDataToAppend.Add(MoveTemp(AdditionalDataCallback));	// -V595 PVS believes that LinkerSave can potentially be nullptr at 
+																				// this point however we test LinkerSave != nullptr to enter this branch.
+}
+
+void FVirtualizedUntypedBulkData::SerializeToPackageTrailer(FLinkerSave& LinkerSave, int64 OffsetPos, FCompressedBuffer PayloadToSerialize, EFlags UpdatedFlags, UObject* Owner)
+{
+	auto OnPayloadWritten = [this, OffsetPos, UpdatedFlags](FLinkerSave& LinkerSave) mutable
+	{
+		checkf(LinkerSave.IsCooking() == false, TEXT("FVirtualizedUntypedBulkData::Serialize should not be called during a cook"));
+
+		int64 PayloadOffset = LinkerSave.PackageTrailerBuilder.FindPayloadOffset(PayloadContentId);
+
+		UpdateArchiveData(LinkerSave, OffsetPos, PayloadOffset);
+
+		// If we are saving the package to disk (we have access to FLinkerSave and its filepath is valid) 
+		// then we should register a callback to be received once the package has actually been saved to 
+		// disk so that we can update the object's members to be redirected to the saved file.
+		if (!LinkerSave.GetFilename().IsEmpty())
+		{
+			auto OnSavePackage = [this, PayloadOffset, UpdatedFlags](const FPackagePath& InPackagePath, FObjectPostSaveContext ObjectSaveContext)
+			{
+				if (!ObjectSaveContext.IsUpdatingLoadedPath())
+				{
+					return;
+				}
+
+				this->PackagePath = InPackagePath;
+				check(!this->PackagePath.IsEmpty()); // LinkerSave guarantees a valid PackagePath if we're updating loaded path
+				this->OffsetInFile = PayloadOffset;
+				this->Flags = UpdatedFlags;
+
+				if (CanUnloadData())
+				{
+					this->CompressionSettings.Reset();
+					this->Payload.Reset();
+				}
+
+				// Update our information in the registry
+				// TODO: Pass Owner into Register once the AssetRegistry has been fixed to use the updated PackageGuid from the save
+				Register(nullptr);
+			};
+
+			LinkerSave.PostSaveCallbacks.Add(MoveTemp(OnSavePackage));
+		}
+	};
+
+	LinkerSave.PackageTrailerBuilder.AddPayload(PayloadContentId, MoveTemp(PayloadToSerialize), MoveTemp(OnPayloadWritten));
+}
+
 void FVirtualizedUntypedBulkData::UpdatePayloadImpl(FSharedBuffer&& InPayload, FPayloadId&& InPayloadID)
 {
 	TRACE_CPUPROFILER_EVENT_SCOPE(FVirtualizedUntypedBulkData::UpdatePayloadImpl);
@@ -1547,6 +1579,16 @@ void FVirtualizedUntypedBulkData::UpdatePayloadId()
 	UpdateKeyIfNeeded();
 }
 
+#if UE_ENABLE_VIRTUALIZATION_TOGGLE
+void FVirtualizedUntypedBulkData::SetVirtualizationOptOut(bool bOptOut)
+{
+	if (ShouldAllowVirtualizationOptOut())
+	{
+		bSkipVirtualization = bOptOut;
+	}	
+}
+#endif //UE_ENABLE_VIRTUALIZATION_TOGGLE
+
 void FVirtualizedUntypedBulkData::UpdateKeyIfNeeded()
 {
 	// If this was created from old BulkData then the key is generated from an older FGuid, we
diff --git a/Engine/Source/Runtime/CoreUObject/Public/Serialization/VirtualizedBulkData.h b/Engine/Source/Runtime/CoreUObject/Public/Serialization/VirtualizedBulkData.h
--- a/Engine/Source/Runtime/CoreUObject/Public/Serialization/VirtualizedBulkData.h
+++ b/Engine/Source/Runtime/CoreUObject/Public/Serialization/VirtualizedBulkData.h
@@ -13,6 +13,10 @@
 class FArchive;
 class UObject;
 
+// When enabled it will be possible for specific bulkdata objects to opt out of being virtualized
+// This is a development feature and not expected to be used
+#define UE_ENABLE_VIRTUALIZATION_TOGGLE 1
+
 //TODO: At some point it might be a good idea to uncomment this to make sure that FVirtualizedUntypedBulkData is
 //		never used at runtime (requires a little too much reworking of some assets for now though)
 //#if WITH_EDITORONLY_DATA
@@ -293,6 +297,11 @@ public:
 	/** Load the payload and set the correct payload id, if the bulkdata has a PlaceholderPayloadId. */
 	void UpdatePayloadId();
 
+#if UE_ENABLE_VIRTUALIZATION_TOGGLE
+	UE_DEPRECATED(5.0, "SetVirtualizationOptOut is an internal feature for development and will be removed without warning!")
+	void SetVirtualizationOptOut(bool bOptOut);
+#endif //UE_ENABLE_VIRTUALIZATION_TOGGLE
+
 protected:
 	enum class ETornOff {};
 	FVirtualizedUntypedBulkData(const FVirtualizedUntypedBulkData& Other, ETornOff);
@@ -337,6 +346,11 @@ private:
 
 	FRIEND_ENUM_CLASS_FLAGS(EFlags);
 
+	/** Old legacy path that saved the payload to the end of the package */
+	void SerializeToLegacyPath(FLinkerSave& LinkerSave, int64 OffsetPos, FCompressedBuffer PayloadToSerialize, EFlags UpdatedFlags, UObject* Owner);
+	/** The new path that saves payloads to the FPackageTrailer which is then appended to the end of the package file */
+	void SerializeToPackageTrailer(FLinkerSave& LinkerSave, int64 OffsetPos, FCompressedBuffer PayloadToSerialize, EFlags UpdatedFlags, UObject* Owner);
+
 	void UpdatePayloadImpl(FSharedBuffer&& InPayload, FPayloadId&& InPayloadID);
 
 	FCompressedBuffer GetDataInternal() const;
@@ -419,6 +433,10 @@ private:
 	/** A 32bit bitfield of flags */
 	EFlags Flags = EFlags::None;
 
+#if UE_ENABLE_VIRTUALIZATION_TOGGLE
+	bool bSkipVirtualization = false;
+#endif //UE_ENABLE_VIRTUALIZATION_TOGGLE
+
 	/** 
 	 * Compression settings to be applied to the payload when the package is next saved. The settings will be reset if
 	 * the payload is unloaded from memory during serialization (i.e. the payload was virtualized or the package was
diff --git a/Engine/Source/Runtime/MeshDescription/Public/MeshDescription.h b/Engine/Source/Runtime/MeshDescription/Public/MeshDescription.h
--- a/Engine/Source/Runtime/MeshDescription/Public/MeshDescription.h
+++ b/Engine/Source/Runtime/MeshDescription/Public/MeshDescription.h
@@ -1295,6 +1295,14 @@ public:
 		: bBulkDataUpdated(false)
 		, bGuidIsHash(false)
 	{
+#if UE_ENABLE_VIRTUALIZATION_TOGGLE
+		// Currently FMeshDescriptionBulkData will cause a hard crash if the virtualized data cannot be found.
+		// Calling ::SetVirtualizationOptOut allows us to prevent the bulkdata from virtualizing if a project
+		// has that option enabled. This is not intended to be a shipping feature.
+		PRAGMA_DISABLE_DEPRECATION_WARNINGS	
+		BulkData.SetVirtualizationOptOut(true);
+		PRAGMA_ENABLE_DEPRECATION_WARNINGS
+#endif //UE_ENABLE_VIRTUALIZATION_TOGGLE
 	}
 
 	/** Serialization */
