Add profile tag in FlushAsyncLoading.

diff --git a/Engine/Source/Runtime/CoreUObject/Private/Serialization/AsyncPackageLoader.cpp b/Engine/Source/Runtime/CoreUObject/Private/Serialization/AsyncPackageLoader.cpp
--- a/Engine/Source/Runtime/CoreUObject/Private/Serialization/AsyncPackageLoader.cpp
+++ b/Engine/Source/Runtime/CoreUObject/Private/Serialization/AsyncPackageLoader.cpp
@@ -602,6 +602,8 @@ bool IsInAsyncLoadingThreadCoreUObjectInternal()
 
 void FlushAsyncLoading(int32 PackageID /* = INDEX_NONE */)
 {
+	TRACE_CPUPROFILER_EVENT_SCOPE(FlushAsyncLoading);
+
 #if defined(WITH_CODE_GUARD_HANDLER) && WITH_CODE_GUARD_HANDLER
 	void CheckImageIntegrityAtRuntime();
 	CheckImageIntegrityAtRuntime();
