Fix potential race conditions introduced in optimization from 17401227.

diff --git a/Engine/Source/Runtime/CoreUObject/Private/Misc/PackageName.cpp b/Engine/Source/Runtime/CoreUObject/Private/Misc/PackageName.cpp
--- a/Engine/Source/Runtime/CoreUObject/Private/Misc/PackageName.cpp
+++ b/Engine/Source/Runtime/CoreUObject/Private/Misc/PackageName.cpp
@@ -686,6 +686,8 @@ bool FPackageName::SplitLongPackageName(const FString& InLongPackageName, FStrin
 	const FLongPackagePathsSingleton& Paths = FLongPackagePathsSingleton::Get();
 
 	const bool bIncludeReadOnlyRoots = true;
+
+	FReadScopeLock ScopeLock(ContentMountPointCriticalSection);
 	const TArray<FString>& ValidRoots = Paths.GetValidLongPackageRoots(bIncludeReadOnlyRoots);
 
 	// Check to see whether our package came from a valid root
@@ -873,6 +875,8 @@ bool FPackageName::IsValidLongPackageName(FStringView InLongPackageName, bool bI
 		{
 			if (Reason == EErrorCode::PackageNamePathNotMounted)
 			{
+				FReadScopeLock ScopeLock(ContentMountPointCriticalSection);
+
 				const FLongPackagePathsSingleton& Paths = FLongPackagePathsSingleton::Get();
 				const TArray<FString>& ValidRoots = Paths.GetValidLongPackageRoots(bIncludeReadOnlyRoots);
 				if (ValidRoots.Num() == 0)
@@ -921,6 +925,7 @@ bool FPackageName::IsValidLongPackageName(FStringView InLongPackageName, bool bI
 	}
 
 	// Check valid roots
+	FReadScopeLock ScopeLock(ContentMountPointCriticalSection);
 	const FLongPackagePathsSingleton& Paths = FLongPackagePathsSingleton::Get();
 	const TArray<FString>& ValidRoots = Paths.GetValidLongPackageRoots(bIncludeReadOnlyRoots);
 	bool bValidRoot = false;
@@ -1036,6 +1041,7 @@ FName FPackageName::GetPackageMountPoint(const FString& InPackagePath, bool InWi
 {
 	FLongPackagePathsSingleton& Paths = FLongPackagePathsSingleton::Get();
 	
+	FReadScopeLock ScopeLock(ContentMountPointCriticalSection);
 	const TArray<FString>& MountPoints = Paths.GetValidLongPackageRoots(true);
 
 	int32 WithoutSlashes = InWithoutSlashes ? 1 : 0;
@@ -1991,7 +1997,11 @@ void FPackageName::IteratePackagesInDirectory(const FString& RootDir, const FPac
 void FPackageName::QueryRootContentPaths(TArray<FString>& OutRootContentPaths, bool bIncludeReadOnlyRoots, bool bWithoutLeadingSlashes, bool bWithoutTrailingSlashes)
 {
 	const FLongPackagePathsSingleton& Paths = FLongPackagePathsSingleton::Get();
-	OutRootContentPaths = Paths.GetValidLongPackageRoots( bIncludeReadOnlyRoots );
+
+	{
+		FReadScopeLock ScopeLock(ContentMountPointCriticalSection);
+		OutRootContentPaths = Paths.GetValidLongPackageRoots( bIncludeReadOnlyRoots );
+	}
 
 	if (bWithoutTrailingSlashes || bWithoutLeadingSlashes)
 	{
