Add an option not to emit warning when calling IsUnsaveable
Reverse check during IsUnsaveable to emit class warning only on valid non transient object

diff --git a/Engine/Source/Runtime/CoreUObject/Private/UObject/SavePackage/SaveContext.cpp b/Engine/Source/Runtime/CoreUObject/Private/UObject/SavePackage/SaveContext.cpp
--- a/Engine/Source/Runtime/CoreUObject/Private/UObject/SavePackage/SaveContext.cpp
+++ b/Engine/Source/Runtime/CoreUObject/Private/UObject/SavePackage/SaveContext.cpp
@@ -37,18 +37,31 @@ void FSaveContext::MarkUnsaveable(UObject* InObject)
 	ensureAlways(!InObject->HasAllFlags(RF_ClassDefaultObject | RF_Transient) || (InObject->GetClass()->ClassGeneratedBy != nullptr && InObject->GetClass()->HasAnyFlags(RF_Transient)));
 }
 
-bool FSaveContext::IsUnsaveable(UObject* InObject) const
+bool FSaveContext::IsUnsaveable(UObject* InObject, bool bEmitWarning) const
 {
 	UObject* Obj = InObject;
 	while (Obj)
 	{
+		// pending kill object are unsaveable
+		if (!IsValidChecked(Obj))
+		{
+			return true;
+		}
+
+		// transient object are considered unsaveable if non native
+		if (Obj->HasAnyFlags(RF_Transient) && !Obj->IsNative())
+		{
+			return true;
+		}
+
 		// if the object class is abstract, has been marked as deprecated, there is a newer version that exist, or the class is marked transient, then the object is unsaveable
 		// @note: Although object instances of a transient class should definitely be unsaveable, it results in discrepancies with the old save algorithm and currently load problems
 		if (Obj->GetClass()->HasAnyClassFlags(CLASS_Abstract | CLASS_Deprecated | CLASS_NewerVersionExists /*| CLASS_Transient*/) 
 			&& !Obj->HasAnyFlags(RF_ClassDefaultObject))
 		{
 			// only warn if the base object is fine but the outer is invalid. If an object is itself unsaveable, the old behavior is to ignore it
-			if (IsValidChecked(InObject) 
+			if (bEmitWarning
+				&& IsValidChecked(Obj)
 				&& InObject->GetOutermost() == GetPackage()
 				&& Obj != InObject)
 			{
@@ -60,18 +73,6 @@ bool FSaveContext::IsUnsaveable(UObject* InObject) const
 			return true;
 		}
 
-		// pending kill object are unsaveable
-		if (!IsValidChecked(Obj))
-		{
-			return true;
-		}
-
-		// transient object are considered unsaveable if non native
-		if (Obj->HasAnyFlags(RF_Transient) && !Obj->IsNative())
-		{
-			return true;
-		}
-
 		Obj = Obj->GetOuter();
 	}
 	return false;	
diff --git a/Engine/Source/Runtime/CoreUObject/Private/UObject/SavePackage/SaveContext.h b/Engine/Source/Runtime/CoreUObject/Private/UObject/SavePackage/SaveContext.h
--- a/Engine/Source/Runtime/CoreUObject/Private/UObject/SavePackage/SaveContext.h
+++ b/Engine/Source/Runtime/CoreUObject/Private/UObject/SavePackage/SaveContext.h
@@ -566,7 +566,7 @@ public:
 
 	void MarkUnsaveable(UObject* InObject);
 
-	bool IsUnsaveable(UObject* InObject) const;
+	bool IsUnsaveable(UObject* InObject, bool bEmitWarning = true) const;
 
 	void RecordIllegalReference(UObject* InFrom, UObject* InTo, EIllegalRefReason InReason)
 	{
diff --git a/Engine/Source/Runtime/CoreUObject/Private/UObject/SavePackage2.cpp b/Engine/Source/Runtime/CoreUObject/Private/UObject/SavePackage2.cpp
--- a/Engine/Source/Runtime/CoreUObject/Private/UObject/SavePackage2.cpp
+++ b/Engine/Source/Runtime/CoreUObject/Private/UObject/SavePackage2.cpp
@@ -271,7 +271,9 @@ ESavePackageResult RoutePresave(FSaveContext& SaveContext)
 	GetObjectsWithPackage(SaveContext.GetPackage(), ObjectsInPackage);
 	for (UObject* Object : ObjectsInPackage)
 	{
-		if (!SaveContext.IsUnsaveable(Object))
+		// Do not emit unsaveable warning while routing presave. 
+		// this is to prevent warning on objects which won't be harvested later since they are unreferenced
+		if (!SaveContext.IsUnsaveable(Object, false/*bEmitWarning*/))
 		{
 			if (SaveContext.IsCooking() && Object->HasAnyFlags(RF_ClassDefaultObject | RF_ArchetypeObject))
 			{
