Removes tons of temporary memory allocations from FAsyncArchive::LogItem (which is disabled by default).

diff --git a/Engine/Source/Runtime/CoreUObject/Private/Serialization/AsyncLoading.cpp b/Engine/Source/Runtime/CoreUObject/Private/Serialization/AsyncLoading.cpp
--- a/Engine/Source/Runtime/CoreUObject/Private/Serialization/AsyncLoading.cpp
+++ b/Engine/Source/Runtime/CoreUObject/Private/Serialization/AsyncLoading.cpp
@@ -292,13 +292,14 @@ DECLARE_MEMORY_STAT(TEXT("FAsyncArchive Buffers"), STAT_FAsyncArchiveMem, STATGR
 
 FORCEINLINE void FAsyncArchive::LogItem(const TCHAR* Item, int64 Offset, int64 Size, double StartTime)
 {
-	FString FileName = PackagePath.GetLocalFullPath();
 	if (UE_LOG_ACTIVE(LogAsyncArchive, Verbose)
 #if defined(ASYNC_WATCH_FILE)
-		|| FileName.Contains(TEXT(ASYNC_WATCH_FILE))
+		|| PackagePath.GetLocalFullPath().Contains(TEXT(ASYNC_WATCH_FILE))
 #endif
 		)
 	{
+		FString FileName = PackagePath.GetLocalFullPath();
+
 		static double GlobalStartTime(FPlatformTime::Seconds());
 		double Now(FPlatformTime::Seconds());
 
@@ -6401,7 +6402,7 @@ EAsyncPackageState::Type FAsyncPackage::PostLoadDeferredObjects(double InTickSta
 
 		PackageScope.ThreadContext.CurrentlyPostLoadedObjectByALT = Object;
 		{
-			TRACE_LOADTIME_POSTLOAD_EXPORT_SCOPE(Object);
+			TRACE_LOADTIME_POSTLOAD_EXPORT_SCOPE(Object);			
 			Object->ConditionalPostLoad();
 		}
 		PackageScope.ThreadContext.CurrentlyPostLoadedObjectByALT = nullptr;
