Ported CL17725160

FCoreUObjectDelegates::OnEndLoadPackage:
Fix the implementation in AsyncLoading1 to handle recursive loads; FAsyncPackage->PackageObjLoaded includes UObjects loaded from packages that were import dependencies of the package, and we need to send OnEndLoadPackage for those as well.

Also filter out the transient package and script and memory packages from all OnEndLoadPackage calls; this was affecting all three of AsyncLoading1, AsyncLoading2, and UObjectGlobals.

Did some comment improvement and removal of dead code in AsyncLoading1 while I was working in the area.

diff --git a/Engine/Source/Runtime/CoreUObject/Private/Serialization/AsyncLoading.cpp b/Engine/Source/Runtime/CoreUObject/Private/Serialization/AsyncLoading.cpp
--- a/Engine/Source/Runtime/CoreUObject/Private/Serialization/AsyncLoading.cpp
+++ b/Engine/Source/Runtime/CoreUObject/Private/Serialization/AsyncLoading.cpp
@@ -4288,7 +4288,6 @@ EAsyncPackageState::Type FAsyncLoadingThread::ProcessLoadedPackages(bool bUseTim
 
 	// Call callbacks in a batch in a stack-local array. This is to ensure that callbacks that recursively enter this function
 	// via FlushAsyncLoading do not stack overflow by doing one recursive call per LoadedPackageToProcess
-	TSet<UPackage*> CompletedUPackages;
 	for (FAsyncPackage* Package : CompletedPackages)
 	{
 		// Call external callbacks
@@ -5269,6 +5268,8 @@ void FAsyncPackage::FlushObjectLinkerCache()
 #if WITH_EDITOR 
 void FAsyncPackage::GetLoadedAssetsAndPackages(TSet<FWeakObjectPtr>& AssetList, TSet<UPackage*>& PackageList)
 {
+	TSet<UObject*> ExistingAssets;
+	TSet<UPackage*> ExistingPackages;
 	for (UObject* Obj : PackageObjLoaded)
 	{
 		if (!IsValid(Obj))
@@ -5277,12 +5278,39 @@ void FAsyncPackage::GetLoadedAssetsAndPackages(TSet<FWeakObjectPtr>& AssetList,
 		}
 		if (Obj->IsAsset())
 		{
+			if (ExistingAssets.Num() == 0)
+			{
+				ExistingAssets.Reserve(AssetList.Num() + PackageObjLoaded.Num());
+				for (FWeakObjectPtr& WeakPtr : AssetList)
+				{
+					UObject* ExistingAsset = WeakPtr.Get();
+					if (ExistingAsset)
+					{
+						ExistingAssets.Add(ExistingAsset);
+					}
+				}
+			}
+
+			ExistingAssets.Add(Obj, &bExists);
 			AssetList.Add(Obj);
 		}
 		UPackage* Package = Obj->GetPackage();
 		if (Package && !Package->HasAnyFlags(RF_Transient) && !Package->HasAnyPackageFlags(PKG_InMemoryOnly))
 		{
-			PackageList.Add(Package);
+			if (ExistingPackages.Num() == 0)
+			{
+				ExistingPackages.Reserve(PackageList.Num() + 1);
+				for (UPackage* ExistingPackage : PackageList)
+				{
+					ExistingPackages.Add(ExistingPackage);
+				}
+			}
+			bool bExists;
+			ExistingPackages.Add(Package, &bExists);
+			if (!bExists)
+			{
+				PackageList.Add(Package);
+			}
 		}
 	}
 }
