Support for FField that does not have a CASTFLAG. Similar to UObject, if the FField does have a cast flag it will be used, else it will loop through all it parent classes.

diff --git a/Engine/Plugins/NetcodeUnitTest/NetcodeUnitTest/Source/NetcodeUnitTest/Public/Net/NUTUtilNet.h b/Engine/Plugins/NetcodeUnitTest/NetcodeUnitTest/Source/NetcodeUnitTest/Public/Net/NUTUtilNet.h
--- a/Engine/Plugins/NetcodeUnitTest/NetcodeUnitTest/Source/NetcodeUnitTest/Public/Net/NUTUtilNet.h
+++ b/Engine/Plugins/NetcodeUnitTest/NetcodeUnitTest/Source/NetcodeUnitTest/Public/Net/NUTUtilNet.h
@@ -342,7 +342,7 @@ private:
  */
 class FNetPropertyHook : public FProperty
 {
-	DECLARE_FIELD(FNetPropertyHook, FProperty, (CASTCLASS_FMulticastSparseDelegateProperty << 1))
+	DECLARE_FIELD(FNetPropertyHook, FProperty, CASTCLASS_None)
 
 private:
 	virtual bool NetSerializeItem(FArchive& Ar, UPackageMap* Map, void* Data, TArray<uint8>* MetaData /*=nullptr*/) const override;
diff --git a/Engine/Source/Runtime/CoreUObject/Private/UObject/Field.cpp b/Engine/Source/Runtime/CoreUObject/Private/UObject/Field.cpp
--- a/Engine/Source/Runtime/CoreUObject/Private/UObject/Field.cpp
+++ b/Engine/Source/Runtime/CoreUObject/Private/UObject/Field.cpp
@@ -10,6 +10,7 @@ Field.cpp: Defines FField property system fundamentals
 #include "Misc/ScopeLock.h"
 #include "Misc/StringBuilder.h"
 #include "Serialization/MemoryWriter.h"
+#include "Misc/AutomationTest.h"
 #include "Misc/ConfigCacheIni.h"
 #include "Misc/OutputDeviceHelper.h"
 #include "Misc/FeedbackContext.h"
@@ -1200,4 +1201,39 @@ FField* FindFPropertyByPath(const TCHAR* InFieldPath)
 	return FoundField;
 }
 
+#if WITH_DEV_AUTOMATION_TESTS
+
+IMPLEMENT_SIMPLE_AUTOMATION_TEST(FFieldCastTest, "UObject.Field Cast", EAutomationTestFlags::EditorContext | EAutomationTestFlags::ApplicationContextMask | EAutomationTestFlags::ServerContext | EAutomationTestFlags::SmokeFilter)
+
+bool FFieldCastTest::RunTest(const FString& Parameters)
+{
+	FBoolProperty* DefaultBoolProperty = static_cast<FBoolProperty*>(FBoolProperty::StaticClass()->GetDefaultObject());
+	FIntProperty* DefaultIntProperty = static_cast<FIntProperty*>(FIntProperty::StaticClass()->GetDefaultObject());
+	FNumericProperty* DefaultNumericProperty = static_cast<FNumericProperty*>(FNumericProperty::StaticClass()->GetDefaultObject());
+	FProperty* BaseProperty = nullptr;
+
+	AddErrorIfFalse(CastField<FBoolProperty>(DefaultBoolProperty) == DefaultBoolProperty, TEXT("DefaultBoolProperty could not be CastField to a FBoolProperty."));
+	BaseProperty = DefaultBoolProperty;
+	AddErrorIfFalse(CastField<FBoolProperty>(BaseProperty) == DefaultBoolProperty, TEXT("Property could not be CastField to a FBoolProperty."));
+
+	AddErrorIfFalse(CastField<FBoolProperty>(DefaultIntProperty) == nullptr, TEXT("DefaultIntProperty was CastField to a FBoolProperty."));
+	BaseProperty = DefaultIntProperty;
+	AddErrorIfFalse(CastField<FBoolProperty>(BaseProperty) == nullptr, TEXT("DefaultIntProperty was CastField to a FBoolProperty."));
+
+	AddErrorIfFalse(CastField<FBoolProperty>(DefaultNumericProperty) == nullptr, TEXT("DefaultNumericProperty was CastField to a FBoolProperty."));
+	BaseProperty = DefaultNumericProperty;
+	AddErrorIfFalse(CastField<FBoolProperty>(BaseProperty) == nullptr, TEXT("BaseProperty was CastField to a FBoolProperty."));
+
+	AddErrorIfFalse(CastField<FNumericProperty>(DefaultIntProperty) == DefaultIntProperty, TEXT("DefaultIntProperty could not be CastField to a FNumericProperty."));
+	BaseProperty = DefaultIntProperty;
+	AddErrorIfFalse(CastField<FNumericProperty>(BaseProperty) == DefaultIntProperty, TEXT("BaseProperty could not be CastField to a FNumericProperty."));
+
+	BaseProperty = nullptr;
+	AddErrorIfFalse(CastField<FNumericProperty>(BaseProperty) == nullptr, TEXT("nullptr was CastField to a FNumericProperty."));
+
+	return true;
+}
+
+#endif
+
 #include "UObject/DefineUPropertyMacros.h"
\ No newline at end of file
diff --git a/Engine/Source/Runtime/CoreUObject/Public/UObject/Field.h b/Engine/Source/Runtime/CoreUObject/Public/UObject/Field.h
--- a/Engine/Source/Runtime/CoreUObject/Public/UObject/Field.h
+++ b/Engine/Source/Runtime/CoreUObject/Public/UObject/Field.h
@@ -92,7 +92,8 @@ public:
 	}
 	inline bool IsChildOf(const FFieldClass* InClass) const
 	{
-		return !!(CastFlags & InClass->GetId());
+		const uint64 OtherClassId = InClass->GetId();
+		return OtherClassId ? !!(CastFlags & OtherClassId) : IsChildOf_Walk(InClass);
 	}
 	FString GetDescription() const;
 	FText GetDisplayNameText() const;
@@ -101,7 +102,7 @@ public:
 		return ConstructFn(InOwner, InName, InFlags);
 	}
 
-	FFieldClass* GetSuperClass()
+	FFieldClass* GetSuperClass() const
 	{
 		return SuperClass;
 	}
@@ -132,6 +133,19 @@ public:
 		return Ar;
 	}
 	COREUOBJECT_API friend FArchive& operator << (FArchive& Ar, FFieldClass*& InOutFieldClass);
+
+private:
+	bool IsChildOf_Walk(const FFieldClass* InBaseClass) const
+	{
+		for (const FFieldClass* TempField = this; TempField; TempField = TempField->GetSuperClass())
+		{
+			if (TempField == InBaseClass)
+			{
+				return true;
+			}
+		}
+		return false;
+	}
 };
 
 #if !CHECK_PUREVIRTUALS
@@ -158,11 +172,11 @@ public: \
 	} \
 	static FFieldClass* StaticClass(); \
 	static FField* Construct(const FFieldVariant& InOwner, const FName& InName, EObjectFlags InObjectFlags); \
-	inline static uint64 StaticClassCastFlagsPrivate() \
+	inline static constexpr uint64 StaticClassCastFlagsPrivate() \
 	{ \
 		return uint64(TStaticFlags); \
 	} \
-	inline static uint64 StaticClassCastFlags() \
+	inline static constexpr uint64 StaticClassCastFlags() \
 	{ \
 		return uint64(TStaticFlags) | Super::StaticClassCastFlags(); \
 	} \
@@ -386,11 +400,11 @@ public:
 
 	static FFieldClass* StaticClass();
 
-	inline static uint64 StaticClassCastFlagsPrivate()
+	inline static constexpr uint64 StaticClassCastFlagsPrivate()
 	{
 		return uint64(CASTCLASS_UField);
 	}
-	inline static uint64 StaticClassCastFlags()
+	inline static constexpr uint64 StaticClassCastFlags()
 	{
 		return uint64(CASTCLASS_UField);
 	}
@@ -508,20 +522,29 @@ public:
 	inline bool IsA(const FFieldClass* FieldType) const
 	{
 		check(FieldType);
-		return !!(GetCastFlags() & FieldType->GetId());
+		return GetClass()->IsChildOf(FieldType);
 	}
 
 	template<typename T>
 	bool IsA() const
 	{
-		return !!(GetCastFlags() & T::StaticClassCastFlagsPrivate());
+		if constexpr (!!(T::StaticClassCastFlagsPrivate()))
+		{
+			return !!(GetCastFlags() & T::StaticClassCastFlagsPrivate());
+		}
+		else
+		{
+			return GetClass()->IsChildOf(T::StaticClass());
+		}
 	}
 
+	UE_DEPRECATED(5.0, "HasAnyCastFlags is deprecated. Not all FField has CastFlag. Use IsA instead.")
 	inline bool HasAnyCastFlags(const uint64 InCastFlags) const
 	{
 		return !!(GetCastFlags() & InCastFlags);
 	}
 
+	UE_DEPRECATED(5.0, "HasAllCastFlags is deprecated. Not all FField has CastFlag. Use IsA instead.")
 	inline bool HasAllCastFlags(const uint64 InCastFlags) const
 	{
 		return (GetCastFlags() & InCastFlags) == InCastFlags;
@@ -840,13 +863,13 @@ public:
 template<typename FieldType>
 FORCEINLINE FieldType* CastField(FField* Src)
 {
-	return Src && Src->HasAnyCastFlags(FieldType::StaticClassCastFlagsPrivate()) ? static_cast<FieldType*>(Src) : nullptr;
+	return Src && Src->IsA<FieldType>() ? static_cast<FieldType*>(Src) : nullptr;
 }
 
 template<typename FieldType>
 FORCEINLINE const FieldType* CastField(const FField* Src)
 {
-	return Src && Src->HasAnyCastFlags(FieldType::StaticClassCastFlagsPrivate()) ? static_cast<const FieldType*>(Src) : nullptr;
+	return Src && Src->IsA<FieldType>() ? static_cast<const FieldType*>(Src) : nullptr;
 }
 
 template<typename FieldType>
@@ -863,7 +886,7 @@ FUNCTION_NON_NULL_RETURN_END
 #if !DO_CHECK
 	return (FieldType*)Src;
 #else
-	FieldType* CastResult = Src && Src->HasAnyCastFlags(FieldType::StaticClassCastFlagsPrivate()) ? (FieldType*)Src : nullptr;
+	FieldType* CastResult = Src && Src->IsA<FieldType>() ? static_cast<FieldType*>(Src) : nullptr;
 	checkf(CastResult, TEXT("CastFieldChecked failed with 0x%016llx"), (int64)(PTRINT)Src);
 	return CastResult;
 #endif // !DO_CHECK
@@ -877,7 +900,7 @@ FUNCTION_NON_NULL_RETURN_END
 #if !DO_CHECK
 	return (const FieldType*)Src;
 #else
-	const FieldType* CastResult = Src && Src->HasAnyCastFlags(FieldType::StaticClassCastFlagsPrivate()) ? (const FieldType*)Src : nullptr;
+	const FieldType* CastResult = Src && Src->IsA<FieldType>() ? static_cast<const FieldType*>(Src) : nullptr;
 	checkf(CastResult, TEXT("CastFieldChecked failed with 0x%016llx"), (int64)(PTRINT)Src);
 	return CastResult;
 #endif // !DO_CHECK
@@ -889,7 +912,7 @@ FORCEINLINE FieldType* CastFieldCheckedNullAllowed(FField* Src)
 #if !DO_CHECK
 	return (FieldType*)Src;
 #else
-	FieldType* CastResult = Src && Src->HasAnyCastFlags(FieldType::StaticClassCastFlagsPrivate()) ? (FieldType*)Src : nullptr;
+	FieldType* CastResult = Src && Src->IsA<FieldType>() ? static_cast<FieldType*>(Src) : nullptr;
 	checkf(CastResult || !Src, TEXT("CastFieldCheckedNullAllowed failed with 0x%016llx"), (int64)(PTRINT)Src);
 	return CastResult;
 #endif // !DO_CHECK
@@ -901,7 +924,7 @@ FORCEINLINE const FieldType* CastFieldCheckedNullAllowed(const FField* Src)
 #if !DO_CHECK
 	return (const FieldType*)Src;
 #else
-	const FieldType* CastResult = Src && Src->HasAnyCastFlags(FieldType::StaticClassCastFlagsPrivate()) ? (const FieldType*)Src : nullptr;
+	const FieldType* CastResult = Src && Src->IsA<FieldType>() ? static_cast<const FieldType*>(Src) : nullptr;
 	checkf(CastResult || !Src, TEXT("CastFieldCheckedNullAllowed failed with 0x%016llx"), (int64)(PTRINT)Src);
 	return CastResult;
 #endif // !DO_CHECK
