Changing E_GCSkipIndexPlaceholder value to something more neutral

diff --git a/Engine/Source/Runtime/CoreUObject/Public/UObject/GarbageCollection.h b/Engine/Source/Runtime/CoreUObject/Public/UObject/GarbageCollection.h
--- a/Engine/Source/Runtime/CoreUObject/Public/UObject/GarbageCollection.h
+++ b/Engine/Source/Runtime/CoreUObject/Public/UObject/GarbageCollection.h
@@ -206,7 +206,7 @@ struct FTokenInfo
 struct COREUOBJECT_API FGCReferenceTokenStream
 {
 	/** Initialization value to ensure that we have the right skip index index */
-	enum EGCArrayInfoPlaceholder { E_GCSkipIndexPlaceholder = 0xDEADBABE }; 
+	enum EGCArrayInfoPlaceholder { E_GCSkipIndexPlaceholder = 0xDEDEDEDE }; 
 
 	/** Constructor */
 	FGCReferenceTokenStream()
