Update FArchiveReplaceOrClearExternalReferences to be able to take EArchiveReplaceObjectFlags consistent with parent class

diff --git a/Engine/Source/Runtime/CoreUObject/Public/Serialization/ArchiveReplaceOrClearExternalReferences.h b/Engine/Source/Runtime/CoreUObject/Public/Serialization/ArchiveReplaceOrClearExternalReferences.h
--- a/Engine/Source/Runtime/CoreUObject/Public/Serialization/ArchiveReplaceOrClearExternalReferences.h
+++ b/Engine/Source/Runtime/CoreUObject/Public/Serialization/ArchiveReplaceOrClearExternalReferences.h
@@ -24,7 +24,22 @@ public:
 		( UObject* InSearchObject
 		, const TMap<T*, T*>& InReplacementMap
 		, UPackage* InDestPackage
-		, bool bDelayStart = false )
+		, EArchiveReplaceObjectFlags Flags = EArchiveReplaceObjectFlags::None)
+		: TSuper(InSearchObject, InReplacementMap, EArchiveReplaceObjectFlags::DelayStart | Flags)
+		, DestPackage(InDestPackage)
+	{
+		if (!(Flags & EArchiveReplaceObjectFlags::DelayStart))
+		{
+			this->SerializeSearchObject();
+		}
+	}
+
+	UE_DEPRECATED(5.0, "Use version that passes Flags instead of just bDelayStart")
+	FArchiveReplaceOrClearExternalReferences
+		( UObject* InSearchObject
+		, const TMap<T*, T*>& InReplacementMap
+		, UPackage* InDestPackage
+		, bool bDelayStart)
 		: TSuper(InSearchObject, InReplacementMap, EArchiveReplaceObjectFlags::DelayStart)
 		, DestPackage(InDestPackage)
 	{
