VirtualizedBulkData: EditorDomain support for loading any format by reference from WorkspaceDomain. Previously this was only supported for legacy WorkspaceDomain packages, but we also need it for packages that have been resaved and use VirtualizedBulkData.
Also fix two other minor issues: Make Move constructor available for FIntVirtualizedBulkData, and make Register public so UObjects can notify the registry that they own a BulkData.
#rnx

diff --git a/Engine/Source/Runtime/CoreUObject/Private/Serialization/VirtualizedBulkData.cpp b/Engine/Source/Runtime/CoreUObject/Private/Serialization/VirtualizedBulkData.cpp
--- a/Engine/Source/Runtime/CoreUObject/Private/Serialization/VirtualizedBulkData.cpp
+++ b/Engine/Source/Runtime/CoreUObject/Private/Serialization/VirtualizedBulkData.cpp
@@ -551,14 +551,14 @@ void FVirtualizedUntypedBulkData::Serialize(FArchive& Ar, UObject* Owner, bool b
 	else if (Ar.IsPersistent() && !Ar.IsObjectReferenceCollector() && !Ar.ShouldSkipBulkData())
 	{
 		FLinkerSave* LinkerSave = nullptr;
-		bool bKeepLegacyDataByReference = false;
+		bool bKeepFileDataByReference = false;
 		if (Ar.IsSaving())
 		{
 			LinkerSave = Cast<FLinkerSave>(Ar.GetLinker());
 			// If we're doing a save that can refer to bulk data by reference, and our legacy data format supports it,
 			// keep any legacy data we have referenced rather than stored, to save space and avoid spending time loading it.
-			bKeepLegacyDataByReference = LinkerSave && LinkerSave->bProceduralSave && PackageSegment == EPackageSegment::Header;
-			if (!bKeepLegacyDataByReference)
+			bKeepFileDataByReference = LinkerSave && LinkerSave->bProceduralSave && PackageSegment == EPackageSegment::Header;
+			if (!bKeepFileDataByReference)
 			{
 				UpdateKeyIfNeeded();
 			}
@@ -604,7 +604,7 @@ void FVirtualizedUntypedBulkData::Serialize(FArchive& Ar, UObject* Owner, bool b
 		{
 			checkf(Ar.IsCooking() == false, TEXT("FVirtualizedUntypedBulkData::Serialize should not be called during a cook"));
 
-			EFlags UpdatedFlags = BuildFlagsForSerialization(Ar, !bKeepLegacyDataByReference);
+			EFlags UpdatedFlags = BuildFlagsForSerialization(Ar, bKeepFileDataByReference);
 
 			// Go back in the archive and update the flags in the archive, we will only apply the updated flags to the current
 			// object later if we detect that the package saved successfully.
@@ -618,9 +618,9 @@ void FVirtualizedUntypedBulkData::Serialize(FArchive& Ar, UObject* Owner, bool b
 
 			// Write out required extra data if we're saving by reference
 			bool bWriteOutPayload = true;
-			if (EnumHasAnyFlags(UpdatedFlags, EFlags::ReferencesLegacyFile))
+			if (IsReferencingByPackagePath(UpdatedFlags))
 			{
-				check(PackageSegment == EPackageSegment::Header); // This should have been checked before setting bKeepLegacyDataByReference=true
+				check(PackageSegment == EPackageSegment::Header); // This should have been checked before setting bKeepFileDataByReference=true
 				FString PackageName = PackagePath.GetPackageName();
 				Ar << PackageName;
 				Ar << OffsetInFile;
@@ -723,14 +723,14 @@ void FVirtualizedUntypedBulkData::Serialize(FArchive& Ar, UObject* Owner, bool b
 
 			if (Trailer != nullptr && Trailer->FindPayloadStatus(PayloadContentId) == EPayloadStatus::StoredVirtualized)
 			{
-				check(!IsReferencingOldBulkData()); // This should not be possible
+				check(!IsReferencingByPackagePath()); // This should not be possible
 
 				EnumAddFlags(Flags, EFlags::IsVirtualized);
 
 				Ar << OffsetInFile;
 				OffsetInFile = INDEX_NONE;
 			}
-			else if (IsReferencingOldBulkData())
+			else if (IsReferencingByPackagePath())
 			{
 				FString PackageNameStr;
 				Ar << PackageNameStr;
@@ -911,7 +911,7 @@ FCompressedBuffer FVirtualizedUntypedBulkData::LoadFromPackageFile() const
 
 	// Open a reader to the file
 	TUniquePtr<FArchive> BulkArchive;
-	if (!IsReferencingOldBulkData() || PackageSegment != EPackageSegment::Header)
+	if (!IsReferencingByPackagePath() || PackageSegment != EPackageSegment::Header)
 	{
 		FOpenPackageResult Result = IPackageResourceManager::Get().OpenReadPackage(PackagePath, PackageSegment);
 		if (Result.Format == EPackageFormat::Binary)
@@ -954,7 +954,7 @@ FCompressedBuffer FVirtualizedUntypedBulkData::LoadFromPackageTrailer() const
 
 	// Open a reader to the file
 	TUniquePtr<FArchive> BulkArchive;
-	if (!IsReferencingOldBulkData() || PackageSegment != EPackageSegment::Header)
+	if (!IsReferencingByPackagePath() || PackageSegment != EPackageSegment::Header)
 	{
 		FOpenPackageResult Result = IPackageResourceManager::Get().OpenReadPackage(PackagePath, PackageSegment);
 		if (Result.Format == EPackageFormat::Binary)
@@ -1095,10 +1095,10 @@ bool FVirtualizedUntypedBulkData::SerializeData(FArchive& Ar, FCompressedBuffer&
 		Ar << InPayload;
 		return true;
 	}
-	else if(Ar.IsLoading() && !EnumHasAnyFlags(PayloadFlags, EFlags::ReferencesLegacyFile)) // Loading from virtualized bulkdata format
+	else if (Ar.IsLoading() && !IsReferencingOldBulkData(PayloadFlags)) // Loading from virtualized bulkdata format
 	{
 		Ar << InPayload;
-		return InPayload.IsNull();	
+		return InPayload.IsNull();
 	}
 	else if (Ar.IsLoading()) // Loading from old bulkdata format
 	{
@@ -1152,7 +1152,7 @@ void FVirtualizedUntypedBulkData::PushData(const FPackagePath& InPackagePath)
 		if (VirtualizationSystem.PushData(PayloadContentId, PayloadToPush, EStorageType::Local, InPackagePath.GetPackageName()))
 		{
 			EnumAddFlags(Flags, EFlags::IsVirtualized);
-			EnumRemoveFlags(Flags, EFlags::ReferencesLegacyFile | EFlags::LegacyFileIsCompressed);
+			EnumRemoveFlags(Flags, EFlags::ReferencesLegacyFile | EFlags::ReferencesWorkspaceDomain | EFlags::LegacyFileIsCompressed);
 			check(!EnumHasAnyFlags(Flags, EFlags::LegacyKeyWasGuidDerived)); // Removed by UpdateKeyIfNeeded
 
 			// Clear members associated with non-virtualized data and release the in-memory
@@ -1265,7 +1265,7 @@ void FVirtualizedUntypedBulkData::DetachFromDisk(FArchive* Ar, bool bEnsurePaylo
 		PackageSegment = EPackageSegment::Header;
 		OffsetInFile = INDEX_NONE;
 
-		EnumRemoveFlags(Flags, EFlags::ReferencesLegacyFile | EFlags::LegacyFileIsCompressed);
+		EnumRemoveFlags(Flags, EFlags::ReferencesLegacyFile | EFlags::ReferencesWorkspaceDomain | EFlags::LegacyFileIsCompressed);
 	}
 
 	AttachedAr = nullptr;	
@@ -1417,6 +1417,7 @@ void FVirtualizedUntypedBulkData::UpdatePayloadImpl(FSharedBuffer&& InPayload, F
 
 	EnumRemoveFlags(Flags,	EFlags::IsVirtualized |
 							EFlags::ReferencesLegacyFile |
+							EFlags::ReferencesWorkspaceDomain |
 							EFlags::LegacyFileIsCompressed |
 							EFlags::LegacyKeyWasGuidDerived);
 
@@ -1651,7 +1652,7 @@ void FVirtualizedUntypedBulkData::RecompressForSerialization(FCompressedBuffer&
 	}
 }
 
-FVirtualizedUntypedBulkData::EFlags FVirtualizedUntypedBulkData::BuildFlagsForSerialization(FArchive& Ar, bool bUpgradeLegacyData) const
+FVirtualizedUntypedBulkData::EFlags FVirtualizedUntypedBulkData::BuildFlagsForSerialization(FArchive& Ar, bool bKeepFileDataByReference) const
 {
 	if (Ar.IsSaving())
 	{
@@ -1661,26 +1662,37 @@ FVirtualizedUntypedBulkData::EFlags FVirtualizedUntypedBulkData::BuildFlagsForSe
 		// Note that these changes are not applied to the current object UNLESS we are saving
 		// the package, in which case the newly modified flags will be applied once we confirm
 		// that the package has saved.
-		
-		const FLinkerSave* LinkerSave = Cast<FLinkerSave>(Ar.GetLinker());
-		if (LinkerSave != nullptr && !LinkerSave->GetFilename().IsEmpty() && ShouldSaveToPackageSidecar())
+
+		bool bIsReferencingByPackagePath = IsReferencingByPackagePath(UpdatedFlags);
+		bool bCanKeepFileDataByReference = bIsReferencingByPackagePath || !PackagePath.IsEmpty();
+		if (bKeepFileDataByReference && bCanKeepFileDataByReference)
 		{
-			EnumAddFlags(UpdatedFlags, EFlags::HasPayloadSidecarFile);
+			if (!bIsReferencingByPackagePath)
+			{
+				EnumAddFlags(UpdatedFlags, EFlags::ReferencesWorkspaceDomain);
+			}
+			EnumRemoveFlags(UpdatedFlags, EFlags::HasPayloadSidecarFile | EFlags::IsVirtualized);
 		}
 		else
 		{
-			EnumRemoveFlags(UpdatedFlags, EFlags::HasPayloadSidecarFile);
-		}
-
-		if (bUpgradeLegacyData)
-		{
-			EnumRemoveFlags(UpdatedFlags, EFlags::ReferencesLegacyFile | EFlags::LegacyFileIsCompressed | EFlags::LegacyKeyWasGuidDerived);
-		}
+			EnumRemoveFlags(UpdatedFlags, EFlags::ReferencesLegacyFile | EFlags::ReferencesWorkspaceDomain | 
+				EFlags::LegacyFileIsCompressed | EFlags::LegacyKeyWasGuidDerived);
 
-		// If we are re-hydrating packages on save then we need to remove the virtualization flag
-		if (LinkerSave != nullptr && bAllowVirtualizationOnSave == false)
-		{
-			EnumRemoveFlags(UpdatedFlags, EFlags::IsVirtualized);
+			const FLinkerSave* LinkerSave = Cast<FLinkerSave>(Ar.GetLinker());
+			if (LinkerSave != nullptr && !LinkerSave->GetFilename().IsEmpty() && ShouldSaveToPackageSidecar())
+			{
+				EnumAddFlags(UpdatedFlags, EFlags::HasPayloadSidecarFile);
+				EnumRemoveFlags(UpdatedFlags, EFlags::IsVirtualized);
+			}
+			else
+			{
+				// If we are re-hydrating packages on save then we need to remove the virtualization flag
+				EnumRemoveFlags(UpdatedFlags, EFlags::HasPayloadSidecarFile);
+				if (LinkerSave != nullptr && bAllowVirtualizationOnSave == false)
+				{
+					EnumRemoveFlags(UpdatedFlags, EFlags::IsVirtualized);
+				}
+			}
 		}
 
 		return UpdatedFlags;
diff --git a/Engine/Source/Runtime/CoreUObject/Public/Serialization/VirtualizedBulkData.h b/Engine/Source/Runtime/CoreUObject/Public/Serialization/VirtualizedBulkData.h
--- a/Engine/Source/Runtime/CoreUObject/Public/Serialization/VirtualizedBulkData.h
+++ b/Engine/Source/Runtime/CoreUObject/Public/Serialization/VirtualizedBulkData.h
@@ -297,6 +297,14 @@ public:
 	/** Load the payload and set the correct payload id, if the bulkdata has a PlaceholderPayloadId. */
 	void UpdatePayloadId();
 
+	/**
+	 * Register this BulkData with the BulkData registry, if it is valid for registration.
+	 * This function is called automatically when necessary by VirtualizedBulkData internals, 
+	 * but external callers can call it to provide information about the owner if it was previously
+	 * registered anonymously by UpdatePayload.
+	 */
+	void Register(UObject* Owner);
+
 #if UE_ENABLE_VIRTUALIZATION_TOGGLE
 	UE_DEPRECATED(5.0, "SetVirtualizationOptOut is an internal feature for development and will be removed without warning!")
 	void SetVirtualizationOptOut(bool bOptOut);
@@ -331,6 +339,8 @@ private:
 		HasRegistered				= 1 << 6,
 		/** The BulkData object is a copy used only to represent the id and payload; it does not communicate with the BulkDataRegistry, and will point DDC jobs toward the original BulkData */
 		IsTornOff					= 1 << 7,
+		/** The bulkdata object references a pyaload stored in a WorkspaceDomain file  */
+		ReferencesWorkspaceDomain	= 1 << 8,
 
 		TransientFlags				= HasRegistered | IsTornOff,
 	};
@@ -373,7 +383,7 @@ private:
 	void UpdateKeyIfNeeded();
 
 	void RecompressForSerialization(FCompressedBuffer& InOutPayload, EFlags PayloadFlags) const;
-	EFlags BuildFlagsForSerialization(FArchive& Ar, bool bUpgradeLegacyData) const;
+	EFlags BuildFlagsForSerialization(FArchive& Ar, bool bKeepFileDataByReference) const;
 
 	bool IsDataVirtualized() const 
 	{ 
@@ -385,12 +395,33 @@ private:
 		return EnumHasAnyFlags(Flags, EFlags::HasPayloadSidecarFile); 
 	}
 
-	bool IsReferencingOldBulkData() const 
+	bool IsReferencingOldBulkData() const
 	{ 
-		return EnumHasAnyFlags(Flags, EFlags::ReferencesLegacyFile); 
+		return IsReferencingOldBulkData(Flags);
+	}
+	static bool IsReferencingOldBulkData(EFlags InFlags)
+	{
+		return EnumHasAnyFlags(InFlags, EFlags::ReferencesLegacyFile);
+	}
+
+	bool IsReferencingWorkspaceDomain() const
+	{
+		return IsReferencingWorkspaceDomain(Flags);
+	}
+	static bool IsReferencingWorkspaceDomain(EFlags InFlags)
+	{
+		return EnumHasAnyFlags(InFlags, EFlags::ReferencesWorkspaceDomain);
+	}
+
+	bool IsReferencingByPackagePath() const
+	{
+		return IsReferencingByPackagePath(Flags);
+	}
+	static bool IsReferencingByPackagePath(EFlags InFlags)
+	{
+		return EnumHasAnyFlags(InFlags, EFlags::ReferencesLegacyFile | EFlags::ReferencesWorkspaceDomain);
 	}
 
-	void Register(UObject* Owner);
 	void Unregister();
 
 	/** Utility to return an apt error message if a payload loaded from this object is considered corrupt. Handles the formatting depending on the various payload flags etc. */
@@ -455,6 +486,10 @@ class TVirtualizedBulkData final : public FVirtualizedUntypedBulkData
 public:
 	TVirtualizedBulkData() = default;
 	~TVirtualizedBulkData() = default;
+	TVirtualizedBulkData(const TVirtualizedBulkData& Other) = default;
+	TVirtualizedBulkData(TVirtualizedBulkData&& Other) = default;
+	TVirtualizedBulkData& operator=(const TVirtualizedBulkData& Other) = default;
+	TVirtualizedBulkData& operator=(TVirtualizedBulkData&& Other) = default;
 
 	TVirtualizedBulkData<DataType> CopyTornOff() const
 	{
