Fix a Linux Development build link error when calling TopInitializedChecked from outside CoreUObject.

diff --git a/Engine/Source/Runtime/CoreUObject/Public/UObject/UObjectThreadContext.h b/Engine/Source/Runtime/CoreUObject/Public/UObject/UObjectThreadContext.h
--- a/Engine/Source/Runtime/CoreUObject/Public/UObject/UObjectThreadContext.h
+++ b/Engine/Source/Runtime/CoreUObject/Public/UObject/UObjectThreadContext.h
@@ -16,7 +16,7 @@ struct FUObjectSerializeContext;
 class FLinkerLoad;
 class IAsyncPackageLoader;
 
-DECLARE_LOG_CATEGORY_EXTERN(LogUObjectThreadContext, Log, All);
+COREUOBJECT_API DECLARE_LOG_CATEGORY_EXTERN(LogUObjectThreadContext, Log, All);
 
 class COREUOBJECT_API FUObjectThreadContext : public TThreadSingleton<FUObjectThreadContext>
 {
