Locking the entire hash maps instead of individual hash buckets when iterating UObject maps to prevent memory stomps when creating / renaming / destroying objects mid iteration

diff --git a/Engine/Source/Editor/UnrealEd/Private/PackageTools.cpp b/Engine/Source/Editor/UnrealEd/Private/PackageTools.cpp
--- a/Engine/Source/Editor/UnrealEd/Private/PackageTools.cpp
+++ b/Engine/Source/Editor/UnrealEd/Private/PackageTools.cpp
@@ -392,15 +392,19 @@ UPackageTools::UPackageTools(const FObjectInitializer& ObjectInitializer)
 				// Flush all pending render commands, as unloading the package may invalidate render resources.
 				FlushRenderingCommands();
 
+				TArray<UObject*> ObjectsInPackage;
+
+				// Can't use ForEachObjectWithPackage here as closing the editor may modify UObject hash tables (known case: renaming objects)
+				GetObjectsWithPackage(PackageBeingUnloaded, ObjectsInPackage, false);
 				// Close any open asset editors.
-				ForEachObjectWithPackage(PackageBeingUnloaded, [](UObject* Obj)
+				for (UObject* Obj : ObjectsInPackage)
 				{
 					if (Obj->IsAsset())
 					{
 						GEditor->GetEditorSubsystem<UAssetEditorSubsystem>()->CloseAllEditorsForAsset(Obj);
 					}
-					return true;
-				}, false);
+				}
+				ObjectsInPackage.Reset();
 
 				PackageBeingUnloaded->bHasBeenFullyLoaded = false;
 				PackageBeingUnloaded->ClearFlags(RF_WasLoaded);
@@ -409,8 +413,9 @@ UPackageTools::UPackageTools(const FObjectInitializer& ObjectInitializer)
 					bScriptPackageWasUnloaded = true;
 				}
 
+				GetObjectsWithPackage(PackageBeingUnloaded, ObjectsInPackage, true, RF_Transient, EInternalObjectFlags::PendingKill);
 				// Notify any Blueprints that are about to be unloaded, and destroy any leftover worlds.
-				ForEachObjectWithPackage(PackageBeingUnloaded, [](UObject* Obj)
+				for (UObject* Obj : ObjectsInPackage)
 				{
 					if (UBlueprint* BP = Cast<UBlueprint>(Obj))
 					{
@@ -423,12 +428,11 @@ UPackageTools::UPackageTools(const FObjectInitializer& ObjectInitializer)
 							World->CleanupWorld();
 						}
 					}
-					return true;
-				}, true, RF_Transient, EInternalObjectFlags::PendingKill);
+				}
+				ObjectsInPackage.Reset();
 
 				// Clear RF_Standalone flag from objects in the package to be unloaded so they get GC'd.
 				{
-					TArray<UObject*> ObjectsInPackage;
 					GetObjectsWithPackage(PackageBeingUnloaded, ObjectsInPackage);
 					for ( UObject* Object : ObjectsInPackage )
 					{
@@ -438,6 +442,7 @@ UPackageTools::UPackageTools(const FObjectInitializer& ObjectInitializer)
 							ObjectsThatHadFlagsCleared.Add(Object);
 						}
 					}
+					ObjectsInPackage.Reset();
 				}
 
 				// Reset loaders
diff --git a/Engine/Source/Runtime/CoreUObject/Private/UObject/UObjectHash.cpp b/Engine/Source/Runtime/CoreUObject/Private/UObject/UObjectHash.cpp
--- a/Engine/Source/Runtime/CoreUObject/Private/UObject/UObjectHash.cpp
+++ b/Engine/Source/Runtime/CoreUObject/Private/UObject/UObjectHash.cpp
@@ -53,22 +53,6 @@
 	*/
 	void *ElementsOrSetPtr[2];
 
-#if !UE_BUILD_SHIPPING
-	/** If true this bucket is being iterated over and no Add or Remove operations are allowed */
-	int32 ReadOnlyLock;
-
-	FORCEINLINE void Lock()
-	{
-		ReadOnlyLock++;
-	}
-
-	FORCEINLINE void Unlock()
-	{
-		ReadOnlyLock--;
-		check(ReadOnlyLock >= 0);
-	}
-#endif // !UE_BUILD_SHIPPING
-
 	FORCEINLINE TSet<UObjectBase*>* GetSet()
 	{
 		if (ElementsOrSetPtr[1] && !ElementsOrSetPtr[0])
@@ -92,9 +76,6 @@
 	{
 		ElementsOrSetPtr[0] = nullptr;
 		ElementsOrSetPtr[1] = nullptr;
-#if !UE_BUILD_SHIPPING
-		ReadOnlyLock = 0;
-#endif // !UE_BUILD_SHIPPING
 	}
 	FORCEINLINE ~FHashBucket()
 	{
@@ -103,11 +84,6 @@
 	/** Adds an Object to the bucket */
 	FORCEINLINE void Add(UObjectBase* Object)
 	{
-#if !UE_BUILD_SHIPPING
-		UE_CLOG(ReadOnlyLock != 0, LogObj, Fatal, TEXT("Trying to add %s to a hash bucket that is currently being iterated over which is not allowed and may lead to undefined behavior!"),
-			*static_cast<UObject*>(Object)->GetFullName());
-#endif // !UE_BUILD_SHIPPING
-
 		TSet<UObjectBase*>* Items = GetSet();
 		if (Items)
 		{
@@ -135,11 +111,6 @@
 	/** Removes an Object from the bucket */
 	FORCEINLINE int32 Remove(UObjectBase* Object)
 	{
-#if !UE_BUILD_SHIPPING
-		UE_CLOG(ReadOnlyLock != 0, LogObj, Fatal, TEXT("Trying to remove %s from a hash bucket that is currently being iterated over which is not allowed and may lead to undefined behavior!"),
-			*static_cast<UObject*>(Object)->GetFullName());
-#endif // !UE_BUILD_SHIPPING
-
 		int32 Result = 0;
 		TSet<UObjectBase*>* Items = GetSet();
 		if (Items)
@@ -277,6 +248,98 @@
 	}
 };
 
+/**
+* Wrapper around a TMap with FHashBucket values that supports read only locks
+*/
+template <typename T>
+class TBucketMap : private TMap<T, FHashBucket>
+{
+	typedef TMap<T, FHashBucket> Super;
+#if !UE_BUILD_SHIPPING
+	int32 ReadOnlyLock = 0;
+#endif // !UE_BUILD_SHIPPING
+public:
+
+	using Super::begin;
+	using Super::end;
+	using Super::GetAllocatedSize;
+	using Super::Find;
+	using Super::Num;
+	using Super::FindChecked;
+
+	FORCEINLINE void LockReadOnly()
+	{
+#if !UE_BUILD_SHIPPING
+		ReadOnlyLock++;
+#endif
+	}
+	FORCEINLINE void UnlockReadOnly()
+	{
+#if !UE_BUILD_SHIPPING
+		ReadOnlyLock--;
+		check(ReadOnlyLock >= 0);
+#endif
+	}
+
+	FORCEINLINE void Compact()
+	{
+#if !UE_BUILD_SHIPPING
+		UE_CLOG(ReadOnlyLock != 0, LogObj, Fatal, TEXT("Trying to modify UObject map (Compact) that is currently being iterated. Please make sure you're not creating new UObjects or Garbage Collecting while iterating UObject hash tables."));
+#endif
+		Super::Compact();
+	}
+
+	FORCEINLINE void Add(const T& Key)
+	{
+#if !UE_BUILD_SHIPPING
+		UE_CLOG(ReadOnlyLock != 0, LogObj, Fatal, TEXT("Trying to modify UObject map (Add) that is currently being iterated. Please make sure you're not creating new UObjects or Garbage Collecting while iterating UObject hash tables."));
+#endif
+		Super::Add(Key);
+	}
+
+	FORCEINLINE void Remove(const T& Key)
+	{
+#if !UE_BUILD_SHIPPING
+		UE_CLOG(ReadOnlyLock != 0, LogObj, Fatal, TEXT("Trying to modify UObject map (Remove) that is currently being iterated. Please make sure you're not creating new UObjects or Garbage Collecting while iterating UObject hash tables."));
+#endif
+		Super::Remove(Key);
+	}
+
+	FORCEINLINE FHashBucket& FindOrAdd(const T& Key)
+	{
+#if !UE_BUILD_SHIPPING
+		UE_CLOG(ReadOnlyLock != 0, LogObj, Fatal, TEXT("Trying to modify UObject map (FindOrAdd) that is currently being iterated. Please make sure you're not creating new UObjects or Garbage Collecting while iterating UObject hash tables."));
+#endif
+		return Super::FindOrAdd(Key);
+	}
+
+	FORCEINLINE FHashBucket& FindOrAdd(T&& Key)
+	{
+#if !UE_BUILD_SHIPPING
+		UE_CLOG(ReadOnlyLock != 0, LogObj, Fatal, TEXT("Trying to modify UObject map (FindOrAdd) that is currently being iterated. Please make sure you're not creating new UObjects or Garbage Collecting while iterating UObject hash tables."));
+#endif
+		return Super::FindOrAdd(Key);
+	}
+};
+
+/**
+ * Simple TBucketMap read only lock scope
+ */
+template <typename T>
+struct TBucketMapLock
+{
+	T& Map;
+	explicit FORCEINLINE TBucketMapLock(T& InMap)
+		: Map(InMap)
+	{
+		Map.LockReadOnly();
+	}
+	FORCEINLINE ~TBucketMapLock()
+	{
+		Map.UnlockReadOnly();
+	}
+};
+
 class FUObjectHashTables
 {
 	/** Critical section that guards against concurrent adds from multiple threads */
@@ -285,17 +348,17 @@
 public:
 
 	/** Hash sets */
-	TMap<int32, FHashBucket> Hash;
+	TBucketMap<int32> Hash;
 	TMultiMap<int32, uint32> HashOuter;
 
 	/** Map of object to their outers, used to avoid an object iterator to find such things. **/
-	TMap<UObjectBase*, FHashBucket> ObjectOuterMap;
-	TMap<UClass*, FHashBucket > ClassToObjectListMap;
+	TBucketMap<UObjectBase*> ObjectOuterMap;
+	TBucketMap<UClass*> ClassToObjectListMap;
 	TMap<UClass*, TSet<UClass*> > ClassToChildListMap;
 	TAtomic<uint64> ClassToChildListMapVersion;
 
 	/** Map of package to the object their contain. */
-	TMap<UPackage*, FHashBucket> PackageToObjectListMap;
+	TBucketMap<UPackage*> PackageToObjectListMap;
 	/** Map of object to their external package. */
 	TMap<UObjectBase*, UPackage*> ObjectToPackageMap;
 
@@ -602,9 +665,6 @@
 	UObject* Result = nullptr;
 	if (FHashBucket* Inners = ThreadHash.PackageToObjectListMap.Find(ObjectPackage))
 	{
-#if !UE_BUILD_SHIPPING
-		Inners->Lock();
-#endif // !UE_BUILD_SHIPPING
 		for (FHashBucketIterator It(*Inners); It; ++It)
 		{
 			UObject* Object = static_cast<UObject*>(*It);
@@ -629,9 +689,6 @@
 				break;
 			}
 		}
-#if !UE_BUILD_SHIPPING
-		Inners->Unlock();
-#endif // !UE_BUILD_SHIPPING
 	}
 	return Result;
 }
@@ -990,6 +1047,7 @@
 	FHashTableLock HashLock(ThreadHash);
 	TArray<FHashBucket*, TInlineAllocator<1> > AllInners;
 
+	TBucketMapLock MapLock(ThreadHash.ObjectOuterMap);
 	if (FHashBucket* Inners = ThreadHash.ObjectOuterMap.Find(Outer))
 	{
 		AllInners.Add(Inners);
@@ -997,9 +1055,6 @@
 	while (AllInners.Num())
 	{
 		FHashBucket* Inners = AllInners.Pop();
-#if !UE_BUILD_SHIPPING
-		Inners->Lock();
-#endif // !UE_BUILD_SHIPPING
 		for (FHashBucketIterator It(*Inners); It; ++It)
 		{
 			UObject *Object = static_cast<UObject*>(*It);
@@ -1019,9 +1074,6 @@
 				}
 			}
 		}
-#if !UE_BUILD_SHIPPING
-		Inners->Unlock();
-#endif // !UE_BUILD_SHIPPING
 	}
 }
 
@@ -1088,6 +1140,9 @@
 	FUObjectHashTables& ThreadHash = FUObjectHashTables::Get();
 	FHashTableLock HashLock(ThreadHash);
 	TArray<FHashBucket*, TInlineAllocator<1> > AllInners;
+	
+	TBucketMapLock PackageToObjectListMapLock(ThreadHash.PackageToObjectListMap);
+	TBucketMapLock ObjectOuterMapLock(ThreadHash.ObjectOuterMap);
 
 	// Add the object bucket that have this package as an external package
 	if (FHashBucket* Inners = ThreadHash.PackageToObjectListMap.Find(Package))
@@ -1102,9 +1157,6 @@
 	while (AllInners.Num())
 	{
 		FHashBucket* Inners = AllInners.Pop();
-#if !UE_BUILD_SHIPPING
-		Inners->Lock();
-#endif // !UE_BUILD_SHIPPING
 		for (FHashBucketIterator It(*Inners); It; ++It)
 		{
 			UObject *Object = static_cast<UObject*>(*It);
@@ -1130,9 +1182,6 @@
 				}
 			}
 		}
-#if !UE_BUILD_SHIPPING
-		Inners->Unlock();
-#endif // !UE_BUILD_SHIPPING
 	}
 }
 
@@ -1194,6 +1243,8 @@
 		ExclusionInternalFlags |= EInternalObjectFlags::AsyncLoading;
 	}
 
+	TBucketMapLock ClassToObjectListMapLock(ThreadHash.ClassToObjectListMap);
+
 	for (const UClass* SearchClass : ClassesToLookFor)
 	{
 		FHashBucket* List = ThreadHash.ClassToObjectListMap.Find(SearchClass);
@@ -1504,7 +1555,7 @@
 	Ar.Logf(TEXT("Total memory allocated for Object Outer Hash: %u bytes."), HashtableAllocatedSize);
 }
 
-void LogHashStatisticsInternal(TMap<int32, FHashBucket>& Hash, FOutputDevice& Ar, const bool bShowHashBucketCollisionInfo)
+void LogHashStatisticsInternal(TBucketMap<int32>& Hash, FOutputDevice& Ar, const bool bShowHashBucketCollisionInfo)
 {
 	TArray<int32> HashBuckets;
 	// Get the set of keys in use, which is the number of hash buckets
diff --git a/Engine/Source/Runtime/Engine/Private/InheritableComponentHandler.cpp b/Engine/Source/Runtime/Engine/Private/InheritableComponentHandler.cpp
--- a/Engine/Source/Runtime/Engine/Private/InheritableComponentHandler.cpp
+++ b/Engine/Source/Runtime/Engine/Private/InheritableComponentHandler.cpp
@@ -444,7 +444,10 @@ void UInheritableComponentHandler::PreloadAllTemplates()
 			{
 				Linker->Preload(Record.ComponentTemplate);
 
-				ForEachObjectWithOuter(Record.ComponentTemplate, [](UObject* SubObj)
+				TArray<UObject*> ComponentTemplateSubobjects;
+				// Can't use ForEachObjectWithOuter here as Preloading may modify UObject hash tables (it will most likely create new objects)
+				GetObjectsWithOuter(Record.ComponentTemplate, ComponentTemplateSubobjects);
+				for (UObject* SubObj : ComponentTemplateSubobjects)
 				{
 					if (SubObj->HasAllFlags(RF_NeedLoad))
 					{
@@ -453,7 +456,7 @@ void UInheritableComponentHandler::PreloadAllTemplates()
 							SubObjLinker->Preload(SubObj);
 						}
 					}
-				});
+				}
 			}
 		}
 	}
