Fixing potential memory stomps during localization data gathering caused by modifying UObject hash maps while iterating them with ForEachObjectWith* functions.

diff --git a/Engine/Source/Runtime/CoreUObject/Private/Serialization/PropertyLocalizationDataGathering.cpp b/Engine/Source/Runtime/CoreUObject/Private/Serialization/PropertyLocalizationDataGathering.cpp
--- a/Engine/Source/Runtime/CoreUObject/Private/Serialization/PropertyLocalizationDataGathering.cpp
+++ b/Engine/Source/Runtime/CoreUObject/Private/Serialization/PropertyLocalizationDataGathering.cpp
@@ -30,11 +30,13 @@ FPropertyLocalizationDataGatherer::FPropertyLocalizationDataGatherer(TArray<FGat
 	}, true, RF_Transient, EInternalObjectFlags::PendingKill);
 
 	// Iterate over each root object in the package
-	ForEachObjectWithPackage(Package, [this](UObject* Object)
+	for (const UObject* Object : AllObjectsInPackage)
 	{
-		GatherLocalizationDataFromObjectWithCallbacks(Object, EPropertyLocalizationGathererTextFlags::None);
-		return true;
-	}, false, RF_Transient, EInternalObjectFlags::PendingKill);
+		if (Object->GetOuter() == Package)
+		{
+			GatherLocalizationDataFromObjectWithCallbacks(Object, EPropertyLocalizationGathererTextFlags::None);
+		}
+	}
 
 	// Iterate any bytecode containing objects
 	for (const FObjectAndGatherFlags& BytecodeToGather : BytecodePendingGather)
@@ -232,14 +234,16 @@ void FPropertyLocalizationDataGatherer::GatherLocalizationDataFromObject(const U
 	// Gather from anything that has us as their outer, as not all objects are reachable via a property pointer.
 	if (!(GatherTextFlags & EPropertyLocalizationGathererTextFlags::SkipSubObjects))
 	{
-		ForEachObjectWithOuter(Object, [this, GatherTextFlags](UObject* ChildObject)
+		TArray<UObject*> InnerObjects;
+		GetObjectsWithOuter(Object, InnerObjects, false, RF_Transient, EInternalObjectFlags::PendingKill);
+		for (UObject* ChildObject : InnerObjects)
 		{
 			// if the child object as a package set, do not gather from it
 			if (!ChildObject->GetExternalPackage())
 			{
 				GatherLocalizationDataFromObjectWithCallbacks(ChildObject, GatherTextFlags);
 			}
-		}, false, RF_Transient, EInternalObjectFlags::PendingKill);
+		}
 	}
 }
 
