Updating comments for ForEachObjectWith* functions to warn agains prohibited operations within their callbacks

diff --git a/Engine/Source/Runtime/CoreUObject/Public/UObject/UObjectHash.h b/Engine/Source/Runtime/CoreUObject/Public/UObject/UObjectHash.h
--- a/Engine/Source/Runtime/CoreUObject/Public/UObject/UObjectHash.h
+++ b/Engine/Source/Runtime/CoreUObject/Public/UObject/UObjectHash.h
@@ -56,6 +56,7 @@ COREUOBJECT_API void GetObjectsWithOuter(const class UObjectBase* Outer, TArray<
 
 /**
  * Performs an operation on objects with a given outer, with the possibility to break iteration
+ * Note that the operation must not modify the UObject hash maps so it can not create, rename or destroy UObjects.
  *
  * @param	Outer						Outer to search for
  * @param	Operation					Function to be called for each object, returning if we want to continue iteration or not
@@ -67,6 +68,7 @@ COREUOBJECT_API void ForEachObjectWithOuterBreakable(const class UObjectBase* Ou
 
 /**
  * Performs an operation on all objects with a given outer
+ * Note that the operation must not modify UObject hash maps so it can not create, rename or destroy UObjects.
  *
  * @param	Outer						Outer to search for
  * @param	Operation					Function to be called for each object
@@ -101,6 +103,7 @@ COREUOBJECT_API void GetObjectsWithPackage(const class UPackage* Outer, TArray<U
 
 /**
  * Performs an operation on all objects found within a given package
+ * Note that the operation must not modify UObject hash maps so it can not create, rename or destroy UObjects.
  *
  * @param	Package						Package to iterate into
  * @param	Operation					Function to be called for each object, return false to break out of the iteration
@@ -123,6 +126,7 @@ COREUOBJECT_API void GetObjectsOfClass(const UClass* ClassToLookFor, TArray<UObj
 
 /**
  * Performs an operation on all objects of the provided class
+ * Note that the operation must not modify UObject hash maps so it can not create, rename or destroy UObjects.
  *
  * @param	Outer						UObject class to loop over instances of
  * @param	Operation					Function to be called for each object
@@ -133,6 +137,7 @@ COREUOBJECT_API void ForEachObjectOfClass(const UClass* ClassToLookFor, TFunctio
 
 /**
  * Performs an operation on all objects of the provided classes
+ * Note that the operation must not modify UObject hash maps so it can not create, rename or destroy UObjects.
  *
  * @param	Classes						UObject Classes to loop over instances of
  * @param	Operation					Function to be called for each object
