Claim last available EObjectFlag for new upcoming UObject allocator

diff --git a/Engine/Source/Runtime/CoreUObject/Public/UObject/ObjectMacros.h b/Engine/Source/Runtime/CoreUObject/Public/UObject/ObjectMacros.h
--- a/Engine/Source/Runtime/CoreUObject/Public/UObject/ObjectMacros.h
+++ b/Engine/Source/Runtime/CoreUObject/Public/UObject/ObjectMacros.h
@@ -535,6 +535,7 @@ enum EObjectFlags
 
 	RF_PendingKill UE_DEPRECATED(5.0, "RF_PendingKill should not be used directly. Make sure references to objects are released using one of the existing engine callbacks or use weak object pointers.") = 0x20000000,	///< Objects that are pending destruction (invalid for gameplay but valid objects). This flag is mirrored in EInternalObjectFlags as PendingKill for performance
 	RF_Garbage UE_DEPRECATED(5.0, "RF_Garbage should not be used directly. Use MarkAsGarbage and ClearGarbage instead.") =0x40000000,	///< Garbage from logical point of view and should not be referenced. This flag is mirrored in EInternalObjectFlags as Garbage for performance
+	RF_AllocatedInSharedPage	=0x80000000,	///< Allocated from a ref-counted page shared with other UObjects
 };
 
 PRAGMA_DISABLE_DEPRECATION_WARNINGS
@@ -544,7 +545,7 @@ constexpr EObjectFlags RF_InternalMirroredFlags = (EObjectFlags)(RF_PendingKill
 PRAGMA_ENABLE_DEPRECATION_WARNINGS
 
 /** Mask for all object flags */
-#define RF_AllFlags				(EObjectFlags)0x7fffffff	///< All flags, used mainly for error checking
+#define RF_AllFlags				(EObjectFlags)0xffffffff	///< All flags, used mainly for error checking
 
 /** Flags to load from unreal asset files */
 #define RF_Load						((EObjectFlags)(RF_Public | RF_Standalone | RF_Transactional | RF_ClassDefaultObject | RF_ArchetypeObject | RF_DefaultSubObject | RF_TextExportTransient | RF_InheritableComponentTemplate | RF_DuplicateTransient | RF_NonPIEDuplicateTransient)) 
