Avoid creating redundant sparse class data instances in cooked games

The process of cooking would ensure that every BP derived from a C++ base using sparse class data would have its own sparse class data instance, as the serialization function would create an instance on save/load. This could potentially result in a lot of redundant instances that don't override any data from their archetype.

These changes allow the API for accessing sparse class data to optionally fallback to using the archetype data when possible (controlled via the EGetSparseClassDataMethod enum), though any existing code already using GetOrCreateSparseClassData (or the UHT wrapper) will continue to always create a new instance.

BPGC serialization has been updated to only save the sparse class data when it has overridden data from its archetype, and the loading code will now avoid creating the instance when no override data has been saved. This allows new code to take advantage of EGetSparseClassDataMethod by choosing to reference the archetype data instead, and the engine has also been updated to avoid calling GetOrCreateSparseClassData where possible.

Anything that requires a mutable sparse class data instance should still use GetOrCreateSparseClassData (or the UHT wrapper that doesn't take EGetSparseClassDataMethod), as otherwise you may edit data belonging to your archetype (which is why GetSparseClassData returns const data).

diff --git a/Engine/Source/Editor/UnrealEd/Private/Kismet2/BlueprintEditorUtils.cpp b/Engine/Source/Editor/UnrealEd/Private/Kismet2/BlueprintEditorUtils.cpp
--- a/Engine/Source/Editor/UnrealEd/Private/Kismet2/BlueprintEditorUtils.cpp
+++ b/Engine/Source/Editor/UnrealEd/Private/Kismet2/BlueprintEditorUtils.cpp
@@ -1265,7 +1265,6 @@ void FBlueprintEditorUtils::RemoveStaleFunctions(UBlueprintGeneratedClass* Class
 		OrphanedClass->ClassAddReferencedObjects = Class->AddReferencedObjects;
 		OrphanedClass->ClassFlags |= CLASS_CompiledFromBlueprint;
 		OrphanedClass->ClassGeneratedBy = Class->ClassGeneratedBy;
-		OrphanedClass->SparseClassData = Class->GetOrCreateSparseClassData();
 
 		const ERenameFlags RenFlags = REN_DontCreateRedirectors | (Blueprint->bIsRegeneratingOnLoad ? REN_ForceNoResetLoaders : 0) | REN_NonTransactional | REN_DoNotDirty;
 
diff --git a/Engine/Source/Programs/UnrealHeaderTool/Private/CodeGenerator.cpp b/Engine/Source/Programs/UnrealHeaderTool/Private/CodeGenerator.cpp
--- a/Engine/Source/Programs/UnrealHeaderTool/Private/CodeGenerator.cpp
+++ b/Engine/Source/Programs/UnrealHeaderTool/Private/CodeGenerator.cpp
@@ -4787,6 +4787,11 @@ void FNativeClassHeaderGenerator::ExportNativeFunctions(FOutputDevice& OutGenera
 		RuntimeStringBuilders.AutogeneratedStaticData.Logf(TEXT("\treturn (F%s*)(GetClass()->GetOrCreateSparseClassData());\r\n"), *SparseClassDataString);
 		RuntimeStringBuilders.AutogeneratedStaticData += TEXT("}\r\n");
 
+		RuntimeStringBuilders.AutogeneratedStaticData.Logf(TEXT("const F%s* Get%s(EGetSparseClassDataMethod GetMethod) const\r\n"), *SparseClassDataString, *SparseClassDataString);
+		RuntimeStringBuilders.AutogeneratedStaticData += TEXT("{\r\n");
+		RuntimeStringBuilders.AutogeneratedStaticData.Logf(TEXT("\treturn (const F%s*)(GetClass()->GetSparseClassData(GetMethod));\r\n"), *SparseClassDataString);
+		RuntimeStringBuilders.AutogeneratedStaticData += TEXT("}\r\n");
+
 		FUnrealScriptStructDefinitionInfo* SparseClassDataStructDef = GTypeDefinitionInfoMap.FindByName<FUnrealScriptStructDefinitionInfo>(*SparseClassDataString);
 		while (SparseClassDataStructDef != nullptr)
 		{
diff --git a/Engine/Source/Runtime/CoreUObject/Private/Serialization/PropertyLocalizationDataGathering.cpp b/Engine/Source/Runtime/CoreUObject/Private/Serialization/PropertyLocalizationDataGathering.cpp
--- a/Engine/Source/Runtime/CoreUObject/Private/Serialization/PropertyLocalizationDataGathering.cpp
+++ b/Engine/Source/Runtime/CoreUObject/Private/Serialization/PropertyLocalizationDataGathering.cpp
@@ -207,8 +207,9 @@ void FPropertyLocalizationDataGatherer::GatherLocalizationDataFromObject(const U
 	// Also gather from the sparse data on UClass types.
 	if (const UClass* Class = Cast<UClass>(Object))
 	{
-		if (const UScriptStruct* SparseDataStruct = Class->GetSparseClassDataStruct())
+		if (Class->OverridesSparseClassDataArchetype())
 		{
+			const UScriptStruct* SparseDataStruct = Class->GetSparseClassDataStruct();
 			const void* SparseData = const_cast<UClass*>(Class)->GetOrCreateSparseClassData();
 			const void* ArchetypeSparseData = Class->GetArchetypeForSparseClassData();
 			GatherLocalizationDataFromStructWithCallbacks(Path + TEXT(".SparseClassData"), SparseDataStruct, SparseData, ArchetypeSparseData, GatherTextFlags);
diff --git a/Engine/Source/Runtime/CoreUObject/Private/UObject/Class.cpp b/Engine/Source/Runtime/CoreUObject/Private/UObject/Class.cpp
--- a/Engine/Source/Runtime/CoreUObject/Private/UObject/Class.cpp
+++ b/Engine/Source/Runtime/CoreUObject/Private/UObject/Class.cpp
@@ -4149,6 +4149,9 @@ void UClass::FinishDestroy()
 
 	ClassDefaultObject = nullptr;
 
+	CleanupSparseClassData();
+	SparseClassDataStruct = nullptr;
+
 #if WITH_EDITORONLY_DATA
 	// If for whatever reason there's still properties that have not been destroyed in PurgeClass, destroy them now
 	DestroyPropertiesPendingDestruction();
@@ -4206,15 +4209,13 @@ void UClass::GetPreloadDependencies(TArray<UObject*>& OutDeps)
 	{
 		OutDeps.Add(SparseClassDataStruct);
 
-		// make sure we always have a sparse class data struct to get the dependencies from
-		GetOrCreateSparseClassData();
-
+		const void* SparseClassDataToUse = GetSparseClassData(EGetSparseClassDataMethod::ArchetypeIfNull);
 		if (UScriptStruct::ICppStructOps* CppStructOps = SparseClassDataStruct->GetCppStructOps())
 		{
-			CppStructOps->GetPreloadDependencies(SparseClassData, OutDeps);
+			CppStructOps->GetPreloadDependencies(const_cast<void*>(SparseClassDataToUse), OutDeps);
 		}
 		// The iterator will recursively loop through all structs in structs/containers too.
-		for (TPropertyValueIterator<FStructProperty> It(SparseClassDataStruct, SparseClassData); It; ++It)
+		for (TPropertyValueIterator<FStructProperty> It(SparseClassDataStruct, SparseClassDataToUse); It; ++It)
 		{
 			const UScriptStruct* StructType = It.Key()->Struct;
 			if (UScriptStruct::ICppStructOps* CppStructOps = StructType->GetCppStructOps())
@@ -4741,14 +4742,9 @@ void UClass::Serialize( FArchive& Ar )
 
 	if (!Ar.IsLoading() && !Ar.IsSaving())
 	{
-		Ar.UsingCustomVersion(FUE5MainStreamObjectVersion::GUID);
-
-		if(Ar.CustomVer(FUE5MainStreamObjectVersion::GUID) >= FUE5MainStreamObjectVersion::SparseClassDataStructSerialization)
-		{
-			Ar << SparseClassDataStruct;
-		}
+		Ar << SparseClassDataStruct;
 
-		if (GetSparseClassDataStruct() != nullptr)
+		if (SparseClassDataStruct && SparseClassData)
 		{
 			SerializeSparseClassData(FStructuredArchiveFromArchive(Ar).GetSlot());
 		}
@@ -4858,11 +4854,11 @@ void UClass::SerializeSparseClassData(FStructuredArchive::FSlot Slot)
 	}
 	else if (UnderlyingArchive.GetPortFlags() != 0)
 	{
-		SparseClassDataStruct->SerializeBinEx(Slot, (uint8*)SparseClassData, SparseClassDataStruct, GetSparseClassDataArchetypeStruct());
+		SparseClassDataStruct->SerializeBinEx(Slot, SparseClassData, GetArchetypeForSparseClassData(), GetSparseClassDataArchetypeStruct());
 	}
 	else
 	{
-		SparseClassDataStruct->SerializeBin(Slot, (uint8*)SparseClassData);
+		SparseClassDataStruct->SerializeBin(Slot, SparseClassData);
 	}
 }
 
@@ -4876,10 +4872,10 @@ FArchive& operator<<(FArchive& Ar, FImplementedInterface& A)
 	return Ar;
 }
 
-void* UClass::GetArchetypeForSparseClassData() const
+const void* UClass::GetArchetypeForSparseClassData() const
 {
 	UClass* SuperClass = GetSuperClass();
-	return SuperClass ? SuperClass->GetOrCreateSparseClassData() : nullptr;
+	return SuperClass ? SuperClass->GetSparseClassData(EGetSparseClassDataMethod::ArchetypeIfNull) : nullptr;
 }
 
 UScriptStruct* UClass::GetSparseClassDataArchetypeStruct() const
@@ -4888,6 +4884,16 @@ UScriptStruct* UClass::GetSparseClassDataArchetypeStruct() const
 	return SuperClass ? SuperClass->GetSparseClassDataStruct() : nullptr;
 }
 
+bool UClass::OverridesSparseClassDataArchetype() const
+{
+	if (SparseClassDataStruct && SparseClassData)
+	{
+		return SparseClassDataStruct != GetSparseClassDataArchetypeStruct()
+			|| SparseClassDataStruct->CompareScriptStruct(SparseClassData, GetArchetypeForSparseClassData(), 0) == false;
+	}
+	return false;
+}
+
 UObject* UClass::GetArchetypeForCDO() const
 {
 	UClass* SuperClass = GetSuperClass();
@@ -5162,7 +5168,7 @@ void* UClass::CreateSparseClassData()
 	if (SparseClassData)
 	{
 		// initialize per class data from the archetype if we have one
-		void* SparseArchetypeData = GetArchetypeForSparseClassData();
+		const void* SparseArchetypeData = GetArchetypeForSparseClassData();
 		UStruct* SparseClassDataArchetypeStruct = GetSparseClassDataArchetypeStruct();
 
 		if (SparseArchetypeData)
@@ -5187,6 +5193,31 @@ void UClass::CleanupSparseClassData()
 	}
 }
 
+const void* UClass::GetSparseClassData(const EGetSparseClassDataMethod GetMethod)
+{
+	if (SparseClassData)
+	{
+		return SparseClassData;
+	}
+
+	switch (GetMethod)
+	{
+	case EGetSparseClassDataMethod::CreateIfNull:
+		return CreateSparseClassData();
+
+	case EGetSparseClassDataMethod::ArchetypeIfNull:
+		// Use the archetype data only when it's the expected type, otherwise the result may be cast to an incorrect type by the caller
+		return SparseClassDataStruct == GetSparseClassDataArchetypeStruct()
+			? GetArchetypeForSparseClassData()
+			: CreateSparseClassData();
+
+	default:
+		break;
+	}
+
+	return nullptr;
+}
+
 UScriptStruct* UClass::GetSparseClassDataStruct() const
 {
 	// this info is specified on the object via code generation so we use it instead of looking at the UClass
diff --git a/Engine/Source/Runtime/CoreUObject/Private/UObject/Obj.cpp b/Engine/Source/Runtime/CoreUObject/Private/UObject/Obj.cpp
--- a/Engine/Source/Runtime/CoreUObject/Private/UObject/Obj.cpp
+++ b/Engine/Source/Runtime/CoreUObject/Private/UObject/Obj.cpp
@@ -1464,9 +1464,19 @@ void UObject::Serialize(FStructuredArchive::FRecord Record)
 		}
 
 		// Keep track of sparse class data for undo/redo
-		if (UnderlyingArchive.IsTransacting() && HasAnyFlags(RF_ClassDefaultObject) && ObjClass->GetSparseClassDataStruct())
+		if (UnderlyingArchive.IsTransacting() && HasAnyFlags(RF_ClassDefaultObject))
 		{
-			ObjClass->SerializeSparseClassData(Record.EnterField(SA_FIELD_NAME(TEXT("SparseClassData"))));
+			UScriptStruct* SerializedSparseClassDataStruct = ObjClass->GetSparseClassDataStruct();
+			if (UnderlyingArchive.IsSaving() && !ObjClass->GetSparseClassData(EGetSparseClassDataMethod::ReturnIfNull))
+			{
+				SerializedSparseClassDataStruct = nullptr;
+			}
+			Record << SA_VALUE(TEXT("SparseClassDataStruct"), SerializedSparseClassDataStruct);
+
+			if (SerializedSparseClassDataStruct)
+			{
+				ObjClass->SerializeSparseClassData(Record.EnterField(SA_FIELD_NAME(TEXT("SparseClassData"))));
+			}
 		}
 
 		// Memory counting (with proper alignment to match C++)
@@ -1933,7 +1943,7 @@ static void GetAssetRegistryTagsFromSearchableProperties(const UObject* Object,
 	UScriptStruct* SparseClassDataStruct = Object->GetClass()->GetSparseClassDataStruct();
 	if (SparseClassDataStruct)
 	{
-		void* SparseClassData = Object->GetClass()->GetOrCreateSparseClassData();
+		const void* SparseClassData = Object->GetClass()->GetSparseClassData(EGetSparseClassDataMethod::ArchetypeIfNull);
 		for (TFieldIterator<FProperty> FieldIt(SparseClassDataStruct); FieldIt; ++FieldIt)
 		{
 			GetAssetRegistryTagFromProperty(SparseClassData, Object, CastField<FProperty>(*FieldIt), OutTags);
diff --git a/Engine/Source/Runtime/CoreUObject/Private/UObject/ScriptCore.cpp b/Engine/Source/Runtime/CoreUObject/Private/UObject/ScriptCore.cpp
--- a/Engine/Source/Runtime/CoreUObject/Private/UObject/ScriptCore.cpp
+++ b/Engine/Source/Runtime/CoreUObject/Private/UObject/ScriptCore.cpp
@@ -2104,13 +2104,9 @@ DEFINE_FUNCTION(UObject::execClassSparseDataVariable)
 	}
 	else
 	{
-		void* SparseDataBaseAddress = P_THIS->GetClass()->GetOrCreateSparseClassData();
+		void* SparseDataBaseAddress = const_cast<void*>(P_THIS->GetClass()->GetSparseClassData(EGetSparseClassDataMethod::ArchetypeIfNull));
 		Stack.MostRecentPropertyAddress = VarProperty->ContainerPtrToValuePtr<uint8>(SparseDataBaseAddress);
 
-		// SPARSEDATA_TODO: remove these two lines once we're sure the math is right
-		int32 Offset = VarProperty->GetOffset_ForInternal();
-		check((uint8*)SparseDataBaseAddress + Offset == Stack.MostRecentPropertyAddress);
-
 		if (RESULT_PARAM)
 		{
 			VarProperty->CopyCompleteValueToScriptVM(RESULT_PARAM, Stack.MostRecentPropertyAddress);
diff --git a/Engine/Source/Runtime/CoreUObject/Public/UObject/Class.h b/Engine/Source/Runtime/CoreUObject/Public/UObject/Class.h
--- a/Engine/Source/Runtime/CoreUObject/Public/UObject/Class.h
+++ b/Engine/Source/Runtime/CoreUObject/Public/UObject/Class.h
@@ -2723,6 +2723,17 @@ struct FClassFunctionLinkInfo
 };
 
 
+enum class EGetSparseClassDataMethod : uint8
+{
+	/** Create a new instance when this class doesn't have any sparse data of its own */
+	CreateIfNull,
+	/** Use the archetype instance (if possible) when this class doesn't have any sparse data of its own */
+	ArchetypeIfNull,
+	/** Return null when this class doesn't have any sparse data of its own */
+	ReturnIfNull,
+};
+
+
 /**
  * An object class.
  */
@@ -2828,23 +2839,27 @@ public:
 
 protected:
 	/** This is where we store the data that is only changed per class instead of per instance */
-	UPROPERTY()
 	void* SparseClassData;
 
 	/** The struct used to store sparse class data. */
-	UPROPERTY()
 	UScriptStruct* SparseClassDataStruct;
 
 public:
+	/**
+	 * Returns a pointer to the sidecar data structure, based on the EGetSparseClassDataMethod.
+	 * @note It is only safe to mutate this data when using "CreateIfNull", as others may return archetype/default data; consider GetOrCreateSparseClassData for this use-case.
+	 */
+	const void* GetSparseClassData(const EGetSparseClassDataMethod GetMethod);
+
 	/**
 	 * Returns a pointer to the sidecar data structure. This function will create an instance of the data structure if one has been specified and it has not yet been created.
 	 */
-	void* GetOrCreateSparseClassData() { return SparseClassData ? SparseClassData : CreateSparseClassData(); }
+	void* GetOrCreateSparseClassData() { return const_cast<void*>(GetSparseClassData(EGetSparseClassDataMethod::CreateIfNull)); }
 
 	/**
 	 * Returns a pointer to the type of the sidecar data structure if one is specified.
 	 */
-	virtual UScriptStruct* GetSparseClassDataStruct() const;
+	UScriptStruct* GetSparseClassDataStruct() const;
 
 	void SetSparseClassDataStruct(UScriptStruct* InSparseClassDataStruct);
 
@@ -2861,11 +2876,12 @@ public:
 	}
 #endif // WITH_EDITOR
 
-private:
+protected:
 	void* CreateSparseClassData();
 
 	void CleanupSparseClassData();
 
+private:
 #if WITH_EDITOR
 	/** Provides access to attributes of the underlying C++ class. Should never be unset. */
 	TOptional<FCppClassTypeInfo> CppTypeInfo;
@@ -3346,11 +3362,14 @@ public:
 	virtual UObject* GetArchetypeForCDO() const;
 
 	/** Returns archetype for sparse class data */
-	virtual void* GetArchetypeForSparseClassData() const;
+	const void* GetArchetypeForSparseClassData() const;
 
 	/** Returns the struct used by the sparse class data archetype */
 	UScriptStruct* GetSparseClassDataArchetypeStruct() const;
 
+	/** Returns whether the sparse class data on this instance overrides that of its archetype (in type or value) */
+	bool OverridesSparseClassDataArchetype() const;
+
 	/**
 	* Returns all objects that should be preloaded before the class default object is serialized at load time. Only used by the EDL.
 	*
diff --git a/Engine/Source/Runtime/CoreUObject/Public/UObject/Object.h b/Engine/Source/Runtime/CoreUObject/Public/UObject/Object.h
--- a/Engine/Source/Runtime/CoreUObject/Public/UObject/Object.h
+++ b/Engine/Source/Runtime/CoreUObject/Public/UObject/Object.h
@@ -627,7 +627,7 @@ public:
 	virtual FString GetDesc() { return TEXT( "" ); }
 
 	/** Return the UStruct corresponding to the sidecar data structure that stores data that is constant for all instances of this class. */
-	virtual UScriptStruct* GetSparseClassDataStruct() const;
+	UScriptStruct* GetSparseClassDataStruct() const;
 
 #if WITH_EDITOR
 	virtual void MoveDataToSparseClassDataStruct() const {}
diff --git a/Engine/Source/Runtime/Engine/Private/BlueprintGeneratedClass.cpp b/Engine/Source/Runtime/Engine/Private/BlueprintGeneratedClass.cpp
--- a/Engine/Source/Runtime/Engine/Private/BlueprintGeneratedClass.cpp
+++ b/Engine/Source/Runtime/Engine/Private/BlueprintGeneratedClass.cpp
@@ -482,15 +482,35 @@ void UBlueprintGeneratedClass::SerializeDefaultObject(UObject* Object, FStructur
 	bool bSkipSparseClassDataSerialization = false;
 	if(UnderlyingArchive.CustomVer(FUE5MainStreamObjectVersion::GUID) >= FUE5MainStreamObjectVersion::SparseClassDataStructSerialization)
 	{
-		UnderlyingArchive << SparseClassDataStruct;
-
-		if(UnderlyingArchive.IsLoading() && SparseClassDataStruct == nullptr && Object->GetLinker())
+		if (UnderlyingArchive.IsSaving())
 		{
-			// Missing or failed to load sparse class data struct - possible that the parent class was deleted, or regenerated on load
-			// so seek past where this CDO was serialized as we cannot read the serialized sparse class data any more.
-			// Note this happens in majority of cases with classes that have no sparse class data attached too. In that case 
-			// this 'skip' over the remaining part of the archive should have no ill effects as the seek should effectively be zero
-			bSkipSparseClassDataSerialization = true;
+			UScriptStruct* SerializedSparseClassDataStruct = SparseClassDataStruct;
+			if (!OverridesSparseClassDataArchetype())
+			{
+				// If this class doesn't override the sparse class data of its archetype, then we can skip saving it 
+				// since it can be lazily regenerated from the archetype data on load
+				SerializedSparseClassDataStruct = nullptr;
+				bSkipSparseClassDataSerialization = true;
+			}
+			UnderlyingArchive << SerializedSparseClassDataStruct;
+		}
+		else if (UnderlyingArchive.IsLoading())
+		{
+			UScriptStruct* SerializedSparseClassDataStruct = nullptr;
+			UnderlyingArchive << SerializedSparseClassDataStruct;
+			if (SparseClassDataStruct != SerializedSparseClassDataStruct)
+			{
+				CleanupSparseClassData();
+				SparseClassDataStruct = SerializedSparseClassDataStruct;
+			}
+			if (!SparseClassDataStruct)
+			{
+				// Missing or failed to load sparse class data struct - possible that the parent class was deleted, or regenerated on load
+				// so seek past where this CDO was serialized as we cannot read the serialized sparse class data any more.
+				// Note this happens in majority of cases with classes that have no sparse class data attached too. In that case 
+				// this 'skip' over the remaining part of the archive should have no ill effects as the seek should effectively be zero
+				bSkipSparseClassDataSerialization = true;
+			}
 		}
 	}
 
@@ -501,16 +521,34 @@ void UBlueprintGeneratedClass::SerializeDefaultObject(UObject* Object, FStructur
 		if(bSkipSparseClassDataSerialization)
 		{
 			// Seek to after sparse class data rather than load it.
-			const FLinkerLoad* Linker = Object->GetLinker();
-			const int32 LinkerIndex = Object->GetLinkerIndex();
-			const FObjectExport& Export = Linker->ExportMap[LinkerIndex];
-			UnderlyingArchive.Seek(Export.SerialOffset + Export.SerialSize);
+			if (UnderlyingArchive.IsLoading())
+			{
+				if (const FLinkerLoad* Linker = Object->GetLinker())
+				{
+					const int32 LinkerIndex = Object->GetLinkerIndex();
+					const FObjectExport& Export = Linker->ExportMap[LinkerIndex];
+					UnderlyingArchive.Seek(Export.SerialOffset + Export.SerialSize);
+				}
+			}
 		}
-		else if (Object->GetSparseClassDataStruct())
+		else if (SparseClassDataStruct)
 		{
 			SerializeSparseClassData(FStructuredArchiveFromArchive(UnderlyingArchive).GetSlot());
 		}
 	}
+
+	if (UnderlyingArchive.IsLoading())
+	{
+		if (SparseClassDataStruct)
+		{
+			// TODO: We should make sure that this instance conforms to its archetype (keeping data where possible, potentially via 
+			// an explicit upgrade path for sparse data), as it may have changed since this BPGC was saved (UE-127121)
+		}
+		else
+		{
+			SparseClassDataStruct = GetSparseClassDataArchetypeStruct();
+		}
+	}
 }
 
 void UBlueprintGeneratedClass::PostLoadDefaultObject(UObject* Object)
