Fix Illegal Reference in external map message
Remove bits of unused code

diff --git a/Engine/Source/Runtime/CoreUObject/Private/UObject/SavePackage2.cpp b/Engine/Source/Runtime/CoreUObject/Private/UObject/SavePackage2.cpp
--- a/Engine/Source/Runtime/CoreUObject/Private/UObject/SavePackage2.cpp
+++ b/Engine/Source/Runtime/CoreUObject/Private/UObject/SavePackage2.cpp
@@ -567,37 +567,9 @@ ESavePackageResult ValidateIllegalReferences(FSaveContext& SaveContext, TArray<U
 	if (ObjectsInOtherMaps.Num() > 0)
 	{
 		UObject* MostLikelyCulprit = nullptr;
-
-		// construct a string containing up to the first 5 objects problem objects
-		FString ObjectNames;
-		int32 MaxNamesToDisplay = 5;
-		bool DisplayIsLimited = true;
-
-		if (ObjectsInOtherMaps.Num() < MaxNamesToDisplay)
-		{
-			MaxNamesToDisplay = ObjectsInOtherMaps.Num();
-			DisplayIsLimited = false;
-		}
-
-		for (int32 Idx = 0; Idx < MaxNamesToDisplay; Idx++)
-		{
-			ObjectNames += ObjectsInOtherMaps[Idx]->GetName() + TEXT("\n");
-		}
-
-		// if there are more than 5 items we indicated this by adding "..." at the end of the list
-		if (DisplayIsLimited)
-		{
-			ObjectNames += TEXT("...\n");
-		}
-
-		Args.Empty();
-		Args.Add(TEXT("FileName"), FText::FromString(SaveContext.GetFilename()));
-		Args.Add(TEXT("ObjectNames"), FText::FromString(ObjectNames));
-		const FText Message = FText::Format(NSLOCTEXT("Core", "LinkedToObjectsInOtherMap_FindCulpritQ", "Can't save {FileName}: Graph is linked to object(s) in external map.\nExternal Object(s):\n{ObjectNames}  \nTry to find the chain of references to that object (may take some time)?"), Args);
-
 		FString CulpritString = TEXT("Unknown");
 		FString Referencer;
-		SavePackageUtilities::FindMostLikelyCulprit(ObjectsInOtherMaps, MostLikelyCulprit, Referencer);
+		SavePackageUtilities::FindMostLikelyCulprit(ObjectsInOtherMaps, MostLikelyCulprit, Referencer, &SaveContext);
 		if (MostLikelyCulprit != nullptr)
 		{
 			CulpritString = FString::Printf(TEXT("%s (%s)"), *MostLikelyCulprit->GetFullName(), *Referencer);
@@ -618,34 +590,6 @@ ESavePackageResult ValidateIllegalReferences(FSaveContext& SaveContext, TArray<U
 	if (PrivateObjects.Num() > 0)
 	{
 		UObject* MostLikelyCulprit = nullptr;
-
-		// construct a string containing up to the first 5 objects problem objects
-		FString ObjectNames;
-		int32 MaxNamesToDisplay = 5;
-		bool DisplayIsLimited = true;
-
-		if (PrivateObjects.Num() < MaxNamesToDisplay)
-		{
-			MaxNamesToDisplay = PrivateObjects.Num();
-			DisplayIsLimited = false;
-		}
-
-		for (int32 Idx = 0; Idx < MaxNamesToDisplay; Idx++)
-		{
-			ObjectNames += PrivateObjects[Idx]->GetName() + TEXT("\n");
-		}
-
-		// if there are more than 5 items we indicated this by adding "..." at the end of the list
-		if (DisplayIsLimited)
-		{
-			ObjectNames += TEXT("...\n");
-		}
-
-		Args.Empty();
-		Args.Add(TEXT("FileName"), FText::FromString(SaveContext.GetFilename()));
-		Args.Add(TEXT("ObjectNames"), FText::FromString(ObjectNames));
-		const FText Message = FText::Format(NSLOCTEXT("Core", "LinkedToPrivateObjectsInOtherPackage_FindCulpritQ", "Can't save {FileName}: Graph is linked to private object(s) in an external package.\nExternal Object(s):\n{ObjectNames}  \nTry to find the chain of references to that object (may take some time)?"), Args);
-
 		FString CulpritString = TEXT("Unknown");
 		FString Referencer;
 		SavePackageUtilities::FindMostLikelyCulprit(PrivateObjects, MostLikelyCulprit, Referencer, &SaveContext);
