A small clean up on how the package trailer is generated.

#rnx

### FPackageTrailer
- FPackageTrailerBuilder::Create now renamed to ::CreateFromTrailer to make it a bit clearer.
- Added 'AccessMode' to the file format documentation.
- FPackageTrailerBuilder no longer tracks infomation about the trailer once it is written to disk. Instead when writing to disk we will create a valid FPackageTrailer as part of the serialization. This will be passed to callbacks and can be queried by virtualized bulkdata.
-- Since the builder no longer tracks data after creating the trailer, we can remove checks from methods like ::AddPayload.
-- Removed FPackageTrailerBuilder::FindPayloadOffset and added a version to FPackageTrailer instead.
- Changed a number of methods to take FPayloadId by reference rather than by value.

### VirtualizedBulkData
- Replaced areas of code responsible for updating previously written out values with the utility function ::UpdateArchiveData to help with code readability.

### Misc
- FLinkerSave no longer creates a FPackageTrailerBuilder by default, instead one should be added to the FLinkerSave if we want to build a trailer for the package being saved.
- FLinkerLoad now considers 0 to be an invalid value for PayloadTocOffset in addition to INDEX_NONE, so we now check if the value is > than 0. This is because when FPackageFileSummary objects are created, all members are set to zero and we know that a FPackageTrailer cannot be 0 and the rest of the package still exist so we can safely consider it invalid.
- SavePackage/SavePackage2 now create the FPackageTrailerBuilder when the linker is created if the package trailer feature is enabled.
-- We now always call BuildAndWriteTrailer and then only skip the creation of the trailer if the builder is missing or conditions are not met.
-- If we do not write out the trailer, we make sure to set PayloadTocOffset in the FPackageFileSummary to INDEX_NONE as it is clearer than leaving it at the default 0.

diff --git a/Engine/Source/Developer/Virtualization/Private/PackageSubmissionChecks.cpp b/Engine/Source/Developer/Virtualization/Private/PackageSubmissionChecks.cpp
--- a/Engine/Source/Developer/Virtualization/Private/PackageSubmissionChecks.cpp
+++ b/Engine/Source/Developer/Virtualization/Private/PackageSubmissionChecks.cpp
@@ -386,7 +386,7 @@ void OnPrePackageSubmission(const TArray<FString>& FilesToSubmit, TArray<FText>&
 			return;
 		}
 
-		FPackageTrailerBuilder TrailerBuilder = FPackageTrailerBuilder::Create(PackageInfo.Trailer, *PackageAr, PackagePath);
+		FPackageTrailerBuilder TrailerBuilder = FPackageTrailerBuilder::CreateFromTrailer(PackageInfo.Trailer, *PackageAr, PackagePath);
 		if (!TrailerBuilder.BuildAndAppendTrailer(nullptr, *CopyAr))
 		{
 			FText Message = FText::Format(	LOCTEXT("Virtualization_TrailerAppend", "Failed to append the trailer to '{0}'"),
diff --git a/Engine/Source/Runtime/CoreUObject/Private/Serialization/VirtualizedBulkData.cpp b/Engine/Source/Runtime/CoreUObject/Private/Serialization/VirtualizedBulkData.cpp
--- a/Engine/Source/Runtime/CoreUObject/Private/Serialization/VirtualizedBulkData.cpp
+++ b/Engine/Source/Runtime/CoreUObject/Private/Serialization/VirtualizedBulkData.cpp
@@ -141,7 +141,7 @@ bool IsValid(const FVirtualizedUntypedBulkData& BulkData, const FCompressedBuffe
 }
 
 /** Utility for finding the FLinkerLoad associated with a given UObject */
-FLinkerLoad* GetLinkerLoadFromOwner(UObject* Owner)
+const FLinkerLoad* GetLinkerLoadFromOwner(UObject* Owner)
 {
 	if (Owner != nullptr)
 	{
@@ -156,9 +156,10 @@ FLinkerLoad* GetLinkerLoadFromOwner(UObject* Owner)
 	}
 }
 
-FPackageTrailer* GetTrailerFromOwner(UObject* Owner)
+/** Utility for finding the FPackageTrailer associated with a given UObject */
+const FPackageTrailer* GetTrailerFromOwner(UObject* Owner)
 {
-	FLinkerLoad* Linker = GetLinkerLoadFromOwner(Owner);
+	const FLinkerLoad* Linker = GetLinkerLoadFromOwner(Owner);
 	if (Linker != nullptr)
 	{
 		return Linker->GetPackageTrailer();
@@ -169,6 +170,18 @@ FPackageTrailer* GetTrailerFromOwner(UObject* Owner)
 	}
 }
 
+/** Utility for updating an existing entry in an Archive before returning the archive to it's original seek position */
+template<typename DataType>
+void UpdateArchiveData(FArchive& Ar, int64 DataPosition, DataType& Data)
+{
+	const int64 OriginalPosition = Ar.Tell();
+
+	Ar.Seek(DataPosition);
+	Ar << Data;
+
+	Ar.Seek(OriginalPosition);
+}
+
 /** Utility for accessing IVirtualizationSourceControlUtilities from the modular feature system. */
 Experimental::IVirtualizationSourceControlUtilities* GetSourceControlInterface()
 {
@@ -477,18 +490,6 @@ void FVirtualizedUntypedBulkData::CreateLegacyUniqueIdentifier(UObject* Owner)
 	}
 }
 
-/** Utility for updating an existing entry in an Archive before returning the archive to it's original seek position */
-template<typename DataType>
-void UpdateArchiveData(FArchive& Ar, int64 DataPosition, DataType& Data)
-{
-	int64 OriginalPosition = Ar.Tell();
-
-	Ar.Seek(DataPosition);
-	Ar << Data;
-
-	Ar.Seek(OriginalPosition);
-}
-
 void FVirtualizedUntypedBulkData::Serialize(FArchive& Ar, UObject* Owner, bool bAllowRegister)
 {
 	TRACE_CPUPROFILER_EVENT_SCOPE(FVirtualizedUntypedBulkData::Serialize);
@@ -611,10 +612,7 @@ void FVirtualizedUntypedBulkData::Serialize(FArchive& Ar, UObject* Owner, bool b
 			// TODO: Not a huge fan of this, might be better to find a way to build the flags during serialization and potential callbacks 
 			// later then go back and update the flags in the Ar. Applying the updated flags only if we are saving a package to disk
 			// and the save succeeds continues to make sense.
-			const int64 RestorePos = Ar.Tell();
-			Ar.Seek(SavedFlagsPos);
-			Ar << UpdatedFlags;
-			Ar.Seek(RestorePos);
+			UpdateArchiveData(Ar, SavedFlagsPos, UpdatedFlags);
 
 			// Write out required extra data if we're saving by reference
 			bool bWriteOutPayload = true;
@@ -675,12 +673,8 @@ void FVirtualizedUntypedBulkData::Serialize(FArchive& Ar, UObject* Owner, bool b
 				// If we do not have a linker OR the linker is in text format then we should just serialize
 				// the payload directly to the archive.
 				if (LinkerSave != nullptr && !LinkerSave->IsTextFormat())
-				{	
-#if UE_ENABLE_VIRTUALIZATION_TOGGLE
-					if (FPackageTrailer::IsEnabled() == false || bSkipVirtualization == true)
-#else
-					if (FPackageTrailer::IsEnabled() == false)
-#endif //UE_ENABLE_VIRTUALIZATION_TOGGLE
+				{
+					if (ShouldUseLegacySerialization(*LinkerSave))
 					{
 						// Legacy path, will save the payload data to the package
 						SerializeToLegacyPath(*LinkerSave, OffsetPos, PayloadToSerialize, UpdatedFlags, Owner);
@@ -699,15 +693,7 @@ void FVirtualizedUntypedBulkData::Serialize(FArchive& Ar, UObject* Owner, bool b
 
 					SerializeData(Ar, PayloadToSerialize, UpdatedFlags);
 					
-					// Record the current archive offset to return to (probably EOF but we cannot be sure)
-					const int64 ReturnPos = Ar.Tell();
-
-					// Update the offset entry that we set up earlier
-					Ar.Seek(OffsetPos);
-					Ar << DataStartOffset;
-
-					// Restore the archive's offset
-					Ar.Seek(ReturnPos);
+					UpdateArchiveData(Ar, OffsetPos, DataStartOffset);
 				}
 			}
 
@@ -719,7 +705,7 @@ void FVirtualizedUntypedBulkData::Serialize(FArchive& Ar, UObject* Owner, bool b
 		}
 		else if (Ar.IsLoading())
 		{
-			FPackageTrailer* Trailer = GetTrailerFromOwner(Owner);
+			const FPackageTrailer* Trailer = GetTrailerFromOwner(Owner);
 
 			if (Trailer != nullptr && Trailer->FindPayloadStatus(PayloadContentId) == EPayloadStatus::StoredVirtualized)
 			{
@@ -1192,7 +1178,7 @@ FPackagePath FVirtualizedUntypedBulkData::GetPackagePathFromOwner(UObject* Owner
 {
 	OutPackageSegment = EPackageSegment::Header;
 
-	FLinkerLoad* Linker = GetLinkerLoadFromOwner(Owner);
+	const FLinkerLoad* Linker = GetLinkerLoadFromOwner(Owner);
 
 	if (Linker != nullptr)
 	{
@@ -1286,14 +1272,7 @@ void FVirtualizedUntypedBulkData::SerializeToLegacyPath(FLinkerSave& LinkerSave,
 
 		SerializeData(DataArchive, PayloadToSerialize, UpdatedFlags);
 
-		// Record the current archive offset (probably EOF but we cannot be sure)
-		const int64 ReturnPos = ExportsArchive.Tell();
-
-		// Update the offset/size entries that we set up during ::Serialize
-		ExportsArchive.Seek(OffsetPos);
-		ExportsArchive << DataStartOffset;
-		// Restore the archive's offset
-		ExportsArchive.Seek(ReturnPos);
+		UpdateArchiveData(ExportsArchive, OffsetPos, DataStartOffset);
 
 		// If we are saving the package to disk (we have access to FLinkerSave and its filepath is valid) 
 		// then we should register a callback to be received once the package has actually been saved to 
@@ -1347,11 +1326,11 @@ void FVirtualizedUntypedBulkData::SerializeToLegacyPath(FLinkerSave& LinkerSave,
 
 void FVirtualizedUntypedBulkData::SerializeToPackageTrailer(FLinkerSave& LinkerSave, int64 OffsetPos, FCompressedBuffer PayloadToSerialize, EFlags UpdatedFlags, UObject* Owner)
 {
-	auto OnPayloadWritten = [this, OffsetPos, UpdatedFlags](FLinkerSave& LinkerSave) mutable
+	auto OnPayloadWritten = [this, OffsetPos, UpdatedFlags](FLinkerSave& LinkerSave, FPackageTrailer& Trailer) mutable
 	{
 		checkf(LinkerSave.IsCooking() == false, TEXT("FVirtualizedUntypedBulkData::Serialize should not be called during a cook"));
 
-		int64 PayloadOffset = LinkerSave.PackageTrailerBuilder.FindPayloadOffset(PayloadContentId);
+		int64 PayloadOffset = Trailer.FindPayloadOffsetInFile(PayloadContentId);
 
 		UpdateArchiveData(LinkerSave, OffsetPos, PayloadOffset);
 
@@ -1387,7 +1366,7 @@ void FVirtualizedUntypedBulkData::SerializeToPackageTrailer(FLinkerSave& LinkerS
 		}
 	};
 
-	LinkerSave.PackageTrailerBuilder.AddPayload(PayloadContentId, MoveTemp(PayloadToSerialize), MoveTemp(OnPayloadWritten));
+	LinkerSave.PackageTrailerBuilder->AddPayload(PayloadContentId, MoveTemp(PayloadToSerialize), MoveTemp(OnPayloadWritten));
 }
 
 void FVirtualizedUntypedBulkData::UpdatePayloadImpl(FSharedBuffer&& InPayload, FPayloadId&& InPayloadID)
@@ -1745,6 +1724,15 @@ FText FVirtualizedUntypedBulkData::GetCorruptedPayloadErrorMsgForSave(FLinkerSav
 								GuidID);
 	}
 }
+
+bool FVirtualizedUntypedBulkData::ShouldUseLegacySerialization(FLinkerSave& LinkerSave) const
+{
+#if UE_ENABLE_VIRTUALIZATION_TOGGLE
+	return !LinkerSave.PackageTrailerBuilder.IsValid() || bSkipVirtualization == true;
+#else
+	return !LinkerSave.PackageTrailerBuilder.IsValid();
+#endif //UE_ENABLE_VIRTUALIZATION_TOGGLE
+}
 	
 FArchive& operator<<(FArchive& Ar, FTocEntry& Entry)
 {
diff --git a/Engine/Source/Runtime/CoreUObject/Private/UObject/LinkerLoad.cpp b/Engine/Source/Runtime/CoreUObject/Private/UObject/LinkerLoad.cpp
--- a/Engine/Source/Runtime/CoreUObject/Private/UObject/LinkerLoad.cpp
+++ b/Engine/Source/Runtime/CoreUObject/Private/UObject/LinkerLoad.cpp
@@ -1521,7 +1521,7 @@ FLinkerLoad::ELinkerStatus FLinkerLoad::SerializePackageTrailer()
 
 	check(PackageTrailer == nullptr);
 
-	if (UE::FPackageTrailer::IsEnabled() && Summary.PayloadTocOffset != INDEX_NONE)
+	if (UE::FPackageTrailer::IsEnabled() && Summary.PayloadTocOffset > 0)
 	{
 		int64 CurPos = Tell();
 		Seek(Summary.PayloadTocOffset);
diff --git a/Engine/Source/Runtime/CoreUObject/Private/UObject/LinkerSave.cpp b/Engine/Source/Runtime/CoreUObject/Private/UObject/LinkerSave.cpp
--- a/Engine/Source/Runtime/CoreUObject/Private/UObject/LinkerSave.cpp
+++ b/Engine/Source/Runtime/CoreUObject/Private/UObject/LinkerSave.cpp
@@ -25,7 +25,6 @@ TMap<FString, TArray<uint8> > FLinkerSave::PackagesToScriptSHAMap;
 FLinkerSave::FLinkerSave(UPackage* InParent, const TCHAR* InFilename, bool bForceByteSwapping, bool bInSaveUnversioned)
 :	FLinker(ELinkerType::Save, InParent)
 ,	Saver(nullptr)
-,	PackageTrailerBuilder(LinkerRoot)
 {
 	SetFilename(InFilename);
 
@@ -80,7 +79,6 @@ FLinkerSave::FLinkerSave(UPackage* InParent, const TCHAR* InFilename, bool bForc
 FLinkerSave::FLinkerSave(UPackage* InParent, FArchive *InSaver, bool bForceByteSwapping, bool bInSaveUnversioned)
 : FLinker(ELinkerType::Save, InParent)
 , Saver(nullptr)
-, PackageTrailerBuilder(LinkerRoot)
 {
 	SetFilename(TEXT("$$Memory$$"));
 	if (FPlatformProperties::HasEditorOnlyData())
@@ -134,7 +132,6 @@ FLinkerSave::FLinkerSave(UPackage* InParent, FArchive *InSaver, bool bForceByteS
 FLinkerSave::FLinkerSave(UPackage* InParent, bool bForceByteSwapping, bool bInSaveUnversioned )
 :	FLinker(ELinkerType::Save, InParent)
 ,	Saver(nullptr)
-,	PackageTrailerBuilder(LinkerRoot)
 {
 	SetFilename(TEXT("$$Memory$$"));
 	if (FPlatformProperties::HasEditorOnlyData())
diff --git a/Engine/Source/Runtime/CoreUObject/Private/UObject/PackageTrailer.cpp b/Engine/Source/Runtime/CoreUObject/Private/UObject/PackageTrailer.cpp
--- a/Engine/Source/Runtime/CoreUObject/Private/UObject/PackageTrailer.cpp
+++ b/Engine/Source/Runtime/CoreUObject/Private/UObject/PackageTrailer.cpp
@@ -72,7 +72,7 @@ FArchive& operator<<(FArchive& Ar, FLookupTableEntry& Entry)
 
 } // namespace Private
 
-FPackageTrailerBuilder FPackageTrailerBuilder::Create(const FPackageTrailer& Trailer, FArchive& Ar, const FPackagePath& PackagePath)
+FPackageTrailerBuilder FPackageTrailerBuilder::CreateFromTrailer(const FPackageTrailer& Trailer, FArchive& Ar, const FPackagePath& PackagePath)
 {
 	FPackageTrailerBuilder Builder(PackagePath.GetPackageFName());
 
@@ -109,8 +109,6 @@ FPackageTrailerBuilder::FPackageTrailerBuilder(FName InPackageName)
 
 void FPackageTrailerBuilder::AddPayload(const Virtualization::FPayloadId& Identifier, FCompressedBuffer Payload, AdditionalDataCallback&& Callback)
 {
-	checkf(TrailerPositionInFile == INDEX_NONE, TEXT("Attempting to add payloads after the trailer has been built. Package '%s'"), *PackageName.ToString());
-
 	Callbacks.Emplace(MoveTemp(Callback));
 
 	if (Identifier.IsValid())
@@ -121,92 +119,91 @@ void FPackageTrailerBuilder::AddPayload(const Virtualization::FPayloadId& Identi
 
 bool FPackageTrailerBuilder::BuildAndAppendTrailer(FLinkerSave* Linker, FArchive& DataArchive)
 {
-	checkf(TrailerPositionInFile == INDEX_NONE, TEXT("Attempting to build the same FPackageTrailer multiple times. Package '%s'"), *PackageName.ToString());
-
 	// Note that we do not serialize containers directly as we want a file format that is 
 	// 100% under our control. This will allow people to create external scripts that can
 	// parse and manipulate the trailer without needing to worry that we might change how
 	// our containers serialize. 
-	TrailerPositionInFile = DataArchive.Tell();
+	FPackageTrailer Trailer;
+	
+	Trailer.TrailerPositionInFile = DataArchive.Tell();
 
-	uint64 HeaderTag = FPackageTrailer::FHeader::HeaderTag;
-	DataArchive << HeaderTag;
+	Trailer.Header.Tag = FPackageTrailer::FHeader::HeaderTag;
+	Trailer.Header.Version = (int32)EPackageTrailerVersion::AUTOMATIC_VERSION;
 
-	EPackageTrailerVersion Version = EPackageTrailerVersion::AUTOMATIC_VERSION;
-	DataArchive << Version;
+	DataArchive << Trailer.Header.Tag;
+	DataArchive << Trailer.Header.Version;
 
 	const uint32 DynamicHeaderSizeOnDisk = ((LocalEntries.Num() + VirtualizedEntries.Num()) * Private::FLookupTableEntry::SizeOnDisk); // Add the length of the lookup table
 	
-	uint32 HeaderSizeOnDisk = FPackageTrailer::FHeader::StaticHeaderSizeOnDisk + DynamicHeaderSizeOnDisk;
-	DataArchive << HeaderSizeOnDisk;
+	Trailer.Header.HeaderLength = FPackageTrailer::FHeader::StaticHeaderSizeOnDisk + DynamicHeaderSizeOnDisk;
+	DataArchive << Trailer.Header.HeaderLength;
 
-	int64 PayloadsDataLength = 0;
-	PayloadLookupTable.Reserve(LocalEntries.Num() + VirtualizedEntries.Num());
+	Trailer.Header.PayloadsDataLength = 0;
+	Trailer.Header.PayloadLookupTable.Reserve(LocalEntries.Num() + VirtualizedEntries.Num());
 
 	for (const TPair<Virtualization::FPayloadId, LocalEntry>& It : LocalEntries)
 	{
 		checkf(It.Key.IsValid(), TEXT("PackageTrailer should not contain invalid FPayloadIds. Package '%s'"), *PackageName.ToString());
 
-		Private::FLookupTableEntry& Entry = PayloadLookupTable.AddDefaulted_GetRef();
+		Private::FLookupTableEntry& Entry = Trailer.Header.PayloadLookupTable.AddDefaulted_GetRef();
 		Entry.Identifier = It.Key;
-		Entry.OffsetInFile = PayloadsDataLength;
+		Entry.OffsetInFile = Trailer.Header.PayloadsDataLength;
 		Entry.CompressedSize = It.Value.Payload.GetCompressedSize();
 		Entry.RawSize = It.Value.Payload.GetRawSize();
 
-		PayloadsDataLength += It.Value.Payload.GetCompressedSize();
+		Trailer.Header.PayloadsDataLength += It.Value.Payload.GetCompressedSize();
 	}
 
 	for (const TPair<Virtualization::FPayloadId, VirtualizedEntry>& It : VirtualizedEntries)
 	{
 		checkf(It.Key.IsValid(), TEXT("PackageTrailer should not contain invalid FPayloadIds. Package '%s'"), *PackageName.ToString());
 
-		Private::FLookupTableEntry& Entry = PayloadLookupTable.AddDefaulted_GetRef();
+		Private::FLookupTableEntry& Entry = Trailer.Header.PayloadLookupTable.AddDefaulted_GetRef();
 		Entry.Identifier = It.Key;
 		Entry.OffsetInFile = INDEX_NONE;
 		Entry.CompressedSize = It.Value.CompressedSize;
 		Entry.RawSize = It.Value.RawSize;
 	}
 
-	DataArchive << PayloadsDataLength;
+	DataArchive << Trailer.Header.PayloadsDataLength;
 
-	// Currently we only support relative access!
-	// Referenced access will come with the editor domain support.
-	EPayloadAccessMode AccessMode = EPayloadAccessMode::Relative;
-	DataArchive << AccessMode;
+	Trailer.Header.AccessMode = AccessMode;
+	DataArchive << Trailer.Header.AccessMode;
 
-	int32 NumPayloads = PayloadLookupTable.Num();
+	int32 NumPayloads = Trailer.Header.PayloadLookupTable.Num();
 	DataArchive << NumPayloads;
 
-	for (Private::FLookupTableEntry& Entry : PayloadLookupTable)
+	for (Private::FLookupTableEntry& Entry : Trailer.Header.PayloadLookupTable)
 	{
 		DataArchive << Entry;
 	}
 
-	checkf((TrailerPositionInFile + HeaderSizeOnDisk) == DataArchive.Tell(), TEXT("Header length was calculated as %d bytes but we wrote %" INT64_FMT " bytes!"), HeaderSizeOnDisk, DataArchive.Tell() - TrailerPositionInFile);
+	checkf((Trailer.TrailerPositionInFile + Trailer.Header.HeaderLength) == DataArchive.Tell(), TEXT("Header length was calculated as %d bytes but we wrote %" INT64_FMT " bytes!"), Trailer.Header.HeaderLength, DataArchive.Tell() - Trailer.TrailerPositionInFile);
 
-	PayloadPosInFile = DataArchive.Tell();
+	const int64 PayloadPosInFile = DataArchive.Tell();
 
 	for (TPair<Virtualization::FPayloadId, LocalEntry>& It : LocalEntries)
 	{
 		DataArchive << It.Value.Payload;
 	}
 
-	checkf((PayloadPosInFile + PayloadsDataLength) == DataArchive.Tell(), TEXT("Total payload length was calculated as %" INT64_FMT " bytes but we wrote %" INT64_FMT " bytes!"), PayloadsDataLength, DataArchive.Tell() - PayloadPosInFile);
-
-	uint64 FooterTag = FPackageTrailer::FFooter::FooterTag;
-	DataArchive << FooterTag;
+	checkf((PayloadPosInFile + Trailer.Header.PayloadsDataLength) == DataArchive.Tell(), TEXT("Total payload length was calculated as %" INT64_FMT " bytes but we wrote %" INT64_FMT " bytes!"), Trailer.Header.PayloadsDataLength, DataArchive.Tell() - PayloadPosInFile);
 
-	int64 TrailerLength = HeaderSizeOnDisk + PayloadsDataLength + FPackageTrailer::FFooter::SizeOnDisk;
-	DataArchive << TrailerLength;
+	FPackageTrailer::FFooter Footer;
 	
-	uint32 PackageTag = PACKAGE_FILE_TAG;
-	DataArchive << PackageTag;
+	Footer.Tag = FPackageTrailer::FFooter::FooterTag;
+	Footer.TrailerLength = Trailer.Header.HeaderLength + Trailer.Header.PayloadsDataLength + FPackageTrailer::FFooter::SizeOnDisk;
+	Footer.PackageTag = PACKAGE_FILE_TAG;
+	
+	DataArchive << Footer.Tag;
+	DataArchive << Footer.TrailerLength;
+	DataArchive << Footer.PackageTag;
 
-	checkf((TrailerPositionInFile + TrailerLength) == DataArchive.Tell(), TEXT("Trailer length was calculated as %" INT64_FMT " bytes but we wrote %" INT64_FMT " bytes!"), TrailerLength, DataArchive.Tell() - TrailerPositionInFile);
+	checkf((Trailer.TrailerPositionInFile + Footer.TrailerLength) == DataArchive.Tell(), TEXT("Trailer length was calculated as %" INT64_FMT " bytes but we wrote %" INT64_FMT " bytes!"), Footer.TrailerLength, DataArchive.Tell() - Trailer.TrailerPositionInFile);
 
 	for (const AdditionalDataCallback& Callback : Callbacks)
 	{
-		Callback(*Linker);
+		Callback(*Linker, Trailer);
 	}
 
 	return !DataArchive.IsError();
@@ -217,26 +214,6 @@ bool FPackageTrailerBuilder::IsEmpty() const
 	return LocalEntries.IsEmpty() && VirtualizedEntries.IsEmpty();
 }
 
-int64 FPackageTrailerBuilder::FindPayloadOffset(const Virtualization::FPayloadId& Identifier) const
-{
-	check(TrailerPositionInFile != INDEX_NONE);
-	check(PayloadPosInFile != INDEX_NONE);
-	check(PayloadLookupTable.Num() == LocalEntries.Num() + VirtualizedEntries.Num());
-
-	if (Identifier.IsValid())
-	{
-		for (const Private::FLookupTableEntry& Entry : PayloadLookupTable)
-		{
-			if (Entry.Identifier == Identifier)
-			{
-				return PayloadPosInFile + Entry.OffsetInFile;
-			}
-		}
-	}
-
-	return INDEX_NONE;
-}
-
 bool FPackageTrailer::IsEnabled()
 {
 	static struct FUsePackageTrailer
@@ -355,7 +332,7 @@ bool FPackageTrailer::TryLoadBackwards(FArchive& Ar)
 	return TryLoad(Ar);
 }
 
-FCompressedBuffer FPackageTrailer::LoadPayload(Virtualization::FPayloadId Id, FArchive& Ar) const
+FCompressedBuffer FPackageTrailer::LoadPayload(const Virtualization::FPayloadId& Id, FArchive& Ar) const
 {
 	const Private::FLookupTableEntry* Entry = Header.PayloadLookupTable.FindByPredicate([&Id](const Private::FLookupTableEntry& Entry)->bool
 		{
@@ -378,7 +355,7 @@ FCompressedBuffer FPackageTrailer::LoadPayload(Virtualization::FPayloadId Id, FA
 	return Payload;
 }
 
-bool FPackageTrailer::UpdatePayloadAsVirtualized(Virtualization::FPayloadId Identifier)
+bool FPackageTrailer::UpdatePayloadAsVirtualized(const Virtualization::FPayloadId& Identifier)
 {
 	Private::FLookupTableEntry* Entry = Header.PayloadLookupTable.FindByPredicate([&Identifier](const Private::FLookupTableEntry& Entry)->bool
 		{
@@ -396,7 +373,7 @@ bool FPackageTrailer::UpdatePayloadAsVirtualized(Virtualization::FPayloadId Iden
 	}
 }
 
-EPayloadStatus FPackageTrailer::FindPayloadStatus(Virtualization::FPayloadId Id) const
+EPayloadStatus FPackageTrailer::FindPayloadStatus(const Virtualization::FPayloadId& Id) const
 {
 	const Private::FLookupTableEntry* Entry = Header.PayloadLookupTable.FindByPredicate([&Id](const Private::FLookupTableEntry& Entry)->bool
 		{
@@ -411,6 +388,28 @@ EPayloadStatus FPackageTrailer::FindPayloadStatus(Virtualization::FPayloadId Id)
 	return Entry->IsVirtualized() ? EPayloadStatus::StoredVirtualized : EPayloadStatus::StoredLocally;
 }
 
+int64 FPackageTrailer::FindPayloadOffsetInFile(const Virtualization::FPayloadId& Id) const
+{
+	if (Id.IsValid())
+	{
+		const Private::FLookupTableEntry* Entry = Header.PayloadLookupTable.FindByPredicate([&Id](const Private::FLookupTableEntry& Entry)->bool
+			{
+				return Entry.Identifier == Id;
+			});
+
+		//TODO Better way to return an error?
+		check(TrailerPositionInFile != INDEX_NONE);
+		check(Header.PayloadsDataLength != INDEX_NONE);
+		check(Entry != nullptr);
+
+		return TrailerPositionInFile + Header.HeaderLength + Entry->OffsetInFile;
+	}
+	else
+	{
+		return INDEX_NONE;
+	}
+}
+
 int64 FPackageTrailer::GetTrailerLength() const
 {
 	return Header.HeaderLength + Header.PayloadsDataLength + FFooter::SizeOnDisk;
@@ -476,6 +475,11 @@ int32 FPackageTrailer::GetNumPayloads(EPayloadFilter Type) const
 	return Count;
 }
 
+EPayloadAccessMode FPackageTrailer::GetAccessMode() const
+{
+	return Header.AccessMode;
+}
+
 bool FindPayloadsInPackageFile(const FPackagePath& PackagePath, EPayloadFilter Filter, TArray<Virtualization::FPayloadId>& OutPayloadIds)
 {
 	if (FPackageName::IsTextPackageExtension(PackagePath.GetHeaderExtension()))
diff --git a/Engine/Source/Runtime/CoreUObject/Private/UObject/SavePackage.cpp b/Engine/Source/Runtime/CoreUObject/Private/UObject/SavePackage.cpp
--- a/Engine/Source/Runtime/CoreUObject/Private/UObject/SavePackage.cpp
+++ b/Engine/Source/Runtime/CoreUObject/Private/UObject/SavePackage.cpp
@@ -1656,25 +1656,26 @@ struct FPackageExportTagger
 	}
 };
 
-[[nodiscard]] ESavePackageResult BuildAndWriteTrailer(FLinkerSave* Linker, FStructuredArchive::FRecord& StructuredArchiveRoot, bool bTextFormat)
+[[nodiscard]] ESavePackageResult BuildAndWriteTrailer(UPackage* Package, FLinkerSave* Linker, FStructuredArchive::FRecord& StructuredArchiveRoot, IPackageWriter* PackageWriter, bool bTextFormat)
 {
+	check(Package != nullptr);
 	check(Linker != nullptr);
-
-	if (bTextFormat)
-	{
-		// We shouldn't have any data in the trailer as text based assets should store 
-		// all payloads inline at the moment.
-		checkf(Linker->PackageTrailerBuilder.IsEmpty(), TEXT("Attempting to build a package trailer for text based asset, this is not supported!"));
-		Linker->Summary.PayloadTocOffset = INDEX_NONE;
-	}
-	else
+	
+	if (Linker->PackageTrailerBuilder.IsValid())
 	{
+		checkf(PackageWriter == nullptr, TEXT("Attempting to build a package trailer with a package writer '%s', this is not supported!"), *Package->GetName());
+		checkf(bTextFormat == false, TEXT("Attempting to build a package trailer for text based asset '%s', this is not supported!"), *Package->GetName());
+
 		Linker->Summary.PayloadTocOffset = Linker->Tell();
-		if (!Linker->PackageTrailerBuilder.BuildAndAppendTrailer(Linker, *Linker))
+		if (!Linker->PackageTrailerBuilder->BuildAndAppendTrailer(Linker, *Linker))
 		{
 			return ESavePackageResult::Error;
 		}
 	}
+	else
+	{
+		Linker->Summary.PayloadTocOffset = INDEX_NONE;
+	}
 
 	return ESavePackageResult::Success;
 }
@@ -2150,6 +2151,11 @@ FSavePackageResultStruct UPackage::Save(UPackage* InOuter, UObject* InAsset, con
 					Linker->bProceduralSave = ObjectSaveContext.bProceduralSave;
 					Linker->bUpdatingLoadedPath = ObjectSaveContext.bUpdatingLoadedPath;
 
+					if (UE::FPackageTrailer::IsEnabled())
+					{
+						Linker->PackageTrailerBuilder = MakeUnique<UE::FPackageTrailerBuilder>(InOuter);
+					}
+
 #if WITH_TEXT_ARCHIVE_SUPPORT
 					if (bTextFormat)
 					{
@@ -3853,13 +3859,10 @@ FSavePackageResultStruct UPackage::Save(UPackage* InOuter, UObject* InAsset, con
 					// Now that the package is written out we can write the package trailer that is appended
 					// to the file. This should be the last thing written to the file!
 					SlowTask.EnterProgressFrame();
-					if (UE::FPackageTrailer::IsEnabled())
+					ESavePackageResult TrailerResult = BuildAndWriteTrailer(InOuter, Linker.Get(), StructuredArchiveRoot, PackageWriter, bTextFormat);
+					if (TrailerResult != ESavePackageResult::Success)
 					{
-						ESavePackageResult TrailerResult = BuildAndWriteTrailer(Linker.Get(), StructuredArchiveRoot, bTextFormat);
-						if (TrailerResult != ESavePackageResult::Success)
-						{
-							return TrailerResult;
-						}
+						return TrailerResult;
 					}
 				}
 
diff --git a/Engine/Source/Runtime/CoreUObject/Private/UObject/SavePackage2.cpp b/Engine/Source/Runtime/CoreUObject/Private/UObject/SavePackage2.cpp
--- a/Engine/Source/Runtime/CoreUObject/Private/UObject/SavePackage2.cpp
+++ b/Engine/Source/Runtime/CoreUObject/Private/UObject/SavePackage2.cpp
@@ -805,6 +805,11 @@ ESavePackageResult CreateLinker(FSaveContext& SaveContext)
 		SaveContext.Linker->bUpdatingLoadedPath = SaveContext.IsUpdatingLoadedPath();
 		SaveContext.Linker->bProceduralSave = SaveContext.IsProceduralSave();
 
+		if (UE::FPackageTrailer::IsEnabled())
+		{
+			SaveContext.Linker->PackageTrailerBuilder = MakeUnique<UE::FPackageTrailerBuilder>(SaveContext.GetPackage());
+		}
+
 #if WITH_TEXT_ARCHIVE_SUPPORT
 		if (SaveContext.IsTextFormat())
 		{
@@ -1715,21 +1720,21 @@ ESavePackageResult WriteExports(FStructuredArchive::FRecord& StructuredArchiveRo
 
 [[nodiscard]] ESavePackageResult BuildAndWriteTrailer(FStructuredArchive::FRecord& StructuredArchiveRoot, FSaveContext& SaveContext)
 {
-	if (SaveContext.IsTextFormat())
-	{
-		// We shouldn't have any data in the trailer as text based assets should store 
-		// all payloads inline at the moment.
-		checkf(SaveContext.Linker->PackageTrailerBuilder.IsEmpty(), TEXT("Attempting to build a package trailer for text based asset '%s', this is not supported!"), *SaveContext.GetPackage()->GetName());
-		SaveContext.Linker->Summary.PayloadTocOffset = INDEX_NONE;
-	}
-	else
+	if (SaveContext.Linker->PackageTrailerBuilder.IsValid())
 	{
+		checkf(SaveContext.GetPackageWriter() == nullptr, TEXT("Attempting to build a package trailer with a package writer '%s', this is not supported!"), *SaveContext.GetPackage()->GetName());
+		checkf(SaveContext.IsTextFormat() == false, TEXT("Attempting to build a package trailer for text based asset '%s', this is not supported!"), *SaveContext.GetPackage()->GetName());
+
 		SaveContext.Linker->Summary.PayloadTocOffset = SaveContext.Linker->Tell();
-		if (!SaveContext.Linker->PackageTrailerBuilder.BuildAndAppendTrailer(SaveContext.Linker.Get(), *SaveContext.Linker.Get()))
+		if (!SaveContext.Linker->PackageTrailerBuilder->BuildAndAppendTrailer(SaveContext.Linker.Get(), *SaveContext.Linker.Get()))
 		{
 			return ESavePackageResult::Error;
 		}
 	}
+	else
+	{	
+		SaveContext.Linker->Summary.PayloadTocOffset = INDEX_NONE;
+	}
 
 	return ESavePackageResult::Success;
 }
@@ -2166,14 +2171,11 @@ ESavePackageResult SaveHarvestedRealms(FSaveContext& SaveContext, ESaveRealm Har
 	// Now that the package is written out we can write the package trailer that is appended
 	// to the file. This should be the last thing written to the file!
 	SlowTask.EnterProgressFrame();
-	if (UE::FPackageTrailer::IsEnabled() && PackageWriter == nullptr)
+	SaveContext.Result = BuildAndWriteTrailer(StructuredArchiveRoot, SaveContext);
+	if (SaveContext.Result != ESavePackageResult::Success)
 	{
-		SaveContext.Result = BuildAndWriteTrailer(StructuredArchiveRoot, SaveContext);
-		if (SaveContext.Result != ESavePackageResult::Success)
-		{
-			return SaveContext.Result;
-		}	
-	}
+		return SaveContext.Result;
+	}	
 
 	int64 ExportsSize = SaveContext.Linker->Tell();
 	if (PackageWriter)
diff --git a/Engine/Source/Runtime/CoreUObject/Public/Serialization/VirtualizedBulkData.h b/Engine/Source/Runtime/CoreUObject/Public/Serialization/VirtualizedBulkData.h
--- a/Engine/Source/Runtime/CoreUObject/Public/Serialization/VirtualizedBulkData.h
+++ b/Engine/Source/Runtime/CoreUObject/Public/Serialization/VirtualizedBulkData.h
@@ -435,6 +435,9 @@ private:
 	 */
 	FText GetCorruptedPayloadErrorMsgForSave(FLinkerSave* Linker) const;
 
+	/** Returns true if we should use legacy serialization instead of the FPackageTrailer system. This can be removed when UE_ENABLE_VIRTUALIZATION_TOGGLE is removed. */
+	bool ShouldUseLegacySerialization(FLinkerSave& LinkerSave) const;
+
 	/** Unique identifier for the bulkdata object itself */
 	FGuid BulkDataId;
 
diff --git a/Engine/Source/Runtime/CoreUObject/Public/UObject/LinkerLoad.h b/Engine/Source/Runtime/CoreUObject/Public/UObject/LinkerLoad.h
--- a/Engine/Source/Runtime/CoreUObject/Public/UObject/LinkerLoad.h
+++ b/Engine/Source/Runtime/CoreUObject/Public/UObject/LinkerLoad.h
@@ -167,7 +167,7 @@ public:
 		return InstancingContext;
 	}
 
-	UE::FPackageTrailer* GetPackageTrailer() const
+	const UE::FPackageTrailer* GetPackageTrailer() const
 	{
 		return PackageTrailer.Get();
 	}
diff --git a/Engine/Source/Runtime/CoreUObject/Public/UObject/LinkerSave.h b/Engine/Source/Runtime/CoreUObject/Public/UObject/LinkerSave.h
--- a/Engine/Source/Runtime/CoreUObject/Public/UObject/LinkerSave.h
+++ b/Engine/Source/Runtime/CoreUObject/Public/UObject/LinkerSave.h
@@ -115,8 +115,8 @@ public:
 	/** Used by FVirtualizedUntypedBulkData to add payloads to be added to the payload sidecar file (currently an experimental feature) */
 	TArray<FSidecarStorageInfo> SidecarDataToAppend;
 	
-	/** Gathers all payloads while save the package, so that they can be stored ina  single data structure @see FPackageTrailer */
-	UE::FPackageTrailerBuilder PackageTrailerBuilder;
+	/** Gathers all payloads while save the package, so that they can be stored in a single data structure @see FPackageTrailer */
+	TUniquePtr<UE::FPackageTrailerBuilder> PackageTrailerBuilder;
 
 	/** 
 	 * Array of callbacks that will be invoked when the package has successfully saved to disk.
diff --git a/Engine/Source/Runtime/CoreUObject/Public/UObject/PackageTrailer.h b/Engine/Source/Runtime/CoreUObject/Public/UObject/PackageTrailer.h
--- a/Engine/Source/Runtime/CoreUObject/Public/UObject/PackageTrailer.h
+++ b/Engine/Source/Runtime/CoreUObject/Public/UObject/PackageTrailer.h
@@ -46,6 +46,7 @@ namespace UE
  * | Version			| uint32			| Version number of the format@see UE::EPackageTrailerVersion										|
  * | HeaderLength		| uint32			| The total size of the header on disk in bytes.													|
  * | PayloadsDataLength	| uint64			| The total size of the payload data on disk in bytes												|
+ * | AccessMode			| uint8				| @see UE::EPayloadAccessMode														|
  * | NumPayloads		| int32				| The number of payloads in LookupTableArray														|
  * | LookupTableArray	| FLookupTableEntry | An array of FLookupTableEntry @see UE::Private::FLookupTableEntry									|
  * |____________________________________________________________________________________________________________________________________________|
@@ -135,14 +136,16 @@ struct FLookupTableEntry
 class COREUOBJECT_API FPackageTrailerBuilder
 {
 public:
-	using AdditionalDataCallback = TFunction<void(FLinkerSave& LinkerSave)>;
+	using AdditionalDataCallback = TFunction<void(FLinkerSave& LinkerSave, class FPackageTrailer& Trailer)>;
 
 	/**
-	 * @param Trailer
-	 * @param Ar
-	 * @param PackagePath
+	 * Creates a builder from an already loaded FPackageTrailer. 
+	 * 
+	 * @param Trailer		The trailer to create the builder from
+	 * @param Ar			An archive that the trailer can use to load payloads from 
+	 * @param PackagePath	Path to the package
 	 */
-	[[nodiscard]] static FPackageTrailerBuilder Create(const class FPackageTrailer& Trailer, FArchive& Ar, const FPackagePath& PackagePath);
+	[[nodiscard]] static FPackageTrailerBuilder CreateFromTrailer(const class FPackageTrailer& Trailer, FArchive& Ar, const FPackagePath& PackagePath);
 
 	FPackageTrailerBuilder() = delete;
 	FPackageTrailerBuilder(UPackage* Package);
@@ -171,9 +174,6 @@ public:
 	/** Returns if the builder has any payload entries or not */
 	[[nodiscard]] bool IsEmpty() const;
 
-	// Methods that can be called after building the trailer and appending it to the package file
-	[[nodiscard]] int64 FindPayloadOffset(const Virtualization::FPayloadId& Identifier) const;
-
 private:
 	
 	/** All of the data required to add a payload that is stored on disk */
@@ -211,24 +211,15 @@ private:
 	/** Name of the package the trailer is being built for, used to give meaningful error messages */
 	FName PackageName;
 
+	/** The access mode that the trailer should have to the payloads */
+	EPayloadAccessMode AccessMode = EPayloadAccessMode::Relative;
+
 	/** Payloads that will be stored locally when the trailer is written to disk */
 	TMap<Virtualization::FPayloadId, LocalEntry> LocalEntries;
 	/** Payloads that are already virtualized and so will not be written to disk */
 	TMap<Virtualization::FPayloadId, VirtualizedEntry> VirtualizedEntries;
-
 	/** Callbacks to invoke once the trailer has been written to the end of a package */
 	TArray<AdditionalDataCallback> Callbacks;
-
-	// Members that are only valid one after building the trailer and appending it to the package file
-
-	/** Where in the package file the trailer is located */
-	int64 TrailerPositionInFile = INDEX_NONE;
-
-	/** Where in the package file that the payload data is located */
-	int64 PayloadPosInFile = INDEX_NONE;
-
-	/** The same look up table that the trailer would have */
-	TArray<Private::FLookupTableEntry> PayloadLookupTable;
 };
 
 /** 
@@ -286,7 +277,7 @@ public:
 	 * @return	The payload in the form of a FCompressedBuffer. If the payload does not exist in the trailer or in the archive then
 	 *			the FCompressedBuffer will be null.
 	 */
-	[[nodiscard]] FCompressedBuffer LoadPayload(Virtualization::FPayloadId Id, FArchive& Ar) const;
+	[[nodiscard]] FCompressedBuffer LoadPayload(const Virtualization::FPayloadId& Id, FArchive& Ar) const;
 
 	/** 
 	 * Calling this indicates that the payload has been virtualized and will no longer be stored on disk. 
@@ -294,10 +285,13 @@ public:
 	 * @param Identifier The payload that has been virtualized
 	 * @return True if the payload was in the trailer, otherwise false
 	 */
-	[[nodiscard]] bool UpdatePayloadAsVirtualized(Virtualization::FPayloadId Identifier);
+	[[nodiscard]] bool UpdatePayloadAsVirtualized(const Virtualization::FPayloadId& Identifier);
 
 	/** Attempt to find the status of the given payload. @See EPayloadStatus */
-	[[nodiscard]] EPayloadStatus FindPayloadStatus(Virtualization::FPayloadId Id) const;
+	[[nodiscard]] EPayloadStatus FindPayloadStatus(const Virtualization::FPayloadId& Id) const;
+
+	/** Returns the absolute offset of the payload in the package file, invalid and virtualized payloads will return INDEX_NONE */
+	[[nodiscard]] int64 FindPayloadOffsetInFile(const Virtualization::FPayloadId& Id) const;
 
 	/** Returns the total size of the of the trailer on disk in bytes */
 	[[nodiscard]] int64 GetTrailerLength() const;
@@ -308,6 +302,8 @@ public:
 	/** Returns the number of payloads that the trailer owns that match the given filter type. @See EPayloadType */
 	[[nodiscard]] int32 GetNumPayloads(EPayloadFilter Type) const;
 
+	[[nodiscard]] EPayloadAccessMode GetAccessMode() const;
+
 	struct FHeader
 	{
 		/** Unique value used to identify the header */
