Fix hang in ReferenceFinder when external packages are involved.

diff --git a/Engine/Source/Runtime/CoreUObject/Private/UObject/ReferencerFinder.cpp b/Engine/Source/Runtime/CoreUObject/Private/UObject/ReferencerFinder.cpp
--- a/Engine/Source/Runtime/CoreUObject/Private/UObject/ReferencerFinder.cpp
+++ b/Engine/Source/Runtime/CoreUObject/Private/UObject/ReferencerFinder.cpp
@@ -74,8 +74,8 @@ TArray<UObject*> FReferencerFinder::GetAllReferencers(const TSet<UObject*>& Refe
 	{
 		FCriticalSection ResultCritical;
 
-		// Lock hashtables so that nothing can add UObjects while we're iterating over the GUObjectArray
-		FScopedUObjectHashTablesLock HashTablesLock;
+		// Lock the global array so that nothing can add UObjects while we're iterating over it
+		GUObjectArray.LockInternalArray();
 
 		const int32 MaxNumberOfObjects = GUObjectArray.GetObjectArrayNum();
 		const int32 NumThreads = FMath::Max(1, FTaskGraphInterface::Get().GetNumWorkerThreads());
@@ -132,6 +132,9 @@ TArray<UObject*> FReferencerFinder::GetAllReferencers(const TSet<UObject*>& Refe
 				Ret.Append(ThreadResult.Array());
 			}
 		}, (GUObjectRegistrationComplete && GAllowParallelReferenceCollection) ? EParallelForFlags::None : EParallelForFlags::ForceSingleThread);
+
+		// Release the global array lock
+		GUObjectArray.UnlockInternalArray();
 	}
 	return Ret;
 }
