Fixing uninitialized struct member tests crashing for some programs

diff --git a/Engine/Source/Runtime/CoreUObject/Private/UObject/Class.cpp b/Engine/Source/Runtime/CoreUObject/Private/UObject/Class.cpp
--- a/Engine/Source/Runtime/CoreUObject/Private/UObject/Class.cpp
+++ b/Engine/Source/Runtime/CoreUObject/Private/UObject/Class.cpp
@@ -3617,10 +3617,15 @@ bool FAutomationTestAttemptToFindUninitializedScriptStructMembers::RunTest(const
 {
 	// This test fails when running tests under UHT because there is no TestUninitializedScriptStructMembersTest, so just skip it in that config.
 #if !HACK_HEADER_GENERATOR
-	return FStructUtils::AttemptToFindUninitializedScriptStructMembers() == 0;
-#else
-	return true;
+	if (UObjectInitialized())
+	{
+		return FStructUtils::AttemptToFindUninitializedScriptStructMembers() == 0;
+	}
+	else
 #endif
+	{
+		return true;
+	}
 }
 
 #endif
