Do not create a FPackageTrailerBuilder when saving a text based asset.

#rnx

- The package trailer system is not yet supported by text based assets.

diff --git a/Engine/Source/Runtime/CoreUObject/Private/UObject/SavePackage.cpp b/Engine/Source/Runtime/CoreUObject/Private/UObject/SavePackage.cpp
--- a/Engine/Source/Runtime/CoreUObject/Private/UObject/SavePackage.cpp
+++ b/Engine/Source/Runtime/CoreUObject/Private/UObject/SavePackage.cpp
@@ -2153,7 +2153,10 @@ FSavePackageResultStruct UPackage::Save(UPackage* InOuter, UObject* InAsset, con
 
 					if (UE::FPackageTrailer::IsEnabled())
 					{
-						Linker->PackageTrailerBuilder = MakeUnique<UE::FPackageTrailerBuilder>(InOuter);
+						if (!bTextFormat) // The package trailer is not supported for text based assets yet
+						{
+							Linker->PackageTrailerBuilder = MakeUnique<UE::FPackageTrailerBuilder>(InOuter);
+						}
 					}
 
 #if WITH_TEXT_ARCHIVE_SUPPORT
diff --git a/Engine/Source/Runtime/CoreUObject/Private/UObject/SavePackage2.cpp b/Engine/Source/Runtime/CoreUObject/Private/UObject/SavePackage2.cpp
--- a/Engine/Source/Runtime/CoreUObject/Private/UObject/SavePackage2.cpp
+++ b/Engine/Source/Runtime/CoreUObject/Private/UObject/SavePackage2.cpp
@@ -807,7 +807,10 @@ ESavePackageResult CreateLinker(FSaveContext& SaveContext)
 
 		if (UE::FPackageTrailer::IsEnabled())
 		{
-			SaveContext.Linker->PackageTrailerBuilder = MakeUnique<UE::FPackageTrailerBuilder>(SaveContext.GetPackage());
+			if (!SaveContext.IsTextFormat()) // The package trailer is not supported for text based assets yet
+			{
+				SaveContext.Linker->PackageTrailerBuilder = MakeUnique<UE::FPackageTrailerBuilder>(SaveContext.GetPackage());
+			}
 		}
 
 #if WITH_TEXT_ARCHIVE_SUPPORT
