Merging CL#18476999

#UE If a parent class changes its sparse data struct and a child class points to that struct, then update the child class as well. This is an extension of the existing fix to update SuperStruct pointers

diff --git a/Engine/Source/Runtime/CoreUObject/Private/UObject/Class.cpp b/Engine/Source/Runtime/CoreUObject/Private/UObject/Class.cpp
--- a/Engine/Source/Runtime/CoreUObject/Private/UObject/Class.cpp
+++ b/Engine/Source/Runtime/CoreUObject/Private/UObject/Class.cpp
@@ -5254,6 +5254,10 @@ void UClass::SetSparseClassDataStruct(UScriptStruct* InSparseClassDataStruct)
 					SubClass->CleanupSparseClassData();
 					SubClassSparseClassDataStruct->SetSuperStruct(InSparseClassDataStruct);
 				}
+				else if (SubClassSparseClassDataStruct == SparseClassDataStruct)
+				{
+					SubClass->SetSparseClassDataStruct(InSparseClassDataStruct);
+				}
 			}
 		}
 		// the old type and new type may not match when we do a reload so get rid of the old data
