Fix assert when cooking with the package trailer enabled.

#rnx

- We do not want to create a FPackageTrailerBuilder during procedural saves.

diff --git a/Engine/Source/Runtime/CoreUObject/Private/UObject/SavePackage.cpp b/Engine/Source/Runtime/CoreUObject/Private/UObject/SavePackage.cpp
--- a/Engine/Source/Runtime/CoreUObject/Private/UObject/SavePackage.cpp
+++ b/Engine/Source/Runtime/CoreUObject/Private/UObject/SavePackage.cpp
@@ -2148,12 +2148,14 @@ FSavePackageResultStruct UPackage::Save(UPackage* InOuter, UObject* InAsset, con
 						TempFilename = FPaths::CreateTempFilename(*FPaths::ProjectSavedDir(), *BaseFilename.Left(32));
 						Linker = MakePimpl<FLinkerSave>(InOuter, *TempFilename.GetValue(), bForceByteSwapping, bSaveUnversionedNative);
 					}
+
 					Linker->bProceduralSave = ObjectSaveContext.bProceduralSave;
 					Linker->bUpdatingLoadedPath = ObjectSaveContext.bUpdatingLoadedPath;
 
 					if (UE::FPackageTrailer::IsEnabled())
 					{
-						if (!bTextFormat) // The package trailer is not supported for text based assets yet
+						// The package trailer is not supported for text based assets yet
+						if (!bTextFormat && !ObjectSaveContext.bProceduralSave) 
 						{
 							Linker->PackageTrailerBuilder = MakeUnique<UE::FPackageTrailerBuilder>(InOuter);
 						}
diff --git a/Engine/Source/Runtime/CoreUObject/Private/UObject/SavePackage2.cpp b/Engine/Source/Runtime/CoreUObject/Private/UObject/SavePackage2.cpp
--- a/Engine/Source/Runtime/CoreUObject/Private/UObject/SavePackage2.cpp
+++ b/Engine/Source/Runtime/CoreUObject/Private/UObject/SavePackage2.cpp
@@ -807,7 +807,8 @@ ESavePackageResult CreateLinker(FSaveContext& SaveContext)
 
 		if (UE::FPackageTrailer::IsEnabled())
 		{
-			if (!SaveContext.IsTextFormat()) // The package trailer is not supported for text based assets yet
+			// The package trailer is not supported for text based assets yet
+			if (!SaveContext.IsTextFormat() && !SaveContext.IsProceduralSave())
 			{
 				SaveContext.Linker->PackageTrailerBuilder = MakeUnique<UE::FPackageTrailerBuilder>(SaveContext.GetPackage());
 			}
