Allow sparse class data to override its copy logic when initializing from its archetype.

diff --git a/Engine/Source/Runtime/CoreUObject/Private/UObject/Class.cpp b/Engine/Source/Runtime/CoreUObject/Private/UObject/Class.cpp
--- a/Engine/Source/Runtime/CoreUObject/Private/UObject/Class.cpp
+++ b/Engine/Source/Runtime/CoreUObject/Private/UObject/Class.cpp
@@ -5174,14 +5174,11 @@ void* UClass::CreateSparseClassData()
 	{
 		// initialize per class data from the archetype if we have one
 		const void* SparseArchetypeData = GetArchetypeForSparseClassData();
-		UStruct* SparseClassDataArchetypeStruct = GetSparseClassDataArchetypeStruct();
+		UScriptStruct* SparseClassDataArchetypeStruct = GetSparseClassDataArchetypeStruct();
 
-		if (SparseArchetypeData)
+		if (SparseArchetypeData && SparseClassDataStruct->IsChildOf(SparseClassDataArchetypeStruct))
 		{
-			for (FProperty* P = SparseClassDataArchetypeStruct->PropertyLink; P; P = P->PropertyLinkNext)
-			{
-				P->CopyCompleteValue_InContainer(SparseClassData, SparseArchetypeData);
-			}
+			SparseClassDataArchetypeStruct->CopyScriptStruct(SparseClassData, SparseArchetypeData);
 		}
 	}
 
