ProcessNewlyLoadedObjects always called for VNI
- also removed installment of FModuleManager::OnProcessLoadedObjectsCallback() in monolithic builds as it's not needed for monolithic builds

#rnx

diff --git a/Engine/Source/Runtime/Core/Private/Modules/ModuleManager.cpp b/Engine/Source/Runtime/Core/Private/Modules/ModuleManager.cpp
--- a/Engine/Source/Runtime/Core/Private/Modules/ModuleManager.cpp
+++ b/Engine/Source/Runtime/Core/Private/Modules/ModuleManager.cpp
@@ -463,7 +463,7 @@ IModuleInterface* FModuleManager::LoadModuleWithFailureReason(const FName InModu
 		{
 			FScopedBootTiming BootScope("LoadModule  - ", InModuleName);
 			TRACE_LOADTIME_REQUEST_GROUP_SCOPE(TEXT("LoadModule - %s"), *InModuleName.ToString());
-#if USE_PER_MODULE_UOBJECT_BOOTSTRAP
+#if USE_PER_MODULE_UOBJECT_BOOTSTRAP || WITH_VERSE
 			{
 				ProcessLoadedObjectsCallback.Broadcast(InModuleName, bCanProcessNewlyLoadedObjects);
 			}
diff --git a/Engine/Source/Runtime/CoreUObject/Private/UObject/Obj.cpp b/Engine/Source/Runtime/CoreUObject/Private/UObject/Obj.cpp
--- a/Engine/Source/Runtime/CoreUObject/Private/UObject/Obj.cpp
+++ b/Engine/Source/Runtime/CoreUObject/Private/UObject/Obj.cpp
@@ -4589,7 +4589,10 @@ void InitUObject()
 
 	FCoreDelegates::OnShutdownAfterError.AddStatic(StaticShutdownAfterError);
 	FCoreDelegates::OnExit.AddStatic(StaticExit);
-#if !USE_PER_MODULE_UOBJECT_BOOTSTRAP // otherwise this is already done
+#if !USE_PER_MODULE_UOBJECT_BOOTSTRAP && !IS_MONOLITHIC
+	// Not necessary when USE_PER_MODULE_UOBJECT_BOOTSTRAP==0 since the callback gets installed elsewhere
+	// Also not necessary for monolithic builds as all pending registrants are available at once on app start
+	// so ProcessNewlyLoadedUObjects needs to only ever be invoked once, not for each module
 	FModuleManager::Get().OnProcessLoadedObjectsCallback().AddStatic(ProcessNewlyLoadedUObjects);
 #endif
 
