Fix crash when cooking child anim BPs

Manually merging CLs 18393342 & 17797050
Additionally patches UClass::ClearSparseClassData to clear all descendent's sparse class data and ensure that the ptr to SparseClassDataStruct is always nulled-out (previously with CL 18393342 this could be left dangling in child classes).

diff --git a/Engine/Source/Runtime/CoreUObject/Private/UObject/Class.cpp b/Engine/Source/Runtime/CoreUObject/Private/UObject/Class.cpp
--- a/Engine/Source/Runtime/CoreUObject/Private/UObject/Class.cpp
+++ b/Engine/Source/Runtime/CoreUObject/Private/UObject/Class.cpp
@@ -5270,17 +5270,25 @@ void UClass::ClearSparseClassDataStruct()
 	{
 		// Find all subclasses and clear their sparse class data struct as well.
 		TArray<UClass*> SubClasses;
+		TArray<UClass*> SubClassesToClear;
 		GetDerivedClasses(this, SubClasses, true /* bRecursive */);
 		for (UClass* SubClass : SubClasses)
 		{
 			UScriptStruct* SubClassSparseClassDataStruct = SubClass->GetSparseClassDataStruct();
-			if (SubClassSparseClassDataStruct && SubClassSparseClassDataStruct->GetSuperStruct() == SparseClassDataStruct)
+			if (SubClassSparseClassDataStruct && SubClassSparseClassDataStruct->IsChildOf(SparseClassDataStruct))
 			{
-				SubClass->CleanupSparseClassData();
-				SubClassSparseClassDataStruct->SetSuperStruct(nullptr);
+				SubClassesToClear.Add(SubClass);
 			}
 		}
 
+		for(UClass* SubClassToClear : SubClassesToClear)
+		{
+			UScriptStruct* SubClassSparseClassDataStruct = SubClassToClear->GetSparseClassDataStruct();
+			SubClassToClear->CleanupSparseClassData();
+			SubClassToClear->SparseClassDataStruct = nullptr;
+			SubClassSparseClassDataStruct->SetSuperStruct(nullptr);	
+		}
+
 		CleanupSparseClassData();
 
 		SparseClassDataStruct = nullptr;
