Fix rare crash when a bulkdata object is destroyed while waiting on an async data load.

- If the asyc read being waited on takes a very long time (1000ms+) then we log a warning that there was a stall. However the warning prints the size of the bulkdata payload to the warning to give the user an idea of if the wait was due to the length of the data or some other reason.
- Calling ::GetBulkDataSize cannot be done in the destructor of the bulkdata object because the method is virtual, which was causing the crash.
- We can however call ::GetBulkDataSizeOnDisk  as it is not virtual and gives a more accurate picture of the length of the data being loaded, as it can be smaller than the result of  ::GetBulkDataSize if compression is used.
- Switched use of %lld to use INT64_FMT instead.

diff --git a/Engine/Source/Runtime/CoreUObject/Private/Serialization/BulkData.cpp b/Engine/Source/Runtime/CoreUObject/Private/Serialization/BulkData.cpp
--- a/Engine/Source/Runtime/CoreUObject/Private/Serialization/BulkData.cpp
+++ b/Engine/Source/Runtime/CoreUObject/Private/Serialization/BulkData.cpp
@@ -2173,7 +2173,7 @@ void FUntypedBulkData::WaitForAsyncLoading()
 	DECLARE_SCOPE_CYCLE_COUNTER(TEXT("FUntypedBulkData::WaitForAsyncLoading"), STAT_UBD_WaitForAsyncLoading, STATGROUP_Memory);
 	while (!SerializeFuture.WaitFor(FTimespan::FromMilliseconds(1000.0)))
 	{
-		UE_LOG(LogSerialization, Warning, TEXT("Waiting for %s bulk data (%lld) to be loaded longer than 1000ms"), *PackagePath.GetDebugName(PackageSegment), GetBulkDataSize());
+		UE_LOG(LogSerialization, Warning, TEXT("Waiting for '%s' bulk data (size %" INT64_FMT ") to be loaded longer than 1000ms"), *PackagePath.GetDebugName(PackageSegment), GetBulkDataSizeOnDisk());
 	}
 	check(BulkDataAsync);
 }
