Unshelved from Sebastien.Lussier

Fixed issue where SavePackage2 would fail to gather custom versions for CDOs
* FPackageHarvester::ProcessExport() is using UClass::SerializeDefaultObject() which will not call CDO->Serialize()
#localization none
#tests Resaved BP assets packages, child classes of AWaterBodyIsland - those now have a FWaterCustomVersion as expected

diff --git a/Engine/Source/Runtime/CoreUObject/Private/UObject/SavePackage/PackageHarvester.cpp b/Engine/Source/Runtime/CoreUObject/Private/UObject/SavePackage/PackageHarvester.cpp
--- a/Engine/Source/Runtime/CoreUObject/Private/UObject/SavePackage/PackageHarvester.cpp
+++ b/Engine/Source/Runtime/CoreUObject/Private/UObject/SavePackage/PackageHarvester.cpp
@@ -245,13 +245,11 @@ void FPackageHarvester::ProcessExport(const FExportWithContext& InProcessContext
 		//@ todo FH: I don't think recursing into the template subobject is necessary, serializing it should catch the necessary sub objects
 		// GetCDOSubobjects??
 	}
-	else
-	{
-		// @todo FH: always serialize???
-		// In the CDO case the above would serialize most of the references, including transient properties
-		// but we still want to serialize the object using the normal path to collect all custom versions it might be using.
 		Export->Serialize(*this);
-	}
+
+	// In the CDO case the above would serialize most of the references, including transient properties
+	// but we still want to serialize the object using the normal path to collect all custom versions it might be using.
+	InObject->Serialize(*this);
 
 	// Gather object preload dependencies
 	if (SaveContext.IsCooking())
