Fix merge issue from robomerge cl 18550665

diff --git a/Engine/Source/Runtime/CoreUObject/Private/UObject/SavePackage/PackageHarvester.cpp b/Engine/Source/Runtime/CoreUObject/Private/UObject/SavePackage/PackageHarvester.cpp
--- a/Engine/Source/Runtime/CoreUObject/Private/UObject/SavePackage/PackageHarvester.cpp
+++ b/Engine/Source/Runtime/CoreUObject/Private/UObject/SavePackage/PackageHarvester.cpp
@@ -245,11 +245,10 @@ void FPackageHarvester::ProcessExport(const FExportWithContext& InProcessContext
 		//@ todo FH: I don't think recursing into the template subobject is necessary, serializing it should catch the necessary sub objects
 		// GetCDOSubobjects??
 	}
-		Export->Serialize(*this);
 
 	// In the CDO case the above would serialize most of the references, including transient properties
 	// but we still want to serialize the object using the normal path to collect all custom versions it might be using.
-	InObject->Serialize(*this);
+	Export->Serialize(*this);
 
 	// Gather object preload dependencies
 	if (SaveContext.IsCooking())
