The package trailer system can now work in projects with the editor domain enable as well.

#rnx

- Now when saving a package to the editor domain we will load the existing package trailer (if one exists) from the workspace domain and then create a new trailer that references those payloads rather than store them locally.
- VirtualizedBulkData serialization is still overly complicated and this code adds more branches there. If we stopped the VBD objects from serializing their offsets at all and rely entirely on the package trailer to know where the payloads are stored then we could simplify this, but that requires the package trailer system to be turned on so will be done as future work.

### VirtualizedBulkData
- Split bAllowVirtualizationOnSave into a second variable. The original now controls if we push the payload on save (off by default and not something we are likely to keep) and the new variable (bForceRehydrationOnSave) controls if we must pull payloads when the package is saved (on by default) which will be moved to a config value once properly supported.
-- The split was done to make the logic in ::Serialize clearer
- We now correctly Register/Unregister the bulkdata object when ::DetachFromDisk is called. This prevents the bulkdata registry from trying to access the workspace domain file after it has been invalidated.
-- In addition we now attach to the FLinkerLoader's cachable archive (if there is one) when using the ::IsReferencingByPackagePath branch of serialization.
- When saving the bulkdata as part of an editor domain save we can keep the virtualization flag as we don't have to worry about the workspace domain file changing without the editor domain version being rebuilt anyway.

### SavePackage
- If the package is part of a procedural save and not a cook we will try to load the trailer from the workspace domain and then convert it to a referenced trailer (which contains the correct info to load the payload from the workspace domain package file but does not contain the payloads itself) which is then appended to the package.

### PackageTrailer
- Removed the fatal error if a project has both the editor domain and the package trailer system enabled.
- Added a method ::CreateReference that creates a trailer that references the payloads in an existing trailer.
- Added a method ::TrySave, allowing direct writing of a trailer to an archive.
-- Now that we have TrySave and TryLoad it probably makes sense to add the << operator as the public interface and move the methods to be private.
- FFooter and FHeader have their own serialization overloads now rather than relying on the higher level structures to do the work.
- When building a new trailer we now create a new FPackageTrailer structure and then write that out in one go, rather than writing each value as it is calculated to make the code easier to read.

diff --git a/Engine/Source/Editor/UnrealEd/Private/EditorDomain/EditorDomainUtils.cpp b/Engine/Source/Editor/UnrealEd/Private/EditorDomain/EditorDomainUtils.cpp
--- a/Engine/Source/Editor/UnrealEd/Private/EditorDomain/EditorDomainUtils.cpp
+++ b/Engine/Source/Editor/UnrealEd/Private/EditorDomain/EditorDomainUtils.cpp
@@ -79,7 +79,7 @@ bool bGGlobalAddedCustomVersionsInitialized = false;
 int64 GMaxBulkDataSize = -1;
 
 // Change to a new guid when EditorDomain needs to be invalidated
-const TCHAR* EditorDomainVersion = TEXT("3237D60426C74B8B93E8F6BA49488E84");
+const TCHAR* EditorDomainVersion = TEXT("83221BBB67B3482A8064C76349ECE07D");
 
 // Identifier of the CacheBuckets for EditorDomain tables
 const TCHAR* EditorDomainPackageBucketName = TEXT("EditorDomainPackage");
diff --git a/Engine/Source/Runtime/CoreUObject/Private/Serialization/VirtualizedBulkData.cpp b/Engine/Source/Runtime/CoreUObject/Private/Serialization/VirtualizedBulkData.cpp
--- a/Engine/Source/Runtime/CoreUObject/Private/Serialization/VirtualizedBulkData.cpp
+++ b/Engine/Source/Runtime/CoreUObject/Private/Serialization/VirtualizedBulkData.cpp
@@ -38,7 +38,6 @@
 
 namespace UE::Virtualization
 {
-
 /** This console variable should only exist for testing */
 static TAutoConsoleVariable<bool> CVarShouldLoadFromSidecar(
 	TEXT("Serialization.LoadFromSidecar"),
@@ -64,9 +63,17 @@ static TAutoConsoleVariable<bool> CVarShouldAllowSidecarSyncing(
 	false,
 	TEXT("When true FVirtualizedUntypedBulkData will attempt to sync it's .upayload file via sourcecontrol if the first attempt to load from it fails"));
 
-/** Might expose this as an option but for now set to false only */
+/** When enabled the bulkdata object will try pushing the payload when saved to disk as part of a package.
+  * This is legacy behavior and likely to be removed
+  */
 static constexpr bool bAllowVirtualizationOnSave = false;
 
+/** When enabled virtualized payloads will be pulled and stored locally when the bulkdata object is saved
+  * to disk as part of a package. This is currently required but will likely be changed to be a config 
+  * option when we can support it.
+  */
+static constexpr bool bForceRehydrationOnSave = true;
+
 /** Wrapper around the config file option [Core.System.Experimental]EnablePackageSidecarSaving */
 bool ShouldSaveToPackageSidecar()
 {
@@ -626,7 +633,7 @@ void FVirtualizedUntypedBulkData::Serialize(FArchive& Ar, UObject* Owner, bool b
 			}
 			else
 			{
-				bWriteOutPayload = !bAllowVirtualizationOnSave || !IsDataVirtualized();
+				bWriteOutPayload = !IsDataVirtualized(UpdatedFlags);
 			}
 
 			if (bWriteOutPayload)
@@ -696,7 +703,7 @@ void FVirtualizedUntypedBulkData::Serialize(FArchive& Ar, UObject* Owner, bool b
 					UpdateArchiveData(Ar, OffsetPos, DataStartOffset);
 				}
 			}
-
+			
 			if (CanUnloadData())
 			{
 				this->CompressionSettings.Reset();
@@ -707,14 +714,32 @@ void FVirtualizedUntypedBulkData::Serialize(FArchive& Ar, UObject* Owner, bool b
 		{
 			const FPackageTrailer* Trailer = GetTrailerFromOwner(Owner);
 
-			if (Trailer != nullptr && Trailer->FindPayloadStatus(PayloadContentId) == EPayloadStatus::StoredVirtualized)
+			if (IsDataVirtualized())
+			{
+				// The payload was already virtualized when the bulkdata object was serialized
+
+				// We aren't going to use these members so reset them
+				OffsetInFile = INDEX_NONE;
+
+				PackagePath.Empty();
+				PackageSegment = EPackageSegment::Header;
+			}
+			else if (Trailer != nullptr && Trailer->FindPayloadStatus(PayloadContentId) == EPayloadStatus::StoredVirtualized)
 			{
+				// The payload was stored locally in the package trailer when the bulkdata object was 
+				// serialized but the trailer was later virtualized.
+
 				check(!IsReferencingByPackagePath()); // This should not be possible
 
 				EnumAddFlags(Flags, EFlags::IsVirtualized);
 
 				Ar << OffsetInFile;
+
+				// We aren't going to use these members so reset them
 				OffsetInFile = INDEX_NONE;
+
+				PackagePath.Empty();
+				PackageSegment = EPackageSegment::Header;
 			}
 			else if (IsReferencingByPackagePath())
 			{
@@ -723,19 +748,20 @@ void FVirtualizedUntypedBulkData::Serialize(FArchive& Ar, UObject* Owner, bool b
 				Ar << OffsetInFile;
 				ensure(FPackagePath::TryFromPackageName(*PackageNameStr, this->PackagePath));
 				PackageSegment = EPackageSegment::Header;
-			}
-			else if (IsDataVirtualized())
-			{
-				// TODO: Remove once safe to do so
-				// Legacy path for packages that were saved as virtualized
-			
-				// We aren't going to use these members so reset them
-				OffsetInFile = INDEX_NONE;
 
-				PackagePath.Empty();
-				PackageSegment = EPackageSegment::Header;
+#if VBD_ALLOW_LINKERLOADER_ATTACHMENT
+				FArchive* CacheableArchive = Ar.GetCacheableArchive();
+				if (Ar.IsAllowingLazyLoading() && CacheableArchive != nullptr)
+				{
+					AttachedAr = CacheableArchive;
+					if (AttachedAr != nullptr)
+					{
+						AttachedAr->AttachBulkData(this);
+					}
+				}
+#endif //VBD_ALLOW_LINKERLOADER_ATTACHMENT
 			}
-			else
+			else 
 			{
 				// If we can lazy load then find the PackagePath, otherwise we will want
 				// to serialize immediately.
@@ -938,6 +964,8 @@ FCompressedBuffer FVirtualizedUntypedBulkData::LoadFromPackageTrailer() const
 
 	UE_LOG(LogVirtualization, Verbose, TEXT("Attempting to load payload from the package trailer: '%s'"), *PackagePath.GetLocalFullPath(PackageSegment));
 
+	// TODO: Could just get the trailer from the owning FLinkerLoad if still attached
+
 	// Open a reader to the file
 	TUniquePtr<FArchive> BulkArchive;
 	if (!IsReferencingByPackagePath() || PackageSegment != EPackageSegment::Header)
@@ -1252,6 +1280,15 @@ void FVirtualizedUntypedBulkData::DetachFromDisk(FArchive* Ar, bool bEnsurePaylo
 		OffsetInFile = INDEX_NONE;
 
 		EnumRemoveFlags(Flags, EFlags::ReferencesLegacyFile | EFlags::ReferencesWorkspaceDomain | EFlags::LegacyFileIsCompressed);
+
+		if (PayloadSize > 0)
+		{
+			Register(nullptr);
+		}
+		else
+		{
+			Unregister();
+		}
 	}
 
 	AttachedAr = nullptr;	
@@ -1326,7 +1363,7 @@ void FVirtualizedUntypedBulkData::SerializeToLegacyPath(FLinkerSave& LinkerSave,
 
 void FVirtualizedUntypedBulkData::SerializeToPackageTrailer(FLinkerSave& LinkerSave, int64 OffsetPos, FCompressedBuffer PayloadToSerialize, EFlags UpdatedFlags, UObject* Owner)
 {
-	auto OnPayloadWritten = [this, OffsetPos, UpdatedFlags](FLinkerSave& LinkerSave, FPackageTrailer& Trailer) mutable
+	auto OnPayloadWritten = [this, OffsetPos, UpdatedFlags](FLinkerSave& LinkerSave, const FPackageTrailer& Trailer) mutable
 	{
 		checkf(LinkerSave.IsCooking() == false, TEXT("FVirtualizedUntypedBulkData::Serialize should not be called during a cook"));
 
@@ -1665,9 +1702,12 @@ FVirtualizedUntypedBulkData::EFlags FVirtualizedUntypedBulkData::BuildFlagsForSe
 			}
 			else
 			{
-				// If we are re-hydrating packages on save then we need to remove the virtualization flag
 				EnumRemoveFlags(UpdatedFlags, EFlags::HasPayloadSidecarFile);
-				if (LinkerSave != nullptr && bAllowVirtualizationOnSave == false)
+
+				// Remove the virtualization flag if we are rehydrating packages on save unless
+				// referencing the payload data is allowed, in which case we can continue to save
+				// as virtualized.
+				if (LinkerSave != nullptr && !bKeepFileDataByReference && bForceRehydrationOnSave)
 				{
 					EnumRemoveFlags(UpdatedFlags, EFlags::IsVirtualized);
 				}
diff --git a/Engine/Source/Runtime/CoreUObject/Private/UObject/PackageTrailer.cpp b/Engine/Source/Runtime/CoreUObject/Private/UObject/PackageTrailer.cpp
--- a/Engine/Source/Runtime/CoreUObject/Private/UObject/PackageTrailer.cpp
+++ b/Engine/Source/Runtime/CoreUObject/Private/UObject/PackageTrailer.cpp
@@ -123,21 +123,17 @@ bool FPackageTrailerBuilder::BuildAndAppendTrailer(FLinkerSave* Linker, FArchive
 	// 100% under our control. This will allow people to create external scripts that can
 	// parse and manipulate the trailer without needing to worry that we might change how
 	// our containers serialize. 
+	
+	// First we build a trailer structure
 	FPackageTrailer Trailer;
 	
-	Trailer.TrailerPositionInFile = DataArchive.Tell();
-
 	Trailer.Header.Tag = FPackageTrailer::FHeader::HeaderTag;
 	Trailer.Header.Version = (int32)EPackageTrailerVersion::AUTOMATIC_VERSION;
 
-	DataArchive << Trailer.Header.Tag;
-	DataArchive << Trailer.Header.Version;
-
 	const uint32 DynamicHeaderSizeOnDisk = ((LocalEntries.Num() + VirtualizedEntries.Num()) * Private::FLookupTableEntry::SizeOnDisk); // Add the length of the lookup table
 	
 	Trailer.Header.HeaderLength = FPackageTrailer::FHeader::StaticHeaderSizeOnDisk + DynamicHeaderSizeOnDisk;
-	DataArchive << Trailer.Header.HeaderLength;
-
+	
 	Trailer.Header.PayloadsDataLength = 0;
 	Trailer.Header.PayloadLookupTable.Reserve(LocalEntries.Num() + VirtualizedEntries.Num());
 
@@ -165,18 +161,13 @@ bool FPackageTrailerBuilder::BuildAndAppendTrailer(FLinkerSave* Linker, FArchive
 		Entry.RawSize = It.Value.RawSize;
 	}
 
-	DataArchive << Trailer.Header.PayloadsDataLength;
-
 	Trailer.Header.AccessMode = AccessMode;
-	DataArchive << Trailer.Header.AccessMode;
 
-	int32 NumPayloads = Trailer.Header.PayloadLookupTable.Num();
-	DataArchive << NumPayloads;
+	// Now that we have the complete trailer we can serialize it to the archive
 
-	for (Private::FLookupTableEntry& Entry : Trailer.Header.PayloadLookupTable)
-	{
-		DataArchive << Entry;
-	}
+	Trailer.TrailerPositionInFile = DataArchive.Tell();
+
+	DataArchive << Trailer.Header;
 
 	checkf((Trailer.TrailerPositionInFile + Trailer.Header.HeaderLength) == DataArchive.Tell(), TEXT("Header length was calculated as %d bytes but we wrote %" INT64_FMT " bytes!"), Trailer.Header.HeaderLength, DataArchive.Tell() - Trailer.TrailerPositionInFile);
 
@@ -189,18 +180,13 @@ bool FPackageTrailerBuilder::BuildAndAppendTrailer(FLinkerSave* Linker, FArchive
 
 	checkf((PayloadPosInFile + Trailer.Header.PayloadsDataLength) == DataArchive.Tell(), TEXT("Total payload length was calculated as %" INT64_FMT " bytes but we wrote %" INT64_FMT " bytes!"), Trailer.Header.PayloadsDataLength, DataArchive.Tell() - PayloadPosInFile);
 
-	FPackageTrailer::FFooter Footer;
-	
-	Footer.Tag = FPackageTrailer::FFooter::FooterTag;
-	Footer.TrailerLength = Trailer.Header.HeaderLength + Trailer.Header.PayloadsDataLength + FPackageTrailer::FFooter::SizeOnDisk;
-	Footer.PackageTag = PACKAGE_FILE_TAG;
-	
-	DataArchive << Footer.Tag;
-	DataArchive << Footer.TrailerLength;
-	DataArchive << Footer.PackageTag;
+	FPackageTrailer::FFooter Footer = Trailer.CreateFooter();
+	DataArchive << Footer;
 
 	checkf((Trailer.TrailerPositionInFile + Footer.TrailerLength) == DataArchive.Tell(), TEXT("Trailer length was calculated as %" INT64_FMT " bytes but we wrote %" INT64_FMT " bytes!"), Footer.TrailerLength, DataArchive.Tell() - Trailer.TrailerPositionInFile);
 
+	// Invoke any registered callbacks and pass in the trailer, this allows the callbacks to poll where 
+	// in the output archive the payload has been stored.
 	for (const AdditionalDataCallback& Callback : Callbacks)
 	{
 		Callback(*Linker, Trailer);
@@ -224,28 +210,6 @@ bool FPackageTrailer::IsEnabled()
 		{		
 			GConfig->GetBool(TEXT("Core.System"), TEXT("UsePackageTrailer"), bEnabled, GEngineIni);
 			UE_LOG(LogSerialization, Log, TEXT("UsePackageTrailer: '%s'"), bEnabled ? TEXT("true") : TEXT("false"));
-
-			// Check to make sure that the editor domain is not also enabled and assert if it is.
-			// Currently the package trailer system will not work correctly with the editor domain and as
-			// it is an opt in feature in development we should just prevent people running with both
-			// options enabled.
-			// We check the config values directly to avoid needing to introduce dependencies between modules.
-			if (bEnabled)
-			{
-				FConfigFile PlatformEngineIni;
-				FConfigCacheIni::LoadLocalIniFile(PlatformEngineIni, TEXT("Editor"), true);
-
-				bool bEditorDomainEnabled = false;
-				if (GConfig->GetBool(TEXT("CookSettings"), TEXT("EditorDomainEnabled"), bEditorDomainEnabled, GEditorIni))
-				{
-					checkf(!bEditorDomainEnabled, TEXT("The package trailer system does not yet work with the editor domain!"));
-				}
-
-				if (GConfig->GetBool(TEXT("EditorDomain"), TEXT("EditorDomainEnabled"), bEditorDomainEnabled, GEditorIni))
-				{
-					checkf(!bEditorDomainEnabled, TEXT("The package trailer system does not yet work with the editor domain!"));
-				}
-			}
 		}
 	} UsePackageTrailer;
 
@@ -269,6 +233,52 @@ bool FPackageTrailer::TryLoadFromPackage(const FPackagePath& PackagePath, FPacka
 	}	
 }
 
+FPackageTrailer FPackageTrailer::CreateReference(const FPackageTrailer& OriginalTrailer)
+{
+	check(OriginalTrailer.TrailerPositionInFile != INDEX_NONE);
+
+	FPackageTrailer ReferenceTrailer;
+	ReferenceTrailer.Header = OriginalTrailer.Header;
+	ReferenceTrailer.Header.AccessMode = EPayloadAccessMode::Referenced; // This trailer will reference payloads stored in another trailer
+	ReferenceTrailer.Header.PayloadsDataLength = 0; // This trailer will contain no payload data
+
+	// Convert the payload offsets into absolute offsets
+	const int64 PayloadOffset = OriginalTrailer.TrailerPositionInFile + OriginalTrailer.Header.HeaderLength;
+	for (Private::FLookupTableEntry& Entry : ReferenceTrailer.Header.PayloadLookupTable)
+	{
+		if (!Entry.IsVirtualized())
+		{
+			Entry.OffsetInFile += PayloadOffset;
+		}
+	}
+
+	return ReferenceTrailer;
+}
+
+bool FPackageTrailer::TrySave(FArchive& Ar)
+{
+	if (Header.AccessMode != EPayloadAccessMode::Referenced)
+	{
+		// Can only save reference trailers directly.
+		return false;
+	}
+
+	check(Header.PayloadsDataLength == 0);
+
+	const int64 PositionInArchive = Ar.Tell();
+
+	Ar << Header;
+	checkf((PositionInArchive + Header.HeaderLength) == Ar.Tell(), TEXT("Header length was calculated as %d bytes but we wrote %" INT64_FMT " bytes!"), Header.HeaderLength, Ar.Tell() - PositionInArchive);
+
+	FFooter Footer = CreateFooter();
+
+	Ar << Footer;
+
+	checkf((PositionInArchive + Footer.TrailerLength) == Ar.Tell(), TEXT("Trailer length was calculated as %" INT64_FMT " bytes but we wrote %" INT64_FMT " bytes!"), Footer.TrailerLength, Ar.Tell() - PositionInArchive);
+
+	return true;
+}
+
 bool FPackageTrailer::TryLoad(FArchive& Ar)
 {
 	check(Ar.IsLoading());
@@ -300,6 +310,11 @@ bool FPackageTrailer::TryLoad(FArchive& Ar)
 		Ar << Entry;
 	}
 
+	if (Header.AccessMode == EPayloadAccessMode::Referenced)
+	{
+		TrailerPositionInFile = INDEX_NONE;
+	}
+
 	return !Ar.IsError();
 }
 
@@ -480,6 +495,45 @@ EPayloadAccessMode FPackageTrailer::GetAccessMode() const
 	return Header.AccessMode;
 }
 
+FPackageTrailer::FFooter FPackageTrailer::CreateFooter() const
+{
+	FFooter Footer;
+
+	Footer.Tag = FPackageTrailer::FFooter::FooterTag;
+	Footer.TrailerLength = Header.HeaderLength + Header.PayloadsDataLength + FPackageTrailer::FFooter::SizeOnDisk;
+	Footer.PackageTag = PACKAGE_FILE_TAG;
+
+	return Footer;
+}
+
+FArchive& operator<<(FArchive& Ar, FPackageTrailer::FHeader& Header)
+{
+	Ar << Header.Tag;
+	Ar << Header.Version;
+	Ar << Header.HeaderLength;
+	Ar << Header.PayloadsDataLength;
+	Ar << Header.AccessMode;
+
+	int32 NumPayloads = Header.PayloadLookupTable.Num();
+	Ar << NumPayloads;
+
+	for (Private::FLookupTableEntry& Entry : Header.PayloadLookupTable)
+	{
+		Ar << Entry;
+	}
+
+	return Ar;
+}
+
+FArchive& operator<<(FArchive& Ar, FPackageTrailer::FFooter& Footer)
+{
+	Ar << Footer.Tag;
+	Ar << Footer.TrailerLength;
+	Ar << Footer.PackageTag;
+
+	return Ar;
+}
+
 bool FindPayloadsInPackageFile(const FPackagePath& PackagePath, EPayloadFilter Filter, TArray<Virtualization::FPayloadId>& OutPayloadIds)
 {
 	if (FPackageName::IsTextPackageExtension(PackagePath.GetHeaderExtension()))
diff --git a/Engine/Source/Runtime/CoreUObject/Private/UObject/SavePackage.cpp b/Engine/Source/Runtime/CoreUObject/Private/UObject/SavePackage.cpp
--- a/Engine/Source/Runtime/CoreUObject/Private/UObject/SavePackage.cpp
+++ b/Engine/Source/Runtime/CoreUObject/Private/UObject/SavePackage.cpp
@@ -1656,11 +1656,13 @@ struct FPackageExportTagger
 	}
 };
 
-[[nodiscard]] ESavePackageResult BuildAndWriteTrailer(UPackage* Package, FLinkerSave* Linker, FStructuredArchive::FRecord& StructuredArchiveRoot, IPackageWriter* PackageWriter, bool bTextFormat)
+[[nodiscard]] ESavePackageResult BuildAndWriteTrailer(UPackage* Package, FLinkerSave* Linker, FStructuredArchive::FRecord& StructuredArchiveRoot, IPackageWriter* PackageWriter, uint32 SaveFlags, bool bTextFormat)
 {
 	check(Package != nullptr);
 	check(Linker != nullptr);
-	
+
+	Linker->Summary.PayloadTocOffset = INDEX_NONE;
+
 	if (Linker->PackageTrailerBuilder.IsValid())
 	{
 		checkf(PackageWriter == nullptr, TEXT("Attempting to build a package trailer with a package writer '%s', this is not supported!"), *Package->GetName());
@@ -1672,9 +1674,20 @@ struct FPackageExportTagger
 			return ESavePackageResult::Error;
 		}
 	}
-	else
+	else if ((SaveFlags & SAVE_BulkDataByReference) != 0)
 	{
-		Linker->Summary.PayloadTocOffset = INDEX_NONE;
+		if (const FLinkerLoad* LinkerLoad = FLinkerLoad::FindExistingLinkerForPackage(Package))
+		{
+			if (const UE::FPackageTrailer* Trailer = LinkerLoad->GetPackageTrailer())
+			{
+				UE::FPackageTrailer ReferenceTrailer = UE::FPackageTrailer::CreateReference(*Trailer);
+				Linker->Summary.PayloadTocOffset = Linker->Tell();
+				if (!ReferenceTrailer.TrySave(*Linker))
+				{
+					return ESavePackageResult::Error;
+				}
+			}
+		}
 	}
 
 	return ESavePackageResult::Success;
@@ -3864,7 +3877,7 @@ FSavePackageResultStruct UPackage::Save(UPackage* InOuter, UObject* InAsset, con
 					// Now that the package is written out we can write the package trailer that is appended
 					// to the file. This should be the last thing written to the file!
 					SlowTask.EnterProgressFrame();
-					ESavePackageResult TrailerResult = BuildAndWriteTrailer(InOuter, Linker.Get(), StructuredArchiveRoot, PackageWriter, bTextFormat);
+					ESavePackageResult TrailerResult = BuildAndWriteTrailer(InOuter, Linker.Get(), StructuredArchiveRoot, PackageWriter, SaveFlags, bTextFormat);
 					if (TrailerResult != ESavePackageResult::Success)
 					{
 						return TrailerResult;
diff --git a/Engine/Source/Runtime/CoreUObject/Private/UObject/SavePackage2.cpp b/Engine/Source/Runtime/CoreUObject/Private/UObject/SavePackage2.cpp
--- a/Engine/Source/Runtime/CoreUObject/Private/UObject/SavePackage2.cpp
+++ b/Engine/Source/Runtime/CoreUObject/Private/UObject/SavePackage2.cpp
@@ -1724,6 +1724,8 @@ ESavePackageResult WriteExports(FStructuredArchive::FRecord& StructuredArchiveRo
 
 [[nodiscard]] ESavePackageResult BuildAndWriteTrailer(FStructuredArchive::FRecord& StructuredArchiveRoot, FSaveContext& SaveContext)
 {
+	SaveContext.Linker->Summary.PayloadTocOffset = INDEX_NONE;
+
 	if (SaveContext.Linker->PackageTrailerBuilder.IsValid())
 	{
 		checkf(SaveContext.GetPackageWriter() == nullptr, TEXT("Attempting to build a package trailer with a package writer '%s', this is not supported!"), *SaveContext.GetPackage()->GetName());
@@ -1735,9 +1737,22 @@ ESavePackageResult WriteExports(FStructuredArchive::FRecord& StructuredArchiveRo
 			return ESavePackageResult::Error;
 		}
 	}
-	else
-	{	
-		SaveContext.Linker->Summary.PayloadTocOffset = INDEX_NONE;
+	else if ((SaveContext.GetSaveArgs().SaveFlags & SAVE_BulkDataByReference) != 0)
+	{
+		check(SaveContext.GetPackage() != nullptr);
+		
+		if (const FLinkerLoad* LinkerLoad = FLinkerLoad::FindExistingLinkerForPackage(SaveContext.GetPackage()))
+		{
+			if (const UE::FPackageTrailer* Trailer = LinkerLoad->GetPackageTrailer())
+			{
+				UE::FPackageTrailer ReferenceTrailer = UE::FPackageTrailer::CreateReference(*Trailer);
+				SaveContext.Linker->Summary.PayloadTocOffset = SaveContext.Linker->Tell();
+				if (!ReferenceTrailer.TrySave(*SaveContext.Linker))
+				{
+					return ESavePackageResult::Error;
+				}
+			}
+		}
 	}
 
 	return ESavePackageResult::Success;
diff --git a/Engine/Source/Runtime/CoreUObject/Public/Serialization/VirtualizedBulkData.h b/Engine/Source/Runtime/CoreUObject/Public/Serialization/VirtualizedBulkData.h
--- a/Engine/Source/Runtime/CoreUObject/Public/Serialization/VirtualizedBulkData.h
+++ b/Engine/Source/Runtime/CoreUObject/Public/Serialization/VirtualizedBulkData.h
@@ -339,12 +339,14 @@ private:
 		HasRegistered				= 1 << 6,
 		/** The BulkData object is a copy used only to represent the id and payload; it does not communicate with the BulkDataRegistry, and will point DDC jobs toward the original BulkData */
 		IsTornOff					= 1 << 7,
-		/** The bulkdata object references a pyaload stored in a WorkspaceDomain file  */
+		/** The bulkdata object references a payload stored in a WorkspaceDomain file  */
 		ReferencesWorkspaceDomain	= 1 << 8,
 
 		TransientFlags				= HasRegistered | IsTornOff,
 	};
 
+	FRIEND_ENUM_CLASS_FLAGS(EFlags);
+
 	/** Used to control what level of error reporting we return from some methods */
 	enum ErrorVerbosity
 	{
@@ -354,8 +356,6 @@ private:
 		All
 	};
 
-	FRIEND_ENUM_CLASS_FLAGS(EFlags);
-
 	/** Old legacy path that saved the payload to the end of the package */
 	void SerializeToLegacyPath(FLinkerSave& LinkerSave, int64 OffsetPos, FCompressedBuffer PayloadToSerialize, EFlags UpdatedFlags, UObject* Owner);
 	/** The new path that saves payloads to the FPackageTrailer which is then appended to the end of the package file */
@@ -387,7 +387,12 @@ private:
 
 	bool IsDataVirtualized() const 
 	{ 
-		return EnumHasAnyFlags(Flags, EFlags::IsVirtualized); 
+		return IsDataVirtualized(Flags);
+	}
+
+	static bool IsDataVirtualized(EFlags InFlags)
+	{
+		return EnumHasAnyFlags(InFlags, EFlags::IsVirtualized);
 	}
 
 	bool HasPayloadSidecarFile() const 
diff --git a/Engine/Source/Runtime/CoreUObject/Public/UObject/PackageTrailer.h b/Engine/Source/Runtime/CoreUObject/Public/UObject/PackageTrailer.h
--- a/Engine/Source/Runtime/CoreUObject/Public/UObject/PackageTrailer.h
+++ b/Engine/Source/Runtime/CoreUObject/Public/UObject/PackageTrailer.h
@@ -136,7 +136,7 @@ struct FLookupTableEntry
 class COREUOBJECT_API FPackageTrailerBuilder
 {
 public:
-	using AdditionalDataCallback = TFunction<void(FLinkerSave& LinkerSave, class FPackageTrailer& Trailer)>;
+	using AdditionalDataCallback = TFunction<void(FLinkerSave& LinkerSave, const class FPackageTrailer& Trailer)>;
 
 	/**
 	 * Creates a builder from an already loaded FPackageTrailer. 
@@ -241,6 +241,13 @@ public:
 	/** Try to load a trailer from a given package path. Note that it will always try to load the trailer from the workspace domain */
 	[[nodiscard]] static bool TryLoadFromPackage(const FPackagePath& PackagePath, FPackageTrailer& OutTrailer);
 
+	/** 
+	 * Create a new trailer that references the payloads stored elsewhere rather than storing them locally.
+	 * 
+	 * @param SourceTrailer The trailer that currently stores the payloads (typically in the workspace domain)
+	 */
+	[[nodiscard]] static FPackageTrailer CreateReference(const FPackageTrailer& SourceTrailer);
+
 	FPackageTrailer() = default;
 	~FPackageTrailer() = default;
 
@@ -250,6 +257,9 @@ public:
 	FPackageTrailer(FPackageTrailer&& Other) = default;
 	FPackageTrailer& operator=(FPackageTrailer&& Other) = default;
 
+	/** Serializes the trailer to the given archive */
+	[[nodiscard]] bool TrySave(FArchive& Ar);
+
 	/** 
 	 * Serializes the trailer from the given archive assuming that the seek position of the archive is already at the correct position
 	 * for the trailer.
@@ -332,6 +342,9 @@ public:
 		EPayloadAccessMode AccessMode = EPayloadAccessMode::Relative;
 		/** Lookup table for the payloads on disk */
 		TArray<Private::FLookupTableEntry> PayloadLookupTable;
+
+		/** Serialization operator */
+		friend FArchive& operator<<(FArchive& Ar, FHeader& Header);
 	};
 
 	struct FFooter
@@ -350,11 +363,17 @@ public:
 		uint64 TrailerLength = 0;	
 		/** End the trailer with PACKAGE_FILE_TAG, which we expect all package files to end with */
 		uint32 PackageTag = 0;
+
+		/** Serialization operator */
+		friend FArchive& operator<<(FArchive& Ar, FFooter& Footer);
 	};
 
 private:
 	friend class FPackageTrailerBuilder;
 
+	/** Create a valid footer for the current trailer */
+	FFooter CreateFooter() const;
+
 	/** Where in the workspace domain package file the trailer is located */
 	int64 TrailerPositionInFile = INDEX_NONE;
 
