Refactoring compiled-in FProperty constructors so that they have the same signature

diff --git a/Engine/Source/Editor/PropertyEditor/Private/PropertyEditorModule.cpp b/Engine/Source/Editor/PropertyEditor/Private/PropertyEditorModule.cpp
--- a/Engine/Source/Editor/PropertyEditor/Private/PropertyEditorModule.cpp
+++ b/Engine/Source/Editor/PropertyEditor/Private/PropertyEditorModule.cpp
@@ -392,9 +392,10 @@ FStructProperty* FPropertyEditorModule::RegisterStructOnScopeProperty(TSharedRef
 			StructOnScopePropertyOwner->AddToRoot();
 		}
 		UScriptStruct* InnerStruct = Cast<UScriptStruct>(const_cast<UStruct*>(StructOnScope->GetStruct()));
-		StructProperty = new FStructProperty(StructOnScopePropertyOwner, *MakeUniqueObjectName(StructOnScopePropertyOwner, UField::StaticClass(), InnerStruct->GetFName()).ToString(), RF_Transient, 0, CPF_None, InnerStruct);
+		StructProperty = new FStructProperty(StructOnScopePropertyOwner, *MakeUniqueObjectName(StructOnScopePropertyOwner, UField::StaticClass(), InnerStruct->GetFName()).ToString(), RF_Transient);
 		StructProperty->Struct = InnerStruct;
 		StructProperty->ElementSize = StructOnScope->GetStruct()->GetStructureSize();
+		StructOnScopePropertyOwner->AddCppProperty(StructProperty);
 
 		RegisteredStructToProxyMap.Add(StructName, StructProperty);
 	}
diff --git a/Engine/Source/Runtime/CoreUObject/Private/UObject/EnumProperty.cpp b/Engine/Source/Runtime/CoreUObject/Private/UObject/EnumProperty.cpp
--- a/Engine/Source/Runtime/CoreUObject/Private/UObject/EnumProperty.cpp
+++ b/Engine/Source/Runtime/CoreUObject/Private/UObject/EnumProperty.cpp
@@ -56,7 +56,9 @@ FEnumProperty::FEnumProperty(FFieldVariant InOwner, const FName& InName, EObject
 }
 
 FEnumProperty::FEnumProperty(FFieldVariant InOwner, const FName& InName, EObjectFlags InObjectFlags, UEnum* InEnum)
+PRAGMA_DISABLE_DEPRECATION_WARNINGS
 	: FProperty(InOwner, InName, InObjectFlags, 0, CPF_HasGetValueTypeHash)
+PRAGMA_ENABLE_DEPRECATION_WARNINGS
 	, Enum(InEnum)
 {
 	// This is expected to be set post-construction by AddCppProperty
@@ -64,13 +66,24 @@ FEnumProperty::FEnumProperty(FFieldVariant InOwner, const FName& InName, EObject
 }
 
 FEnumProperty::FEnumProperty(FFieldVariant InOwner, const FName& InName, EObjectFlags InObjectFlags, int32 InOffset, EPropertyFlags InFlags, UEnum* InEnum)
+PRAGMA_DISABLE_DEPRECATION_WARNINGS
 	: FProperty(InOwner, InName, InObjectFlags, InOffset, InFlags | CPF_HasGetValueTypeHash)
+PRAGMA_ENABLE_DEPRECATION_WARNINGS
 	, Enum(InEnum)
 {
 	// This is expected to be set post-construction by AddCppProperty
 	UnderlyingProp = nullptr;
 }
 
+FEnumProperty::FEnumProperty(FFieldVariant InOwner, const UECodeGen_Private::FEnumPropertyParams& Prop)
+	: FProperty(InOwner, (const UECodeGen_Private::FPropertyParamsBaseWithOffset&)Prop, CPF_HasGetValueTypeHash)
+{
+	Enum = Prop.EnumFunc ? Prop.EnumFunc() : nullptr;
+
+	// This is expected to be set post-construction by AddCppProperty
+	UnderlyingProp = nullptr;
+}
+
 #if WITH_EDITORONLY_DATA
 FEnumProperty::FEnumProperty(UField* InField)
 	: FProperty(InField)
diff --git a/Engine/Source/Runtime/CoreUObject/Private/UObject/FieldPathProperty.cpp b/Engine/Source/Runtime/CoreUObject/Private/UObject/FieldPathProperty.cpp
--- a/Engine/Source/Runtime/CoreUObject/Private/UObject/FieldPathProperty.cpp
+++ b/Engine/Source/Runtime/CoreUObject/Private/UObject/FieldPathProperty.cpp
@@ -14,6 +14,12 @@
 
 IMPLEMENT_FIELD(FFieldPathProperty)
 
+FFieldPathProperty::FFieldPathProperty(FFieldVariant InOwner, const UECodeGen_Private::FFieldPathPropertyParams& Prop)
+	: FFieldPathProperty_Super(InOwner, (const UECodeGen_Private::FPropertyParamsBaseWithOffset&)Prop)
+{
+	PropertyClass = Prop.PropertyClassFunc();
+}
+
 #if WITH_EDITORONLY_DATA
 FFieldPathProperty::FFieldPathProperty(UField* InField)
 	: FFieldPathProperty_Super(InField)
diff --git a/Engine/Source/Runtime/CoreUObject/Private/UObject/Property.cpp b/Engine/Source/Runtime/CoreUObject/Private/UObject/Property.cpp
--- a/Engine/Source/Runtime/CoreUObject/Private/UObject/Property.cpp
+++ b/Engine/Source/Runtime/CoreUObject/Private/UObject/Property.cpp
@@ -721,6 +721,40 @@ FProperty::FProperty(FFieldVariant InOwner, const FName& InName, EObjectFlags In
 	Init();
 }
 
+FProperty::FProperty(FFieldVariant InOwner, const UECodeGen_Private::FPropertyParamsBaseWithOffset& Prop, EPropertyFlags AdditionalPropertyFlags /*= CPF_None*/)
+	: FField(InOwner, UTF8_TO_TCHAR(Prop.NameUTF8), Prop.ObjectFlags)
+	, ArrayDim(1)
+	, ElementSize(0)
+	, PropertyFlags(Prop.PropertyFlags | AdditionalPropertyFlags)
+	, RepIndex(0)
+	, BlueprintReplicationCondition(COND_None)
+	, Offset_Internal(0)
+	, PropertyLinkNext(nullptr)
+	, NextRef(nullptr)
+	, DestructorLinkNext(nullptr)
+	, PostConstructLinkNext(nullptr)
+{
+	this->Offset_Internal = Prop.Offset;
+
+	Init();
+}
+
+FProperty::FProperty(FFieldVariant InOwner, const UECodeGen_Private::FPropertyParamsBaseWithoutOffset& Prop, EPropertyFlags AdditionalPropertyFlags /*= CPF_None*/)
+	: FField(InOwner, UTF8_TO_TCHAR(Prop.NameUTF8), Prop.ObjectFlags)
+	, ArrayDim(1)
+	, ElementSize(0)
+	, PropertyFlags(Prop.PropertyFlags | AdditionalPropertyFlags)
+	, RepIndex(0)
+	, BlueprintReplicationCondition(COND_None)
+	, Offset_Internal(0)
+	, PropertyLinkNext(nullptr)
+	, NextRef(nullptr)
+	, DestructorLinkNext(nullptr)
+	, PostConstructLinkNext(nullptr)
+{
+	Init();
+}
+
 #if WITH_EDITORONLY_DATA
 FProperty::FProperty(UField* InField)
 	: Super(InField)
diff --git a/Engine/Source/Runtime/CoreUObject/Private/UObject/PropertyArray.cpp b/Engine/Source/Runtime/CoreUObject/Private/UObject/PropertyArray.cpp
--- a/Engine/Source/Runtime/CoreUObject/Private/UObject/PropertyArray.cpp
+++ b/Engine/Source/Runtime/CoreUObject/Private/UObject/PropertyArray.cpp
@@ -17,6 +17,14 @@
 -----------------------------------------------------------------------------*/
 IMPLEMENT_FIELD(FArrayProperty)
 
+FArrayProperty::FArrayProperty(FFieldVariant InOwner, const UECodeGen_Private::FArrayPropertyParams& Prop)
+	: Super(InOwner, (const UECodeGen_Private::FPropertyParamsBaseWithOffset&)Prop)
+	, Inner(nullptr)
+{
+	ArrayFlags = Prop.ArrayFlags;
+	SetElementSize();
+}
+
 #if WITH_EDITORONLY_DATA
 FArrayProperty::FArrayProperty(UField* InField)
 	: FArrayProperty_Super(InField)
diff --git a/Engine/Source/Runtime/CoreUObject/Private/UObject/PropertyBaseObject.cpp b/Engine/Source/Runtime/CoreUObject/Private/UObject/PropertyBaseObject.cpp
--- a/Engine/Source/Runtime/CoreUObject/Private/UObject/PropertyBaseObject.cpp
+++ b/Engine/Source/Runtime/CoreUObject/Private/UObject/PropertyBaseObject.cpp
@@ -20,6 +20,17 @@
 -----------------------------------------------------------------------------*/
 IMPLEMENT_FIELD(FObjectPropertyBase)
 
+FObjectPropertyBase::FObjectPropertyBase(FFieldVariant InOwner, const UECodeGen_Private::FObjectPropertyParams& Prop, EPropertyFlags AdditionalPropertyFlags /*= CPF_None*/)
+	: FProperty(InOwner, (const UECodeGen_Private::FPropertyParamsBaseWithOffset&)Prop, AdditionalPropertyFlags)
+{
+	PropertyClass = Prop.ClassFunc ? Prop.ClassFunc() : nullptr;
+}
+FObjectPropertyBase::FObjectPropertyBase(FFieldVariant InOwner, const UECodeGen_Private::FObjectPropertyParamsWithoutClass& Prop, EPropertyFlags AdditionalPropertyFlags /*= CPF_None*/)
+	: FProperty(InOwner, (const UECodeGen_Private::FPropertyParamsBaseWithOffset&)Prop)
+	, PropertyClass(nullptr)
+{
+}
+
 #if WITH_EDITORONLY_DATA
 FObjectPropertyBase::FObjectPropertyBase(UField* InField)
 	: FProperty(InField)
diff --git a/Engine/Source/Runtime/CoreUObject/Private/UObject/PropertyBool.cpp b/Engine/Source/Runtime/CoreUObject/Private/UObject/PropertyBool.cpp
--- a/Engine/Source/Runtime/CoreUObject/Private/UObject/PropertyBool.cpp
+++ b/Engine/Source/Runtime/CoreUObject/Private/UObject/PropertyBool.cpp
@@ -29,7 +29,9 @@ FBoolProperty::FBoolProperty(FFieldVariant InOwner, const FName& InName, EObject
 }
 
 FBoolProperty::FBoolProperty(FFieldVariant InOwner, const FName& InName, EObjectFlags InObjectFlags, int32 InOffset, EPropertyFlags InFlags, uint32 InBitMask, uint32 InElementSize, bool bIsNativeBool)
+PRAGMA_DISABLE_DEPRECATION_WARNINGS
 	: FProperty(InOwner, InName, InObjectFlags, InOffset, InFlags | CPF_HasGetValueTypeHash)
+PRAGMA_ENABLE_DEPRECATION_WARNINGS
 	, FieldSize(0)
 	, ByteOffset(0)
 	, ByteMask(1)
@@ -38,6 +40,46 @@ FBoolProperty::FBoolProperty(FFieldVariant InOwner, const FName& InName, EObject
 	SetBoolSize(InElementSize, bIsNativeBool, InBitMask);
 }
 
+FBoolProperty::FBoolProperty(FFieldVariant InOwner, const UECodeGen_Private::FBoolPropertyParams& Prop)
+	: FProperty(InOwner, (const UECodeGen_Private::FPropertyParamsBaseWithoutOffset&)Prop)
+	, FieldSize(0)
+	, ByteOffset(0)
+	, ByteMask(1)
+	, FieldMask(1)
+{
+	auto DoDetermineBitfieldOffsetAndMask = [](uint32& Offset, uint32& BitMask, void (*SetBit)(void* Obj), const SIZE_T SizeOf)
+	{
+		TUniquePtr<uint8[]> Buffer = MakeUnique<uint8[]>(SizeOf);
+
+		SetBit(Buffer.Get());
+
+		// Here we are making the assumption that bitfields are aligned in the struct. Probably true.
+		// If not, it may be ok unless we are on a page boundary or something, but the check will fire in that case.
+		// Have faith.
+		for (uint32 TestOffset = 0; TestOffset < SizeOf; TestOffset++)
+		{
+			if (uint8 Mask = Buffer[TestOffset])
+			{
+				Offset = TestOffset;
+				BitMask = (uint32)Mask;
+				check(FMath::RoundUpToPowerOfTwo(BitMask) == BitMask); // better be only one bit on
+				break;
+			}
+		}
+	};
+
+	uint32 Offset = 0;
+	uint32 BitMask = 0;
+	if (Prop.SetBitFunc)
+	{
+		DoDetermineBitfieldOffsetAndMask(Offset, BitMask, Prop.SetBitFunc, Prop.SizeOfOuter);
+		check(BitMask);
+	}
+
+	SetOffset_Internal(Offset);
+	SetBoolSize(Prop.ElementSize, !!(Prop.Flags & UECodeGen_Private::EPropertyGenFlags::NativeBool), BitMask);
+}
+
 #if WITH_EDITORONLY_DATA
 FBoolProperty::FBoolProperty(UField* InField)
 	: FProperty(InField)
diff --git a/Engine/Source/Runtime/CoreUObject/Private/UObject/PropertyByte.cpp b/Engine/Source/Runtime/CoreUObject/Private/UObject/PropertyByte.cpp
--- a/Engine/Source/Runtime/CoreUObject/Private/UObject/PropertyByte.cpp
+++ b/Engine/Source/Runtime/CoreUObject/Private/UObject/PropertyByte.cpp
@@ -22,6 +22,12 @@
 
 IMPLEMENT_FIELD(FByteProperty)
 
+FByteProperty::FByteProperty(FFieldVariant InOwner, const UECodeGen_Private::FBytePropertyParams& Prop)
+	: TProperty_Numeric(InOwner, (const UECodeGen_Private::FPropertyParamsBaseWithOffset&)Prop)
+{
+	this->Enum = Prop.EnumFunc ? Prop.EnumFunc() : nullptr;
+}
+
 #if WITH_EDITORONLY_DATA
 FByteProperty::FByteProperty(UField* InField)
 	: TProperty_Numeric(InField)
diff --git a/Engine/Source/Runtime/CoreUObject/Private/UObject/PropertyClass.cpp b/Engine/Source/Runtime/CoreUObject/Private/UObject/PropertyClass.cpp
--- a/Engine/Source/Runtime/CoreUObject/Private/UObject/PropertyClass.cpp
+++ b/Engine/Source/Runtime/CoreUObject/Private/UObject/PropertyClass.cpp
@@ -18,6 +18,16 @@
 -----------------------------------------------------------------------------*/
 IMPLEMENT_FIELD(FClassProperty)
 
+FClassProperty::FClassProperty(FFieldVariant InOwner, const UECodeGen_Private::FClassPropertyParams& Prop)
+	: FObjectProperty(InOwner, (const UECodeGen_Private::FObjectPropertyParams&)Prop)
+{
+	if (!PropertyClass)
+	{
+		PropertyClass = UClass::StaticClass();
+	}
+	MetaClass = Prop.MetaClassFunc ? Prop.MetaClassFunc() : nullptr;
+}
+
 #if WITH_EDITORONLY_DATA
 FClassProperty::FClassProperty(UField* InField)
 	: FObjectProperty(InField)
diff --git a/Engine/Source/Runtime/CoreUObject/Private/UObject/PropertyDelegate.cpp b/Engine/Source/Runtime/CoreUObject/Private/UObject/PropertyDelegate.cpp
--- a/Engine/Source/Runtime/CoreUObject/Private/UObject/PropertyDelegate.cpp
+++ b/Engine/Source/Runtime/CoreUObject/Private/UObject/PropertyDelegate.cpp
@@ -18,6 +18,12 @@
 -----------------------------------------------------------------------------*/
 IMPLEMENT_FIELD(FDelegateProperty)
 
+FDelegateProperty::FDelegateProperty(FFieldVariant InOwner, const UECodeGen_Private::FDelegatePropertyParams& Prop)
+	: FDelegateProperty_Super(InOwner, (const UECodeGen_Private::FPropertyParamsBaseWithOffset&)Prop)
+{
+	SignatureFunction = Prop.SignatureFunctionFunc ? Prop.SignatureFunctionFunc() : nullptr;
+}
+
 #if WITH_EDITORONLY_DATA
 FDelegateProperty::FDelegateProperty(UField* InField)
 	: FDelegateProperty_Super(InField)
diff --git a/Engine/Source/Runtime/CoreUObject/Private/UObject/PropertyInterface.cpp b/Engine/Source/Runtime/CoreUObject/Private/UObject/PropertyInterface.cpp
--- a/Engine/Source/Runtime/CoreUObject/Private/UObject/PropertyInterface.cpp
+++ b/Engine/Source/Runtime/CoreUObject/Private/UObject/PropertyInterface.cpp
@@ -21,6 +21,13 @@
 -----------------------------------------------------------------------------*/
 IMPLEMENT_FIELD(FInterfaceProperty)
 
+FInterfaceProperty::FInterfaceProperty(FFieldVariant InOwner, const UECodeGen_Private::FInterfacePropertyParams& Prop)
+	: FInterfaceProperty_Super(InOwner, (const UECodeGen_Private::FPropertyParamsBaseWithOffset&)Prop)
+{
+	this->PropertyFlags &= (~CPF_InterfaceClearMask);
+	this->InterfaceClass = Prop.InterfaceClassFunc ? Prop.InterfaceClassFunc() : nullptr;
+}
+
 #if WITH_EDITORONLY_DATA
 FInterfaceProperty::FInterfaceProperty(UField* InField)
 	: FInterfaceProperty_Super(InField)
diff --git a/Engine/Source/Runtime/CoreUObject/Private/UObject/PropertyLazyObjectPtr.cpp b/Engine/Source/Runtime/CoreUObject/Private/UObject/PropertyLazyObjectPtr.cpp
--- a/Engine/Source/Runtime/CoreUObject/Private/UObject/PropertyLazyObjectPtr.cpp
+++ b/Engine/Source/Runtime/CoreUObject/Private/UObject/PropertyLazyObjectPtr.cpp
@@ -11,6 +11,11 @@
 -----------------------------------------------------------------------------*/
 IMPLEMENT_FIELD(FLazyObjectProperty)
 
+FLazyObjectProperty::FLazyObjectProperty(FFieldVariant InOwner, const UECodeGen_Private::FLazyObjectPropertyParams& Prop)
+	: TFObjectPropertyBase(InOwner, Prop)
+{
+}
+
 FString FLazyObjectProperty::GetCPPType(FString* ExtendedTypeText/*=NULL*/, uint32 CPPExportFlags/*=0*/) const
 {
 	return GetCPPTypeCustom(ExtendedTypeText, CPPExportFlags,
diff --git a/Engine/Source/Runtime/CoreUObject/Private/UObject/PropertyMap.cpp b/Engine/Source/Runtime/CoreUObject/Private/UObject/PropertyMap.cpp
--- a/Engine/Source/Runtime/CoreUObject/Private/UObject/PropertyMap.cpp
+++ b/Engine/Source/Runtime/CoreUObject/Private/UObject/PropertyMap.cpp
@@ -194,7 +194,9 @@ FMapProperty::FMapProperty(FFieldVariant InOwner, const FName& InName, EObjectFl
 }
 
 FMapProperty::FMapProperty(FFieldVariant InOwner, const FName& InName, EObjectFlags InObjectFlags, int32 InOffset, EPropertyFlags InFlags, EMapPropertyFlags InMapFlags)
+PRAGMA_DISABLE_DEPRECATION_WARNINGS
 	: FMapProperty_Super(InOwner, InName, InObjectFlags, InOffset, InFlags)
+PRAGMA_ENABLE_DEPRECATION_WARNINGS
 {
 	// These are expected to be set post-construction by AddCppProperty
 	KeyProp   = nullptr;
@@ -203,6 +205,16 @@ FMapProperty::FMapProperty(FFieldVariant InOwner, const FName& InName, EObjectFl
 	MapFlags = InMapFlags;
 }
 
+FMapProperty::FMapProperty(FFieldVariant InOwner, const UECodeGen_Private::FMapPropertyParams& Prop)
+	: FMapProperty_Super(InOwner, (const UECodeGen_Private::FPropertyParamsBaseWithOffset&)Prop)
+{
+	// These are expected to be set post-construction by AddCppProperty
+	KeyProp = nullptr;
+	ValueProp = nullptr;
+
+	MapFlags = Prop.MapFlags;
+}
+
 #if WITH_EDITORONLY_DATA
 FMapProperty::FMapProperty(UField* InField)
 	: FMapProperty_Super(InField)
diff --git a/Engine/Source/Runtime/CoreUObject/Private/UObject/PropertyMulticastDelegate.cpp b/Engine/Source/Runtime/CoreUObject/Private/UObject/PropertyMulticastDelegate.cpp
--- a/Engine/Source/Runtime/CoreUObject/Private/UObject/PropertyMulticastDelegate.cpp
+++ b/Engine/Source/Runtime/CoreUObject/Private/UObject/PropertyMulticastDelegate.cpp
@@ -14,6 +14,12 @@
 
 FMulticastScriptDelegate::FInvocationList FMulticastDelegateProperty::EmptyList;
 
+FMulticastDelegateProperty::FMulticastDelegateProperty(FFieldVariant InOwner, const UECodeGen_Private::FMulticastDelegatePropertyParams& Prop, EPropertyFlags AdditionalPropertyFlags /*= CPF_None*/)
+	: FProperty(InOwner, (const UECodeGen_Private::FPropertyParamsBaseWithOffset&)Prop, AdditionalPropertyFlags)
+{
+	SignatureFunction = Prop.SignatureFunctionFunc ? Prop.SignatureFunctionFunc() : nullptr;
+}
+
 #if WITH_EDITORONLY_DATA
 FMulticastDelegateProperty::FMulticastDelegateProperty(UField* InField)
 	: FProperty(InField)
@@ -373,6 +379,11 @@ void FMulticastDelegateProperty::AddReferencedObjects(FReferenceCollector& Colle
 
 IMPLEMENT_FIELD(FMulticastDelegateProperty)
 
+FMulticastInlineDelegateProperty::FMulticastInlineDelegateProperty(FFieldVariant InOwner, const UECodeGen_Private::FMulticastDelegatePropertyParams& Prop)
+	: TProperty_MulticastDelegate(InOwner, Prop)
+{
+}
+
 const FMulticastScriptDelegate* FMulticastInlineDelegateProperty::GetMulticastDelegate(const void* PropertyValue) const
 {
 	return (const FMulticastScriptDelegate*)PropertyValue;
@@ -441,6 +452,11 @@ void FMulticastInlineDelegateProperty::ClearDelegate(UObject* Parent, void* Prop
 
 IMPLEMENT_FIELD(FMulticastInlineDelegateProperty)
 
+FMulticastSparseDelegateProperty::FMulticastSparseDelegateProperty(FFieldVariant InOwner, const UECodeGen_Private::FMulticastDelegatePropertyParams& Prop)
+	: TProperty_MulticastDelegate(InOwner, Prop)
+{
+}
+
 const FMulticastScriptDelegate* FMulticastSparseDelegateProperty::GetMulticastDelegate(const void* PropertyValue) const
 {
 	const FSparseDelegate* SparseDelegate = (const FSparseDelegate*)PropertyValue;
diff --git a/Engine/Source/Runtime/CoreUObject/Private/UObject/PropertyName.cpp b/Engine/Source/Runtime/CoreUObject/Private/UObject/PropertyName.cpp
--- a/Engine/Source/Runtime/CoreUObject/Private/UObject/PropertyName.cpp
+++ b/Engine/Source/Runtime/CoreUObject/Private/UObject/PropertyName.cpp
@@ -11,6 +11,11 @@
 -----------------------------------------------------------------------------*/
 IMPLEMENT_FIELD(FNameProperty)
 
+FNameProperty::FNameProperty(FFieldVariant InOwner, const UECodeGen_Private::FNamePropertyParams& Prop)
+	: FNameProperty_Super(InOwner, (const UECodeGen_Private::FPropertyParamsBaseWithOffset&)Prop)
+{
+}
+
 void FNameProperty::ExportTextItem( FString& ValueStr, const void* PropertyValue, const void* DefaultValue, UObject* Parent, int32 PortFlags, UObject* ExportRootScope ) const
 {
 	FName Temp = *(FName*)PropertyValue;
diff --git a/Engine/Source/Runtime/CoreUObject/Private/UObject/PropertyNumeric.cpp b/Engine/Source/Runtime/CoreUObject/Private/UObject/PropertyNumeric.cpp
--- a/Engine/Source/Runtime/CoreUObject/Private/UObject/PropertyNumeric.cpp
+++ b/Engine/Source/Runtime/CoreUObject/Private/UObject/PropertyNumeric.cpp
@@ -19,6 +19,11 @@ IMPLEMENT_FIELD(FFloatProperty)
 IMPLEMENT_FIELD(FDoubleProperty)
 IMPLEMENT_FIELD(FLargeWorldCoordinatesRealProperty)
 
+FNumericProperty::FNumericProperty(FFieldVariant InOwner, const UECodeGen_Private::FPropertyParamsBaseWithOffset& Prop, EPropertyFlags AdditionalPropertyFlags /*= CPF_None*/)
+	: FProperty(InOwner, Prop, AdditionalPropertyFlags)
+{
+}
+
 int64 FNumericProperty::ReadEnumAsInt64(FStructuredArchive::FSlot Slot, UStruct* DefaultsStruct, const FPropertyTag& Tag)
 {
 	//@warning: mirrors loading code in FByteProperty::SerializeItem() and FEnumProperty::SerializeItem()
@@ -232,3 +237,48 @@ FString FNumericProperty::GetNumericPropertyValueToString(void const* Data) cons
 	check(0);
 	return FString();
 }
+
+FInt8Property::FInt8Property(FFieldVariant InOwner, const UECodeGen_Private::FInt8PropertyParams& Prop)
+	: TProperty_Numeric(InOwner, (const UECodeGen_Private::FPropertyParamsBaseWithOffset&)Prop)
+{
+}
+
+FInt16Property::FInt16Property(FFieldVariant InOwner, const UECodeGen_Private::FInt16PropertyParams& Prop)
+	: TProperty_Numeric(InOwner, (const UECodeGen_Private::FPropertyParamsBaseWithOffset&)Prop)
+{
+}
+
+FIntProperty::FIntProperty(FFieldVariant InOwner, const UECodeGen_Private::FIntPropertyParams& Prop)
+	: TProperty_Numeric(InOwner, (const UECodeGen_Private::FPropertyParamsBaseWithOffset&)Prop)
+{
+}
+
+FInt64Property::FInt64Property(FFieldVariant InOwner, const UECodeGen_Private::FInt64PropertyParams& Prop)
+	: TProperty_Numeric(InOwner, (const UECodeGen_Private::FPropertyParamsBaseWithOffset&)Prop)
+{
+}
+
+FUInt16Property::FUInt16Property(FFieldVariant InOwner, const UECodeGen_Private::FFInt16PropertyParams& Prop)
+	: TProperty_Numeric(InOwner, (const UECodeGen_Private::FPropertyParamsBaseWithOffset&)Prop)
+{
+}
+
+FUInt32Property::FUInt32Property(FFieldVariant InOwner, const UECodeGen_Private::FUInt32PropertyParams& Prop)
+	: TProperty_Numeric(InOwner, (const UECodeGen_Private::FPropertyParamsBaseWithOffset&)Prop)
+{
+}
+
+FUInt64Property::FUInt64Property(FFieldVariant InOwner, const UECodeGen_Private::FUnsizedIntPropertyParams& Prop)
+	: TProperty_Numeric(InOwner, (const UECodeGen_Private::FPropertyParamsBaseWithOffset&)Prop)
+{
+}
+
+FFloatProperty::FFloatProperty(FFieldVariant InOwner, const UECodeGen_Private::FFloatPropertyParams& Prop)
+	: TProperty_Numeric(InOwner, (const UECodeGen_Private::FPropertyParamsBaseWithOffset&)Prop)
+{
+}
+
+FDoubleProperty::FDoubleProperty(FFieldVariant InOwner, const UECodeGen_Private::FDoublePropertyParams& Prop)
+	: TProperty_Numeric(InOwner, (const UECodeGen_Private::FPropertyParamsBaseWithOffset&)Prop)
+{
+}
\ No newline at end of file
diff --git a/Engine/Source/Runtime/CoreUObject/Private/UObject/PropertyObject.cpp b/Engine/Source/Runtime/CoreUObject/Private/UObject/PropertyObject.cpp
--- a/Engine/Source/Runtime/CoreUObject/Private/UObject/PropertyObject.cpp
+++ b/Engine/Source/Runtime/CoreUObject/Private/UObject/PropertyObject.cpp
@@ -14,6 +14,11 @@
 -----------------------------------------------------------------------------*/
 IMPLEMENT_FIELD(FObjectProperty)
 
+FObjectProperty::FObjectProperty(FFieldVariant InOwner, const UECodeGen_Private::FObjectPropertyParams& Prop)
+	: TFObjectPropertyBase(InOwner, Prop)
+{
+}
+
 FString FObjectProperty::GetCPPTypeCustom(FString* ExtendedTypeText, uint32 CPPExportFlags, const FString& InnerNativeTypeName)  const
 {
 	ensure(!InnerNativeTypeName.IsEmpty());
diff --git a/Engine/Source/Runtime/CoreUObject/Private/UObject/PropertySet.cpp b/Engine/Source/Runtime/CoreUObject/Private/UObject/PropertySet.cpp
--- a/Engine/Source/Runtime/CoreUObject/Private/UObject/PropertySet.cpp
+++ b/Engine/Source/Runtime/CoreUObject/Private/UObject/PropertySet.cpp
@@ -173,7 +173,16 @@ FSetProperty::FSetProperty(FFieldVariant InOwner, const FName& InName, EObjectFl
 }
 
 FSetProperty::FSetProperty(FFieldVariant InOwner, const FName& InName, EObjectFlags InObjectFlags, int32 InOffset, EPropertyFlags InFlags)
+PRAGMA_DISABLE_DEPRECATION_WARNINGS
 : FSetProperty_Super(InOwner, InName, InObjectFlags, InOffset, InFlags)
+PRAGMA_ENABLE_DEPRECATION_WARNINGS
+{
+	// This is expected to be set post-construction by AddCppProperty
+	ElementProp = nullptr;
+}
+
+FSetProperty::FSetProperty(FFieldVariant InOwner, const UECodeGen_Private::FSetPropertyParams& Prop)
+	: FSetProperty_Super(InOwner, (const UECodeGen_Private::FPropertyParamsBaseWithOffset&)Prop)
 {
 	// This is expected to be set post-construction by AddCppProperty
 	ElementProp = nullptr;
diff --git a/Engine/Source/Runtime/CoreUObject/Private/UObject/PropertySoftClassPtr.cpp b/Engine/Source/Runtime/CoreUObject/Private/UObject/PropertySoftClassPtr.cpp
--- a/Engine/Source/Runtime/CoreUObject/Private/UObject/PropertySoftClassPtr.cpp
+++ b/Engine/Source/Runtime/CoreUObject/Private/UObject/PropertySoftClassPtr.cpp
@@ -15,6 +15,12 @@
 -----------------------------------------------------------------------------*/
 IMPLEMENT_FIELD(FSoftClassProperty)
 
+FSoftClassProperty::FSoftClassProperty(FFieldVariant InOwner, const UECodeGen_Private::FSoftClassPropertyParams& Prop)
+	: Super(InOwner, (const UECodeGen_Private::FObjectPropertyParamsWithoutClass&)Prop, UClass::StaticClass())
+{
+	this->MetaClass = Prop.MetaClassFunc ? Prop.MetaClassFunc() : nullptr;
+}
+
 #if WITH_EDITORONLY_DATA
 FSoftClassProperty::FSoftClassProperty(UField* InField)
 	: Super(InField)
diff --git a/Engine/Source/Runtime/CoreUObject/Private/UObject/PropertySoftObjectPtr.cpp b/Engine/Source/Runtime/CoreUObject/Private/UObject/PropertySoftObjectPtr.cpp
--- a/Engine/Source/Runtime/CoreUObject/Private/UObject/PropertySoftObjectPtr.cpp
+++ b/Engine/Source/Runtime/CoreUObject/Private/UObject/PropertySoftObjectPtr.cpp
@@ -12,6 +12,16 @@
 -----------------------------------------------------------------------------*/
 IMPLEMENT_FIELD(FSoftObjectProperty)
 
+FSoftObjectProperty::FSoftObjectProperty(FFieldVariant InOwner, const UECodeGen_Private::FSoftObjectPropertyParams& Prop)
+	: TFObjectPropertyBase(InOwner, Prop)
+{
+}
+
+FSoftObjectProperty::FSoftObjectProperty(FFieldVariant InOwner, const UECodeGen_Private::FObjectPropertyParamsWithoutClass& Prop, UClass* InClass)
+	: TFObjectPropertyBase(InOwner, Prop, InClass)
+{
+}
+
 FString FSoftObjectProperty::GetCPPTypeCustom(FString* ExtendedTypeText, uint32 CPPExportFlags, const FString& InnerNativeTypeName) const
 {
 	ensure(!InnerNativeTypeName.IsEmpty());
diff --git a/Engine/Source/Runtime/CoreUObject/Private/UObject/PropertyStr.cpp b/Engine/Source/Runtime/CoreUObject/Private/UObject/PropertyStr.cpp
--- a/Engine/Source/Runtime/CoreUObject/Private/UObject/PropertyStr.cpp
+++ b/Engine/Source/Runtime/CoreUObject/Private/UObject/PropertyStr.cpp
@@ -12,6 +12,11 @@
 -----------------------------------------------------------------------------*/
 IMPLEMENT_FIELD(FStrProperty)
 
+FStrProperty::FStrProperty(FFieldVariant InOwner, const UECodeGen_Private::FStrPropertyParams& Prop)
+	: FStrProperty_Super(InOwner, (const UECodeGen_Private::FPropertyParamsBaseWithOffset&)Prop)
+{
+}
+
 EConvertFromTypeResult FStrProperty::ConvertFromType(const FPropertyTag& Tag, FStructuredArchive::FSlot Slot, uint8* Data, UStruct* DefaultsStruct)
 {
 	// Convert serialized text to string.
diff --git a/Engine/Source/Runtime/CoreUObject/Private/UObject/PropertyStruct.cpp b/Engine/Source/Runtime/CoreUObject/Private/UObject/PropertyStruct.cpp
--- a/Engine/Source/Runtime/CoreUObject/Private/UObject/PropertyStruct.cpp
+++ b/Engine/Source/Runtime/CoreUObject/Private/UObject/PropertyStruct.cpp
@@ -54,12 +54,31 @@ FStructProperty::FStructProperty(FFieldVariant InOwner, const FName& InName, EOb
 }
 
 FStructProperty::FStructProperty(FFieldVariant InOwner, const FName& InName, EObjectFlags InObjectFlags, int32 InOffset, EPropertyFlags InFlags, UScriptStruct* InStruct)
+PRAGMA_DISABLE_DEPRECATION_WARNINGS
 	: FProperty(InOwner, InName, InObjectFlags, InOffset, InStruct->GetCppStructOps() ? InStruct->GetCppStructOps()->GetComputedPropertyFlags() | InFlags : InFlags)
-	,	Struct( InStruct )
+PRAGMA_ENABLE_DEPRECATION_WARNINGS
+	, Struct( InStruct )
 {
 	ElementSize = Struct->PropertiesSize;
 }
 
+static EPropertyFlags GetStructComputedPropertyFlags(const UECodeGen_Private::FStructPropertyParams& Prop)
+{
+	EPropertyFlags ComputedPropertyFlags = CPF_None;
+	UScriptStruct* Struct = Prop.ScriptStructFunc ? Prop.ScriptStructFunc() : nullptr;
+	if (Struct && Struct->GetCppStructOps())
+	{
+		ComputedPropertyFlags = Struct->GetCppStructOps()->GetComputedPropertyFlags();
+	}
+	return ComputedPropertyFlags;
+}
+
+FStructProperty::FStructProperty(FFieldVariant InOwner, const UECodeGen_Private::FStructPropertyParams& Prop)
+	: FProperty(InOwner, (const UECodeGen_Private::FPropertyParamsBaseWithOffset&)Prop, GetStructComputedPropertyFlags(Prop))
+{
+	Struct = Prop.ScriptStructFunc ? Prop.ScriptStructFunc() : nullptr;
+}
+
 #if WITH_EDITORONLY_DATA
 FStructProperty::FStructProperty(UField* InField)
 	: FProperty(InField)
diff --git a/Engine/Source/Runtime/CoreUObject/Private/UObject/PropertyWeakObjectPtr.cpp b/Engine/Source/Runtime/CoreUObject/Private/UObject/PropertyWeakObjectPtr.cpp
--- a/Engine/Source/Runtime/CoreUObject/Private/UObject/PropertyWeakObjectPtr.cpp
+++ b/Engine/Source/Runtime/CoreUObject/Private/UObject/PropertyWeakObjectPtr.cpp
@@ -9,6 +9,11 @@
 -----------------------------------------------------------------------------*/
 IMPLEMENT_FIELD(FWeakObjectProperty)
 
+FWeakObjectProperty::FWeakObjectProperty(FFieldVariant InOwner, const UECodeGen_Private::FWeakObjectPropertyParams& Prop)
+	: TFObjectPropertyBase(InOwner, Prop)
+{
+}
+
 FString FWeakObjectProperty::GetCPPType( FString* ExtendedTypeText/*=NULL*/, uint32 CPPExportFlags/*=0*/ ) const
 {
 	return GetCPPTypeCustom(ExtendedTypeText, CPPExportFlags,
diff --git a/Engine/Source/Runtime/CoreUObject/Private/UObject/TextProperty.cpp b/Engine/Source/Runtime/CoreUObject/Private/UObject/TextProperty.cpp
--- a/Engine/Source/Runtime/CoreUObject/Private/UObject/TextProperty.cpp
+++ b/Engine/Source/Runtime/CoreUObject/Private/UObject/TextProperty.cpp
@@ -12,6 +12,11 @@
 
 IMPLEMENT_FIELD(FTextProperty)
 
+FTextProperty::FTextProperty(FFieldVariant InOwner, const UECodeGen_Private::FTextPropertyParams& Prop)
+	: FTextProperty_Super(InOwner, (const UECodeGen_Private::FPropertyParamsBaseWithOffset&)Prop)
+{
+}
+
 EConvertFromTypeResult FTextProperty::ConvertFromType(const FPropertyTag& Tag, FStructuredArchive::FSlot Slot, uint8* Data, UStruct* DefaultsStruct)
 {
 	// Convert serialized string to text.
diff --git a/Engine/Source/Runtime/CoreUObject/Private/UObject/UObjectGlobals.cpp b/Engine/Source/Runtime/CoreUObject/Private/UObject/UObjectGlobals.cpp
--- a/Engine/Source/Runtime/CoreUObject/Private/UObject/UObjectGlobals.cpp
+++ b/Engine/Source/Runtime/CoreUObject/Private/UObject/UObjectGlobals.cpp
@@ -4350,17 +4350,30 @@ FString FAssetMsg::GetAssetLogString(const UObject* Object, const FString& Messa
 
 namespace UECodeGen_Private
 {
+	template <typename PropertyType, typename PropertyParamsType>
+	PropertyType* NewFProperty(FFieldVariant Outer, const FPropertyParamsBase& PropBase)
+	{
+		const PropertyParamsType& Prop = (const PropertyParamsType&)PropBase;
+		PropertyType* NewProp = new PropertyType(Outer, Prop);
+
+#if WITH_METADATA
+		if (Prop.NumMetaData)
+		{
+			for (const FMetaDataPairParam* MetaDataData = Prop.MetaDataArray, *MetaDataEnd = MetaDataData + Prop.NumMetaData; MetaDataData != MetaDataEnd; ++MetaDataData)
+			{
+				NewProp->SetMetaData(UTF8_TO_TCHAR(MetaDataData->NameUTF8), UTF8_TO_TCHAR(MetaDataData->ValueUTF8));
+			}
+		}
+#endif
+		return NewProp;
+	}
+
 	void ConstructFProperty(FFieldVariant Outer, const FPropertyParamsBase* const*& PropertyArray, int32& NumProperties)
 	{
 		const FPropertyParamsBase* PropBase = *--PropertyArray;
 
 		uint32 ReadMore = 0;
 
-#if WITH_METADATA
-		const FMetaDataPairParam* MetaDataArray = nullptr;
-		int32                     NumMetaData   = 0;
-#endif
-
 		FProperty* NewProp = nullptr;
 		switch (PropBase->Flags & PropertyTypeMask)
 		{
@@ -4372,121 +4385,61 @@ namespace UECodeGen_Private
 
 			case EPropertyGenFlags::Byte:
 			{
-				const FBytePropertyParams* Prop = (const FBytePropertyParams*)PropBase;
-				NewProp = new FByteProperty(Outer, UTF8_TO_TCHAR(Prop->NameUTF8), Prop->ObjectFlags, Prop->Offset, Prop->PropertyFlags, Prop->EnumFunc ? Prop->EnumFunc() : nullptr);
-
-#if WITH_METADATA
-				MetaDataArray = Prop->MetaDataArray;
-				NumMetaData   = Prop->NumMetaData;
-#endif
+				NewProp = NewFProperty<FByteProperty, FBytePropertyParams>(Outer, *PropBase);
 			}
 			break;
 
 			case EPropertyGenFlags::Int8:
 			{
-				const FInt8PropertyParams* Prop = (const FInt8PropertyParams*)PropBase;
-				NewProp = new FInt8Property(Outer, UTF8_TO_TCHAR(Prop->NameUTF8), Prop->ObjectFlags, Prop->Offset, Prop->PropertyFlags);
-
-#if WITH_METADATA
-				MetaDataArray = Prop->MetaDataArray;
-				NumMetaData   = Prop->NumMetaData;
-#endif
+				NewProp = NewFProperty<FInt8Property, FInt8PropertyParams>(Outer, *PropBase);
 			}
 			break;
 
 			case EPropertyGenFlags::Int16:
 			{
-				const FInt16PropertyParams* Prop = (const FInt16PropertyParams*)PropBase;
-				NewProp = new FInt16Property(Outer, UTF8_TO_TCHAR(Prop->NameUTF8), Prop->ObjectFlags, Prop->Offset, Prop->PropertyFlags);
-
-#if WITH_METADATA
-				MetaDataArray = Prop->MetaDataArray;
-				NumMetaData   = Prop->NumMetaData;
-#endif
+				NewProp = NewFProperty<FInt16Property, FInt16PropertyParams>(Outer, *PropBase);
 			}
 			break;
 
 			case EPropertyGenFlags::Int:
 			{
-				const FIntPropertyParams* Prop = (const FIntPropertyParams*)PropBase;
-				NewProp = new FIntProperty(Outer, UTF8_TO_TCHAR(Prop->NameUTF8), Prop->ObjectFlags, Prop->Offset, Prop->PropertyFlags);
-
-#if WITH_METADATA
-				MetaDataArray = Prop->MetaDataArray;
-				NumMetaData   = Prop->NumMetaData;
-#endif
+				NewProp = NewFProperty<FIntProperty, FIntPropertyParams>(Outer, *PropBase);
 			}
 			break;
 
 			case EPropertyGenFlags::Int64:
 			{
-				const FInt64PropertyParams* Prop = (const FInt64PropertyParams*)PropBase;
-				NewProp = new FInt64Property(Outer, UTF8_TO_TCHAR(Prop->NameUTF8), Prop->ObjectFlags, Prop->Offset, Prop->PropertyFlags);
-
-#if WITH_METADATA
-				MetaDataArray = Prop->MetaDataArray;
-				NumMetaData   = Prop->NumMetaData;
-#endif
+				NewProp = NewFProperty<FInt64Property, FInt64PropertyParams>(Outer, *PropBase);
 			}
 			break;
 
 			case EPropertyGenFlags::UInt16:
 			{
-				const FFInt16PropertyParams* Prop = (const FFInt16PropertyParams*)PropBase;
-				NewProp = new FUInt16Property(Outer, UTF8_TO_TCHAR(Prop->NameUTF8), Prop->ObjectFlags, Prop->Offset, Prop->PropertyFlags);
-
-#if WITH_METADATA
-				MetaDataArray = Prop->MetaDataArray;
-				NumMetaData   = Prop->NumMetaData;
-#endif
+				NewProp = NewFProperty<FUInt16Property, FFInt16PropertyParams>(Outer, *PropBase);
 			}
 			break;
 
 			case EPropertyGenFlags::UInt32:
 			{
-				const FUInt32PropertyParams* Prop = (const FUInt32PropertyParams*)PropBase;
-				NewProp = new FUInt32Property(Outer, UTF8_TO_TCHAR(Prop->NameUTF8), Prop->ObjectFlags, Prop->Offset, Prop->PropertyFlags);
-
-#if WITH_METADATA
-				MetaDataArray = Prop->MetaDataArray;
-				NumMetaData   = Prop->NumMetaData;
-#endif
+				NewProp = NewFProperty<FUInt32Property, FUInt32PropertyParams>(Outer, *PropBase);
 			}
 			break;
 
 			case EPropertyGenFlags::UInt64:
 			{
-				const FFInt64PropertyParams* Prop = (const FFInt64PropertyParams*)PropBase;
-				NewProp = new FUInt64Property(Outer, UTF8_TO_TCHAR(Prop->NameUTF8), Prop->ObjectFlags, Prop->Offset, Prop->PropertyFlags);
-
-#if WITH_METADATA
-				MetaDataArray = Prop->MetaDataArray;
-				NumMetaData   = Prop->NumMetaData;
-#endif
+				NewProp = NewFProperty<FUInt64Property, FFInt64PropertyParams>(Outer, *PropBase);
 			}
 			break;
 
 			case EPropertyGenFlags::UnsizedInt:
 			{
-				const FUnsizedIntPropertyParams* Prop = (const FUnsizedIntPropertyParams*)PropBase;
-				NewProp = new FUInt64Property(Outer, UTF8_TO_TCHAR(Prop->NameUTF8), Prop->ObjectFlags, Prop->Offset, Prop->PropertyFlags);
-
-#if WITH_METADATA
-				MetaDataArray = Prop->MetaDataArray;
-				NumMetaData   = Prop->NumMetaData;
-#endif
+				NewProp = NewFProperty<FUInt64Property, FUnsizedIntPropertyParams>(Outer, *PropBase);
 			}
 			break;
 
 			case EPropertyGenFlags::UnsizedUInt:
 			{
-				const FUnsizedFIntPropertyParams* Prop = (const FUnsizedFIntPropertyParams*)PropBase;
-				NewProp = new FUInt64Property(Outer, UTF8_TO_TCHAR(Prop->NameUTF8), Prop->ObjectFlags, Prop->Offset, Prop->PropertyFlags);
-
-#if WITH_METADATA
-				MetaDataArray = Prop->MetaDataArray;
-				NumMetaData   = Prop->NumMetaData;
-#endif
+				NewProp = NewFProperty<FUInt64Property, FUnsizedFIntPropertyParams>(Outer, *PropBase);
 			}
 			break;
 
@@ -4495,13 +4448,7 @@ namespace UECodeGen_Private
 #endif
 			case EPropertyGenFlags::Float:
 			{
-				const FFloatPropertyParams* Prop = (const FFloatPropertyParams*)PropBase;
-				NewProp = new FFloatProperty(Outer, UTF8_TO_TCHAR(Prop->NameUTF8), Prop->ObjectFlags, Prop->Offset, Prop->PropertyFlags);
-
-#if WITH_METADATA
-				MetaDataArray = Prop->MetaDataArray;
-				NumMetaData   = Prop->NumMetaData;
-#endif
+				NewProp = NewFProperty<FFloatProperty, FFloatPropertyParams>(Outer, *PropBase);
 			}
 			break;
 
@@ -4510,54 +4457,13 @@ namespace UECodeGen_Private
 #endif
 			case EPropertyGenFlags::Double:
 			{
-				const FDoublePropertyParams* Prop = (const FDoublePropertyParams*)PropBase;
-				NewProp = new FDoubleProperty(Outer, UTF8_TO_TCHAR(Prop->NameUTF8), Prop->ObjectFlags, Prop->Offset, Prop->PropertyFlags);
-
-#if WITH_METADATA
-				MetaDataArray = Prop->MetaDataArray;
-				NumMetaData   = Prop->NumMetaData;
-#endif
+				NewProp = NewFProperty<FDoubleProperty, FDoublePropertyParams>(Outer, *PropBase);
 			}
 			break;
 
 			case EPropertyGenFlags::Bool:
 			{
-				auto DoDetermineBitfieldOffsetAndMask = [](uint32& Offset, uint32& BitMask, void (*SetBit)(void* Obj), const SIZE_T SizeOf)
-				{
-					TUniquePtr<uint8[]> Buffer = MakeUnique<uint8[]>(SizeOf);
-
-					SetBit(Buffer.Get());
-
-					// Here we are making the assumption that bitfields are aligned in the struct. Probably true.
-					// If not, it may be ok unless we are on a page boundary or something, but the check will fire in that case.
-					// Have faith.
-					for (uint32 TestOffset = 0; TestOffset < SizeOf; TestOffset++)
-					{
-						if (uint8 Mask = Buffer[TestOffset])
-						{
-							Offset = TestOffset;
-							BitMask = (uint32)Mask;
-							check(FMath::RoundUpToPowerOfTwo(BitMask) == BitMask); // better be only one bit on
-							break;
-						}
-					}
-				};
-
-				const FBoolPropertyParams* Prop = (const FBoolPropertyParams*)PropBase;
-				uint32 Offset = 0;
-				uint32 BitMask = 0;
-				if (Prop->SetBitFunc)
-				{
-					DoDetermineBitfieldOffsetAndMask(Offset, BitMask, Prop->SetBitFunc, Prop->SizeOfOuter);
-					check(BitMask);
-				}
-
-				NewProp = new FBoolProperty(Outer, UTF8_TO_TCHAR(Prop->NameUTF8), Prop->ObjectFlags, Offset, Prop->PropertyFlags, BitMask, Prop->ElementSize, !!(Prop->Flags & EPropertyGenFlags::NativeBool));
-
-#if WITH_METADATA
-				MetaDataArray = Prop->MetaDataArray;
-				NumMetaData   = Prop->NumMetaData;
-#endif
+				NewProp = NewFProperty<FBoolProperty, FBoolPropertyParams>(Outer, *PropBase);
 			}
 			break;
 
@@ -4565,60 +4471,30 @@ namespace UECodeGen_Private
 			{
 				if (EnumHasAllFlags(PropBase->Flags, EPropertyGenFlags::ObjectPtr))
 				{
-					const FObjectPtrPropertyParams* Prop = (const FObjectPtrPropertyParams*)PropBase;
-					NewProp = new FObjectPtrProperty(Outer, UTF8_TO_TCHAR(Prop->NameUTF8), Prop->ObjectFlags, Prop->Offset, Prop->PropertyFlags, Prop->ClassFunc ? Prop->ClassFunc() : nullptr);
-
-#if WITH_METADATA
-					MetaDataArray = Prop->MetaDataArray;
-					NumMetaData   = Prop->NumMetaData;
-#endif
+					NewProp = NewFProperty<FObjectPtrProperty, FObjectPtrPropertyParams>(Outer, *PropBase);
 				}
 				else
 				{
-					const FObjectPropertyParams* Prop = (const FObjectPropertyParams*)PropBase;
-					NewProp = new FObjectProperty(Outer, UTF8_TO_TCHAR(Prop->NameUTF8), Prop->ObjectFlags, Prop->Offset, Prop->PropertyFlags, Prop->ClassFunc ? Prop->ClassFunc() : nullptr);
-
-#if WITH_METADATA
-					MetaDataArray = Prop->MetaDataArray;
-					NumMetaData   = Prop->NumMetaData;
-#endif
+					NewProp = NewFProperty<FObjectProperty, FObjectPropertyParams>(Outer, *PropBase);
 				}
 			}
 			break;
 
 			case EPropertyGenFlags::WeakObject:
 			{
-				const FWeakObjectPropertyParams* Prop = (const FWeakObjectPropertyParams*)PropBase;
-				NewProp = new FWeakObjectProperty(Outer, UTF8_TO_TCHAR(Prop->NameUTF8), Prop->ObjectFlags, Prop->Offset, Prop->PropertyFlags, Prop->ClassFunc ? Prop->ClassFunc() : nullptr);
-
-#if WITH_METADATA
-				MetaDataArray = Prop->MetaDataArray;
-				NumMetaData   = Prop->NumMetaData;
-#endif
+				NewProp = NewFProperty<FWeakObjectProperty, FWeakObjectPropertyParams>(Outer, *PropBase);
 			}
 			break;
 
 			case EPropertyGenFlags::LazyObject:
 			{
-				const FLazyObjectPropertyParams* Prop = (const FLazyObjectPropertyParams*)PropBase;
-				NewProp = new FLazyObjectProperty(Outer, UTF8_TO_TCHAR(Prop->NameUTF8), Prop->ObjectFlags, Prop->Offset, Prop->PropertyFlags, Prop->ClassFunc ? Prop->ClassFunc() : nullptr);
-
-#if WITH_METADATA
-				MetaDataArray = Prop->MetaDataArray;
-				NumMetaData   = Prop->NumMetaData;
-#endif
+				NewProp = NewFProperty<FLazyObjectProperty, FLazyObjectPropertyParams>(Outer, *PropBase);
 			}
 			break;
 
 			case EPropertyGenFlags::SoftObject:
 			{
-				const FSoftObjectPropertyParams* Prop = (const FSoftObjectPropertyParams*)PropBase;
-				NewProp = new FSoftObjectProperty(Outer, UTF8_TO_TCHAR(Prop->NameUTF8), Prop->ObjectFlags, Prop->Offset, Prop->PropertyFlags, Prop->ClassFunc ? Prop->ClassFunc() : nullptr);
-
-#if WITH_METADATA
-				MetaDataArray = Prop->MetaDataArray;
-				NumMetaData   = Prop->NumMetaData;
-#endif
+				NewProp = NewFProperty<FSoftObjectProperty, FSoftObjectPropertyParams>(Outer, *PropBase);
 			}
 			break;
 
@@ -4626,218 +4502,112 @@ namespace UECodeGen_Private
 			{
 				if (EnumHasAllFlags(PropBase->Flags, EPropertyGenFlags::ObjectPtr))
 				{
-					const FClassPtrPropertyParams* Prop = (const FClassPtrPropertyParams*)PropBase;
-					NewProp = new FClassPtrProperty(Outer, UTF8_TO_TCHAR(Prop->NameUTF8), Prop->ObjectFlags, Prop->Offset, Prop->PropertyFlags, Prop->MetaClassFunc ? Prop->MetaClassFunc() : nullptr, Prop->ClassFunc ? Prop->ClassFunc() : nullptr);
-
-#if WITH_METADATA
-					MetaDataArray = Prop->MetaDataArray;
-					NumMetaData   = Prop->NumMetaData;
-#endif
+					NewProp = NewFProperty<FClassPtrProperty, FClassPtrPropertyParams>(Outer, *PropBase);
 				}
 				else
 				{
-					const FClassPropertyParams* Prop = (const FClassPropertyParams*)PropBase;
-					NewProp = new FClassProperty(Outer, UTF8_TO_TCHAR(Prop->NameUTF8), Prop->ObjectFlags, Prop->Offset, Prop->PropertyFlags, Prop->MetaClassFunc ? Prop->MetaClassFunc() : nullptr, Prop->ClassFunc ? Prop->ClassFunc() : nullptr);
-
-#if WITH_METADATA
-					MetaDataArray = Prop->MetaDataArray;
-					NumMetaData   = Prop->NumMetaData;
-#endif
+					NewProp = NewFProperty<FClassProperty, FClassPropertyParams>(Outer, *PropBase);
 				}
 			}
 			break;
 
 			case EPropertyGenFlags::SoftClass:
 			{
-				const FSoftClassPropertyParams* Prop = (const FSoftClassPropertyParams*)PropBase;
-				NewProp = new FSoftClassProperty(Outer, UTF8_TO_TCHAR(Prop->NameUTF8), Prop->ObjectFlags, Prop->Offset, Prop->PropertyFlags, Prop->MetaClassFunc ? Prop->MetaClassFunc() : nullptr);
-
-#if WITH_METADATA
-				MetaDataArray = Prop->MetaDataArray;
-				NumMetaData   = Prop->NumMetaData;
-#endif
+				NewProp = NewFProperty<FSoftClassProperty, FSoftClassPropertyParams>(Outer, *PropBase);
 			}
 			break;
 
 			case EPropertyGenFlags::Interface:
 			{
-				const FInterfacePropertyParams* Prop = (const FInterfacePropertyParams*)PropBase;
-				NewProp = new FInterfaceProperty(Outer, UTF8_TO_TCHAR(Prop->NameUTF8), Prop->ObjectFlags, Prop->Offset, Prop->PropertyFlags, Prop->InterfaceClassFunc ? Prop->InterfaceClassFunc() : nullptr);
-
-#if WITH_METADATA
-				MetaDataArray = Prop->MetaDataArray;
-				NumMetaData   = Prop->NumMetaData;
-#endif
+				NewProp = NewFProperty<FInterfaceProperty, FInterfacePropertyParams>(Outer, *PropBase);
 			}
 			break;
 
 			case EPropertyGenFlags::Name:
 			{
-				const FNamePropertyParams* Prop = (const FNamePropertyParams*)PropBase;
-				NewProp = new FNameProperty(Outer, UTF8_TO_TCHAR(Prop->NameUTF8), Prop->ObjectFlags, Prop->Offset, Prop->PropertyFlags);
-
-#if WITH_METADATA
-				MetaDataArray = Prop->MetaDataArray;
-				NumMetaData   = Prop->NumMetaData;
-#endif
+				NewProp = NewFProperty<FNameProperty, FNamePropertyParams>(Outer, *PropBase);
 			}
 			break;
 
 			case EPropertyGenFlags::Str:
 			{
-				const FStrPropertyParams* Prop = (const FStrPropertyParams*)PropBase;
-				NewProp = new FStrProperty(Outer, UTF8_TO_TCHAR(Prop->NameUTF8), Prop->ObjectFlags, Prop->Offset, Prop->PropertyFlags);
-
-#if WITH_METADATA
-				MetaDataArray = Prop->MetaDataArray;
-				NumMetaData   = Prop->NumMetaData;
-#endif
+				NewProp = NewFProperty<FStrProperty, FStrPropertyParams>(Outer, *PropBase);
 			}
 			break;
 
 			case EPropertyGenFlags::Array:
 			{
-				const FArrayPropertyParams* Prop = (const FArrayPropertyParams*)PropBase;
-				NewProp = new FArrayProperty(Outer, UTF8_TO_TCHAR(Prop->NameUTF8), Prop->ObjectFlags, Prop->Offset, Prop->PropertyFlags, Prop->ArrayFlags);
+				NewProp = NewFProperty<FArrayProperty, FArrayPropertyParams>(Outer, *PropBase);
 
 				// Next property is the array inner
 				ReadMore = 1;
-
-#if WITH_METADATA
-				MetaDataArray = Prop->MetaDataArray;
-				NumMetaData   = Prop->NumMetaData;
-#endif
 			}
 			break;
 
 			case EPropertyGenFlags::Map:
 			{
-				const FMapPropertyParams* Prop = (const FMapPropertyParams*)PropBase;
-				NewProp = new FMapProperty(Outer, UTF8_TO_TCHAR(Prop->NameUTF8), Prop->ObjectFlags, Prop->Offset, Prop->PropertyFlags, Prop->MapFlags);
+				NewProp = NewFProperty<FMapProperty, FMapPropertyParams>(Outer, *PropBase);
 
 				// Next two properties are the map key and value inners
 				ReadMore = 2;
-
-#if WITH_METADATA
-				MetaDataArray = Prop->MetaDataArray;
-				NumMetaData   = Prop->NumMetaData;
-#endif
 			}
 			break;
 
 			case EPropertyGenFlags::Set:
 			{
-				const FSetPropertyParams* Prop = (const FSetPropertyParams*)PropBase;
-				NewProp = new FSetProperty(Outer, UTF8_TO_TCHAR(Prop->NameUTF8), Prop->ObjectFlags, Prop->Offset, Prop->PropertyFlags);
+				NewProp = NewFProperty<FSetProperty, FSetPropertyParams>(Outer, *PropBase);
 
 				// Next property is the set inner
 				ReadMore = 1;
-
-#if WITH_METADATA
-				MetaDataArray = Prop->MetaDataArray;
-				NumMetaData   = Prop->NumMetaData;
-#endif
 			}
 			break;
 
 			case EPropertyGenFlags::Struct:
 			{
-				const FStructPropertyParams* Prop = (const FStructPropertyParams*)PropBase;
-				NewProp = new FStructProperty(Outer, UTF8_TO_TCHAR(Prop->NameUTF8), Prop->ObjectFlags, Prop->Offset, Prop->PropertyFlags, Prop->ScriptStructFunc ? Prop->ScriptStructFunc() : nullptr);
-
-#if WITH_METADATA
-				MetaDataArray = Prop->MetaDataArray;
-				NumMetaData   = Prop->NumMetaData;
-#endif
+				NewProp = NewFProperty<FStructProperty, FStructPropertyParams>(Outer, *PropBase);
 			}
 			break;
 
 			case EPropertyGenFlags::Delegate:
 			{
-				const FDelegatePropertyParams* Prop = (const FDelegatePropertyParams*)PropBase;
-				NewProp = new FDelegateProperty(Outer, UTF8_TO_TCHAR(Prop->NameUTF8), Prop->ObjectFlags, Prop->Offset, Prop->PropertyFlags, Prop->SignatureFunctionFunc ? Prop->SignatureFunctionFunc() : nullptr);
-
-#if WITH_METADATA
-				MetaDataArray = Prop->MetaDataArray;
-				NumMetaData   = Prop->NumMetaData;
-#endif
+				NewProp = NewFProperty<FDelegateProperty, FDelegatePropertyParams>(Outer, *PropBase);
 			}
 			break;
 
 			case EPropertyGenFlags::InlineMulticastDelegate:
 			{
-				const FMulticastDelegatePropertyParams* Prop = (const FMulticastDelegatePropertyParams*)PropBase;
-				NewProp = new FMulticastInlineDelegateProperty(Outer, UTF8_TO_TCHAR(Prop->NameUTF8), Prop->ObjectFlags, Prop->Offset, Prop->PropertyFlags, Prop->SignatureFunctionFunc ? Prop->SignatureFunctionFunc() : nullptr);
-
-#if WITH_METADATA
-				MetaDataArray = Prop->MetaDataArray;
-				NumMetaData   = Prop->NumMetaData;
-#endif
+				NewProp = NewFProperty<FMulticastInlineDelegateProperty, FMulticastDelegatePropertyParams>(Outer, *PropBase);
 			}
 			break;
 
 			case EPropertyGenFlags::SparseMulticastDelegate:
 			{
-				const FMulticastDelegatePropertyParams* Prop = (const FMulticastDelegatePropertyParams*)PropBase;
-				NewProp = new FMulticastSparseDelegateProperty(Outer, UTF8_TO_TCHAR(Prop->NameUTF8), Prop->ObjectFlags, Prop->Offset, Prop->PropertyFlags, Prop->SignatureFunctionFunc ? Prop->SignatureFunctionFunc() : nullptr);
-
-#if WITH_METADATA
-				MetaDataArray = Prop->MetaDataArray;
-				NumMetaData = Prop->NumMetaData;
-#endif
+				NewProp = NewFProperty<FMulticastSparseDelegateProperty, FMulticastDelegatePropertyParams>(Outer, *PropBase);
 			}
 			break;
 
 			case EPropertyGenFlags::Text:
 			{
-				const FTextPropertyParams* Prop = (const FTextPropertyParams*)PropBase;
-				NewProp = new FTextProperty(Outer, UTF8_TO_TCHAR(Prop->NameUTF8), Prop->ObjectFlags, Prop->Offset, Prop->PropertyFlags);
-
-#if WITH_METADATA
-				MetaDataArray = Prop->MetaDataArray;
-				NumMetaData   = Prop->NumMetaData;
-#endif
+				NewProp = NewFProperty<FTextProperty, FTextPropertyParams>(Outer, *PropBase);
 			}
 			break;
 
 			case EPropertyGenFlags::Enum:
 			{
-				const FEnumPropertyParams* Prop = (const FEnumPropertyParams*)PropBase;
-				NewProp = new FEnumProperty(Outer, UTF8_TO_TCHAR(Prop->NameUTF8), Prop->ObjectFlags, Prop->Offset, Prop->PropertyFlags, Prop->EnumFunc ? Prop->EnumFunc() : nullptr);
+				NewProp = NewFProperty<FEnumProperty, FEnumPropertyParams>(Outer, *PropBase);
 
 				// Next property is the underlying integer property
 				ReadMore = 1;
-
-#if WITH_METADATA
-				MetaDataArray = Prop->MetaDataArray;
-				NumMetaData   = Prop->NumMetaData;
-#endif
 			}
 			break;
 
 			case EPropertyGenFlags::FieldPath:
 			{
-				const FFieldPathPropertyParams* Prop = (const FFieldPathPropertyParams*)PropBase;
-				NewProp = new FFieldPathProperty(Outer, UTF8_TO_TCHAR(Prop->NameUTF8), Prop->ObjectFlags, Prop->Offset, Prop->PropertyFlags, Prop->PropertyClassFunc());
-
-#if WITH_METADATA
-				MetaDataArray = Prop->MetaDataArray;
-				NumMetaData = Prop->NumMetaData;
-#endif
+				NewProp = NewFProperty<FFieldPathProperty, FFieldPathPropertyParams>(Outer, *PropBase);
 			}
 			break;
 		}
 
-#if WITH_METADATA
-		if (NumMetaData)
-		{
-			for (const FMetaDataPairParam* MetaDataData = MetaDataArray, *MetaDataEnd = MetaDataData + NumMetaData; MetaDataData != MetaDataEnd; ++MetaDataData)
-			{
-				NewProp->SetMetaData(UTF8_TO_TCHAR(MetaDataData->NameUTF8), UTF8_TO_TCHAR(MetaDataData->ValueUTF8));
-			}
-		}
-#endif
-
 		NewProp->ArrayDim = PropBase->ArrayDim;
 		if (PropBase->RepNotifyFuncUTF8)
 		{
diff --git a/Engine/Source/Runtime/CoreUObject/Public/UObject/EnumProperty.h b/Engine/Source/Runtime/CoreUObject/Public/UObject/EnumProperty.h
--- a/Engine/Source/Runtime/CoreUObject/Public/UObject/EnumProperty.h
+++ b/Engine/Source/Runtime/CoreUObject/Public/UObject/EnumProperty.h
@@ -17,6 +17,14 @@ public:
 	FEnumProperty(FFieldVariant InOwner, const FName& InName, EObjectFlags InObjectFlags);
 	FEnumProperty(FFieldVariant InOwner, const FName& InName, EObjectFlags InObjectFlags, UEnum* InEnum);
 	FEnumProperty(FFieldVariant InOwner, const FName& InName, EObjectFlags InObjectFlags, int32 InOffset, EPropertyFlags InFlags, UEnum* InEnum);
+
+	/**
+	 * Constructor used for constructing compiled in properties
+	 * @param InOwner Owner of the property
+	 * @param Prop Pointer to the compiled in structure describing the property
+	 **/
+	FEnumProperty(FFieldVariant InOwner, const UECodeGen_Private::FEnumPropertyParams& Prop);
+
 #if WITH_EDITORONLY_DATA
 	explicit FEnumProperty(UField* InField);
 #endif // WITH_EDITORONLY_DATA
diff --git a/Engine/Source/Runtime/CoreUObject/Public/UObject/FieldPathProperty.h b/Engine/Source/Runtime/CoreUObject/Public/UObject/FieldPathProperty.h
--- a/Engine/Source/Runtime/CoreUObject/Public/UObject/FieldPathProperty.h
+++ b/Engine/Source/Runtime/CoreUObject/Public/UObject/FieldPathProperty.h
@@ -25,12 +25,22 @@ public:
 	{
 	}
 
+	UE_DEPRECATED(5.1, "Compiled-in property constructor is deprecated, use other constructors instead.")
 	FFieldPathProperty(FFieldVariant InOwner, const FName& InName, EObjectFlags InObjectFlags, int32 InOffset, EPropertyFlags InFlags, FFieldClass* InPropertyClass)
+PRAGMA_DISABLE_DEPRECATION_WARNINGS
 		: FFieldPathProperty_Super(InOwner, InName, InObjectFlags, InOffset, InFlags)
+PRAGMA_ENABLE_DEPRECATION_WARNINGS
 		, PropertyClass(InPropertyClass)
 	{
 	}
 
+	/**
+	 * Constructor used for constructing compiled in properties
+	 * @param InOwner Owner of the property
+	 * @param PropBase Pointer to the compiled in structure describing the property
+	 **/
+	FFieldPathProperty(FFieldVariant InOwner, const UECodeGen_Private::FFieldPathPropertyParams& Prop);
+
 #if WITH_EDITORONLY_DATA
 	explicit FFieldPathProperty(UField* InField);
 #endif // WITH_EDITORONLY_DATA
diff --git a/Engine/Source/Runtime/CoreUObject/Public/UObject/TextProperty.h b/Engine/Source/Runtime/CoreUObject/Public/UObject/TextProperty.h
--- a/Engine/Source/Runtime/CoreUObject/Public/UObject/TextProperty.h
+++ b/Engine/Source/Runtime/CoreUObject/Public/UObject/TextProperty.h
@@ -23,11 +23,21 @@ public:
 	{
 	}
 
+	UE_DEPRECATED(5.1, "Compiled-in property constructor is deprecated, use other constructors instead.")
 	FTextProperty(FFieldVariant InOwner, const FName& InName, EObjectFlags InObjectFlags, int32 InOffset, EPropertyFlags InFlags)
+PRAGMA_DISABLE_DEPRECATION_WARNINGS
 		: FTextProperty_Super(InOwner, InName, InObjectFlags, InOffset, InFlags)
+PRAGMA_ENABLE_DEPRECATION_WARNINGS
 	{
 	}
 
+	/**
+	 * Constructor used for constructing compiled in properties
+	 * @param InOwner Owner of the property
+	 * @param PropBase Pointer to the compiled in structure describing the property
+	 **/
+	FTextProperty(FFieldVariant InOwner, const UECodeGen_Private::FTextPropertyParams& Prop);
+
 #if WITH_EDITORONLY_DATA
 	explicit FTextProperty(UField* InField)
 		: FTextProperty_Super(InField)
diff --git a/Engine/Source/Runtime/CoreUObject/Public/UObject/UObjectGlobals.h b/Engine/Source/Runtime/CoreUObject/Public/UObject/UObjectGlobals.h
--- a/Engine/Source/Runtime/CoreUObject/Public/UObject/UObjectGlobals.h
+++ b/Engine/Source/Runtime/CoreUObject/Public/UObject/UObjectGlobals.h
@@ -2748,6 +2748,16 @@ namespace UECodeGen_Private
 		int32          ArrayDim;
 	};
 
+	struct FPropertyParamsBaseWithoutOffset // : FPropertyParamsBase
+	{
+		const char* NameUTF8;
+		const char* RepNotifyFuncUTF8;
+		EPropertyFlags    PropertyFlags;
+		EPropertyGenFlags Flags;
+		EObjectFlags   ObjectFlags;
+		int32          ArrayDim;
+	};
+
 	struct FPropertyParamsBaseWithOffset // : FPropertyParamsBase
 	{
 		const char*    NameUTF8;
@@ -2790,7 +2800,7 @@ namespace UECodeGen_Private
 #endif
 	};
 
-	struct FBoolPropertyParams // : FPropertyParamsBase
+	struct FBoolPropertyParams // : FPropertyParamsBaseWithoutOffset
 	{
 		const char*      NameUTF8;
 		const char*         RepNotifyFuncUTF8;
@@ -2807,7 +2817,38 @@ namespace UECodeGen_Private
 #endif
 	};
 
-	struct FObjectPropertyParams // : FPropertyParamsBaseWithOffset
+	struct FObjectPropertyParamsWithoutClass // : FPropertyParamsBaseWithOffset
+	{
+		const char* NameUTF8;
+		const char* RepNotifyFuncUTF8;
+		EPropertyFlags      PropertyFlags;
+		EPropertyGenFlags   Flags;
+		EObjectFlags     ObjectFlags;
+		int32            ArrayDim;
+		int32            Offset;
+#if WITH_METADATA
+		const FMetaDataPairParam* MetaDataArray;
+		int32                               NumMetaData;
+#endif
+	};
+
+	struct FObjectPropertyParamsWithClass // : FPropertyParamsBaseWithOffset
+	{
+		const char* NameUTF8;
+		const char* RepNotifyFuncUTF8;
+		EPropertyFlags      PropertyFlags;
+		EPropertyGenFlags   Flags;
+		EObjectFlags     ObjectFlags;
+		int32            ArrayDim;
+		int32            Offset;
+		UClass*        (*ClassFunc)();
+#if WITH_METADATA
+		const FMetaDataPairParam* MetaDataArray;
+		int32                               NumMetaData;
+#endif
+	};
+
+	struct FObjectPropertyParams // : FObjectPropertyParamsWithClass
 	{
 		const char*      NameUTF8;
 		const char*         RepNotifyFuncUTF8;
@@ -2823,7 +2864,7 @@ namespace UECodeGen_Private
 #endif
 	};
 
-	struct FClassPropertyParams // : FPropertyParamsBaseWithOffset
+	struct FClassPropertyParams // : FObjectPropertyParams
 	{
 		const char*      NameUTF8;
 		const char*         RepNotifyFuncUTF8;
@@ -2832,15 +2873,15 @@ namespace UECodeGen_Private
 		EObjectFlags     ObjectFlags;
 		int32            ArrayDim;
 		int32            Offset;
-		UClass*        (*MetaClassFunc)();
 		UClass*        (*ClassFunc)();
+		UClass*        (*MetaClassFunc)();		
 #if WITH_METADATA
 		const FMetaDataPairParam*           MetaDataArray;
 		int32                               NumMetaData;
 #endif
 	};
 
-	struct FSoftClassPropertyParams // : FPropertyParamsBaseWithOffset
+	struct FSoftClassPropertyParams // : FObjectPropertyParamsWithoutClass
 	{
 		const char*      NameUTF8;
 		const char*         RepNotifyFuncUTF8;
diff --git a/Engine/Source/Runtime/CoreUObject/Public/UObject/UnrealType.h b/Engine/Source/Runtime/CoreUObject/Public/UObject/UnrealType.h
--- a/Engine/Source/Runtime/CoreUObject/Public/UObject/UnrealType.h
+++ b/Engine/Source/Runtime/CoreUObject/Public/UObject/UnrealType.h
@@ -206,7 +206,24 @@ public:
 public:
 	// Constructors.
 	FProperty(FFieldVariant InOwner, const FName& InName, EObjectFlags InObjectFlags);
+
+	UE_DEPRECATED(5.1, "Compiled-in property constructor is deprecated, use other constructors instead.")
 	FProperty(FFieldVariant InOwner, const FName& InName, EObjectFlags InObjectFlags, int32 InOffset, EPropertyFlags InFlags);
+
+	/**
+	 * Constructor used for constructing compiled-in properties
+	 * @param InOwner Owner of the property
+	 * @param PropBase Pointer to the compiled in structure describing the property
+	 **/
+	FProperty(FFieldVariant InOwner, const UECodeGen_Private::FPropertyParamsBaseWithOffset& Prop, EPropertyFlags AdditionalPropertyFlags = CPF_None);
+
+	/**
+	 * Constructor used for constructing compiled-in properties
+	 * @param InOwner Owner of the property
+	 * @param PropBase Pointer to the compiled in structure describing the property
+	 **/
+	FProperty(FFieldVariant InOwner, const UECodeGen_Private::FPropertyParamsBaseWithoutOffset& Prop, EPropertyFlags AdditionalPropertyFlags = CPF_None);
+
 #if WITH_EDITORONLY_DATA
 	explicit FProperty(UField* InField);
 #endif // WITH_EDITORONLY_DATA
@@ -1142,12 +1159,29 @@ public:
 		SetElementSize();
 	}
 
+	UE_DEPRECATED(5.1, "Compiled-in property constructor is deprecated, use other constructors instead.")
 	TProperty(FFieldVariant InOwner, const FName& InName, EObjectFlags InObjectFlags, int32 InOffset, EPropertyFlags InFlags)
+PRAGMA_DISABLE_DEPRECATION_WARNINGS
 		: Super(InOwner, InName, InObjectFlags, InOffset, InFlags | TTypeFundamentals::GetComputedFlagsPropertyFlags())
+PRAGMA_ENABLE_DEPRECATION_WARNINGS
 	{
 		SetElementSize();
 	}
 
+	/**
+	 * Constructor used for constructing compiled in properties
+	 * @param InOwner Owner of the property
+	 * @param PropBase Pointer to the compiled in structure describing the property
+	 **/
+	template <typename PropertyParamsType>
+	TProperty(FFieldVariant InOwner, PropertyParamsType& Prop)
+		: Super(InOwner, Prop, TTypeFundamentals::GetComputedFlagsPropertyFlags())
+	{
+		SetElementSize();
+	}
+
+public:
+
 #if WITH_EDITORONLY_DATA
 	explicit TProperty(UField* InField)
 		: Super(InField)
@@ -1258,8 +1292,21 @@ public:
 	{
 	}
 
+	UE_DEPRECATED(5.1, "Compiled-in property constructor is deprecated, use other constructors instead.")
 	TProperty_WithEqualityAndSerializer(FFieldVariant InOwner, const FName& InName, EObjectFlags InObjectFlags, int32 InOffset, EPropertyFlags InFlags)
+PRAGMA_DISABLE_DEPRECATION_WARNINGS
 		: Super(InOwner, InName, InObjectFlags, InOffset, InFlags)
+PRAGMA_ENABLE_DEPRECATION_WARNINGS
+	{
+	}
+
+	/**
+	 * Constructor used for constructing compiled in properties
+	 * @param InOwner Owner of the property
+	 * @param PropBase Pointer to the compiled in structure describing the property
+	 **/
+	TProperty_WithEqualityAndSerializer(FFieldVariant InOwner, const UECodeGen_Private::FPropertyParamsBaseWithOffset& Prop)
+		: Super(InOwner, Prop)
 	{
 	}
 
@@ -1293,11 +1340,23 @@ class COREUOBJECT_API FNumericProperty : public FProperty
 
 	FNumericProperty(FFieldVariant InOwner, const FName& InName, EObjectFlags InObjectFlags)
 		: FProperty(InOwner, InName, InObjectFlags)
-	{}
+	{
+	}
 
+	UE_DEPRECATED(5.1, "Compiled-in property constructor is deprecated, use other constructors instead.")
 	FNumericProperty(FFieldVariant InOwner, const FName& InName, EObjectFlags InObjectFlags, int32 InOffset, EPropertyFlags InFlags)
+PRAGMA_DISABLE_DEPRECATION_WARNINGS
 		: FProperty(InOwner, InName, InObjectFlags, InOffset, InFlags)
-	{}
+PRAGMA_ENABLE_DEPRECATION_WARNINGS
+	{
+	}
+
+	/**
+	 * Constructor used for constructing compiled in properties
+	 * @param InOwner Owner of the property
+	 * @param PropBase Pointer to the compiled in structure describing the property
+	 **/
+	FNumericProperty(FFieldVariant InOwner, const UECodeGen_Private::FPropertyParamsBaseWithOffset& Prop, EPropertyFlags AdditionalPropertyFlags = CPF_None);
 
 #if WITH_EDITORONLY_DATA
 	explicit FNumericProperty(UField* InField)
@@ -1432,8 +1491,21 @@ public:
 	{
 	}
 
+	UE_DEPRECATED(5.1, "Compiled-in property constructor is deprecated, use other constructors instead.")
 	TProperty_Numeric(FFieldVariant InOwner, const FName& InName, EObjectFlags InObjectFlags, int32 InOffset, EPropertyFlags InFlags)
+PRAGMA_DISABLE_DEPRECATION_WARNINGS
 		: Super(InOwner, InName, InObjectFlags, InOffset, InFlags)
+PRAGMA_ENABLE_DEPRECATION_WARNINGS
+	{
+	}
+
+	/**
+	 * Constructor used for constructing compiled in properties
+	 * @param InOwner Owner of the property
+	 * @param PropBase Pointer to the compiled in structure describing the property
+	 **/
+	TProperty_Numeric(FFieldVariant InOwner, const UECodeGen_Private::FPropertyParamsBaseWithOffset& Prop)
+		: Super(InOwner, Prop)
 	{
 	}
 
@@ -1656,12 +1728,22 @@ class COREUOBJECT_API FByteProperty : public TProperty_Numeric<uint8>
 	{
 	}
 
+	UE_DEPRECATED(5.1, "Compiled-in property constructor is deprecated, use other constructors instead.")
 	FByteProperty(FFieldVariant InOwner, const FName& InName, EObjectFlags InObjectFlags, int32 InOffset, EPropertyFlags InFlags, UEnum* InEnum = nullptr)
+PRAGMA_DISABLE_DEPRECATION_WARNINGS
 		: TProperty_Numeric(InOwner, InName, InObjectFlags, InOffset, InFlags)
+PRAGMA_ENABLE_DEPRECATION_WARNINGS
 	,	Enum( InEnum )
 	{
 	}
 
+	/**
+	 * Constructor used for constructing compiled in properties
+	 * @param InOwner Owner of the property
+	 * @param PropBase Pointer to the compiled in structure describing the property
+	 **/
+	FByteProperty(FFieldVariant InOwner, const UECodeGen_Private::FBytePropertyParams& Prop);
+
 #if WITH_EDITORONLY_DATA
 	explicit FByteProperty(UField* InField);
 #endif // WITH_EDITORONLY_DATA
@@ -1714,11 +1796,21 @@ class COREUOBJECT_API FInt8Property : public TProperty_Numeric<int8>
 	{
 	}
 
+	UE_DEPRECATED(5.1, "Compiled-in property constructor is deprecated, use other constructors instead.")
 	FInt8Property(FFieldVariant InOwner, const FName& InName, EObjectFlags InObjectFlags, int32 InOffset, EPropertyFlags InFlags)
+PRAGMA_DISABLE_DEPRECATION_WARNINGS
 		: TProperty_Numeric(InOwner, InName, InObjectFlags, InOffset, InFlags)
+PRAGMA_ENABLE_DEPRECATION_WARNINGS
 	{
 	}
 
+	/**
+	 * Constructor used for constructing compiled in properties
+	 * @param InOwner Owner of the property
+	 * @param PropBase Pointer to the compiled in structure describing the property
+	 **/
+	FInt8Property(FFieldVariant InOwner, const UECodeGen_Private::FInt8PropertyParams& Prop);
+
 #if WITH_EDITORONLY_DATA
 	explicit FInt8Property(UField* InField)
 		: TProperty_Numeric(InField)
@@ -1743,11 +1835,21 @@ class COREUOBJECT_API FInt16Property : public TProperty_Numeric<int16>
 	{
 	}
 
+	UE_DEPRECATED(5.1, "Compiled-in property constructor is deprecated, use other constructors instead.")
 	FInt16Property(FFieldVariant InOwner, const FName& InName, EObjectFlags InObjectFlags, int32 InOffset, EPropertyFlags InFlags)
+PRAGMA_DISABLE_DEPRECATION_WARNINGS
 		: TProperty_Numeric(InOwner, InName, InObjectFlags, InOffset, InFlags)
+PRAGMA_ENABLE_DEPRECATION_WARNINGS
 	{
 	}
 
+	/**
+	 * Constructor used for constructing compiled in properties
+	 * @param InOwner Owner of the property
+	 * @param PropBase Pointer to the compiled in structure describing the property
+	 **/
+	FInt16Property(FFieldVariant InOwner, const UECodeGen_Private::FInt16PropertyParams& Prop);
+
 #if WITH_EDITORONLY_DATA
 	explicit FInt16Property(UField* InField)
 		: TProperty_Numeric(InField)
@@ -1773,11 +1875,21 @@ class COREUOBJECT_API FIntProperty : public TProperty_Numeric<int32>
 	{
 	}
 
+	UE_DEPRECATED(5.1, "Compiled-in property constructor is deprecated, use other constructors instead.")
 	FIntProperty(FFieldVariant InOwner, const FName& InName, EObjectFlags InObjectFlags, int32 InOffset, EPropertyFlags InFlags)
+PRAGMA_DISABLE_DEPRECATION_WARNINGS
 		: TProperty_Numeric(InOwner, InName, InObjectFlags, InOffset, InFlags)
+PRAGMA_ENABLE_DEPRECATION_WARNINGS
 	{
 	}
 
+	/**
+	 * Constructor used for constructing compiled in properties
+	 * @param InOwner Owner of the property
+	 * @param PropBase Pointer to the compiled in structure describing the property
+	 **/
+	FIntProperty(FFieldVariant InOwner, const UECodeGen_Private::FIntPropertyParams& Prop);
+
 #if WITH_EDITORONLY_DATA
 	explicit FIntProperty(UField* InField)
 		: TProperty_Numeric(InField)
@@ -1802,11 +1914,21 @@ class COREUOBJECT_API FInt64Property : public TProperty_Numeric<int64>
 	{
 	}
 
+	UE_DEPRECATED(5.1, "Compiled-in property constructor is deprecated, use other constructors instead.")
 	FInt64Property(FFieldVariant InOwner, const FName& InName, EObjectFlags InObjectFlags, int32 InOffset, EPropertyFlags InFlags)
+PRAGMA_DISABLE_DEPRECATION_WARNINGS
 		: TProperty_Numeric(InOwner, InName, InObjectFlags, InOffset, InFlags)
+PRAGMA_ENABLE_DEPRECATION_WARNINGS
 	{
 	}
 
+	/**
+	 * Constructor used for constructing compiled in properties
+	 * @param InOwner Owner of the property
+	 * @param PropBase Pointer to the compiled in structure describing the property
+	 **/
+	FInt64Property(FFieldVariant InOwner, const UECodeGen_Private::FInt64PropertyParams& Prop);
+
 #if WITH_EDITORONLY_DATA
 	explicit FInt64Property(UField* InField)
 		: TProperty_Numeric(InField)
@@ -1831,11 +1953,21 @@ class COREUOBJECT_API FUInt16Property : public TProperty_Numeric<uint16>
 	{
 	}
 
+	UE_DEPRECATED(5.1, "Compiled-in property constructor is deprecated, use other constructors instead.")
 	FUInt16Property(FFieldVariant InOwner, const FName& InName, EObjectFlags InObjectFlags, int32 InOffset, EPropertyFlags InFlags)
+PRAGMA_DISABLE_DEPRECATION_WARNINGS
 		: TProperty_Numeric(InOwner, InName, InObjectFlags, InOffset, InFlags)
+PRAGMA_ENABLE_DEPRECATION_WARNINGS
 	{
 	}
 
+	/**
+	 * Constructor used for constructing compiled in properties
+	 * @param InOwner Owner of the property
+	 * @param PropBase Pointer to the compiled in structure describing the property
+	 **/
+	FUInt16Property(FFieldVariant InOwner, const UECodeGen_Private::FFInt16PropertyParams& Prop);
+
 #if WITH_EDITORONLY_DATA
 	explicit FUInt16Property(UField* InField)
 		: TProperty_Numeric(InField)
@@ -1860,11 +1992,21 @@ class COREUOBJECT_API FUInt32Property : public TProperty_Numeric<uint32>
 	{
 	}
 
+	UE_DEPRECATED(5.1, "Compiled-in property constructor is deprecated, use other constructors instead.")
 	FUInt32Property(FFieldVariant InOwner, const FName& InName, EObjectFlags InObjectFlags, int32 InOffset, EPropertyFlags InFlags)
+PRAGMA_DISABLE_DEPRECATION_WARNINGS
 	:	TProperty_Numeric(InOwner, InName, InObjectFlags, InOffset, InFlags)
+PRAGMA_ENABLE_DEPRECATION_WARNINGS
 	{
 	}
 
+	/**
+	 * Constructor used for constructing compiled in properties
+	 * @param InOwner Owner of the property
+	 * @param PropBase Pointer to the compiled in structure describing the property
+	 **/
+	FUInt32Property(FFieldVariant InOwner, const UECodeGen_Private::FUInt32PropertyParams& Prop);
+
 #if WITH_EDITORONLY_DATA
 	explicit FUInt32Property(UField* InField)
 		: TProperty_Numeric(InField)
@@ -1889,11 +2031,21 @@ class COREUOBJECT_API FUInt64Property : public TProperty_Numeric<uint64>
 	{
 	}
 
+	UE_DEPRECATED(5.1, "Compiled-in property constructor is deprecated, use other constructors instead.")
 	FUInt64Property(FFieldVariant InOwner, const FName& InName, EObjectFlags InObjectFlags, int32 InOffset, EPropertyFlags InFlags)
+PRAGMA_DISABLE_DEPRECATION_WARNINGS
 		: TProperty_Numeric(InOwner, InName, InObjectFlags, InOffset, InFlags)
+PRAGMA_ENABLE_DEPRECATION_WARNINGS
 	{
 	}
 
+	/**
+	 * Constructor used for constructing compiled in properties
+	 * @param InOwner Owner of the property
+	 * @param PropBase Pointer to the compiled in structure describing the property
+	 **/
+	FUInt64Property(FFieldVariant InOwner, const UECodeGen_Private::FUnsizedIntPropertyParams& Prop);
+
 #if WITH_EDITORONLY_DATA
 	explicit FUInt64Property(UField* InField)
 		: TProperty_Numeric(InField)
@@ -1941,11 +2093,21 @@ class COREUOBJECT_API FFloatProperty : public TProperty_Numeric<float>
 	{
 	}
 
+	UE_DEPRECATED(5.1, "Compiled-in property constructor is deprecated, use other constructors instead.")
 	FFloatProperty(FFieldVariant InOwner, const FName& InName, EObjectFlags InObjectFlags, int32 InOffset, EPropertyFlags InFlags)
+PRAGMA_DISABLE_DEPRECATION_WARNINGS
 		: TProperty_Numeric(InOwner, InName, InObjectFlags, InOffset, InFlags)
+PRAGMA_ENABLE_DEPRECATION_WARNINGS
 	{
 	}
 
+	/**
+	 * Constructor used for constructing compiled in properties
+	 * @param InOwner Owner of the property
+	 * @param PropBase Pointer to the compiled in structure describing the property
+	 **/
+	FFloatProperty(FFieldVariant InOwner, const UECodeGen_Private::FFloatPropertyParams& Prop);
+
 #if WITH_EDITORONLY_DATA
 	explicit FFloatProperty(UField* InField)
 		: TProperty_Numeric(InField)
@@ -1974,11 +2136,21 @@ class COREUOBJECT_API FDoubleProperty : public TProperty_Numeric<double>
 	{
 	}
 
+	UE_DEPRECATED(5.1, "Compiled-in property constructor is deprecated, use other constructors instead.")
 	FDoubleProperty(FFieldVariant InOwner, const FName& InName, EObjectFlags InObjectFlags, int32 InOffset, EPropertyFlags InFlags)
+PRAGMA_DISABLE_DEPRECATION_WARNINGS
 		: TProperty_Numeric(InOwner, InName, InObjectFlags, InOffset, InFlags)
+PRAGMA_ENABLE_DEPRECATION_WARNINGS
 	{
 	}
 
+	/**
+	 * Constructor used for constructing compiled in properties
+	 * @param InOwner Owner of the property
+	 * @param PropBase Pointer to the compiled in structure describing the property
+	 **/
+	FDoubleProperty(FFieldVariant InOwner, const UECodeGen_Private::FDoublePropertyParams& Prop);
+
 #if WITH_EDITORONLY_DATA
 	explicit FDoubleProperty(UField* InField)
 		: TProperty_Numeric(InField)
@@ -2039,8 +2211,16 @@ public:
 	 * @param InElementSize Sizeof of the boolean type this property represents.
 	 * @param bIsNativeBool true if this property represents C++ bool type.
 	 */
+	UE_DEPRECATED(5.1, "Compiled-in property constructor is deprecated, use other constructors instead.")
 	FBoolProperty(FFieldVariant InOwner, const FName& InName, EObjectFlags InObjectFlags, int32 InOffset, EPropertyFlags InFlags, uint32 InBitMask, uint32 InElementSize, bool bIsNativeBool);
 
+	/**
+	 * Constructor used for constructing compiled in properties
+	 * @param InOwner Owner of the property
+	 * @param PropBase Pointer to the compiled in structure describing the property
+	 **/
+	FBoolProperty(FFieldVariant InOwner, const UECodeGen_Private::FBoolPropertyParams& Prop);
+
 #if WITH_EDITORONLY_DATA
 	explicit FBoolProperty(UField* InField);
 #endif // WITH_EDITORONLY_DATA
@@ -2158,11 +2338,22 @@ public:
 		, PropertyClass(nullptr)
 	{}
 
+	UE_DEPRECATED(5.1, "Compiled-in property constructor is deprecated, use other constructors instead.")
 	FObjectPropertyBase(FFieldVariant InOwner, const FName& InName, EObjectFlags InObjectFlags, int32 InOffset, EPropertyFlags InFlags, UClass* InClass = NULL)
+PRAGMA_DISABLE_DEPRECATION_WARNINGS
 		: FProperty(InOwner, InName, InObjectFlags, InOffset, InFlags)
+PRAGMA_ENABLE_DEPRECATION_WARNINGS
 		, PropertyClass(InClass)
 	{}
 
+	/**
+	 * Constructor used for constructing compiled-in properties
+	 * @param InOwner Owner of the property
+	 * @param PropBase Pointer to the compiled in structure describing the property
+	 **/
+	FObjectPropertyBase(FFieldVariant InOwner, const UECodeGen_Private::FObjectPropertyParams& Prop, EPropertyFlags AdditionalPropertyFlags = CPF_None);
+	FObjectPropertyBase(FFieldVariant InOwner, const UECodeGen_Private::FObjectPropertyParamsWithoutClass& Prop, EPropertyFlags AdditionalPropertyFlags = CPF_None);
+
 #if WITH_EDITORONLY_DATA
 	explicit FObjectPropertyBase(UField* InField);
 #endif // WITH_EDITORONLY_DATA
@@ -2317,13 +2508,38 @@ public:
 		this->PropertyClass = nullptr;
 	}
 
-
+	UE_DEPRECATED(5.1, "Compiled-in property constructor is deprecated, use other constructors instead.")
 	TFObjectPropertyBase(FFieldVariant InOwner, const FName& InName, EObjectFlags InObjectFlags, int32 InOffset, EPropertyFlags InFlags, UClass* InClass)
+PRAGMA_DISABLE_DEPRECATION_WARNINGS
 		: Super(InOwner, InName, InObjectFlags, InOffset, InFlags)
+PRAGMA_ENABLE_DEPRECATION_WARNINGS
 	{
 		this->PropertyClass = InClass;
 	}
 
+	/**
+	 * Constructor used for constructing compiled-in properties
+	 * @param InOwner Owner of the property
+	 * @param Prop Pointer to the compiled in structure describing the property
+	 **/
+	TFObjectPropertyBase(FFieldVariant InOwner, const UECodeGen_Private::FObjectPropertyParams& Prop)
+		: Super(InOwner, Prop)
+	{
+		this->PropertyClass = Prop.ClassFunc ? Prop.ClassFunc() : nullptr;
+	}
+	/**
+	 * Constructor used for constructing compiled-in properties
+	 * @param InOwner Owner of the property
+	 * @param Prop Pointer to the compiled in structure describing the property
+	 * @param InClass Class of the object this property represents
+	 **/
+	TFObjectPropertyBase(FFieldVariant InOwner, const UECodeGen_Private::FObjectPropertyParamsWithoutClass& Prop, UClass* InClass)
+		: Super(InOwner, Prop)
+	{
+		this->PropertyClass = InClass;
+	}
+	
+
 #if WITH_EDITORONLY_DATA
 	explicit TFObjectPropertyBase(UField* InField)
 		: Super(InField)
@@ -2361,11 +2577,21 @@ class COREUOBJECT_API FObjectProperty : public TFObjectPropertyBase<UObject*>
 	{
 	}
 
+	UE_DEPRECATED(5.1, "Compiled-in property constructor is deprecated, use other constructors instead.")
 	FObjectProperty(FFieldVariant InOwner, const FName& InName, EObjectFlags InObjectFlags, int32 InOffset, EPropertyFlags InFlags, UClass* InClass)
+PRAGMA_DISABLE_DEPRECATION_WARNINGS
 		: TFObjectPropertyBase(InOwner, InName, InObjectFlags, InOffset, InFlags, InClass)
+PRAGMA_ENABLE_DEPRECATION_WARNINGS
 	{
 	}
 
+	/**
+	 * Constructor used for constructing compiled in properties
+	 * @param InOwner Owner of the property
+	 * @param PropBase Pointer to the compiled in structure describing the property
+	 **/
+	FObjectProperty(FFieldVariant InOwner, const UECodeGen_Private::FObjectPropertyParams& Prop);
+
 #if WITH_EDITORONLY_DATA
 	explicit FObjectProperty(UField* InField)
 		: TFObjectPropertyBase(InField)
@@ -2450,11 +2676,21 @@ class COREUOBJECT_API FWeakObjectProperty : public TFObjectPropertyBase<FWeakObj
 	{
 	}
 
+	UE_DEPRECATED(5.1, "Compiled-in property constructor is deprecated, use other constructors instead.")
 	FWeakObjectProperty(FFieldVariant InOwner, const FName& InName, EObjectFlags InObjectFlags, int32 InOffset, EPropertyFlags InFlags, UClass* InClass)
+PRAGMA_DISABLE_DEPRECATION_WARNINGS
 		: TFObjectPropertyBase(InOwner, InName, InObjectFlags, InOffset, InFlags, InClass)
+PRAGMA_ENABLE_DEPRECATION_WARNINGS
 	{
 	}
 	
+	/**
+	 * Constructor used for constructing compiled in properties
+	 * @param InOwner Owner of the property
+	 * @param PropBase Pointer to the compiled in structure describing the property
+	 **/
+	FWeakObjectProperty(FFieldVariant InOwner, const UECodeGen_Private::FWeakObjectPropertyParams& Prop);
+
 #if WITH_EDITORONLY_DATA
 	explicit FWeakObjectProperty(UField* InField)
 		: TFObjectPropertyBase(InField)
@@ -2495,11 +2731,21 @@ class COREUOBJECT_API FLazyObjectProperty : public TFObjectPropertyBase<FLazyObj
 	{
 	}
 
+	UE_DEPRECATED(5.1, "Compiled-in property constructor is deprecated, use other constructors instead.")
 	FLazyObjectProperty(FFieldVariant InOwner, const FName& InName, EObjectFlags InObjectFlags, int32 InOffset, EPropertyFlags InFlags, UClass* InClass)
+PRAGMA_DISABLE_DEPRECATION_WARNINGS
 		: TFObjectPropertyBase(InOwner, InName, InObjectFlags, InOffset, InFlags, InClass)
+PRAGMA_ENABLE_DEPRECATION_WARNINGS
 	{
 	}
 
+	/**
+	 * Constructor used for constructing compiled in properties
+	 * @param InOwner Owner of the property
+	 * @param PropBase Pointer to the compiled in structure describing the property
+	 **/
+	FLazyObjectProperty(FFieldVariant InOwner, const UECodeGen_Private::FLazyObjectPropertyParams& Prop);
+
 #if WITH_EDITORONLY_DATA
 	explicit FLazyObjectProperty(UField* InField)
 		: TFObjectPropertyBase(InField)
@@ -2541,11 +2787,31 @@ class COREUOBJECT_API FSoftObjectProperty : public TFObjectPropertyBase<FSoftObj
 
 	FSoftObjectProperty(FFieldVariant InOwner, const FName& InName, EObjectFlags InObjectFlags)
 		: TFObjectPropertyBase(InOwner, InName, InObjectFlags)
-	{}
+	{
+	}
 
+	UE_DEPRECATED(5.1, "Compiled-in property constructor is deprecated, use other constructors instead.")
 	FSoftObjectProperty(FFieldVariant InOwner, const FName& InName, EObjectFlags InObjectFlags, int32 InOffset, EPropertyFlags InFlags, UClass* InClass)
+PRAGMA_DISABLE_DEPRECATION_WARNINGS
 		: TFObjectPropertyBase(InOwner, InName, InObjectFlags, InOffset, InFlags, InClass)
-	{}
+PRAGMA_ENABLE_DEPRECATION_WARNINGS
+	{
+	}
+
+	/**
+	 * Constructor used for constructing compiled-in properties
+	 * @param InOwner Owner of the property
+	 * @param PropBase Pointer to the compiled in structure describing the property
+	 **/
+	FSoftObjectProperty(FFieldVariant InOwner, const UECodeGen_Private::FSoftObjectPropertyParams& Prop);
+
+	/**
+	 * Constructor used for constructing compiled-in properties
+	 * @param InOwner Owner of the property
+	 * @param PropBase Pointer to the compiled in structure describing the property
+	 * @param Class Class of the object this property represents
+	 **/
+	FSoftObjectProperty(FFieldVariant InOwner, const UECodeGen_Private::FObjectPropertyParamsWithoutClass& Prop, UClass* InClass);
 
 #if WITH_EDITORONLY_DATA
 	explicit FSoftObjectProperty(UField* InField)
@@ -2623,12 +2889,22 @@ public:
 	{
 	}
 
+	UE_DEPRECATED(5.1, "Compiled-in property constructor is deprecated, use other constructors instead.")
 	FClassProperty(FFieldVariant InOwner, const FName& InName, EObjectFlags InObjectFlags, int32 InOffset, EPropertyFlags InFlags, UClass* InMetaClass, UClass* InClassType)
+PRAGMA_DISABLE_DEPRECATION_WARNINGS
 		: FObjectProperty(InOwner, InName, InObjectFlags, InOffset, InFlags, InClassType ? InClassType : UClass::StaticClass())
+PRAGMA_ENABLE_DEPRECATION_WARNINGS
 		, MetaClass(InMetaClass)
 	{
 	}
 
+	/**
+	 * Constructor used for constructing compiled in properties
+	 * @param InOwner Owner of the property
+	 * @param PropBase Pointer to the compiled in structure describing the property
+	 **/
+	FClassProperty(FFieldVariant InOwner, const UECodeGen_Private::FClassPropertyParams& Prop);
+
 #if WITH_EDITORONLY_DATA
 	explicit FClassProperty(UField* InField);
 #endif // WITH_EDITORONLY_DATA
@@ -2731,11 +3007,21 @@ public:
 		, MetaClass(nullptr)
 	{}
 
+	UE_DEPRECATED(5.1, "Compiled-in property constructor is deprecated, use other constructors instead.")
 	FSoftClassProperty(FFieldVariant InOwner, const FName& InName, EObjectFlags InObjectFlags, int32 InOffset, EPropertyFlags InFlags, UClass* InMetaClass)
+PRAGMA_DISABLE_DEPRECATION_WARNINGS
 		: Super(InOwner, InName, InObjectFlags, InOffset, InFlags, UClass::StaticClass())
+PRAGMA_ENABLE_DEPRECATION_WARNINGS
 		, MetaClass(InMetaClass)
 	{}
 
+	/**
+	 * Constructor used for constructing compiled in properties
+	 * @param InOwner Owner of the property
+	 * @param PropBase Pointer to the compiled in structure describing the property
+	 **/
+	FSoftClassProperty(FFieldVariant InOwner, const UECodeGen_Private::FSoftClassPropertyParams& Prop);
+
 #if WITH_EDITORONLY_DATA
 	explicit FSoftClassProperty(UField* InField);
 #endif // WITH_EDITORONLY_DATA
@@ -2805,12 +3091,22 @@ public:
 	{
 	}
 
+	UE_DEPRECATED(5.1, "Compiled-in property constructor is deprecated, use other constructors instead.")
 	FInterfaceProperty(FFieldVariant InOwner, const FName& InName, EObjectFlags InObjectFlags, int32 InOffset, EPropertyFlags InFlags, UClass* InInterfaceClass)
+PRAGMA_DISABLE_DEPRECATION_WARNINGS
 		: FInterfaceProperty_Super(InOwner, InName, InObjectFlags, InOffset, (InFlags & ~CPF_InterfaceClearMask))
+PRAGMA_ENABLE_DEPRECATION_WARNINGS
 		, InterfaceClass(InInterfaceClass)
 	{
 	}
 
+	/**
+	 * Constructor used for constructing compiled in properties
+	 * @param InOwner Owner of the property
+	 * @param PropBase Pointer to the compiled in structure describing the property
+	 **/
+	FInterfaceProperty(FFieldVariant InOwner, const UECodeGen_Private::FInterfacePropertyParams& Prop);
+
 #if WITH_EDITORONLY_DATA
 	explicit FInterfaceProperty(UField* InField);
 #endif // WITH_EDITORONLY_DATA
@@ -2885,11 +3181,21 @@ public:
 	{
 	}
 
+	UE_DEPRECATED(5.1, "Compiled-in property constructor is deprecated, use other constructors instead.")
 	FNameProperty(FFieldVariant InOwner, const FName& InName, EObjectFlags InObjectFlags, int32 InOffset, EPropertyFlags InFlags)
+PRAGMA_DISABLE_DEPRECATION_WARNINGS
 		: FNameProperty_Super(InOwner, InName, InObjectFlags, InOffset, InFlags)
+PRAGMA_ENABLE_DEPRECATION_WARNINGS
 	{
 	}
 
+	/**
+	 * Constructor used for constructing compiled in properties
+	 * @param InOwner Owner of the property
+	 * @param PropBase Pointer to the compiled in structure describing the property
+	 **/
+	FNameProperty(FFieldVariant InOwner, const UECodeGen_Private::FNamePropertyParams& Prop);
+
 #if WITH_EDITORONLY_DATA
 	explicit FNameProperty(UField* InField)
 		: FNameProperty_Super(InField)
@@ -2929,11 +3235,21 @@ public:
 	{
 	}
 
+	UE_DEPRECATED(5.1, "Compiled-in property constructor is deprecated, use other constructors instead.")
 	FStrProperty(FFieldVariant InOwner, const FName& InName, EObjectFlags InObjectFlags, int32 InOffset, EPropertyFlags InFlags)
+PRAGMA_DISABLE_DEPRECATION_WARNINGS
 		: FStrProperty_Super(InOwner, InName, InObjectFlags, InOffset, InFlags)
+PRAGMA_ENABLE_DEPRECATION_WARNINGS
 	{
 	}
 
+	/**
+	 * Constructor used for constructing compiled in properties
+	 * @param InOwner Owner of the property
+	 * @param PropBase Pointer to the compiled in structure describing the property
+	 **/
+	FStrProperty(FFieldVariant InOwner, const UECodeGen_Private::FStrPropertyParams& Prop);
+
 #if WITH_EDITORONLY_DATA
 	explicit FStrProperty(UField* InField)
 		: FStrProperty_Super(InField)
@@ -3000,14 +3316,24 @@ public:
 		SetElementSize();
 	}
 
+	UE_DEPRECATED(5.1, "Compiled-in property constructor is deprecated, use other constructors instead.")
 	FArrayProperty(FFieldVariant InOwner, const FName& InName, EObjectFlags InObjectFlags, int32 InOffset, EPropertyFlags InFlags, EArrayPropertyFlags InArrayPropertyFlags)
+PRAGMA_DISABLE_DEPRECATION_WARNINGS
 		: FArrayProperty_Super(InOwner, InName, InObjectFlags, InOffset, InFlags)
+PRAGMA_ENABLE_DEPRECATION_WARNINGS
 		, Inner(nullptr)
 	{
 		ArrayFlags = InArrayPropertyFlags;
 		SetElementSize();
 	}
 
+	/** 
+	 * Constructor used for constructing compiled in properties
+	 * @param InOwner Owner of the property
+	 * @param PropBase Pointer to the compiled in structure describing the property
+	 **/
+	FArrayProperty(FFieldVariant InOwner, const UECodeGen_Private::FArrayPropertyParams& Prop);
+
 	virtual ~FArrayProperty();
 
 #if WITH_EDITORONLY_DATA
@@ -3130,8 +3456,17 @@ public:
 	typedef TTypeFundamentals::TCppType TCppType;
 
 	FMapProperty(FFieldVariant InOwner, const FName& InName, EObjectFlags InObjectFlags, EMapPropertyFlags InMapFlags=EMapPropertyFlags::None);
+
+	UE_DEPRECATED(5.1, "Compiled-in property constructor is deprecated, use other constructors instead.")
 	FMapProperty(FFieldVariant InOwner, const FName& InName, EObjectFlags InObjectFlags, int32 InOffset, EPropertyFlags InFlags, EMapPropertyFlags InMapFlags);
 
+	/**
+	 * Constructor used for constructing compiled in properties
+	 * @param InOwner Owner of the property
+	 * @param PropBase Pointer to the compiled in structure describing the property
+	 **/
+	FMapProperty(FFieldVariant InOwner, const UECodeGen_Private::FMapPropertyParams& Prop);
+
 #if WITH_EDITORONLY_DATA
 	explicit FMapProperty(UField* InField);
 #endif // WITH_EDITORONLY_DATA
@@ -3261,8 +3596,17 @@ public:
 	typedef TTypeFundamentals::TCppType TCppType;
 
 	FSetProperty(FFieldVariant InOwner, const FName& InName, EObjectFlags InObjectFlags);
+
+	UE_DEPRECATED(5.1, "Compiled-in property constructor is deprecated, use other constructors instead.")
 	FSetProperty(FFieldVariant InOwner, const FName& InName, EObjectFlags InObjectFlags, int32 InOffset, EPropertyFlags InFlags);
 
+	/**
+	 * Constructor used for constructing compiled in properties
+	 * @param InOwner Owner of the property
+	 * @param PropBase Pointer to the compiled in structure describing the property
+	 **/
+	FSetProperty(FFieldVariant InOwner, const UECodeGen_Private::FSetPropertyParams& Prop);
+
 #if WITH_EDITORONLY_DATA
 	explicit FSetProperty(UField* InField);
 #endif // WITH_EDITORONLY_DATA
@@ -5090,8 +5434,17 @@ class COREUOBJECT_API FStructProperty : public FProperty
 	class UScriptStruct* Struct;
 public:
 	FStructProperty(FFieldVariant InOwner, const FName& InName, EObjectFlags InObjectFlags);
+
+	UE_DEPRECATED(5.1, "Compiled-in property constructor is deprecated, use other constructors instead.")
 	FStructProperty(FFieldVariant InOwner, const FName& InName, EObjectFlags InObjectFlags, int32 InOffset, EPropertyFlags InFlags, UScriptStruct* InStruct);
 
+	/**
+	 * Constructor used for constructing compiled in properties
+	 * @param InOwner Owner of the property
+	 * @param PropBase Pointer to the compiled in structure describing the property
+	 **/
+	FStructProperty(FFieldVariant InOwner, const UECodeGen_Private::FStructPropertyParams& Prop);
+
 #if WITH_EDITORONLY_DATA
 	explicit FStructProperty(UField* InField);
 #endif // WITH_EDITORONLY_DATA
@@ -5176,12 +5529,22 @@ public:
 	{
 	}
 
+	UE_DEPRECATED(5.1, "Compiled-in property constructor is deprecated, use other constructors instead.")
 	FDelegateProperty(FFieldVariant InOwner, const FName& InName, EObjectFlags InObjectFlags, int32 InOffset, EPropertyFlags InFlags, UFunction* InSignatureFunction = NULL)
+PRAGMA_DISABLE_DEPRECATION_WARNINGS
 		: FDelegateProperty_Super(InOwner, InName, InObjectFlags, InOffset, InFlags)
+PRAGMA_ENABLE_DEPRECATION_WARNINGS
 		, SignatureFunction(InSignatureFunction)
 	{
 	}
 
+	/**
+	 * Constructor used for constructing compiled in properties
+	 * @param InOwner Owner of the property
+	 * @param PropBase Pointer to the compiled in structure describing the property
+	 **/
+	FDelegateProperty(FFieldVariant InOwner, const UECodeGen_Private::FDelegatePropertyParams& Prop);
+
 #if WITH_EDITORONLY_DATA
 	explicit FDelegateProperty(UField* InField);
 #endif // WITH_EDITORONLY_DATA
@@ -5236,12 +5599,22 @@ public:
 	{
 	}
 
+	UE_DEPRECATED(5.1, "Compiled-in property constructor is deprecated, use other constructors instead.")
 	FMulticastDelegateProperty(FFieldVariant InOwner, const FName& InName, EObjectFlags InObjectFlags, int32 InOffset, EPropertyFlags InFlags, UFunction* InSignatureFunction = NULL)
+PRAGMA_DISABLE_DEPRECATION_WARNINGS
 		: FProperty(InOwner, InName, InObjectFlags, InOffset, InFlags)
+PRAGMA_ENABLE_DEPRECATION_WARNINGS
 		, SignatureFunction(InSignatureFunction)
 	{
 	}
 
+	/**
+	 * Constructor used for constructing compiled in properties
+	 * @param InOwner Owner of the property
+	 * @param PropBase Pointer to the compiled in structure describing the property
+	 **/
+	FMulticastDelegateProperty(FFieldVariant InOwner, const UECodeGen_Private::FMulticastDelegatePropertyParams& Prop, EPropertyFlags AdditionalPropertyFlags = CPF_None);
+
 #if WITH_EDITORONLY_DATA
 	explicit FMulticastDelegateProperty(UField* InField);
 #endif // WITH_EDITORONLY_DATA
@@ -5307,8 +5680,11 @@ public:
 		this->SignatureFunction = nullptr;
 	}
 
+	UE_DEPRECATED(5.1, "Compiled-in property constructor is deprecated, use other constructors instead.")
 	TProperty_MulticastDelegate(FFieldVariant InOwner, const FName& InName, EObjectFlags InObjectFlags, int32 InOffset, EPropertyFlags InFlags, UFunction* InSignatureFunction = nullptr)
+PRAGMA_DISABLE_DEPRECATION_WARNINGS
 		: Super(InOwner, InName, InObjectFlags, InOffset, InFlags)
+PRAGMA_ENABLE_DEPRECATION_WARNINGS
 	{
 		this->SignatureFunction = InSignatureFunction;
 	}
@@ -5318,6 +5694,17 @@ public:
 	{
 	}
 
+	/**
+	 * Constructor used for constructing compiled in properties
+	 * @param InOwner Owner of the property
+	 * @param PropBase Pointer to the compiled in structure describing the property
+	 **/
+	TProperty_MulticastDelegate(FFieldVariant InOwner, const UECodeGen_Private::FMulticastDelegatePropertyParams& Prop)
+		: Super(InOwner, Prop)
+	{
+		this->SignatureFunction = Prop.SignatureFunctionFunc ? Prop.SignatureFunctionFunc() : nullptr;
+	}
+
 #if WITH_EDITORONLY_DATA
 	explicit TProperty_MulticastDelegate(UField* InField)
 		: Super(InField)
@@ -5344,11 +5731,21 @@ public:
 	{
 	}
 
+	UE_DEPRECATED(5.1, "Compiled-in property constructor is deprecated, use other constructors instead.")
 	FMulticastInlineDelegateProperty(FFieldVariant InOwner, const FName& InName, EObjectFlags InObjectFlags, int32 InOffset, EPropertyFlags InFlags, UFunction* InSignatureFunction = nullptr)
+PRAGMA_DISABLE_DEPRECATION_WARNINGS
 		: TProperty_MulticastDelegate(InOwner, InName, InObjectFlags, InOffset, InFlags, InSignatureFunction)
+PRAGMA_ENABLE_DEPRECATION_WARNINGS
 	{
 	}
 
+	/**
+	 * Constructor used for constructing compiled in properties
+	 * @param InOwner Owner of the property
+	 * @param PropBase Pointer to the compiled in structure describing the property
+	 **/
+	FMulticastInlineDelegateProperty(FFieldVariant InOwner, const UECodeGen_Private::FMulticastDelegatePropertyParams& Prop);
+
 #if WITH_EDITORONLY_DATA
 	explicit FMulticastInlineDelegateProperty(UField* InField)
 		: TProperty_MulticastDelegate(InField)
@@ -5385,11 +5782,21 @@ public:
 	{
 	}
 
+	UE_DEPRECATED(5.1, "Compiled-in property constructor is deprecated, use other constructors instead.")
 	FMulticastSparseDelegateProperty(FFieldVariant InOwner, const FName& InName, EObjectFlags InObjectFlags, int32 InOffset, EPropertyFlags InFlags, UFunction* InSignatureFunction = nullptr)
+PRAGMA_DISABLE_DEPRECATION_WARNINGS
 		: TProperty_MulticastDelegate(InOwner, InName, InObjectFlags, InOffset, InFlags, InSignatureFunction)
+PRAGMA_ENABLE_DEPRECATION_WARNINGS
 	{
 	}
 
+	/**
+	 * Constructor used for constructing compiled in properties
+	 * @param InOwner Owner of the property
+	 * @param PropBase Pointer to the compiled in structure describing the property
+	 **/
+	FMulticastSparseDelegateProperty(FFieldVariant InOwner, const UECodeGen_Private::FMulticastDelegatePropertyParams& Prop);
+
 #if WITH_EDITORONLY_DATA
 	explicit FMulticastSparseDelegateProperty(UField* InField)
 		: TProperty_MulticastDelegate(InField)
