Convert BulkData to use a separate file handle (not AttachedAr) when MakeSureBulkDataIsLoaded is called from non-gamethread. This prevents unguarded multithreaded access to AttachedAr.
#rnx

diff --git a/Engine/Source/Runtime/CoreUObject/Private/Serialization/BulkData.cpp b/Engine/Source/Runtime/CoreUObject/Private/Serialization/BulkData.cpp
--- a/Engine/Source/Runtime/CoreUObject/Private/Serialization/BulkData.cpp
+++ b/Engine/Source/Runtime/CoreUObject/Private/Serialization/BulkData.cpp
@@ -2141,6 +2141,13 @@ void FUntypedBulkData::MakeSureBulkDataIsLoaded()
 	// Nothing to do if data is already loaded.
 	if( !BulkData )
 	{
+		if (!IsInGameThread())
+		{
+			// BulkDatas in the same package share AttachedAr that they get from the LinkerLoad of the package.
+			// To make calls to those BulkDatas threadsafe, we need to not use AttachedAr when called multithreaded.
+			// LoadBulkDataWithFileReader does so by using a separate Archive to the PackagePath instead of AttachedAr.
+			LoadBulkDataWithFileReader();
+		}
 		// Look for async request first
 		if (SerializeFuture.IsValid())
 		{
