Saves 8MB of fixed memory by using a dynamically sized queue instead for async loading

diff --git a/Engine/Source/Runtime/CoreUObject/Private/Serialization/AsyncLoading2.cpp b/Engine/Source/Runtime/CoreUObject/Private/Serialization/AsyncLoading2.cpp
--- a/Engine/Source/Runtime/CoreUObject/Private/Serialization/AsyncLoading2.cpp
+++ b/Engine/Source/Runtime/CoreUObject/Private/Serialization/AsyncLoading2.cpp
@@ -49,6 +49,7 @@
 #include "Serialization/UnversionedPropertySerialization.h"
 #include "Serialization/Zenaphore.h"
 #include "UObject/GCObject.h"
+#include "Experimental/Containers/FAAArrayQueue.h"
 #include "UObject/ObjectRedirector.h"
 #include "Serialization/BulkData.h"
 #include "Serialization/LargeMemoryReader.h"
@@ -1424,9 +1425,7 @@ public:
 
 private:
 	FZenaphore* Zenaphore = nullptr;
-	TAtomic<uint64> Head { 0 };
-	TAtomic<uint64> Tail { 0 };
-	TAtomic<FEventLoadNode2*> Entries[524288];
+	FAAArrayQueue<FEventLoadNode2> Queue;
 };
 
 struct FAsyncLoadEventSpec
@@ -2922,7 +2921,6 @@ void FEventLoadNode2::ProcessDependencies(FAsyncLoadingThreadState2& ThreadState
 
 FAsyncLoadEventQueue2::FAsyncLoadEventQueue2()
 {
-	FMemory::Memzero(Entries, sizeof(Entries));
 }
 
 FAsyncLoadEventQueue2::~FAsyncLoadEventQueue2()
@@ -2931,12 +2929,9 @@ FAsyncLoadEventQueue2::~FAsyncLoadEventQueue2()
 
 void FAsyncLoadEventQueue2::Push(FEventLoadNode2* Node)
 {
-	uint64 LocalHead = Head.IncrementExchange();
-	FEventLoadNode2* Expected = nullptr;
-	if (!Entries[LocalHead % UE_ARRAY_COUNT(Entries)].CompareExchange(Expected, Node))
-	{
-		*(volatile int*)0 = 0; // queue is full: TODO
-	}
+	LLM_SCOPE_BYNAME(TEXT("AsyncLoadEventQueue2"));
+	Queue.enqueue(Node);
+
 	if (Zenaphore)
 	{
 		Zenaphore->NotifyOne();
@@ -2956,28 +2951,7 @@ bool FAsyncLoadEventQueue2::PopAndExecute(FAsyncLoadingThreadState2& ThreadState
 		return true;
 	}
 
-	FEventLoadNode2* Node = nullptr;
-	{
-		uint64 LocalHead = Head.Load();
-		uint64 LocalTail = Tail.Load();
-		for (;;)
-		{
-			if (LocalTail >= LocalHead)
-			{
-				break;
-			}
-			if (Tail.CompareExchange(LocalTail, LocalTail + 1))
-			{
-				while (!Node)
-				{
-					Node = Entries[LocalTail % UE_ARRAY_COUNT(Entries)].Exchange(nullptr);
-				}
-				break;
-			}
-		}
-	}
-
-	if (Node)
+	if (FEventLoadNode2* Node = Queue.dequeue())
 	{
 		//TRACE_CPUPROFILER_EVENT_SCOPE(Execute);
 		Node->Execute(ThreadState);
