Clean up the logging when a virtualized payload fails to load in some way to prevent multiple duplicate log entries.

#rnx

- We now only validate the payloads contents when ::GetDataInternal is called. This prevents duplicate log messages being raised.
-- Duplicate calls to IsValid have been removed.
-- This also means we don't need ::GetCorruptedPayloadErrorMsgForLoad anymore as it is only used in a single place.
- Renamed ::IsValid to ::IsDataValid to make it clearer that we are checking the contents of the payload.
--  The utility no longer checks if the payload fails to load (size should be >0 but the payload is null) but only checks if the payload is corrupted or does not match the expected content hash.
--- This is because the payload failing to load is an error, but it is something that we expect to see happen on occasion, perhaps when the network is flaky etc. However the wrong or corrupted data being loaded is not expected and so needs to be reported as it's own specific error.
--- Splitting these types of failures up, should allow us to consider enabling VBD_CORRUPTED_PAYLOAD_IS_FATAL again. Although now we prevent the saving of packages with corrupted payloads we probably don't need to treat a corrupted payload as a reason to terminate the editor.
- Removed warning from FVirtualizedUntypedBulkData::PullData when the payload failed to pull as IVirtualizationSystem will have already logged an error and the warning was incorrect as it was displaying the payload identifier as the bulkdata guid.
- VBD_CORRUPTED_DATA_SEVERITY is now set as an Error if VBD_CORRUPTED_PAYLOAD_IS_FATAL is disabled.

diff --git a/Engine/Source/Runtime/CoreUObject/Private/Serialization/VirtualizedBulkData.cpp b/Engine/Source/Runtime/CoreUObject/Private/Serialization/VirtualizedBulkData.cpp
--- a/Engine/Source/Runtime/CoreUObject/Private/Serialization/VirtualizedBulkData.cpp
+++ b/Engine/Source/Runtime/CoreUObject/Private/Serialization/VirtualizedBulkData.cpp
@@ -33,7 +33,7 @@
 #if VBD_CORRUPTED_PAYLOAD_IS_FATAL
 	#define VBD_CORRUPTED_DATA_SEVERITY Fatal
 #else
-	#define VBD_CORRUPTED_DATA_SEVERITY Warning
+	#define VBD_CORRUPTED_DATA_SEVERITY Error
 #endif // VBD_CORRUPTED_PAYLOAD_IS_FATAL
 
 namespace UE::Virtualization
@@ -126,22 +126,26 @@ void LogPackageOpenFailureMessage(const FPackagePath& PackagePath, EPackageSegme
 	}
 }
 
-/** Utility for checking if a payload (in FCompressedBuffer format) matches the expectations of a bulkdata's members or not. */
-bool IsValid(const FVirtualizedUntypedBulkData& BulkData, const FCompressedBuffer& Payload)
+/** 
+ * Utility used to validate the contents of a recently loaded payload.
+ * If the given payload is null, then we assume that the load failed and errors would've been raised else
+ * where in code and there is no need to validate the contents.
+ * If the contents are validated we check the loaded result against the members of a bulkdata object to 
+ * see if they match.
+ */
+bool IsDataValid(const FVirtualizedUntypedBulkData& BulkData, const FCompressedBuffer& Payload)
 {
-	if (Payload.IsNull() && BulkData.GetPayloadSize() > 0)
-	{
-		return false;
-	}
-
-	if (!BulkData.HasPlaceholderPayloadId() && BulkData.GetPayloadId() != FIoHash(Payload.GetRawHash()))
+	if (!Payload.IsNull())
 	{
-		return false;
-	}
+		if (!BulkData.HasPlaceholderPayloadId() && BulkData.GetPayloadId() != FIoHash(Payload.GetRawHash()))
+		{
+			return false;
+		}
 
-	if (Payload.GetRawSize() != BulkData.GetPayloadSize())
-	{
-		return false;
+		if (Payload.GetRawSize() != BulkData.GetPayloadSize())
+		{
+			return false;
+		}
 	}
 
 	return true;
@@ -644,8 +648,8 @@ void FVirtualizedUntypedBulkData::Serialize(FArchive& Ar, UObject* Owner, bool b
 				// Checks to make sure that the payload we are saving is what we expect it to be. If not then we need to log an error to the
 				// user, mark the archive as having an error and return.
 				// This error is not reasonable to expect, if it occurs then it means something has gone very wrong, which is why there is
-				// also an ensure to make sure that any instances of it occuring as properly recorded and can be investigated.
-				if (!IsValid(*this, PayloadToSerialize))
+				// also an ensure to make sure that any instances of it occurring as properly recorded and can be investigated.
+				if (!IsDataValid(*this, PayloadToSerialize) || (GetPayloadSize() > 0 && PayloadToSerialize.IsNull()))
 				{
 					ensureMsgf(false, TEXT("%s"), *GetCorruptedPayloadErrorMsgForSave(LinkerSave).ToString());
 
@@ -1194,11 +1198,7 @@ FCompressedBuffer FVirtualizedUntypedBulkData::PullData() const
 				PayloadSize,
 				PulledPayload.GetRawSize());
 	}
-	else
-	{
-		UE_LOG(LogVirtualization, Warning, TEXT("Failed to pull virtual data with guid (%s)"), *PayloadContentId.ToString());
-	}
-
+	
 	return PulledPayload;
 }
 
@@ -1269,9 +1269,6 @@ void FVirtualizedUntypedBulkData::DetachFromDisk(FArchive* Ar, bool bEnsurePaylo
 		if (Payload.IsNull() && bEnsurePayloadIsLoaded)
 		{
 			FCompressedBuffer CompressedPayload = GetDataInternal();
-
-			UE_CLOG(!IsValid(*this, CompressedPayload), LogVirtualization, VBD_CORRUPTED_DATA_SEVERITY, TEXT("%s"), *GetCorruptedPayloadErrorMsgForLoad());
-
 			Payload = CompressedPayload.Decompress();
 		}
 
@@ -1477,7 +1474,9 @@ FCompressedBuffer FVirtualizedUntypedBulkData::GetDataInternal() const
 		FCompressedBuffer CompressedPayload = PullData();
 		
 		checkf(Payload.IsNull(), TEXT("Pulling data somehow assigned it to the bulk data object!")); //Make sure that we did not assign the buffer internally
-		UE_CLOG(!IsValid(*this, CompressedPayload), LogVirtualization, VBD_CORRUPTED_DATA_SEVERITY, TEXT("%s"), *GetCorruptedPayloadErrorMsgForLoad());
+
+		UE_CLOG(CompressedPayload.IsNull(), LogVirtualization, Error, TEXT("Failed to pull payload '%s'"), *PayloadContentId.ToString());
+		UE_CLOG(!IsDataValid(*this, CompressedPayload), LogVirtualization, VBD_CORRUPTED_DATA_SEVERITY, TEXT("Virtualized payload '%s' is corrupt! Check the backend storage."), *PayloadContentId.ToString());
 		
 		return CompressedPayload;
 	}
@@ -1486,8 +1485,12 @@ FCompressedBuffer FVirtualizedUntypedBulkData::GetDataInternal() const
 		FCompressedBuffer CompressedPayload = LoadFromDisk();
 		
 		check(Payload.IsNull()); //Make sure that we did not assign the buffer internally
-		UE_CLOG(!IsValid(*this, CompressedPayload), LogVirtualization, VBD_CORRUPTED_DATA_SEVERITY, TEXT("%s"), *GetCorruptedPayloadErrorMsgForLoad());
-		
+
+		UE_CLOG(CompressedPayload.IsNull(), LogVirtualization, Error, TEXT("Failed to load payload '%s"), *PayloadContentId.ToString());
+		UE_CLOG(!IsDataValid(*this, CompressedPayload), LogVirtualization, VBD_CORRUPTED_DATA_SEVERITY, TEXT("Payload '%s' loaded from '%s' is corrupt! Check the file on disk."),
+			*PayloadContentId.ToString(), 
+			*PackagePath.GetDebugName());
+
 		return CompressedPayload;
 	}
 }
@@ -1510,8 +1513,6 @@ TFuture<FSharedBuffer> FVirtualizedUntypedBulkData::GetPayload() const
 	{
 		FCompressedBuffer CompressedPayload = GetDataInternal();
 
-		UE_CLOG(!IsValid(*this, CompressedPayload), LogVirtualization, VBD_CORRUPTED_DATA_SEVERITY, TEXT("%s"), *GetCorruptedPayloadErrorMsgForLoad());
-		
 		// TODO: Not actually async yet!
 		Promise.SetValue(CompressedPayload.Decompress());
 	}
@@ -1722,22 +1723,6 @@ FVirtualizedUntypedBulkData::EFlags FVirtualizedUntypedBulkData::BuildFlagsForSe
 	}
 }
 
-FString FVirtualizedUntypedBulkData::GetCorruptedPayloadErrorMsgForLoad() const
-{
-	TStringBuilder<512> Message;
-	
-	if (IsDataVirtualized())
-	{
-		Message << TEXT("Virtualized payload '") << PayloadContentId << TEXT("' is corrupt! Check the backend storage.");
-	}
-	else
-	{
-		Message << TEXT("Payload ' ") << PayloadContentId << TEXT(" loaded from '") << *PackagePath.GetDebugName() << TEXT("' is corrupt! Check the file on disk.");
-	}
-
-	return Message.ToString();
-}
-
 FText FVirtualizedUntypedBulkData::GetCorruptedPayloadErrorMsgForSave(FLinkerSave* Linker) const
 {
 	const FText GuidID = FText::FromString(GetIdentifier().ToString());
diff --git a/Engine/Source/Runtime/CoreUObject/Public/Serialization/VirtualizedBulkData.h b/Engine/Source/Runtime/CoreUObject/Public/Serialization/VirtualizedBulkData.h
--- a/Engine/Source/Runtime/CoreUObject/Public/Serialization/VirtualizedBulkData.h
+++ b/Engine/Source/Runtime/CoreUObject/Public/Serialization/VirtualizedBulkData.h
@@ -429,9 +429,6 @@ private:
 
 	void Unregister();
 
-	/** Utility to return an apt error message if a payload loaded from this object is considered corrupt. Handles the formatting depending on the various payload flags etc. */
-	FString GetCorruptedPayloadErrorMsgForLoad() const;
-
 	/** 
 	 * Utility to return an apt error message if the payload is invalid when trying to save the bulkdata. It will try to provide the best info from the given options.
 	 * 
