Fix MountPointExists to include mount points that are not plugins such as Game and Engine

diff --git a/Engine/Source/Runtime/CoreUObject/Private/Misc/PackageName.cpp b/Engine/Source/Runtime/CoreUObject/Private/Misc/PackageName.cpp
--- a/Engine/Source/Runtime/CoreUObject/Private/Misc/PackageName.cpp
+++ b/Engine/Source/Runtime/CoreUObject/Private/Misc/PackageName.cpp
@@ -326,7 +326,7 @@ struct FLongPackagePathsSingleton
 	bool MountPointExists(const FString& RootPath)
 	{
 		FReadScopeLock ScopeLock(ContentMountPointCriticalSection);
-		return MountPointRootPaths.Contains(RootPath);
+		return GetValidLongPackageRoots(true/*bIncludeReadOnlyRoots*/).Contains(RootPath);
 	}
 
 private:
