Making FScopedObjectFlagMarker respect both mirrored object flags (PendingKill and Garbage)

diff --git a/Engine/Source/Runtime/CoreUObject/Private/UObject/UObjectGlobals.cpp b/Engine/Source/Runtime/CoreUObject/Private/UObject/UObjectGlobals.cpp
--- a/Engine/Source/Runtime/CoreUObject/Private/UObject/UObjectGlobals.cpp
+++ b/Engine/Source/Runtime/CoreUObject/Private/UObject/UObjectGlobals.cpp
@@ -3598,16 +3598,16 @@ void FScopedObjectFlagMarker::RestoreObjectFlags()
 		UObject* Object = It.Key();
 		FStoredObjectFlags& PreviousObjectFlags = It.Value();
 
-		// clear all flags, frist clear the PendingKill flag as we don't allow clearing it through ClearFlags
-		Object->ClearGarbage();
+		// clear all flags, frist clear the mirrored flags as we don't allow clearing them through ClearFlags
+		Object->ClearGarbage(); // The currently mirrored flags are mutually exclusive and this will take care of both
 		Object->ClearFlags(RF_AllFlags);
 		Object->ClearInternalFlags(EInternalObjectFlags::AllFlags);
 
-		// then reset the ones that were originally set, start with PendingKill as we don't allow setting it through SetFlags
-		if ((PreviousObjectFlags.InternalFlags & EInternalObjectFlags::PendingKill) == EInternalObjectFlags::PendingKill)
+		// then reset the ones that were originally set
+		if (!!(PreviousObjectFlags.InternalFlags & EInternalObjectFlags::MirroredFlags))
 		{
 PRAGMA_DISABLE_DEPRECATION_WARNINGS
-			checkf((PreviousObjectFlags.Flags & RF_PendingKill) == RF_PendingKill, TEXT("Object %s had EInternalObjectFlags::PendingKill flag set but no RF_PendingKill"), *Object->GetFullName());
+			checkf(!!(PreviousObjectFlags.Flags & RF_InternalMirroredFlags), TEXT("Object %s had internal mirrored flag set but it was not matched in object flags"), *Object->GetFullName());
 PRAGMA_ENABLE_DEPRECATION_WARNINGS
 			Object->MarkAsGarbage();
 		}
diff --git a/Engine/Source/Runtime/CoreUObject/Public/UObject/ObjectMacros.h b/Engine/Source/Runtime/CoreUObject/Public/UObject/ObjectMacros.h
--- a/Engine/Source/Runtime/CoreUObject/Public/UObject/ObjectMacros.h
+++ b/Engine/Source/Runtime/CoreUObject/Public/UObject/ObjectMacros.h
@@ -573,6 +573,7 @@ enum class EInternalObjectFlags : int32
 	PendingConstruction = 1 << 31, ///< Object didn't have its class constructor called yet (only the UObjectBase one to initialize its most basic members)
 
 	GarbageCollectionKeepFlags = Native | Async | AsyncLoading,
+	MirroredFlags = Garbage | PendingKill, /// Flags mirrored in EObjectFlags
 
 	//~ Make sure this is up to date!
 	AllFlags = Garbage | PersistentGarbage | ReachableInCluster | ClusterRoot | Native | Async | AsyncLoading | Unreachable | PendingKill | RootSet | PendingConstruction
