Misc changes required to build consoles as modules

diff --git a/Engine/Plugins/Online/OnlineSubsystem/Source/Public/Interfaces/OnlineIdentityInterface.h b/Engine/Plugins/Online/OnlineSubsystem/Source/Public/Interfaces/OnlineIdentityInterface.h
--- a/Engine/Plugins/Online/OnlineSubsystem/Source/Public/Interfaces/OnlineIdentityInterface.h
+++ b/Engine/Plugins/Online/OnlineSubsystem/Source/Public/Interfaces/OnlineIdentityInterface.h
@@ -207,7 +207,7 @@ class IOnlineIdentity
 {
 
 protected:
-	IOnlineIdentity() {};	
+	ONLINESUBSYSTEM_API IOnlineIdentity() {};
 
 public:
 
@@ -239,7 +239,7 @@ public:
 		NetworkConnectionUnavailable =	1 << 10,
 	};
 
-	virtual ~IOnlineIdentity() {};
+	ONLINESUBSYSTEM_API virtual ~IOnlineIdentity() {};
 
 	/**
 	 * Delegate called when a player logs in/out
diff --git a/Engine/Source/Runtime/CoreUObject/Private/Misc/PackageName.cpp b/Engine/Source/Runtime/CoreUObject/Private/Misc/PackageName.cpp
--- a/Engine/Source/Runtime/CoreUObject/Private/Misc/PackageName.cpp
+++ b/Engine/Source/Runtime/CoreUObject/Private/Misc/PackageName.cpp
@@ -1831,6 +1831,30 @@ FString FPackageName::GetLocalizedPackagePath(const FString& InSourcePackagePath
 	return (LocalizedPackageName.IsNone()) ? InSourcePackagePath : LocalizedPackageName.ToString();
 }
 
+const FString& FPackageName::GetAssetPackageExtension()
+{
+	static FString AssetPackageExtension(LexToString(EPackageExtension::Asset));
+	return AssetPackageExtension;
+}
+
+const FString& FPackageName::GetMapPackageExtension()
+{
+	static FString MapPackageExtension(LexToString(EPackageExtension::Map));
+	return MapPackageExtension;
+}
+
+const FString& FPackageName::GetTextAssetPackageExtension()
+{
+	static FString TextAssetPackageExtension(LexToString(EPackageExtension::TextAsset));
+	return TextAssetPackageExtension;
+}
+
+const FString& FPackageName::GetTextMapPackageExtension()
+{
+	static FString TextMapPackageExtension(LexToString(EPackageExtension::TextMap));
+	return TextMapPackageExtension;
+}
+
 FString FPackageName::PackageFromPath(const TCHAR* InPathName)
 {
 	FString PackageName;
diff --git a/Engine/Source/Runtime/CoreUObject/Public/Misc/PackageName.h b/Engine/Source/Runtime/CoreUObject/Public/Misc/PackageName.h
--- a/Engine/Source/Runtime/CoreUObject/Public/Misc/PackageName.h
+++ b/Engine/Source/Runtime/CoreUObject/Public/Misc/PackageName.h
@@ -450,43 +450,27 @@ public:
 	 *
 	 * @return	file extension for asset packages ( dot included )
 	 */
-	static FORCEINLINE const FString& GetAssetPackageExtension()
-	{
-		static FString AssetPackageExtension(LexToString(EPackageExtension::Asset));
-		return AssetPackageExtension;
-	}
+	static const FString& GetAssetPackageExtension();
 	/** 
 	 * Returns the file extension for packages containing assets.
 	 *
 	 * @return	file extension for asset packages ( dot included )
 	 */
-	static FORCEINLINE const FString& GetMapPackageExtension()
-	{
-		static FString MapPackageExtension(LexToString(EPackageExtension::Map));
-		return MapPackageExtension;
-	}
+	static const FString& GetMapPackageExtension();
 
 	/**
 	* Returns the file extension for packages containing text assets.
 	*
 	* @return	file extension for text asset packages ( dot included )
 	*/
-	static FORCEINLINE const FString& GetTextAssetPackageExtension()
-	{
-		static FString TextAssetPackageExtension(LexToString(EPackageExtension::TextAsset));
-		return TextAssetPackageExtension;
-	}
+	static const FString& GetTextAssetPackageExtension();
 
 	/**
 	* Returns the file extension for packages containing text maps.
 	*
 	* @return	file extension for text map packages ( dot included )
 	*/
-	static FORCEINLINE const FString& GetTextMapPackageExtension()
-	{
-		static FString TextMapPackageExtension(LexToString(EPackageExtension::TextMap));
-		return TextMapPackageExtension;
-	}
+	static const FString& GetTextMapPackageExtension();
 
 	/**
 	 * Returns whether the passed in extension is a valid text package
