BulkData: Remove the obsolete option "AllowBulkDataInIoStore" that let games run with package data in iostore and bulkdata in pak.
The default configuration for ue5 is to run with both package data and bulk data in iostore.
The only remaining alternative is to disable iostore completely and run with both package data and bulk data in pak.

#rnx

diff --git a/Engine/Config/BaseEngine.ini b/Engine/Config/BaseEngine.ini
--- a/Engine/Config/BaseEngine.ini
+++ b/Engine/Config/BaseEngine.ini
@@ -1109,7 +1109,6 @@ UseSeperateBulkDataFiles=False
 GameThreadHeartBeatStartSuspended=False
 GameThreadHeartBeatHitchDuration=-1.0
 GameThreadHeartBeatStackWalk=false
-AllowBulkDataInIoStore=True
 IgnoreInlineBulkDataReloadEnsures=False
 UsePackageTrailer=False
 HangDuration=0.0
diff --git a/Engine/Source/Programs/AutomationTool/Scripts/CopyBuildToStagingDirectory.Automation.cs b/Engine/Source/Programs/AutomationTool/Scripts/CopyBuildToStagingDirectory.Automation.cs
--- a/Engine/Source/Programs/AutomationTool/Scripts/CopyBuildToStagingDirectory.Automation.cs
+++ b/Engine/Source/Programs/AutomationTool/Scripts/CopyBuildToStagingDirectory.Automation.cs
@@ -2870,12 +2870,6 @@ namespace AutomationScripts
 						Dictionary<string, string> UnrealPakResponseFile = PakParams.UnrealPakResponseFile;
 						if (ShouldCreateIoStoreContainerFiles(Params, SC))
 						{
-							bool bAllowBulkDataInIoStore = true;
-							if(!PlatformEngineConfig.GetBool("Core.System", "AllowBulkDataInIoStore", out bAllowBulkDataInIoStore))
-							{
-								bAllowBulkDataInIoStore = true; // Default is to allow it in the IoStore
-							}
-
 							bool bAllowShadersInIoStore;
 							if (!PlatformEngineConfig.GetBool("Core.System", "AllowShadersInIoStore", out bAllowShadersInIoStore))
 							{
@@ -2902,14 +2896,7 @@ namespace AutomationScripts
 								else if(Path.GetExtension(Entry.Key).Contains(".ubulk") ||
 										Path.GetExtension(Entry.Key).Contains(".uptnl"))
 								{
-									if(bAllowBulkDataInIoStore)
-									{
-										IoStoreResponseFile.Add(Entry.Key, Entry.Value);
-									}
-									else
-									{
-										UnrealPakResponseFile.Add(Entry.Key, Entry.Value);
-									}
+									IoStoreResponseFile.Add(Entry.Key, Entry.Value);
 								}
 								else if (Path.GetExtension(Entry.Key).Contains(".ushaderbytecode"))
 								{
diff --git a/Engine/Source/Runtime/CoreUObject/Private/Serialization/BulkData2.cpp b/Engine/Source/Runtime/CoreUObject/Private/Serialization/BulkData2.cpp
--- a/Engine/Source/Runtime/CoreUObject/Private/Serialization/BulkData2.cpp
+++ b/Engine/Source/Runtime/CoreUObject/Private/Serialization/BulkData2.cpp
@@ -38,27 +38,6 @@ namespace
 {
 	const uint16 InvalidBulkDataIndex = ~uint16(0);
 
-	/** Return true if bulkdata is allowed to be stored in the IoStore system and false if it must be stored in the older PakFile system */
-	bool ShouldAllowBulkDataInIoStore()
-	{
-		static struct FAllowBulkDataInIoStore
-		{
-			bool bEnabled = true;
-
-			FAllowBulkDataInIoStore()
-			{
-				FConfigFile PlatformEngineIni;
-				FConfigCacheIni::LoadLocalIniFile(PlatformEngineIni, TEXT("Engine"), true, ANSI_TO_TCHAR(FPlatformProperties::IniPlatformName()));
-
-				PlatformEngineIni.GetBool(TEXT("Core.System"), TEXT("AllowBulkDataInIoStore"), bEnabled);
-
-				UE_LOG(LogSerialization, Display, TEXT("AllowBulkDataInIoStore: '%s'"), bEnabled ?  TEXT("true") :  TEXT("false"));
-			}
-		} AllowBulkDataInIoStore;
-
-		return AllowBulkDataInIoStore.bEnabled;
-	}
-
 	/** 
 	 * Returns true if we should not trigger an ensure if we detect an inline bulkdata reload request that will not work with the 
 	 * IoStore system. This can be done by setting [Core.System]IgnoreInlineBulkDataReloadEnsures to true in the config file.
@@ -767,13 +746,8 @@ FBulkDataBase::~FBulkDataBase()
 void FBulkDataBase::ConditionalSetInlineAlwaysAllowDiscard()
 {
 	// If PackagePath is null and we do not have BULKDATA_UsesIoDispatcher, we will not be able to reload the bulkdata.
-	// We will not have a PackagePath if the engine is using IoStore, unless the setting ShouldAllowBulkDataInIoStore()
-	// has been set to false.
-	// If IoStore with !ShouldAllowBulkDataInIoStore, we set the PackagePath, but only when BulkData is in a separate
-	// BulkData file. We can load the separate BulkData files, but we we can not load the file containing the
-	// exports (which is where inline BulkData or end-of-file BulkData is stored) because those files are managed
-	// separately by IoStore .
-	// So in IoStore  with BulkData stored either inline or in the end-of-file section, we can not reload the bulkdata.
+	// We will not have a PackagePath if the engine is using IoStore.
+	// So in IoStore with BulkData stored inline, we can not reload the bulkdata.
 	// We do not need to consider the end-of-file bulkdata section because IoStore  guarantees during cooking
 	// (SaveBulkData) that only inlined data is in the package file; there is no end-of-file bulkdata section.
 
@@ -854,7 +828,6 @@ void FBulkDataBase::Serialize(FArchive& Ar, UObject* Owner, int32 /*Index*/, boo
 		checkf(Package != nullptr, TEXT("FBulkDataBase::Serialize requires an Owner that returns a valid UPackage"));
 
 		const bool bEngineUsesIoStore = FIoDispatcher::IsInitialized();
-		const bool bSeparateFilesUseIoStore = bEngineUsesIoStore && ShouldAllowBulkDataInIoStore();
 		const FPackagePath* PackagePath = nullptr;
 		const FLinkerLoad* Linker = nullptr;
 		if (bEngineUsesIoStore)
@@ -865,25 +838,8 @@ void FBulkDataBase::Serialize(FArchive& Ar, UObject* Owner, int32 /*Index*/, boo
 				TEXT("Non-inline BulkData must be stored in a separate file when EDL is enabled!"));
 			if (IsInSeparateFile())
 			{
-				if (bSeparateFilesUseIoStore)
-				{
-					Data.PackageID = Package->GetPackageIdToLoad().Value();
-					SetRuntimeBulkDataFlags(BULKDATA_UsesIoDispatcher); // Indicates that this BulkData should use the FIoChunkId rather than a PackagePath
-				}
-				else
-				{
-					// When using iostore with !ShouldAllowBulkDataInIoStore, the separate files for BulkData are loadable,
-					// but we still cannot load the package files themselves (iostore handles those in a special manner).
-					PackagePath = &Package->GetLoadedPath();
-					if (PackagePath->IsEmpty())
-					{
-						// Note that this Bulkdata object will end up with an invalid token and will end up resolving to an empty file path!
-						UE_LOG(LogSerialization, Warning, TEXT("Could not get PackagePath from package for package %s!"), *Package->GetName());
-						PackagePath = nullptr;
-					}
-					// Reset the token even though it should already be invalid (it will be set later when registered)
-					Data.Token = InvalidToken;
-				}
+				Data.PackageID = Package->GetPackageIdToLoad().Value();
+				SetRuntimeBulkDataFlags(BULKDATA_UsesIoDispatcher); // Indicates that this BulkData should use the FIoChunkId rather than a PackagePath
 			}
 			else
 			{
