Split FPackageName's initialization into two events so that FPackageName::TryConvertFilenameToLongPackageName can be called with limited functionality before pluginmanager and coreobject have been initialized.
#rnx

diff --git a/Engine/Source/Runtime/CoreUObject/Private/Misc/PackageName.cpp b/Engine/Source/Runtime/CoreUObject/Private/Misc/PackageName.cpp
--- a/Engine/Source/Runtime/CoreUObject/Private/Misc/PackageName.cpp
+++ b/Engine/Source/Runtime/CoreUObject/Private/Misc/PackageName.cpp
@@ -216,6 +216,15 @@ struct FLongPackagePathsSingleton
 		return Singleton;
 	}
 
+	/** Initialization function to setup content paths that cannot be run until CoreUObject/PluginManager have been initialized */
+	void OnCoreUObjectInitialized()
+	{
+		// Allow the plugin manager to mount new content paths by exposing access through a delegate.  PluginManager is 
+		// a Core class, but content path functionality is added at the CoreUObject level.
+		IPluginManager::Get().SetRegisterMountPointDelegate(IPluginManager::FRegisterMountPointDelegate::CreateStatic(&FPackageName::RegisterMountPoint));
+		IPluginManager::Get().SetUnRegisterMountPointDelegate(IPluginManager::FRegisterMountPointDelegate::CreateStatic(&FPackageName::UnRegisterMountPoint));
+	}
+
 	const TArray<FString>& GetValidLongPackageRoots(bool bIncludeReadOnlyRoots) const
 	{
 		return ValidLongPackageRoots[bIncludeReadOnlyRoots ? 1 : 0];
@@ -332,6 +341,8 @@ struct FLongPackagePathsSingleton
 private:
 	FLongPackagePathsSingleton()
 	{
+		SCOPED_BOOT_TIMING("FPackageName::FLongPackagePathsSingleton");
+
 		ConfigRootPath = TEXT("/Config/");
 		EngineRootPath = TEXT("/Engine/");
 		GameRootPath   = TEXT("/Game/");
@@ -398,11 +409,6 @@ private:
 		AddPathPair(ContentRootToPath, TempRootPath,   GameSavedPathRebased);
 		AddPathPair(ContentRootToPath, ConfigRootPath, GameConfigPathRebased);
 
-		// Allow the plugin manager to mount new content paths by exposing access through a delegate.  PluginManager is 
-		// a Core class, but content path functionality is added at the CoreUObject level.
-		IPluginManager::Get().SetRegisterMountPointDelegate( IPluginManager::FRegisterMountPointDelegate::CreateStatic( &FPackageName::RegisterMountPoint ) );
-		IPluginManager::Get().SetUnRegisterMountPointDelegate( IPluginManager::FRegisterMountPointDelegate::CreateStatic( &FPackageName::UnRegisterMountPoint ) );
-
 		UpdateValidLongPackageRoots();
 	}
 };
@@ -2044,10 +2050,9 @@ void FPackageName::QueryRootContentPaths(TArray<FString>& OutRootContentPaths, b
 	}
 }
 
-void FPackageName::EnsureContentPathsAreRegistered()
+void FPackageName::OnCoreUObjectInitialized()
 {
-	SCOPED_BOOT_TIMING("FPackageName::EnsureContentPathsAreRegistered");
-	FLongPackagePathsSingleton::Get();
+	FLongPackagePathsSingleton::Get().OnCoreUObjectInitialized();
 }
 
 bool FPackageName::ParseExportTextPath(const FString& InExportTextPath, FString* OutClassName, FString* OutObjectPath)
diff --git a/Engine/Source/Runtime/CoreUObject/Private/UObject/CoreNative.cpp b/Engine/Source/Runtime/CoreUObject/Private/UObject/CoreNative.cpp
--- a/Engine/Source/Runtime/CoreUObject/Private/UObject/CoreNative.cpp
+++ b/Engine/Source/Runtime/CoreUObject/Private/UObject/CoreNative.cpp
@@ -52,7 +52,7 @@ public:
 #endif
 
 		// Make sure that additional content mount points can be registered after CoreUObject loads
-		FPackageName::EnsureContentPathsAreRegistered();
+		FPackageName::OnCoreUObjectInitialized();
 
 #if DO_BLUEPRINT_GUARD
 		FFrame::InitPrintScriptCallstack();
diff --git a/Engine/Source/Runtime/CoreUObject/Public/Misc/PackageName.h b/Engine/Source/Runtime/CoreUObject/Public/Misc/PackageName.h
--- a/Engine/Source/Runtime/CoreUObject/Public/Misc/PackageName.h
+++ b/Engine/Source/Runtime/CoreUObject/Public/Misc/PackageName.h
@@ -610,7 +610,7 @@ public:
 	static void QueryRootContentPaths( TArray<FString>& OutRootContentPaths, bool bIncludeReadOnlyRoots = false, bool bWithoutLeadingSlashes = false, bool bWithoutTrailingSlashes = false);
 	
 	/** If the FLongPackagePathsSingleton is not created yet, this function will create it and thus allow mount points to be added */
-	static void EnsureContentPathsAreRegistered();
+	static void OnCoreUObjectInitialized();
 
 	/** 
 	 * Converts the supplied export text path to an object path and class name.
