Clean up virtualized bulkdata serialization and allow the package trailer to contain payloads with different access types.

#rnx

### PackageTrailer
- We no longer store the access mode for all payloads in the trailer. Instead each payload entry in the lookup table will store the access mode for that specific payload. This allow us to store many payloads of different access types in the same trailer.
- Note that this format change is currently versioned (going from 0 to 1) and is compatible with older version.
-- To allow this, the << operator has been removed from FLookupTableEntry and replaced with a ::Serialize method that takes the version of the package trailer header as an input so that it can run it's own versioning checks.
- When loading trailers with version 0 we apply the access mode to all non virtualized payload entries in the lookup table to reproduce the single access mode for the entire trailer.
- A FPackageTrailerBuilder can now be created directly from an existing FPackageTrailer that will reference all of the trailers local payloads via the new static method ::CreateReferenceToTrailer
- Removed FPackageTrailer::TrySave as it is no longer used when creating a trailer for the editor domain.
- Renamed FPackageTrailer::LoadPayload to ::LoadLocalPayload to make it clear the type of payload that can be loaded.
-- Someday the package trailer should be able to load payloads of any type.

### VirtualizedBulkData Serialization Changes
- We now record if the bulkdata's payload is stored in a package trailer or not via the 'StoredInPackageTrailer' flag, which makes it easier to identify critical errors in serialization (the trailer is missing for example)
-- Storing this info as a flag is now possible as the design for the trailer has changed and a fully virtualized package will still contain a trailer.
- We now only write out OffsetInFile if the payload is not stored in a FPackageTrailer, if the bulkdata object does use a package trailer via the new flag then we take the offset from the trailer instead.
- We also no longer write out a path when saving a payload as referenced to the editor domain, it is assumed that the trailer will know where the workspace domain data is (namely the owning package path)
-- Due to the 'StoredInPackageTrailer' only being added now, this means that packages already saved with a package trailer will work with this code logic as those packages will have serialized the offsets to disk, but won't have the new flag and so will serialize the offsets when loading until re-saved. (worth noting that nobody should encounter this as the package trailer system is disabled by default)
- The loading code is a lot easier to read now and a lot of the branches have been removed although we still have a few extra checks for older formats that might have been saved during development.
- When loading the only real consideration we have is if the payload is in the trailer or not.
- Future Work:
-- If we are not writing to the trailer then we use the legacy serialization system which appends the bulkdata to the end of the package. This is used in two cases, 1) When the package trailer system is disabled 2) If the payload is updated in postload and we are writing out the package to the editor domain as the package trailer does not support writing out local payloads outside of the workspace domain at the moment. The trailer should handle this itself and not rely on the legacy serialization system so that we can remove it.
-- Note that if a payload is stored by reference, we are not adding it explicitly to the payload trailer and we are relying on the trailer builder already containing the reference information when it is created in SavePackage. We do however run an assert to ensure that the reference does exist in the trailer but it might be worth revisiting this.

### VirtualizedBulkData Misc
- The member 'PackageSegment' is now initialized in the constructor to EPackageSegment::Header.
- Moved ::GetPackagePathFromOwner from being a method to a utility function as it accessed no members.
- Validating the payload for saving during serialization has been moved from ::Serialize to it's own method ::TryPayloadValidationForSaving.

### PackageSubmissionChecks
- Added an error message if we detect referenced payloads in a trailer we are trying to virtualize as this condition should never occur. If we encounter it, we give a user facing error message and fail the submit process. It might be worth demoting this to an assert instead as it is not an error the user is expected to see.

diff --git a/Engine/Source/Developer/Virtualization/Private/PackageSubmissionChecks.cpp b/Engine/Source/Developer/Virtualization/Private/PackageSubmissionChecks.cpp
--- a/Engine/Source/Developer/Virtualization/Private/PackageSubmissionChecks.cpp
+++ b/Engine/Source/Developer/Virtualization/Private/PackageSubmissionChecks.cpp
@@ -173,12 +173,22 @@ void OnPrePackageSubmission(const TArray<FString>& FilesToSubmit, TArray<FText>&
 	{
 		FPackagePath PackagePath = FPackagePath::FromLocalPath(AbsoluteFilePath);
 
-		// TODO: How to handle text packages
+		// TODO: How to handle text packages?
 		if (FPackageName::IsPackageExtension(PackagePath.GetHeaderExtension()) || FPackageName::IsTextPackageExtension(PackagePath.GetHeaderExtension()))
 		{
 			FPackageTrailer Trailer;
 			if (FPackageTrailer::TryLoadFromPackage(PackagePath, Trailer))
 			{
+				// The following is not expected to ever happen, currently we give a user facing error but it generally means that the asset is broken somehow.
+				ensureMsgf(Trailer.GetNumPayloads(EPayloadFilter::Referenced) == 0, TEXT("Trying to virtualize a package that already contains payload references which the workspace file should not ever contain!"));
+				if (Trailer.GetNumPayloads(EPayloadFilter::Referenced) > 0)
+				{
+					FText Message = FText::Format(	LOCTEXT("Virtualization_PkgHasReferences", "Cannot virtualize the package '{1}' as it has referenced payloads in the trailer"),
+													FText::FromString(PackagePath.GetDebugName()));
+					Errors.Add(Message);
+					return;
+				}
+
 				FPackageInfo PkgInfo;
 
 				PkgInfo.Path = MoveTemp(PackagePath);
@@ -277,7 +287,7 @@ void OnPrePackageSubmission(const TArray<FString>& FilesToSubmit, TArray<FText>&
 		{
 			checkf(PayloadId.IsValid(), TEXT("PackageTrailer for package '%s' should not contain invalid FPayloadIds"), *PackageInfo.Path.GetDebugName());
 			
-			FCompressedBuffer Payload = PackageInfo.Trailer.LoadPayload(PayloadId, *PackageAr);
+			FCompressedBuffer Payload = PackageInfo.Trailer.LoadLocalPayload(PayloadId, *PackageAr);
 
 			if (PayloadId != FIoHash(Payload.GetRawHash()))
 			{
@@ -386,7 +396,7 @@ void OnPrePackageSubmission(const TArray<FString>& FilesToSubmit, TArray<FText>&
 			return;
 		}
 
-		FPackageTrailerBuilder TrailerBuilder = FPackageTrailerBuilder::CreateFromTrailer(PackageInfo.Trailer, *PackageAr, PackagePath);
+		FPackageTrailerBuilder TrailerBuilder = FPackageTrailerBuilder::CreateFromTrailer(PackageInfo.Trailer, *PackageAr, PackagePath.GetPackageFName());
 		if (!TrailerBuilder.BuildAndAppendTrailer(nullptr, *CopyAr))
 		{
 			FText Message = FText::Format(	LOCTEXT("Virtualization_TrailerAppend", "Failed to append the trailer to '{0}'"),
diff --git a/Engine/Source/Editor/UnrealEd/Private/EditorDomain/EditorDomainUtils.cpp b/Engine/Source/Editor/UnrealEd/Private/EditorDomain/EditorDomainUtils.cpp
--- a/Engine/Source/Editor/UnrealEd/Private/EditorDomain/EditorDomainUtils.cpp
+++ b/Engine/Source/Editor/UnrealEd/Private/EditorDomain/EditorDomainUtils.cpp
@@ -79,7 +79,7 @@ bool bGGlobalAddedCustomVersionsInitialized = false;
 int64 GMaxBulkDataSize = -1;
 
 // Change to a new guid when EditorDomain needs to be invalidated
-const TCHAR* EditorDomainVersion = TEXT("83221BBB67B3482A8064C76349ECE07D");
+const TCHAR* EditorDomainVersion = TEXT("54D98F079D2A46A9802B2B990560285D");
 
 // Identifier of the CacheBuckets for EditorDomain tables
 const TCHAR* EditorDomainPackageBucketName = TEXT("EditorDomainPackage");
diff --git a/Engine/Source/Runtime/CoreUObject/Private/Serialization/VirtualizedBulkData.cpp b/Engine/Source/Runtime/CoreUObject/Private/Serialization/VirtualizedBulkData.cpp
--- a/Engine/Source/Runtime/CoreUObject/Private/Serialization/VirtualizedBulkData.cpp
+++ b/Engine/Source/Runtime/CoreUObject/Private/Serialization/VirtualizedBulkData.cpp
@@ -181,6 +181,23 @@ const FPackageTrailer* GetTrailerFromOwner(UObject* Owner)
 	}
 }
 
+/** Utility for finding the package path associated with a given UObject */
+FPackagePath GetPackagePathFromOwner(UObject* Owner, EPackageSegment& OutPackageSegment)
+{
+	OutPackageSegment = EPackageSegment::Header;
+
+	const FLinkerLoad* Linker = GetLinkerLoadFromOwner(Owner);
+
+	if (Linker != nullptr)
+	{
+		return Linker->GetPackagePath();
+	}
+	else
+	{
+		return FPackagePath();
+	}
+}
+
 /** Utility for updating an existing entry in an Archive before returning the archive to it's original seek position */
 template<typename DataType>
 void UpdateArchiveData(FArchive& Ar, int64 DataPosition, DataType& Data)
@@ -616,7 +633,7 @@ void FVirtualizedUntypedBulkData::Serialize(FArchive& Ar, UObject* Owner, bool b
 		{
 			checkf(Ar.IsCooking() == false, TEXT("FVirtualizedUntypedBulkData::Serialize should not be called during a cook"));
 
-			EFlags UpdatedFlags = BuildFlagsForSerialization(Ar, bKeepFileDataByReference);
+			const EFlags UpdatedFlags = BuildFlagsForSerialization(Ar, bKeepFileDataByReference);
 
 			// Go back in the archive and update the flags in the archive, we will only apply the updated flags to the current
 			// object later if we detect that the package saved successfully.
@@ -630,9 +647,11 @@ void FVirtualizedUntypedBulkData::Serialize(FArchive& Ar, UObject* Owner, bool b
 			if (IsReferencingByPackagePath(UpdatedFlags))
 			{
 				check(PackageSegment == EPackageSegment::Header); // This should have been checked before setting bKeepFileDataByReference=true
-				FString PackageName = PackagePath.GetPackageName();
-				Ar << PackageName;
-				Ar << OffsetInFile;
+				if (!IsStoredInPackageTrailer(UpdatedFlags))
+				{
+					Ar << OffsetInFile;
+				}
+
 				bWriteOutPayload = false;
 			}
 			else
@@ -645,25 +664,9 @@ void FVirtualizedUntypedBulkData::Serialize(FArchive& Ar, UObject* Owner, bool b
 				// Need to load the payload so that we can write it out
 				FCompressedBuffer PayloadToSerialize = GetDataInternal();
 				
-				// Checks to make sure that the payload we are saving is what we expect it to be. If not then we need to log an error to the
-				// user, mark the archive as having an error and return.
-				// This error is not reasonable to expect, if it occurs then it means something has gone very wrong, which is why there is
-				// also an ensure to make sure that any instances of it occurring as properly recorded and can be investigated.
-				if (!IsDataValid(*this, PayloadToSerialize) || (GetPayloadSize() > 0 && PayloadToSerialize.IsNull()))
+				if (!TryPayloadValidationForSaving(PayloadToSerialize, LinkerSave))
 				{
-					ensureMsgf(false, TEXT("%s"), *GetCorruptedPayloadErrorMsgForSave(LinkerSave).ToString());
-
-					if (LinkerSave && LinkerSave->GetOutputDevice() != nullptr)
-					{
-						LinkerSave->GetOutputDevice()->Logf(ELogVerbosity::Error, TEXT("%s"), *GetCorruptedPayloadErrorMsgForSave(LinkerSave).ToString());
-					}
-					else
-					{
-						UE_LOG(LogVirtualization, Error, TEXT("%s"), *GetCorruptedPayloadErrorMsgForSave(LinkerSave).ToString());
-					}
-
 					Ar.SetError();
-
 					return;
 				}
 
@@ -673,33 +676,34 @@ void FVirtualizedUntypedBulkData::Serialize(FArchive& Ar, UObject* Owner, bool b
 				// to prevent potentially bad data being saved to disk.
 				checkf(PayloadToSerialize || GetPayloadSize() == 0, TEXT("Failed to acquire the payload for saving!"));
 
-				// Write out a dummy value that we will write over once the payload is serialized
-				int64 PlaceholderValue = INDEX_NONE;
-
-				const int64 OffsetPos = Ar.Tell();
-				Ar << PlaceholderValue; // OffsetInFile
-
 				// If we have a valid linker then we will defer serialization of the payload so that it will
 				// be placed at the end of the output file so we don't have to seek past the payload on load.
 				// If we do not have a linker OR the linker is in text format then we should just serialize
 				// the payload directly to the archive.
 				if (LinkerSave != nullptr && !LinkerSave->IsTextFormat())
-				{
-					if (ShouldUseLegacySerialization(*LinkerSave))
+				{	
+					if (IsStoredInPackageTrailer(UpdatedFlags))
 					{
-						// Legacy path, will save the payload data to the package
-						SerializeToLegacyPath(*LinkerSave, OffsetPos, PayloadToSerialize, UpdatedFlags, Owner);
+						// New path that will save the payload to the package trailer
+						SerializeToPackageTrailer(*LinkerSave, PayloadToSerialize, UpdatedFlags, Owner);	
 					}
 					else 
 					{
-						// New path that will save the payload to the package trailer
-						SerializeToPackageTrailer(*LinkerSave, OffsetPos, PayloadToSerialize, UpdatedFlags, Owner);
+						// Legacy path, will save the payload data to the end of the package
+						SerializeToLegacyPath(*LinkerSave, PayloadToSerialize, UpdatedFlags, Owner);
 					}	
 				}
 				else
 				{
-					checkf(Ar.IsCooking() == false, TEXT("FVirtualizedUntypedBulkData::Serialize should not be called during a cook"));
-					
+					// Not saving to a package so serialize inline into the archive
+					check(IsStoredInPackageTrailer(UpdatedFlags) == false);
+
+					const int64 OffsetPos = Ar.Tell();
+
+					// Write out a dummy value that we will write over once the payload is serialized
+					int64 PlaceholderValue = INDEX_NONE;
+					Ar << PlaceholderValue; // OffsetInFile
+
 					int64 DataStartOffset = Ar.Tell();
 
 					SerializeData(Ar, PayloadToSerialize, UpdatedFlags);
@@ -707,6 +711,28 @@ void FVirtualizedUntypedBulkData::Serialize(FArchive& Ar, UObject* Owner, bool b
 					UpdateArchiveData(Ar, OffsetPos, DataStartOffset);
 				}
 			}
+
+			// Some paranoid checks to make sure everything is working as expected
+			// TODO: Probably remove this when FPackageTrailer::IsEnabled is removed
+			if(IsStoredInPackageTrailer(UpdatedFlags) && PayloadContentId.IsValid())
+			{
+				check(LinkerSave != nullptr);
+				check(LinkerSave->PackageTrailerBuilder.IsValid());
+				checkf(!(IsDataVirtualized(UpdatedFlags) && IsReferencingByPackagePath(UpdatedFlags)), TEXT("Payload cannot be both virtualized and a reference"));
+
+				if (IsReferencingByPackagePath(UpdatedFlags))
+				{
+					check(LinkerSave->PackageTrailerBuilder->IsReferencedPayloadEntry(PayloadContentId));
+				}
+				else if (IsDataVirtualized(UpdatedFlags))
+				{
+					check(LinkerSave->PackageTrailerBuilder->IsVirtualizedPayloadEntry(PayloadContentId));
+				}
+				else
+				{
+					check(LinkerSave->PackageTrailerBuilder->IsLocalPayloadEntry(PayloadContentId));
+				}
+			}
 			
 			if (CanUnloadData())
 			{
@@ -716,59 +742,44 @@ void FVirtualizedUntypedBulkData::Serialize(FArchive& Ar, UObject* Owner, bool b
 		}
 		else if (Ar.IsLoading())
 		{
-			const FPackageTrailer* Trailer = GetTrailerFromOwner(Owner);
+			OffsetInFile = INDEX_NONE;
+			PackagePath.Empty();
+			PackageSegment = EPackageSegment::Header;
 
-			if (IsDataVirtualized())
+			const FPackageTrailer* Trailer = GetTrailerFromOwner(Owner);
+			checkf(Trailer != nullptr || !IsStoredInPackageTrailer(), TEXT("Payload was stored in a package trailer, but there no trailer loaded"));
+		
+			if (IsStoredInPackageTrailer())
 			{
-				// The payload was already virtualized when the bulkdata object was serialized
-
-				// We aren't going to use these members so reset them
-				OffsetInFile = INDEX_NONE;
-
-				PackagePath.Empty();
-				PackageSegment = EPackageSegment::Header;
+				// Cache the offset from the trailer (if we move the loading of the payload to the trailer 
+				// at a later point then we can skip this)
+				OffsetInFile = Trailer->FindPayloadOffsetInFile(PayloadContentId);
 			}
-			else if (Trailer != nullptr && Trailer->FindPayloadStatus(PayloadContentId) == EPayloadStatus::StoredVirtualized)
+			else
 			{
-				// The payload was stored locally in the package trailer when the bulkdata object was 
-				// serialized but the trailer was later virtualized.
-
-				check(!IsReferencingByPackagePath()); // This should not be possible
+				// TODO: This check is for older virtualized formats that might be seen in older test projects.
+				UE_CLOG(IsDataVirtualized(), LogVirtualization, Error, TEXT("Payload in '%s' is virtualized in an older format and should be re-saved!"), *Owner->GetName());
+				if (!IsDataVirtualized())
+				{
+					Ar << OffsetInFile;
+				}
+			}
 
+			// This cannot be inside the above ::IsStoredInPackageTrailer branch due to the original prototype assets using the trailer without the StoredInPackageTrailer flag
+			if (Trailer != nullptr && Trailer->FindPayloadStatus(PayloadContentId) == EPayloadStatus::StoredVirtualized)
+			{
+				// As the virtualization process happens outside of serialization we need
+				// to check with the trailer to see if the payload is virtualized or not
 				EnumAddFlags(Flags, EFlags::IsVirtualized);
-
-				Ar << OffsetInFile;
-
-				// We aren't going to use these members so reset them
 				OffsetInFile = INDEX_NONE;
-
-				PackagePath.Empty();
-				PackageSegment = EPackageSegment::Header;
 			}
-			else if (IsReferencingByPackagePath())
-			{
-				FString PackageNameStr;
-				Ar << PackageNameStr;
-				Ar << OffsetInFile;
-				ensure(FPackagePath::TryFromPackageName(*PackageNameStr, this->PackagePath));
-				PackageSegment = EPackageSegment::Header;
 
-#if VBD_ALLOW_LINKERLOADER_ATTACHMENT
-				FArchive* CacheableArchive = Ar.GetCacheableArchive();
-				if (Ar.IsAllowingLazyLoading() && CacheableArchive != nullptr)
-				{
-					AttachedAr = CacheableArchive;
-					if (AttachedAr != nullptr)
-					{
-						AttachedAr->AttachBulkData(this);
-					}
-				}
-#endif //VBD_ALLOW_LINKERLOADER_ATTACHMENT
-			}
-			else 
+			checkf(!(IsDataVirtualized() && IsReferencingByPackagePath()), TEXT("Payload cannot be both virtualized and a reference"));
+			checkf(!IsDataVirtualized() || OffsetInFile == INDEX_NONE, TEXT("Virtualized payloads should have an invalid offset"));
+			
+			if (!IsDataVirtualized())
 			{
-				// If we can lazy load then find the PackagePath, otherwise we will want
-				// to serialize immediately.
+				// If we can lazy load then find the PackagePath, otherwise we will want to serialize immediately.
 				FArchive* CacheableArchive = Ar.GetCacheableArchive();
 				if (Ar.IsAllowingLazyLoading() && CacheableArchive != nullptr)
 				{
@@ -779,24 +790,21 @@ void FVirtualizedUntypedBulkData::Serialize(FArchive& Ar, UObject* Owner, bool b
 					PackagePath.Empty();
 					PackageSegment = EPackageSegment::Header;
 				}
-				
-				OffsetInFile = INDEX_NONE;
-				Ar << OffsetInFile;
-
-				if (!PackagePath.IsEmpty())
+					
+				if (!PackagePath.IsEmpty() && CacheableArchive != nullptr)
 				{
 #if VBD_ALLOW_LINKERLOADER_ATTACHMENT
 					AttachedAr = CacheableArchive;
-					if (AttachedAr != nullptr)
-					{
-						AttachedAr->AttachBulkData(this);
-					}
+					AttachedAr->AttachBulkData(this);
 #endif //VBD_ALLOW_LINKERLOADER_ATTACHMENT
 				}
 				else
 				{
-					// If we have no packagepath then we need to load the data immediately
-					// as we will not be able to load it on demand.
+					checkf(Ar.Tell() == OffsetInFile, TEXT("Attempting to load an inline payload but the offset does not match"));
+
+					// If the package path is invalid or the archive is not cacheable then we
+					// cannot rely on loading the payload at a future point on demand so we need 
+					// to load the data immediately.
 					FCompressedBuffer CompressedPayload;
 					SerializeData(Ar, CompressedPayload, Flags);
 					
@@ -1002,7 +1010,7 @@ FCompressedBuffer FVirtualizedUntypedBulkData::LoadFromPackageTrailer() const
 	
 	if (Trailer.TryLoadBackwards(*BulkArchive))
 	{
-		return Trailer.LoadPayload(PayloadContentId, *BulkArchive);
+		return Trailer.LoadLocalPayload(PayloadContentId, *BulkArchive);
 	}
 	else
 	{
@@ -1202,22 +1210,6 @@ FCompressedBuffer FVirtualizedUntypedBulkData::PullData() const
 	return PulledPayload;
 }
 
-FPackagePath FVirtualizedUntypedBulkData::GetPackagePathFromOwner(UObject* Owner, EPackageSegment& OutPackageSegment) const
-{
-	OutPackageSegment = EPackageSegment::Header;
-
-	const FLinkerLoad* Linker = GetLinkerLoadFromOwner(Owner);
-
-	if (Linker != nullptr)
-	{
-		return Linker->GetPackagePath();
-	}
-	else
-	{
-		return FPackagePath();
-	}
-}
-
 bool FVirtualizedUntypedBulkData::CanUnloadData() const
 {
 	// We cannot unload the data if are unable to reload it from a file
@@ -1297,8 +1289,14 @@ FGuid FVirtualizedUntypedBulkData::GetIdentifier() const
 	return BulkDataId;
 }
 
-void FVirtualizedUntypedBulkData::SerializeToLegacyPath(FLinkerSave& LinkerSave, int64 OffsetPos, FCompressedBuffer PayloadToSerialize, EFlags UpdatedFlags, UObject* Owner)
+void FVirtualizedUntypedBulkData::SerializeToLegacyPath(FLinkerSave& LinkerSave, FCompressedBuffer PayloadToSerialize, EFlags UpdatedFlags, UObject* Owner)
 {
+	const int64 OffsetPos = LinkerSave.Tell();
+
+	// Write out a dummy value that we will write over once the payload is serialized
+	int64 PlaceholderValue = INDEX_NONE;
+	LinkerSave << PlaceholderValue; // OffsetInFile
+
 	// The lambda is mutable so that PayloadToSerialize is not const (due to FArchive api not accepting const values)
 	auto SerializePayload = [this, OffsetPos, PayloadToSerialize, UpdatedFlags, Owner](FLinkerSave& LinkerSave, FArchive& ExportsArchive, FArchive& DataArchive, int64 DataStartOffset) mutable
 	{
@@ -1358,16 +1356,14 @@ void FVirtualizedUntypedBulkData::SerializeToLegacyPath(FLinkerSave& LinkerSave,
 																				// this point however we test LinkerSave != nullptr to enter this branch.
 }
 
-void FVirtualizedUntypedBulkData::SerializeToPackageTrailer(FLinkerSave& LinkerSave, int64 OffsetPos, FCompressedBuffer PayloadToSerialize, EFlags UpdatedFlags, UObject* Owner)
+void FVirtualizedUntypedBulkData::SerializeToPackageTrailer(FLinkerSave& LinkerSave, FCompressedBuffer PayloadToSerialize, EFlags UpdatedFlags, UObject* Owner)
 {
-	auto OnPayloadWritten = [this, OffsetPos, UpdatedFlags](FLinkerSave& LinkerSave, const FPackageTrailer& Trailer) mutable
+	auto OnPayloadWritten = [this, UpdatedFlags](FLinkerSave& LinkerSave, const FPackageTrailer& Trailer) mutable
 	{
 		checkf(LinkerSave.IsCooking() == false, TEXT("FVirtualizedUntypedBulkData::Serialize should not be called during a cook"));
 
 		int64 PayloadOffset = Trailer.FindPayloadOffsetInFile(PayloadContentId);
 
-		UpdateArchiveData(LinkerSave, OffsetPos, PayloadOffset);
-
 		// If we are saving the package to disk (we have access to FLinkerSave and its filepath is valid) 
 		// then we should register a callback to be received once the package has actually been saved to 
 		// disk so that we can update the object's members to be redirected to the saved file.
@@ -1675,6 +1671,8 @@ FVirtualizedUntypedBulkData::EFlags FVirtualizedUntypedBulkData::BuildFlagsForSe
 	{
 		EFlags UpdatedFlags = Flags;
 
+		const FLinkerSave* LinkerSave = Cast<FLinkerSave>(Ar.GetLinker());
+
 		// Now update any changes to the flags that we might need to make when serializing.
 		// Note that these changes are not applied to the current object UNLESS we are saving
 		// the package, in which case the newly modified flags will be applied once we confirm
@@ -1695,7 +1693,7 @@ FVirtualizedUntypedBulkData::EFlags FVirtualizedUntypedBulkData::BuildFlagsForSe
 			EnumRemoveFlags(UpdatedFlags, EFlags::ReferencesLegacyFile | EFlags::ReferencesWorkspaceDomain | 
 				EFlags::LegacyFileIsCompressed | EFlags::LegacyKeyWasGuidDerived);
 
-			const FLinkerSave* LinkerSave = Cast<FLinkerSave>(Ar.GetLinker());
+			
 			if (LinkerSave != nullptr && !LinkerSave->GetFilename().IsEmpty() && ShouldSaveToPackageSidecar())
 			{
 				EnumAddFlags(UpdatedFlags, EFlags::HasPayloadSidecarFile);
@@ -1715,6 +1713,15 @@ FVirtualizedUntypedBulkData::EFlags FVirtualizedUntypedBulkData::BuildFlagsForSe
 			}
 		}
 
+		if (ShouldUseLegacySerialization(LinkerSave) == false)
+		{
+			EnumAddFlags(UpdatedFlags, EFlags::StoredInPackageTrailer);
+		}
+		else
+		{
+			EnumRemoveFlags(UpdatedFlags, EFlags::StoredInPackageTrailer);
+		}
+
 		return UpdatedFlags;
 	}
 	else
@@ -1723,6 +1730,31 @@ FVirtualizedUntypedBulkData::EFlags FVirtualizedUntypedBulkData::BuildFlagsForSe
 	}
 }
 
+bool FVirtualizedUntypedBulkData::TryPayloadValidationForSaving(const FCompressedBuffer& PayloadForSaving, FLinkerSave* LinkerSave) const
+{
+	if (!IsDataValid(*this, PayloadForSaving) || (GetPayloadSize() > 0 && PayloadForSaving.IsNull()))
+	{
+		const FString ErrorMessage = GetCorruptedPayloadErrorMsgForSave(LinkerSave).ToString();
+
+		ensureMsgf(false, TEXT("%s"), *ErrorMessage);
+
+		if (LinkerSave != nullptr && LinkerSave->GetOutputDevice() != nullptr)
+		{
+			LinkerSave->GetOutputDevice()->Logf(ELogVerbosity::Error, TEXT("%s"), *ErrorMessage);
+		}
+		else
+		{
+			UE_LOG(LogVirtualization, Error, TEXT("%s"), *ErrorMessage);
+		}
+
+		return false;
+	}
+	else
+	{
+		return true;
+	}
+}
+
 FText FVirtualizedUntypedBulkData::GetCorruptedPayloadErrorMsgForSave(FLinkerSave* Linker) const
 {
 	const FText GuidID = FText::FromString(GetIdentifier().ToString());
@@ -1750,12 +1782,17 @@ FText FVirtualizedUntypedBulkData::GetCorruptedPayloadErrorMsgForSave(FLinkerSav
 	}
 }
 
-bool FVirtualizedUntypedBulkData::ShouldUseLegacySerialization(FLinkerSave& LinkerSave) const
+bool FVirtualizedUntypedBulkData::ShouldUseLegacySerialization(const FLinkerSave* LinkerSave) const
 {
+	if (LinkerSave == nullptr)
+	{
+		return true;
+	}
+
 #if UE_ENABLE_VIRTUALIZATION_TOGGLE
-	return !LinkerSave.PackageTrailerBuilder.IsValid() || bSkipVirtualization == true;
+	return (!LinkerSave->PackageTrailerBuilder.IsValid()) || bSkipVirtualization == true;
 #else
-	return !LinkerSave.PackageTrailerBuilder.IsValid();
+	return !LinkerSave->PackageTrailerBuilder.IsValid();
 #endif //UE_ENABLE_VIRTUALIZATION_TOGGLE
 }
 	
diff --git a/Engine/Source/Runtime/CoreUObject/Private/UObject/PackageTrailer.cpp b/Engine/Source/Runtime/CoreUObject/Private/UObject/PackageTrailer.cpp
--- a/Engine/Source/Runtime/CoreUObject/Private/UObject/PackageTrailer.cpp
+++ b/Engine/Source/Runtime/CoreUObject/Private/UObject/PackageTrailer.cpp
@@ -18,6 +18,8 @@ enum class EPackageTrailerVersion : uint32
 {
 	// The original trailer format when it was first added
 	INITIAL = 0,
+	// Access mode is now per payload and found in FLookupTableEntry 
+	ACCESS_PER_PAYLOAD = 1,
 
 	// -----<new versions can be added before this line>-------------------------------------------------
 	// - this needs to be the last line (see note below)
@@ -27,8 +29,8 @@ enum class EPackageTrailerVersion : uint32
 
 // These asserts are here to make sure that any changes to the size of disk constants are intentional.
 // If the change was intentional then just update the assert.
-static_assert(FPackageTrailer::FHeader::StaticHeaderSizeOnDisk == 29, "FPackageTrailer::FHeader size has been changed, if this was intentional then update this assert");
-static_assert(Private::FLookupTableEntry::SizeOnDisk == 44, "FLookupTableEntry size has been changed, if this was intentional then update this assert");
+static_assert(FPackageTrailer::FHeader::StaticHeaderSizeOnDisk == 28, "FPackageTrailer::FHeader size has been changed, if this was intentional then update this assert");
+static_assert(Private::FLookupTableEntry::SizeOnDisk == 45, "FLookupTableEntry size has been changed, if this was intentional then update this assert");
 static_assert(FPackageTrailer::FFooter::SizeOnDisk == 20, "FPackageTrailer::FFooter size has been changed, if this was intentional then update this assert");
 
 /** Utility for recording failed package open reasons */
@@ -60,49 +62,103 @@ FLookupTableEntry::FLookupTableEntry(const Virtualization::FPayloadId& InIdentif
 
 }
 
-FArchive& operator<<(FArchive& Ar, FLookupTableEntry& Entry)
+void FLookupTableEntry::Serialize(FArchive& Ar, EPackageTrailerVersion PackageTrailerVersion)
 {
-	Ar << Entry.Identifier;
-	Ar << Entry.OffsetInFile;
-	Ar << Entry.CompressedSize;
-	Ar << Entry.RawSize;
+	Ar << Identifier;
+	Ar << OffsetInFile;
+	Ar << CompressedSize;
+	Ar << RawSize;
 
-	return Ar;
+	if (Ar.IsSaving() || PackageTrailerVersion >= EPackageTrailerVersion::ACCESS_PER_PAYLOAD)
+	{
+		Ar << AccessMode;
+	}
 }
 
 } // namespace Private
 
-FPackageTrailerBuilder FPackageTrailerBuilder::CreateFromTrailer(const FPackageTrailer& Trailer, FArchive& Ar, const FPackagePath& PackagePath)
+FPackageTrailerBuilder FPackageTrailerBuilder::CreateFromTrailer(const FPackageTrailer& Trailer, FArchive& Ar, const FName& PackageName)
 {
-	FPackageTrailerBuilder Builder(PackagePath.GetPackageFName());
+	FPackageTrailerBuilder Builder(PackageName);
 
 	for (const Private::FLookupTableEntry& Entry : Trailer.Header.PayloadLookupTable)
 	{
-		checkf(Entry.Identifier.IsValid(), TEXT("PackageTrailer for package should not contain invalid FPayloadIds. Package '%s'"), *PackagePath.GetPackageName());
+		checkf(Entry.Identifier.IsValid(), TEXT("PackageTrailer for package should not contain invalid FPayloadIds. Package '%s'"), *PackageName.ToString());
 
-		if (Entry.IsVirtualized())
+		switch (Entry.AccessMode)
 		{
-			Builder.VirtualizedEntries.Add(Entry.Identifier, VirtualizedEntry(Entry.CompressedSize, Entry.RawSize));
-		}
-		else
-		{
-			FCompressedBuffer Payload = Trailer.LoadPayload(Entry.Identifier, Ar);
-			Builder.LocalEntries.Add(Entry.Identifier, LocalEntry(MoveTemp(Payload)));
+			case EPayloadAccessMode::Local:
+			{
+				FCompressedBuffer Payload = Trailer.LoadLocalPayload(Entry.Identifier, Ar);
+				Builder.LocalEntries.Add(Entry.Identifier, LocalEntry(MoveTemp(Payload)));
+			}
+			break;
+
+			case EPayloadAccessMode::Referenced:
+			{
+				Builder.ReferencedEntries.Add(Entry.Identifier, ReferencedEntry(Entry.OffsetInFile, Entry.CompressedSize, Entry.RawSize));
+			}
+			break;
+
+			case EPayloadAccessMode::Virtualized:
+			{
+				Builder.VirtualizedEntries.Add(Entry.Identifier, VirtualizedEntry(Entry.CompressedSize, Entry.RawSize));
+			}
+			break;
+
+			default:
+			{
+				checkNoEntry();
+			}
 		}
 	}
 
 	return Builder;
 }
 
-FPackageTrailerBuilder::FPackageTrailerBuilder(UPackage* Package)
+TUniquePtr<UE::FPackageTrailerBuilder> FPackageTrailerBuilder::CreateReferenceToTrailer(const class FPackageTrailer& Trailer, const FName& PackageName)
 {
-	if (Package != nullptr)
+	TUniquePtr<UE::FPackageTrailerBuilder> Builder = MakeUnique<UE::FPackageTrailerBuilder>(PackageName);
+
+	for (const Private::FLookupTableEntry& Entry : Trailer.Header.PayloadLookupTable)
 	{
-		PackageName = Package->GetFName();
+		checkf(Entry.Identifier.IsValid(), TEXT("PackageTrailer for package should not contain invalid FPayloadIds. Package '%s'"), *PackageName.ToString());
+
+		switch (Entry.AccessMode)
+		{
+			case EPayloadAccessMode::Local:
+			{
+				const int64 AbsoluteOffset = Trailer.FindPayloadOffsetInFile(Entry.Identifier);
+				checkf(AbsoluteOffset != INDEX_NONE, TEXT("PackageTrailer for package should not contain invalid payload offsets. Package '%s'"), *PackageName.ToString());
+
+				Builder->ReferencedEntries.Add(Entry.Identifier, ReferencedEntry(AbsoluteOffset, Entry.CompressedSize, Entry.RawSize));
+			}
+			break;
+			
+			case EPayloadAccessMode::Referenced:
+			{
+				checkf(false, TEXT("Attempting to create a reference to a trailer that already contains reference payload entries. Package '%s'"), *PackageName.ToString());
+			}
+			break;
+			
+			case EPayloadAccessMode::Virtualized:
+			{
+				Builder->VirtualizedEntries.Add(Entry.Identifier, VirtualizedEntry(Entry.CompressedSize, Entry.RawSize));
+			}
+			break;
+
+			default:
+			{
+				checkNoEntry();
+			}
+
+		}
 	}
+
+	return Builder;
 }
 
-FPackageTrailerBuilder::FPackageTrailerBuilder(FName InPackageName)
+FPackageTrailerBuilder::FPackageTrailerBuilder(const FName& InPackageName)
 	: PackageName(InPackageName)
 {
 }
@@ -130,7 +186,7 @@ bool FPackageTrailerBuilder::BuildAndAppendTrailer(FLinkerSave* Linker, FArchive
 	Trailer.Header.Tag = FPackageTrailer::FHeader::HeaderTag;
 	Trailer.Header.Version = (int32)EPackageTrailerVersion::AUTOMATIC_VERSION;
 
-	const uint32 DynamicHeaderSizeOnDisk = ((LocalEntries.Num() + VirtualizedEntries.Num()) * Private::FLookupTableEntry::SizeOnDisk); // Add the length of the lookup table
+	const uint32 DynamicHeaderSizeOnDisk = (GetNumPayloads() * Private::FLookupTableEntry::SizeOnDisk); // Add the length of the lookup table
 	
 	Trailer.Header.HeaderLength = FPackageTrailer::FHeader::StaticHeaderSizeOnDisk + DynamicHeaderSizeOnDisk;
 	
@@ -146,10 +202,23 @@ bool FPackageTrailerBuilder::BuildAndAppendTrailer(FLinkerSave* Linker, FArchive
 		Entry.OffsetInFile = Trailer.Header.PayloadsDataLength;
 		Entry.CompressedSize = It.Value.Payload.GetCompressedSize();
 		Entry.RawSize = It.Value.Payload.GetRawSize();
+		Entry.AccessMode = EPayloadAccessMode::Local;
 
 		Trailer.Header.PayloadsDataLength += It.Value.Payload.GetCompressedSize();
 	}
 
+	for (const TPair<Virtualization::FPayloadId, ReferencedEntry>& It : ReferencedEntries)
+	{
+		checkf(It.Key.IsValid(), TEXT("PackageTrailer should not contain invalid FPayloadIds. Package '%s'"), *PackageName.ToString());
+
+		Private::FLookupTableEntry& Entry = Trailer.Header.PayloadLookupTable.AddDefaulted_GetRef();
+		Entry.Identifier = It.Key;
+		Entry.OffsetInFile = It.Value.Offset;
+		Entry.CompressedSize = It.Value.CompressedSize;
+		Entry.RawSize = It.Value.RawSize;
+		Entry.AccessMode = EPayloadAccessMode::Referenced;
+	}
+
 	for (const TPair<Virtualization::FPayloadId, VirtualizedEntry>& It : VirtualizedEntries)
 	{
 		checkf(It.Key.IsValid(), TEXT("PackageTrailer should not contain invalid FPayloadIds. Package '%s'"), *PackageName.ToString());
@@ -159,10 +228,9 @@ bool FPackageTrailerBuilder::BuildAndAppendTrailer(FLinkerSave* Linker, FArchive
 		Entry.OffsetInFile = INDEX_NONE;
 		Entry.CompressedSize = It.Value.CompressedSize;
 		Entry.RawSize = It.Value.RawSize;
+		Entry.AccessMode = EPayloadAccessMode::Virtualized;
 	}
 
-	Trailer.Header.AccessMode = AccessMode;
-
 	// Now that we have the complete trailer we can serialize it to the archive
 
 	Trailer.TrailerPositionInFile = DataArchive.Tell();
@@ -197,7 +265,42 @@ bool FPackageTrailerBuilder::BuildAndAppendTrailer(FLinkerSave* Linker, FArchive
 
 bool FPackageTrailerBuilder::IsEmpty() const
 {
-	return LocalEntries.IsEmpty() && VirtualizedEntries.IsEmpty();
+	return LocalEntries.IsEmpty() && ReferencedEntries.IsEmpty() && VirtualizedEntries.IsEmpty();
+}
+
+bool FPackageTrailerBuilder::IsLocalPayloadEntry(const Virtualization::FPayloadId& Identifier) const
+{
+	return LocalEntries.Find(Identifier) != nullptr;
+}
+
+bool FPackageTrailerBuilder::IsReferencedPayloadEntry(const Virtualization::FPayloadId& Identifier) const
+{
+	return ReferencedEntries.Find(Identifier) != nullptr;
+}
+
+bool FPackageTrailerBuilder::IsVirtualizedPayloadEntry(const Virtualization::FPayloadId& Identifier) const
+{
+	return VirtualizedEntries.Find(Identifier) != nullptr;
+}
+
+int32 FPackageTrailerBuilder::GetNumPayloads() const
+{
+	return GetNumLocalPayloads() + GetNumReferencedPayloads() + GetNumVirtualizedPayloads();
+}
+
+int32 FPackageTrailerBuilder::GetNumLocalPayloads() const
+{
+	return LocalEntries.Num();
+}
+
+int32 FPackageTrailerBuilder::GetNumReferencedPayloads() const
+{
+	return ReferencedEntries.Num();
+}
+
+int32 FPackageTrailerBuilder::GetNumVirtualizedPayloads() const
+{
+	return VirtualizedEntries.Num();
 }
 
 bool FPackageTrailer::IsEnabled()
@@ -233,52 +336,6 @@ bool FPackageTrailer::TryLoadFromPackage(const FPackagePath& PackagePath, FPacka
 	}	
 }
 
-FPackageTrailer FPackageTrailer::CreateReference(const FPackageTrailer& OriginalTrailer)
-{
-	check(OriginalTrailer.TrailerPositionInFile != INDEX_NONE);
-
-	FPackageTrailer ReferenceTrailer;
-	ReferenceTrailer.Header = OriginalTrailer.Header;
-	ReferenceTrailer.Header.AccessMode = EPayloadAccessMode::Referenced; // This trailer will reference payloads stored in another trailer
-	ReferenceTrailer.Header.PayloadsDataLength = 0; // This trailer will contain no payload data
-
-	// Convert the payload offsets into absolute offsets
-	const int64 PayloadOffset = OriginalTrailer.TrailerPositionInFile + OriginalTrailer.Header.HeaderLength;
-	for (Private::FLookupTableEntry& Entry : ReferenceTrailer.Header.PayloadLookupTable)
-	{
-		if (!Entry.IsVirtualized())
-		{
-			Entry.OffsetInFile += PayloadOffset;
-		}
-	}
-
-	return ReferenceTrailer;
-}
-
-bool FPackageTrailer::TrySave(FArchive& Ar)
-{
-	if (Header.AccessMode != EPayloadAccessMode::Referenced)
-	{
-		// Can only save reference trailers directly.
-		return false;
-	}
-
-	check(Header.PayloadsDataLength == 0);
-
-	const int64 PositionInArchive = Ar.Tell();
-
-	Ar << Header;
-	checkf((PositionInArchive + Header.HeaderLength) == Ar.Tell(), TEXT("Header length was calculated as %d bytes but we wrote %" INT64_FMT " bytes!"), Header.HeaderLength, Ar.Tell() - PositionInArchive);
-
-	FFooter Footer = CreateFooter();
-
-	Ar << Footer;
-
-	checkf((PositionInArchive + Footer.TrailerLength) == Ar.Tell(), TEXT("Trailer length was calculated as %" INT64_FMT " bytes but we wrote %" INT64_FMT " bytes!"), Footer.TrailerLength, Ar.Tell() - PositionInArchive);
-
-	return true;
-}
-
 bool FPackageTrailer::TryLoad(FArchive& Ar)
 {
 	check(Ar.IsLoading());
@@ -297,7 +354,12 @@ bool FPackageTrailer::TryLoad(FArchive& Ar)
 
 	Ar << Header.HeaderLength;
 	Ar << Header.PayloadsDataLength;
-	Ar << Header.AccessMode;
+
+	EPayloadAccessMode LegacyAccessMode = EPayloadAccessMode::Local;
+	if (Header.Version < (uint32)EPackageTrailerVersion::ACCESS_PER_PAYLOAD)
+	{
+		Ar << LegacyAccessMode;
+	}
 
 	int32 NumPayloads = 0;
 	Ar << NumPayloads;
@@ -307,12 +369,12 @@ bool FPackageTrailer::TryLoad(FArchive& Ar)
 	for (int32 Index = 0; Index < NumPayloads; ++Index)
 	{
 		Private::FLookupTableEntry& Entry = Header.PayloadLookupTable.AddDefaulted_GetRef();
-		Ar << Entry;
-	}
+		Entry.Serialize(Ar, (EPackageTrailerVersion)Header.Version);
 
-	if (Header.AccessMode == EPayloadAccessMode::Referenced)
-	{
-		TrailerPositionInFile = INDEX_NONE;
+		if (Header.Version < (uint32)EPackageTrailerVersion::ACCESS_PER_PAYLOAD)
+		{
+			Entry.AccessMode = Entry.OffsetInFile != INDEX_NONE ? LegacyAccessMode : EPayloadAccessMode::Virtualized;
+		}
 	}
 
 	return !Ar.IsError();
@@ -347,21 +409,21 @@ bool FPackageTrailer::TryLoadBackwards(FArchive& Ar)
 	return TryLoad(Ar);
 }
 
-FCompressedBuffer FPackageTrailer::LoadPayload(const Virtualization::FPayloadId& Id, FArchive& Ar) const
+FCompressedBuffer FPackageTrailer::LoadLocalPayload(const Virtualization::FPayloadId& Id, FArchive& Ar) const
 {
+	// TODO: This method should be able to load the payload in all cases, but we need a good way of passing the Archive/PackagePath 
+	// to the trailer etc. Would work if we stored the package path in the trailer.
 	const Private::FLookupTableEntry* Entry = Header.PayloadLookupTable.FindByPredicate([&Id](const Private::FLookupTableEntry& Entry)->bool
 		{
 			return Entry.Identifier == Id;
 		});
 
-	// TODO: Should we check if payload is correct here or let caller do that?
-
-	if (Entry == nullptr)
+	if (Entry == nullptr || Entry->AccessMode != EPayloadAccessMode::Local)
 	{
 		return FCompressedBuffer();
 	}
 
-	const int64 OffsetInFile = TrailerPositionInFile + + Header.HeaderLength + Entry->OffsetInFile;
+	const int64 OffsetInFile = TrailerPositionInFile + Header.HeaderLength + Entry->OffsetInFile;
 	Ar.Seek(OffsetInFile);
 
 	FCompressedBuffer Payload;
@@ -379,6 +441,7 @@ bool FPackageTrailer::UpdatePayloadAsVirtualized(const Virtualization::FPayloadI
 
 	if (Entry != nullptr)
 	{
+		Entry->AccessMode = EPayloadAccessMode::Virtualized;
 		Entry->OffsetInFile = INDEX_NONE;
 		return true;
 	}
@@ -400,7 +463,25 @@ EPayloadStatus FPackageTrailer::FindPayloadStatus(const Virtualization::FPayload
 		return EPayloadStatus::NotFound;
 	}
 
-	return Entry->IsVirtualized() ? EPayloadStatus::StoredVirtualized : EPayloadStatus::StoredLocally;
+	switch (Entry->AccessMode)
+	{
+		case EPayloadAccessMode::Local:
+			return EPayloadStatus::StoredLocally;
+			break;
+
+		case EPayloadAccessMode::Referenced:
+			return EPayloadStatus::StoredAsReference;
+			break;
+
+		case EPayloadAccessMode::Virtualized:
+			return EPayloadStatus::StoredVirtualized;
+			break;
+
+		default:
+			checkNoEntry();
+			return EPayloadStatus::NotFound;
+			break;
+	}
 }
 
 int64 FPackageTrailer::FindPayloadOffsetInFile(const Virtualization::FPayloadId& Id) const
@@ -417,7 +498,25 @@ int64 FPackageTrailer::FindPayloadOffsetInFile(const Virtualization::FPayloadId&
 		check(Header.PayloadsDataLength != INDEX_NONE);
 		check(Entry != nullptr);
 
-		return TrailerPositionInFile + Header.HeaderLength + Entry->OffsetInFile;
+		switch (Entry->AccessMode)
+		{
+			case EPayloadAccessMode::Local:
+				return TrailerPositionInFile + Header.HeaderLength + Entry->OffsetInFile;
+				break;
+
+			case EPayloadAccessMode::Referenced:
+				return Entry->OffsetInFile;
+				break;
+
+			case EPayloadAccessMode::Virtualized:
+				return INDEX_NONE;
+				break;
+
+			default:
+				checkNoEntry();
+				return INDEX_NONE;
+				break;
+		}
 	}
 	else
 	{
@@ -471,30 +570,29 @@ int32 FPackageTrailer::GetNumPayloads(EPayloadFilter Type) const
 
 	switch (Type)
 	{
-	case EPayloadFilter::All:
-		Count = Header.PayloadLookupTable.Num();
-		break;
+		case EPayloadFilter::All:
+			Count = Header.PayloadLookupTable.Num();
+			break;
 
-	case EPayloadFilter::Local:
-		Count = Algo::CountIf(Header.PayloadLookupTable, [](const Private::FLookupTableEntry& Entry) { return !Entry.IsVirtualized(); });
-		break;
+		case EPayloadFilter::Local:
+			Count = Algo::CountIf(Header.PayloadLookupTable, [](const Private::FLookupTableEntry& Entry) { return Entry.AccessMode == EPayloadAccessMode::Local; });
+			break;
+
+		case EPayloadFilter::Referenced:
+			Count = Algo::CountIf(Header.PayloadLookupTable, [](const Private::FLookupTableEntry& Entry) { return Entry.AccessMode == EPayloadAccessMode::Referenced; });
+			break;
 
-	case EPayloadFilter::Virtualized:
-		Count = Algo::CountIf(Header.PayloadLookupTable, [](const Private::FLookupTableEntry& Entry) { return Entry.IsVirtualized(); });
-		break;
+		case EPayloadFilter::Virtualized:
+			Count = Algo::CountIf(Header.PayloadLookupTable, [](const Private::FLookupTableEntry& Entry) { return Entry.AccessMode == EPayloadAccessMode::Virtualized; });
+			break;
 
-	default:
-		checkNoEntry();
+		default:
+			checkNoEntry();
 	}
 	
 	return Count;
 }
 
-EPayloadAccessMode FPackageTrailer::GetAccessMode() const
-{
-	return Header.AccessMode;
-}
-
 FPackageTrailer::FFooter FPackageTrailer::CreateFooter() const
 {
 	FFooter Footer;
@@ -508,18 +606,20 @@ FPackageTrailer::FFooter FPackageTrailer::CreateFooter() const
 
 FArchive& operator<<(FArchive& Ar, FPackageTrailer::FHeader& Header)
 {
+	// Make sure that we save the most up to date version	
+	Header.Version = (int32)EPackageTrailerVersion::AUTOMATIC_VERSION;
+
 	Ar << Header.Tag;
 	Ar << Header.Version;
 	Ar << Header.HeaderLength;
 	Ar << Header.PayloadsDataLength;
-	Ar << Header.AccessMode;
 
 	int32 NumPayloads = Header.PayloadLookupTable.Num();
 	Ar << NumPayloads;
 
 	for (Private::FLookupTableEntry& Entry : Header.PayloadLookupTable)
 	{
-		Ar << Entry;
+		Entry.Serialize(Ar, (EPackageTrailerVersion)Header.Version);
 	}
 
 	return Ar;
diff --git a/Engine/Source/Runtime/CoreUObject/Private/UObject/SavePackage.cpp b/Engine/Source/Runtime/CoreUObject/Private/UObject/SavePackage.cpp
--- a/Engine/Source/Runtime/CoreUObject/Private/UObject/SavePackage.cpp
+++ b/Engine/Source/Runtime/CoreUObject/Private/UObject/SavePackage.cpp
@@ -1665,7 +1665,14 @@ struct FPackageExportTagger
 
 	if (Linker->PackageTrailerBuilder.IsValid())
 	{
-		checkf(PackageWriter == nullptr, TEXT("Attempting to build a package trailer with a package writer '%s', this is not supported!"), *Package->GetName());
+		// At the moment we assume that we cannot have reference payloads in the trailer if SAVE_BulkDataByReference is not set and we
+		// cannot have locally stored payloads if SAVE_BulkDataByReference is set.
+		checkf((SaveFlags & SAVE_BulkDataByReference) != 0 || Linker->PackageTrailerBuilder->GetNumReferencedPayloads() == 0,
+			TEXT("Attempting to build a package trailer with referenced payloads but the SAVE_BulkDataByReference flag is not set. '%s'"), *Package->GetName());
+
+		checkf((SaveFlags & SAVE_BulkDataByReference) != 0 || Linker->PackageTrailerBuilder->GetNumLocalPayloads() == 0,
+			TEXT("Attempting to build a package trailer with local payloads but the SAVE_BulkDataByReference flag is set. '%s'"), *Package->GetName());
+
 		checkf(bTextFormat == false, TEXT("Attempting to build a package trailer for text based asset '%s', this is not supported!"), *Package->GetName());
 
 		Linker->Summary.PayloadTocOffset = Linker->Tell();
@@ -1673,23 +1680,10 @@ struct FPackageExportTagger
 		{
 			return ESavePackageResult::Error;
 		}
-	}
-	else if ((SaveFlags & SAVE_BulkDataByReference) != 0)
-	{
-		if (const FLinkerLoad* LinkerLoad = FLinkerLoad::FindExistingLinkerForPackage(Package))
-		{
-			if (const UE::FPackageTrailer* Trailer = LinkerLoad->GetPackageTrailer())
-			{
-				UE::FPackageTrailer ReferenceTrailer = UE::FPackageTrailer::CreateReference(*Trailer);
-				Linker->Summary.PayloadTocOffset = Linker->Tell();
-				if (!ReferenceTrailer.TrySave(*Linker))
-				{
-					return ESavePackageResult::Error;
-				}
-			}
-		}
-	}
 
+		Linker->PackageTrailerBuilder.Reset();
+	}
+	
 	return ESavePackageResult::Success;
 }
 
@@ -2170,7 +2164,17 @@ FSavePackageResultStruct UPackage::Save(UPackage* InOuter, UObject* InAsset, con
 						// The package trailer is not supported for text based assets yet
 						if (!bTextFormat && !ObjectSaveContext.bProceduralSave) 
 						{
-							Linker->PackageTrailerBuilder = MakeUnique<UE::FPackageTrailerBuilder>(InOuter);
+							Linker->PackageTrailerBuilder = MakeUnique<UE::FPackageTrailerBuilder>(InOuter->GetFName());
+						}
+						else if ((SaveFlags & SAVE_BulkDataByReference) != 0)
+						{
+							if (const FLinkerLoad* LinkerLoad = FLinkerLoad::FindExistingLinkerForPackage(InOuter))
+							{
+								if (const UE::FPackageTrailer* Trailer = LinkerLoad->GetPackageTrailer())
+								{
+									Linker->PackageTrailerBuilder = UE::FPackageTrailerBuilder::CreateReferenceToTrailer(*Trailer, InOuter->GetFName());
+								}
+							}
 						}
 					}
 
diff --git a/Engine/Source/Runtime/CoreUObject/Private/UObject/SavePackage2.cpp b/Engine/Source/Runtime/CoreUObject/Private/UObject/SavePackage2.cpp
--- a/Engine/Source/Runtime/CoreUObject/Private/UObject/SavePackage2.cpp
+++ b/Engine/Source/Runtime/CoreUObject/Private/UObject/SavePackage2.cpp
@@ -810,7 +810,17 @@ ESavePackageResult CreateLinker(FSaveContext& SaveContext)
 			// The package trailer is not supported for text based assets yet
 			if (!SaveContext.IsTextFormat() && !SaveContext.IsProceduralSave())
 			{
-				SaveContext.Linker->PackageTrailerBuilder = MakeUnique<UE::FPackageTrailerBuilder>(SaveContext.GetPackage());
+				SaveContext.Linker->PackageTrailerBuilder = MakeUnique<UE::FPackageTrailerBuilder>(SaveContext.GetPackage()->GetFName());
+			}
+			else if ((SaveContext.GetSaveArgs().SaveFlags & SAVE_BulkDataByReference) != 0)
+			{
+				if (const FLinkerLoad* LinkerLoad = FLinkerLoad::FindExistingLinkerForPackage(SaveContext.GetPackage()))
+				{
+					if (const UE::FPackageTrailer* Trailer = LinkerLoad->GetPackageTrailer())
+					{
+						SaveContext.Linker->PackageTrailerBuilder = UE::FPackageTrailerBuilder::CreateReferenceToTrailer(*Trailer, SaveContext.GetPackage()->GetFName());
+					}
+				}
 			}
 		}
 
@@ -1728,7 +1738,14 @@ ESavePackageResult WriteExports(FStructuredArchive::FRecord& StructuredArchiveRo
 
 	if (SaveContext.Linker->PackageTrailerBuilder.IsValid())
 	{
-		checkf(SaveContext.GetPackageWriter() == nullptr, TEXT("Attempting to build a package trailer with a package writer '%s', this is not supported!"), *SaveContext.GetPackage()->GetName());
+		// At the moment we assume that we cannot have reference payloads in the trailer if SAVE_BulkDataByReference is not set and we
+		// cannot have locally stored payloads if SAVE_BulkDataByReference is set.
+		checkf((SaveContext.GetSaveArgs().SaveFlags & SAVE_BulkDataByReference) != 0 || SaveContext.Linker->PackageTrailerBuilder->GetNumReferencedPayloads() == 0,
+			TEXT("Attempting to build a package trailer with referenced payloads but the SAVE_BulkDataByReference flag is not set. '%s'"), *SaveContext.GetPackage()->GetName());
+
+		checkf(	(SaveContext.GetSaveArgs().SaveFlags & SAVE_BulkDataByReference) == 0 || SaveContext.Linker->PackageTrailerBuilder->GetNumLocalPayloads() == 0, 
+				TEXT("Attempting to build a package trailer with local payloads but the SAVE_BulkDataByReference flag is set. '%s'"), *SaveContext.GetPackage()->GetName());
+
 		checkf(SaveContext.IsTextFormat() == false, TEXT("Attempting to build a package trailer for text based asset '%s', this is not supported!"), *SaveContext.GetPackage()->GetName());
 
 		SaveContext.Linker->Summary.PayloadTocOffset = SaveContext.Linker->Tell();
@@ -1736,23 +1753,8 @@ ESavePackageResult WriteExports(FStructuredArchive::FRecord& StructuredArchiveRo
 		{
 			return ESavePackageResult::Error;
 		}
-	}
-	else if ((SaveContext.GetSaveArgs().SaveFlags & SAVE_BulkDataByReference) != 0)
-	{
-		check(SaveContext.GetPackage() != nullptr);
-		
-		if (const FLinkerLoad* LinkerLoad = FLinkerLoad::FindExistingLinkerForPackage(SaveContext.GetPackage()))
-		{
-			if (const UE::FPackageTrailer* Trailer = LinkerLoad->GetPackageTrailer())
-			{
-				UE::FPackageTrailer ReferenceTrailer = UE::FPackageTrailer::CreateReference(*Trailer);
-				SaveContext.Linker->Summary.PayloadTocOffset = SaveContext.Linker->Tell();
-				if (!ReferenceTrailer.TrySave(*SaveContext.Linker))
-				{
-					return ESavePackageResult::Error;
-				}
-			}
-		}
+
+		SaveContext.Linker->PackageTrailerBuilder.Reset();
 	}
 
 	return ESavePackageResult::Success;
diff --git a/Engine/Source/Runtime/CoreUObject/Public/Serialization/VirtualizedBulkData.h b/Engine/Source/Runtime/CoreUObject/Public/Serialization/VirtualizedBulkData.h
--- a/Engine/Source/Runtime/CoreUObject/Public/Serialization/VirtualizedBulkData.h
+++ b/Engine/Source/Runtime/CoreUObject/Public/Serialization/VirtualizedBulkData.h
@@ -8,6 +8,7 @@
 #include "Memory/SharedBuffer.h"
 #include "Misc/Guid.h"
 #include "Misc/PackagePath.h"
+#include "Misc/PackageSegment.h"
 #include "Virtualization/PayloadId.h"
 
 class FArchive;
@@ -341,6 +342,8 @@ private:
 		IsTornOff					= 1 << 7,
 		/** The bulkdata object references a payload stored in a WorkspaceDomain file  */
 		ReferencesWorkspaceDomain	= 1 << 8,
+		/** The payload is stored in a package trailer, so the bulkdata object will have to poll the trailer to find the payload offset */
+		StoredInPackageTrailer		= 1 << 9,
 
 		TransientFlags				= HasRegistered | IsTornOff,
 	};
@@ -357,9 +360,9 @@ private:
 	};
 
 	/** Old legacy path that saved the payload to the end of the package */
-	void SerializeToLegacyPath(FLinkerSave& LinkerSave, int64 OffsetPos, FCompressedBuffer PayloadToSerialize, EFlags UpdatedFlags, UObject* Owner);
+	void SerializeToLegacyPath(FLinkerSave& LinkerSave, FCompressedBuffer PayloadToSerialize, EFlags UpdatedFlags, UObject* Owner);
 	/** The new path that saves payloads to the FPackageTrailer which is then appended to the end of the package file */
-	void SerializeToPackageTrailer(FLinkerSave& LinkerSave, int64 OffsetPos, FCompressedBuffer PayloadToSerialize, EFlags UpdatedFlags, UObject* Owner);
+	void SerializeToPackageTrailer(FLinkerSave& LinkerSave, FCompressedBuffer PayloadToSerialize, EFlags UpdatedFlags, UObject* Owner);
 
 	void UpdatePayloadImpl(FSharedBuffer&& InPayload, FPayloadId&& InPayloadID);
 
@@ -376,8 +379,6 @@ private:
 	void PushData(const FPackagePath& InPackagePath);
 	FCompressedBuffer PullData() const;
 
-	FPackagePath GetPackagePathFromOwner(UObject* Owner, EPackageSegment& OutPackageSegment) const;
-
 	bool CanUnloadData() const;
 
 	void UpdateKeyIfNeeded();
@@ -427,8 +428,25 @@ private:
 		return EnumHasAnyFlags(InFlags, EFlags::ReferencesLegacyFile | EFlags::ReferencesWorkspaceDomain);
 	}
 
+	bool IsStoredInPackageTrailer() const
+	{
+		return IsStoredInPackageTrailer(Flags);
+	}
+	static bool IsStoredInPackageTrailer(EFlags InFlags)
+	{
+		return EnumHasAnyFlags(InFlags, EFlags::StoredInPackageTrailer);
+	}
+
 	void Unregister();
 
+	/**
+	 * Checks to make sure that the payload we are saving is what we expect it to be. If not then we need to log an error to the
+	 * user, mark the archive as having an error and return.
+	 * This error is not reasonable to expect, if it occurs then it means something has gone very wrong, which is why there is
+	 * also an ensure to make sure that any instances of it occurring as properly recorded and can be investigated.
+	 */
+	bool TryPayloadValidationForSaving(const FCompressedBuffer& PayloadForSaving, FLinkerSave* LinkerSave) const;
+
 	/** 
 	 * Utility to return an apt error message if the payload is invalid when trying to save the bulkdata. It will try to provide the best info from the given options.
 	 * 
@@ -438,7 +456,7 @@ private:
 	FText GetCorruptedPayloadErrorMsgForSave(FLinkerSave* Linker) const;
 
 	/** Returns true if we should use legacy serialization instead of the FPackageTrailer system. This can be removed when UE_ENABLE_VIRTUALIZATION_TOGGLE is removed. */
-	bool ShouldUseLegacySerialization(FLinkerSave& LinkerSave) const;
+	bool ShouldUseLegacySerialization(const FLinkerSave* LinkerSave) const;
 
 	/** Unique identifier for the bulkdata object itself */
 	FGuid BulkDataId;
@@ -464,7 +482,7 @@ private:
 	FPackagePath PackagePath;
 
 	/** PackageSegment to load with the packagepath (unused if the payload does not come from PackageResourceManager) */
-	EPackageSegment PackageSegment;
+	EPackageSegment PackageSegment= EPackageSegment::Header;
 
 	/** A 32bit bitfield of flags */
 	EFlags Flags = EFlags::None;
diff --git a/Engine/Source/Runtime/CoreUObject/Public/UObject/PackageTrailer.h b/Engine/Source/Runtime/CoreUObject/Public/UObject/PackageTrailer.h
--- a/Engine/Source/Runtime/CoreUObject/Public/UObject/PackageTrailer.h
+++ b/Engine/Source/Runtime/CoreUObject/Public/UObject/PackageTrailer.h
@@ -46,7 +46,6 @@ namespace UE
  * | Version			| uint32			| Version number of the format@see UE::EPackageTrailerVersion										|
  * | HeaderLength		| uint32			| The total size of the header on disk in bytes.													|
  * | PayloadsDataLength	| uint64			| The total size of the payload data on disk in bytes												|
- * | AccessMode			| uint8				| @see UE::EPayloadAccessMode														|
  * | NumPayloads		| int32				| The number of payloads in LookupTableArray														|
  * | LookupTableArray	| FLookupTableEntry | An array of FLookupTableEntry @see UE::Private::FLookupTableEntry									|
  * |____________________________________________________________________________________________________________________________________________|
@@ -54,7 +53,7 @@ namespace UE
  * | Array				| FCompressedBuffer | A binary blob containing all of the payloads. Individual payloads can be found via				|
  * |										 the LookupTableArray found in the header.															|
  * |____________________________________________________________________________________________________________________________________________|
- * [Footer]
+ * | [Footer]																																	|
  * | Tag				| uint64			| Should match FFooter::FooterTag, used to identify that the data being read is an FPackageTrailer	|
  * | TrailerLength		| uint64			| The total size of the trailer on disk in bytes. Can be used to find the start of the trailer when	|
  * |										  reading backwards.																				|
@@ -68,8 +67,10 @@ enum class EPayloadFilter
 {
 	/** All payload types. */
 	All,
-	/** All payloads stored locally in the package file. */
+	/** All payloads stored locally in the package trailer. */
 	Local,
+	/** All payloads that are a reference to payloads stored in the workspace domain trailer*/
+	Referenced,
 	/** All payloads stored in a virtualized backend. */
 	Virtualized
 };
@@ -79,8 +80,10 @@ enum class EPayloadStatus
 {
 	/** The payload is not registered in the package trailer */
 	NotFound = 0,
-	/** The payload is stored locally on disk */
+	/** The payload is stored locally inside the current package trailer where ever that is written to disk */
 	StoredLocally,
+	/** The payload is stored in the workspace domain trailer */
+	StoredAsReference,
 	/** The payload is virtualized and needs to be accessed via the IVirtualizationSystem */
 	StoredVirtualized,
 };
@@ -88,31 +91,36 @@ enum class EPayloadStatus
 /** Lists the various methods of payload access that the trailer supports */
 enum class EPayloadAccessMode : uint8
 {
-	/** The payloads are stored in the Payload Data segment of the trailer and the offsets in FLookupTableEntry will be relative to the start of this segment */
-	Relative = 0,
-	/** The payloads are stored in the trailer of another file (most likely the workspace domain package file) and the offsets in FLookupTableEntry are absolute offsets in that external file */
-	Referenced
+	/** The payload is stored in the Payload Data segment of the trailer and the offsets in FLookupTableEntry will be relative to the start of this segment */
+	Local = 0,
+	/** The payload is stored in another package trailer (most likely the workspace domain package file) and the offsets in FLookupTableEntry are absolute offsets in that external file */
+	Referenced,
+	/** The payload is virtualized and needs to be accessed via IVirtualizationSystem */
+	Virtualized
 };
 
+enum class EPackageTrailerVersion : uint32;
+
 namespace Private
 {
 
 struct FLookupTableEntry
 {
 	/** Size of the entry when serialized to disk in bytes */
-	static constexpr uint32 SizeOnDisk = 44;	// Identifier		| 20 bytes
+	static constexpr uint32 SizeOnDisk = 45;	// Identifier		| 20 bytes
 												// OffsetInFile		| 8 bytes
 												// CompressedSize	| 8 bytes
 												// RawSize			| 8 bytes
+												// AccessMode		| 1 byte
 
 	FLookupTableEntry() = default;
 	FLookupTableEntry(const Virtualization::FPayloadId& InIdentifier, uint64 InRawSize);
 
-	friend FArchive& operator<<(FArchive& Ar, FLookupTableEntry& Entry);
-
+	void Serialize(FArchive& Ar, EPackageTrailerVersion PackageTrailerVersion);
+	
 	[[nodiscard]] bool IsVirtualized() const
 	{
-		return OffsetInFile == INDEX_NONE;
+		return AccessMode == EPayloadAccessMode::Virtualized;
 	}
 
 	/** Identifier for the payload */
@@ -123,6 +131,8 @@ struct FLookupTableEntry
 	uint64 CompressedSize = INDEX_NONE;
 	/** The size of the payload when uncompressed. */
 	uint64 RawSize = INDEX_NONE;
+
+	EPayloadAccessMode AccessMode = EPayloadAccessMode::Local;
 };
 
 } // namespace Private
@@ -139,17 +149,28 @@ public:
 	using AdditionalDataCallback = TFunction<void(FLinkerSave& LinkerSave, const class FPackageTrailer& Trailer)>;
 
 	/**
-	 * Creates a builder from an already loaded FPackageTrailer. 
+	 * Creates a builder from a pre-existing FPackageTrailer.
+	 * Payloads stored locally in the source trailer will be loaded from disk via the provided archive so that the
+	 * builder can write them to any future trailer that it creates.
 	 * 
 	 * @param Trailer		The trailer to create the builder from
 	 * @param Ar			An archive that the trailer can use to load payloads from 
-	 * @param PackagePath	Path to the package
+	 * @param PackageName	The name of the package that owns the trailer. Used for error messages.
 	 */
-	[[nodiscard]] static FPackageTrailerBuilder CreateFromTrailer(const class FPackageTrailer& Trailer, FArchive& Ar, const FPackagePath& PackagePath);
+	[[nodiscard]] static FPackageTrailerBuilder CreateFromTrailer(const class FPackageTrailer& Trailer, FArchive& Ar, const FName& PackageName);
+
+	/**
+	 * Creates a builder from a pre-existing FPackageTrailer that will will reference the local payloads of the
+	 * source trailer. 
+	 * This means that there is no need to load the payloads.
+	 *
+	 * @param Trailer		The trailer to create the reference from.
+	 * @param PackageName	The name of the package that owns the trailer. Used for error messages.
+	 */
+	[[nodiscard]] static TUniquePtr<UE::FPackageTrailerBuilder> CreateReferenceToTrailer(const class FPackageTrailer& Trailer, const FName& PackageName);
 
 	FPackageTrailerBuilder() = delete;
-	FPackageTrailerBuilder(UPackage* Package);
-	FPackageTrailerBuilder(FName InPackageName);
+	FPackageTrailerBuilder(const FName& InPackageName);
 	~FPackageTrailerBuilder() = default;
 
 	// Methods that can be called while building the trailer
@@ -174,9 +195,23 @@ public:
 	/** Returns if the builder has any payload entries or not */
 	[[nodiscard]] bool IsEmpty() const;
 
+	[[nodiscard]] bool IsLocalPayloadEntry(const Virtualization::FPayloadId& Identifier) const;
+	[[nodiscard]] bool IsReferencedPayloadEntry(const Virtualization::FPayloadId& Identifier) const;
+	[[nodiscard]] bool IsVirtualizedPayloadEntry(const Virtualization::FPayloadId& Identifier) const;
+
+	/** Returns the total number of payload entries in the builder */
+	[[nodiscard]] int32 GetNumPayloads() const;
+	
+	/** Returns the number of payload entries in the builder with the access mode EPayloadAccessMode::Local */
+	[[nodiscard]] int32 GetNumLocalPayloads() const;
+	/** Returns the number of payload entries in the builder with the access mode EPayloadAccessMode::Referenced */
+	[[nodiscard]] int32 GetNumReferencedPayloads() const;
+	/** Returns the number of payload entries in the builder with the access mode EPayloadAccessMode::Virtualized */
+	[[nodiscard]] int32 GetNumVirtualizedPayloads() const;
+
 private:
 	
-	/** All of the data required to add a payload that is stored on disk */
+	/** All of the data required to add a payload that is stored locally within the trailer */
 	struct LocalEntry
 	{
 		LocalEntry() = default;
@@ -190,6 +225,25 @@ private:
 		FCompressedBuffer Payload;
 	};
 
+	/** All of the data required to add a reference to a payload stored in another trailer */
+	struct ReferencedEntry
+	{
+		ReferencedEntry() = default;
+		ReferencedEntry(int64 InOffset, int64 InCompressedSize, int64 InRawSize)
+			: Offset(InOffset)
+			, CompressedSize(InCompressedSize)
+			, RawSize(InRawSize)
+		{
+
+		}
+		~ReferencedEntry() = default;
+
+		int64 Offset = INDEX_NONE;
+		int64 CompressedSize = INDEX_NONE;
+		int64 RawSize = INDEX_NONE;
+	};
+
+
 	/** All of the data required to add a payload that is virtualized */
 	struct VirtualizedEntry
 	{
@@ -211,11 +265,10 @@ private:
 	/** Name of the package the trailer is being built for, used to give meaningful error messages */
 	FName PackageName;
 
-	/** The access mode that the trailer should have to the payloads */
-	EPayloadAccessMode AccessMode = EPayloadAccessMode::Relative;
-
 	/** Payloads that will be stored locally when the trailer is written to disk */
 	TMap<Virtualization::FPayloadId, LocalEntry> LocalEntries;
+	/** Payloads that reference entries in another trailer */
+	TMap<Virtualization::FPayloadId, ReferencedEntry> ReferencedEntries;
 	/** Payloads that are already virtualized and so will not be written to disk */
 	TMap<Virtualization::FPayloadId, VirtualizedEntry> VirtualizedEntries;
 	/** Callbacks to invoke once the trailer has been written to the end of a package */
@@ -241,13 +294,6 @@ public:
 	/** Try to load a trailer from a given package path. Note that it will always try to load the trailer from the workspace domain */
 	[[nodiscard]] static bool TryLoadFromPackage(const FPackagePath& PackagePath, FPackageTrailer& OutTrailer);
 
-	/** 
-	 * Create a new trailer that references the payloads stored elsewhere rather than storing them locally.
-	 * 
-	 * @param SourceTrailer The trailer that currently stores the payloads (typically in the workspace domain)
-	 */
-	[[nodiscard]] static FPackageTrailer CreateReference(const FPackageTrailer& SourceTrailer);
-
 	FPackageTrailer() = default;
 	~FPackageTrailer() = default;
 
@@ -257,9 +303,6 @@ public:
 	FPackageTrailer(FPackageTrailer&& Other) = default;
 	FPackageTrailer& operator=(FPackageTrailer&& Other) = default;
 
-	/** Serializes the trailer to the given archive */
-	[[nodiscard]] bool TrySave(FArchive& Ar);
-
 	/** 
 	 * Serializes the trailer from the given archive assuming that the seek position of the archive is already at the correct position
 	 * for the trailer.
@@ -279,15 +322,16 @@ public:
 	[[nodiscard]] bool TryLoadBackwards(FArchive& Ar);
 
 	/** 
-	 * Loads the a payload from the provided archive 
+	 * Loads a payload that is stored locally within the package trailer. Payloads stored externally (either referenced
+	 * or virtualized) will not load.
 	 * 
 	 * @param Id The payload to load
-	 * @param Ar The archive to load the payload from, it is assumed that this archive is the package file in the workspace domain
+	 * @param Ar The archive from which the payload trailer was also loaded from
 	 * 
-	 * @return	The payload in the form of a FCompressedBuffer. If the payload does not exist in the trailer or in the archive then
-	 *			the FCompressedBuffer will be null.
+	 * @return	The payload in the form of a FCompressedBuffer. If the payload does not exist in the trailer or is not
+	 *			stored locally in the trailer then the FCompressedBuffer will be null.
 	 */
-	[[nodiscard]] FCompressedBuffer LoadPayload(const Virtualization::FPayloadId& Id, FArchive& Ar) const;
+	[[nodiscard]] FCompressedBuffer LoadLocalPayload(const Virtualization::FPayloadId& Id, FArchive& Ar) const;
 
 	/** 
 	 * Calling this indicates that the payload has been virtualized and will no longer be stored on disk. 
@@ -312,8 +356,6 @@ public:
 	/** Returns the number of payloads that the trailer owns that match the given filter type. @See EPayloadType */
 	[[nodiscard]] int32 GetNumPayloads(EPayloadFilter Type) const;
 
-	[[nodiscard]] EPayloadAccessMode GetAccessMode() const;
-
 	struct FHeader
 	{
 		/** Unique value used to identify the header */
@@ -323,11 +365,10 @@ public:
 		 * Size of the static header data when serialized to disk in bytes. Note that we still need to 
 		 * add the size of the data in PayloadLookupTable to get the final header size on disk 
 		 */
-		static constexpr uint32 StaticHeaderSizeOnDisk = 29;	// HeaderTag			| 8 bytes
+		static constexpr uint32 StaticHeaderSizeOnDisk = 28;	// HeaderTag			| 8 bytes
 																// Version				| 4 bytes
 																// HeaderLength			| 4 bytes
 																// PayloadsDataLength	| 8 bytes
-																// AccessMode			| 1 byte
 																// NumPayloads			| 4 bytes
 
 		/** Expected tag at the start of the header */
@@ -338,8 +379,6 @@ public:
 		uint32 HeaderLength = 0;
 		/** Total length of the payloads on disk in bytes */
 		uint64 PayloadsDataLength = 0;
-		/** What sort of access to the payloads does the trailer have */
-		EPayloadAccessMode AccessMode = EPayloadAccessMode::Relative;
 		/** Lookup table for the payloads on disk */
 		TArray<Private::FLookupTableEntry> PayloadLookupTable;
 
