Added support to Loose File Package Writer to write multiple package output per input
Added similar support to the package store manifest
Zen Store support will be added separately
Added some extra data to linker object resources to be used by the iostore package optimizer to properly support optional package chunks.

Multiple cooking outputs is used to support "optional objects" which are editor data package exports that would be put in separate iostore containers than the normal cooked game content.

diff --git a/Engine/Source/Developer/IoStoreUtilities/Private/IoStoreUtilities.cpp b/Engine/Source/Developer/IoStoreUtilities/Private/IoStoreUtilities.cpp
--- a/Engine/Source/Developer/IoStoreUtilities/Private/IoStoreUtilities.cpp
+++ b/Engine/Source/Developer/IoStoreUtilities/Private/IoStoreUtilities.cpp
@@ -4578,7 +4578,7 @@ int32 Staged2Zen(const FString& BuildPath, const FKeyChain& KeyChain, const FStr
 		ZenStoreWriter->BeginPackage(BeginPackageInfo);
 
 		IPackageWriter::FPackageInfo PackageStorePackageInfo;
-		PackageStorePackageInfo.PackageName = PackageInfo.PackageName;
+		PackageStorePackageInfo.OutputPackageName = PackageStorePackageInfo.InputPackageName = PackageInfo.PackageName;
 		PackageStorePackageInfo.LooseFilePath = PackageInfo.FileName;
 		PackageStorePackageInfo.ChunkId = PackageInfo.Chunk.Value;
 
@@ -4588,7 +4588,7 @@ int32 Staged2Zen(const FString& BuildPath, const FKeyChain& KeyChain, const FStr
 		for (const FBulkDataInfo& BulkDataInfo : PackageInfo.BulkData)
 		{
 			IPackageWriter::FBulkDataInfo PackageStoreBulkDataInfo;
-			PackageStoreBulkDataInfo.PackageName = PackageInfo.PackageName;
+			PackageStoreBulkDataInfo.OutputPackageName = PackageStoreBulkDataInfo.InputPackageName = PackageInfo.PackageName;
 			PackageStoreBulkDataInfo.LooseFilePath = BulkDataInfo.FileName;
 			PackageStoreBulkDataInfo.ChunkId = BulkDataInfo.Chunk.Value;
 			PackageStoreBulkDataInfo.BulkDataType = BulkDataInfo.BulkDataType;
diff --git a/Engine/Source/Developer/IoStoreUtilities/Private/PackageStoreManifest.cpp b/Engine/Source/Developer/IoStoreUtilities/Private/PackageStoreManifest.cpp
--- a/Engine/Source/Developer/IoStoreUtilities/Private/PackageStoreManifest.cpp
+++ b/Engine/Source/Developer/IoStoreUtilities/Private/PackageStoreManifest.cpp
@@ -15,27 +15,27 @@ FPackageStoreManifest::FPackageStoreManifest(const FString& InCookedOutputPath)
 	FPaths::NormalizeFilename(CookedOutputPath);
 }
 
-void FPackageStoreManifest::BeginPackage(FName PackageName)
+void FPackageStoreManifest::BeginPackage(FName InputPackageName)
 {
 	FScopeLock Lock(&CriticalSection);
-	FPackageInfo& PackageInfo = PackageInfoByNameMap.FindOrAdd(PackageName);
-	PackageInfo.PackageName = PackageName;
+	FPackageInfo& PackageInfo = MainPackageInfoByNameMap.FindOrAdd(InputPackageName);
+	PackageInfo.PackageName = InputPackageName;
 	if (PackageInfo.PackageChunkId.IsValid())
 	{
 		FileNameByChunkIdMap.Remove(PackageInfo.PackageChunkId);
-		PackageInfo.PackageChunkId = FIoChunkId();
 	}
 	for (const FIoChunkId& BulkDataChunkId : PackageInfo.BulkDataChunkIds)
 	{
 		FileNameByChunkIdMap.Remove(BulkDataChunkId);
 	}
-	PackageInfo.BulkDataChunkIds.Empty();
+	PackageInfo.BulkDataChunkIds.Reset();
+	OptionalPackageInfoByNameMap.Remove(InputPackageName);
 }
 
-void FPackageStoreManifest::AddPackageData(FName PackageName, const FString& FileName, const FIoChunkId& ChunkId)
+void FPackageStoreManifest::AddPackageData(FName InputPackageName, FName OutputPackageName, const FString& FileName, const FIoChunkId& ChunkId)
 {
 	FScopeLock Lock(&CriticalSection);
-	FPackageInfo* PackageInfo = PackageInfoByNameMap.Find(PackageName);
+	FPackageInfo* PackageInfo = GetPackageInfo_NoLock(InputPackageName, OutputPackageName);
 	check(PackageInfo);
 	PackageInfo->PackageChunkId = ChunkId;
 	if (!FileName.IsEmpty())
@@ -44,10 +44,10 @@ void FPackageStoreManifest::AddPackageData(FName PackageName, const FString& Fil
 	}
 }
 
-void FPackageStoreManifest::AddBulkData(FName PackageName, const FString& FileName, const FIoChunkId& ChunkId)
+void FPackageStoreManifest::AddBulkData(FName InputPackageName, FName OutputPackageName, const FString& FileName, const FIoChunkId& ChunkId)
 {
 	FScopeLock Lock(&CriticalSection);
-	FPackageInfo* PackageInfo = PackageInfoByNameMap.Find(PackageName);
+	FPackageInfo* PackageInfo = GetPackageInfo_NoLock(InputPackageName, OutputPackageName);
 	check(PackageInfo);
 	PackageInfo->BulkDataChunkIds.Add(ChunkId);
 	if (!FileName.IsEmpty())
@@ -94,11 +94,16 @@ FIoStatus FPackageStoreManifest::Save(const TCHAR* Filename) const
 	}
 	Writer->WriteArrayEnd();
 
-	Writer->WriteArrayStart(TEXT("Packages"));
-	for (const auto& KV : PackageInfoByNameMap)
+	auto WritePackageInfoObject = [Writer, &ChunkIdToString](const FPackageInfo& PackageInfo, const TCHAR* Name = nullptr)
 	{
-		const FPackageInfo& PackageInfo = KV.Value;
-		Writer->WriteObjectStart();
+		if (Name)
+		{
+			Writer->WriteObjectStart(Name);
+		}
+		else
+		{
+			Writer->WriteObjectStart();
+		}
 		Writer->WriteValue(TEXT("Name"), PackageInfo.PackageName.ToString());
 		Writer->WriteValue(TEXT("PackageChunkId"), ChunkIdToString(PackageInfo.PackageChunkId));
 		if (!PackageInfo.BulkDataChunkIds.IsEmpty())
@@ -111,9 +116,30 @@ FIoStatus FPackageStoreManifest::Save(const TCHAR* Filename) const
 			Writer->WriteArrayEnd();
 		}
 		Writer->WriteObjectEnd();
+	};
+
+	Writer->WriteArrayStart(TEXT("Packages"));
+	for (const auto& PackageNameInfoPair : MainPackageInfoByNameMap)
+	{
+		const FPackageInfo& PackageInfo = PackageNameInfoPair.Value;
+		WritePackageInfoObject(PackageInfo);
+	}
+	Writer->WriteArrayEnd();
+
+	Writer->WriteArrayStart(TEXT("OptionalPackages"));
+	for (const auto& PackageNameInfoPair : OptionalPackageInfoByNameMap)
+	{
+		Writer->WriteObjectStart();
+		Writer->WriteValue(TEXT("PackageName"), PackageNameInfoPair.Key.ToString());
+		{
+			const FPackageInfo& PackageInfo = PackageNameInfoPair.Value;
+			WritePackageInfoObject(PackageInfo, TEXT("PackageInfo"));
+		}
+		Writer->WriteObjectEnd();
 	}
 	Writer->WriteArrayEnd();
 
+
 	Writer->WriteObjectEnd();
 	Writer->Close();
 
@@ -128,7 +154,8 @@ FIoStatus FPackageStoreManifest::Save(const TCHAR* Filename) const
 FIoStatus FPackageStoreManifest::Load(const TCHAR* Filename)
 {
 	FScopeLock Lock(&CriticalSection);
-	PackageInfoByNameMap.Empty();
+	MainPackageInfoByNameMap.Empty();
+	OptionalPackageInfoByNameMap.Empty();
 	FileNameByChunkIdMap.Empty();
 
 	auto ChunkIdFromString = [](const FString& ChunkIdString)
@@ -181,15 +208,10 @@ FIoStatus FPackageStoreManifest::Load(const TCHAR* Filename)
 		FileNameByChunkIdMap.Add(ChunkId, FPaths::Combine(CookedOutputPath, RelativePath));
 	}
 
-	TArray<TSharedPtr<FJsonValue>> PackagesArray = JsonObject->Values.FindRef(TEXT("Packages"))->AsArray();
-	PackageInfoByNameMap.Reserve(PackagesArray.Num());
-	for (const TSharedPtr<FJsonValue>& PackageValue : PackagesArray)
-	{
-		TSharedPtr<FJsonObject> PackageObject = PackageValue->AsObject();
-		FName PackageName = FName(PackageObject->Values.FindRef(TEXT("Name"))->AsString());
 
-		FPackageInfo& PackageInfo = PackageInfoByNameMap.FindOrAdd(PackageName);
-		PackageInfo.PackageName = PackageName;
+	auto ReadPackageInfo = [&ChunkIdFromString](TSharedPtr<FJsonObject> PackageObject, FPackageInfo& PackageInfo)
+	{
+		check(!PackageInfo.PackageName.IsNone());
 		PackageInfo.PackageChunkId = ChunkIdFromString(PackageObject->Values.FindRef(TEXT("PackageChunkId"))->AsString());
 
 		TSharedPtr<FJsonValue> BulkDataChunkIdsValue = PackageObject->Values.FindRef(TEXT("BulkDataChunkIds"));
@@ -202,6 +224,34 @@ FIoStatus FPackageStoreManifest::Load(const TCHAR* Filename)
 				PackageInfo.BulkDataChunkIds.Add(ChunkIdFromString(BulkDataChunkIdValue->AsString()));
 			}
 		}
+	};
+
+	TArray<TSharedPtr<FJsonValue>> PackagesArray = JsonObject->Values.FindRef(TEXT("Packages"))->AsArray();
+	MainPackageInfoByNameMap.Reserve(PackagesArray.Num());
+	for (const TSharedPtr<FJsonValue>& PackageValue : PackagesArray)
+	{
+		TSharedPtr<FJsonObject> PackageObject = PackageValue->AsObject();
+		FName PackageName = FName(PackageObject->Values.FindRef(TEXT("Name"))->AsString());
+
+		FPackageInfo& PackageInfo = MainPackageInfoByNameMap.FindOrAdd(PackageName);
+		PackageInfo.PackageName = PackageName;
+		ReadPackageInfo(PackageObject, PackageInfo);
+	}
+
+	TArray<TSharedPtr<FJsonValue>> OptionalPackagesArray = JsonObject->Values.FindRef(TEXT("OptionalPackages"))->AsArray();
+	OptionalPackageInfoByNameMap.Reserve(OptionalPackagesArray.Num());
+	for (const TSharedPtr<FJsonValue>& PackageNameInfoPair : OptionalPackagesArray)
+	{
+		TSharedPtr<FJsonObject> PairObject = PackageNameInfoPair->AsObject();
+		FName InputPackageName = FName(PairObject->Values.FindRef(TEXT("PackageName"))->AsString());
+
+		TSharedPtr<FJsonObject> PackageObject = PairObject->Values.FindRef(TEXT("PackageInfo"))->AsObject();
+		FName OutputPackageName = FName(PackageObject->Values.FindRef(TEXT("Name"))->AsString());
+
+		FPackageInfo PackageInfo;
+		PackageInfo.PackageName = OutputPackageName;
+		ReadPackageInfo(PackageObject, PackageInfo);
+		OptionalPackageInfoByNameMap.Add(InputPackageName, MoveTemp(PackageInfo));
 	}
 
 	return FIoStatus::Ok;
@@ -223,7 +273,10 @@ TArray<FPackageStoreManifest::FPackageInfo> FPackageStoreManifest::GetPackages()
 {
 	FScopeLock Lock(&CriticalSection);
 	TArray<FPackageInfo> Packages;
-	PackageInfoByNameMap.GenerateValueArray(Packages);
+	MainPackageInfoByNameMap.GenerateValueArray(Packages);
+	TArray<FPackageInfo> OptionalPackages;
+	OptionalPackageInfoByNameMap.GenerateValueArray(OptionalPackages);
+	Packages.Append(MoveTemp(OptionalPackages));
 	return Packages;
 }
 
@@ -243,3 +296,32 @@ const FPackageStoreManifest::FZenServerInfo* FPackageStoreManifest::ReadZenServe
 	return ZenServerInfo.Get();
 }
 
+FPackageStoreManifest::FPackageInfo* FPackageStoreManifest::GetPackageInfo_NoLock(FName InputPackageName, FName OutputPackageName)
+{
+	FPackageInfo* PackageInfo = nullptr;
+	if (InputPackageName == OutputPackageName)
+	{
+		PackageInfo = MainPackageInfoByNameMap.Find(InputPackageName);
+	}
+	else
+	{
+		TArray<FPackageInfo*, TInlineAllocator<2>> PackageInfos;
+		OptionalPackageInfoByNameMap.MultiFindPointer(InputPackageName, PackageInfos);
+		FPackageInfo** PackageInfoPtr = PackageInfos.FindByPredicate([OutputPackageName](FPackageInfo* Info)
+			{
+				return Info->PackageName == OutputPackageName;
+			});
+		// If we didn't find any, create one
+		if (PackageInfoPtr == nullptr)
+		{
+			PackageInfo = &OptionalPackageInfoByNameMap.Add(InputPackageName);
+			PackageInfo->PackageName = OutputPackageName;
+		}
+		else
+		{
+			PackageInfo = *PackageInfoPtr;
+		}
+	}
+	return PackageInfo;
+}
+
diff --git a/Engine/Source/Developer/IoStoreUtilities/Private/ZenStoreWriter.cpp b/Engine/Source/Developer/IoStoreUtilities/Private/ZenStoreWriter.cpp
--- a/Engine/Source/Developer/IoStoreUtilities/Private/ZenStoreWriter.cpp
+++ b/Engine/Source/Developer/IoStoreUtilities/Private/ZenStoreWriter.cpp
@@ -224,7 +224,7 @@ void FZenStoreWriter::WritePackageData(const FPackageInfo& Info, FLargeMemoryWri
 
 	FIoBuffer CookedHeaderBuffer = FIoBuffer(PackageData.Data(), Info.HeaderSize, PackageData);
 	FIoBuffer CookedExportsBuffer = FIoBuffer(PackageData.Data() + Info.HeaderSize, PackageData.DataSize() - Info.HeaderSize, PackageData);
-	TUniquePtr<FPackageStorePackage> Package{PackageStoreOptimizer->CreatePackageFromCookedHeader(Info.PackageName, CookedHeaderBuffer)};
+	TUniquePtr<FPackageStorePackage> Package{PackageStoreOptimizer->CreatePackageFromCookedHeader(Info.OutputPackageName, CookedHeaderBuffer)};
 	PackageStoreOptimizer->FinalizePackage(Package.Get());
 	TArray<FFileRegion> FileRegionsCopy(FileRegions);
 	for (FFileRegion& Region : FileRegionsCopy)
@@ -234,7 +234,7 @@ void FZenStoreWriter::WritePackageData(const FPackageInfo& Info, FLargeMemoryWri
 	}
 	FIoBuffer PackageBuffer = PackageStoreOptimizer->CreatePackageBuffer(Package.Get(), CookedExportsBuffer, &FileRegionsCopy);
 	FIoChunkId ChunkId = CreateIoChunkId(Package->GetId().Value(), 0, EIoChunkType::ExportBundleData);
-	PackageStoreManifest.AddPackageData(Info.PackageName, Info.LooseFilePath, Info.ChunkId);
+	PackageStoreManifest.AddPackageData(Info.InputPackageName, Info.OutputPackageName, Info.LooseFilePath, Info.ChunkId);
 	for (FFileRegion& Region : FileRegionsCopy)
 	{
 		// Adjust regions once more so they are relative to the exports bundle buffer
@@ -246,7 +246,7 @@ void FZenStoreWriter::WritePackageData(const FPackageInfo& Info, FLargeMemoryWri
 
 	FCbObjectId ChunkOid = ToObjectId(Info.ChunkId);
 
-	FPendingPackageState& ExistingState = GetPendingPackage(Info.PackageName);
+	FPendingPackageState& ExistingState = GetPendingPackage(Info.InputPackageName);
 
 	FPackageDataEntry& Entry = ExistingState.PackageData;
 
@@ -267,12 +267,12 @@ void FZenStoreWriter::WriteIoStorePackageData(const FPackageInfo& Info, const FI
 
 	TRACE_CPUPROFILER_EVENT_SCOPE(WriteIoStorePackageData);
 
-	PackageStoreManifest.AddPackageData(Info.PackageName, Info.LooseFilePath, Info.ChunkId);
+	PackageStoreManifest.AddPackageData(Info.InputPackageName, Info.OutputPackageName, Info.LooseFilePath, Info.ChunkId);
 	//WriteFileRegions(*FPaths::ChangeExtension(Info.LooseFilePath, FString(".uexp") + FFileRegion::RegionsFileExtension), FileRegionsCopy);
 
 	FCbObjectId ChunkOid = ToObjectId(Info.ChunkId);
 
-	FPendingPackageState& ExistingState = GetPendingPackage(Info.PackageName);
+	FPendingPackageState& ExistingState = GetPendingPackage(Info.InputPackageName);
 
 	FPackageDataEntry& Entry = ExistingState.PackageData;
 
@@ -295,7 +295,7 @@ void FZenStoreWriter::WriteBulkData(const FBulkDataInfo& Info, const FIoBuffer&
 
 	FCbObjectId ChunkOid = ToObjectId(Info.ChunkId);
 
-	FPendingPackageState& ExistingState = GetPendingPackage(Info.PackageName);
+	FPendingPackageState& ExistingState = GetPendingPackage(Info.InputPackageName);
 
 	FBulkDataEntry& BulkEntry = ExistingState.BulkData.AddDefaulted_GetRef(); 
 
@@ -310,7 +310,7 @@ void FZenStoreWriter::WriteBulkData(const FBulkDataInfo& Info, const FIoBuffer&
 	BulkEntry.ChunkId	= ChunkOid;
 	BulkEntry.IsValid	= true;
 
-	PackageStoreManifest.AddBulkData(Info.PackageName, Info.LooseFilePath, Info.ChunkId);
+	PackageStoreManifest.AddBulkData(Info.InputPackageName, Info.OutputPackageName, Info.LooseFilePath, Info.ChunkId);
 
 	//	WriteFileRegions(*(Info.LooseFilePath + FFileRegion::RegionsFileExtension), FileRegions);
 }
@@ -319,7 +319,7 @@ void FZenStoreWriter::WriteAdditionalFile(const FAdditionalFileInfo& Info, const
 {
 	const FZenFileSystemManifest::FManifestEntry& ManifestEntry = ZenFileSystemManifest->CreateManifestEntry(Info.Filename);
 	
-	FPendingPackageState& ExistingState = GetPendingPackage(Info.PackageName);
+	FPendingPackageState& ExistingState = GetPendingPackage(Info.InputPackageName);
 
 	FFileDataEntry& FileEntry = ExistingState.FileData.AddDefaulted_GetRef();
 	
@@ -737,6 +737,7 @@ void FZenStoreWriter::CommitPackageInternal(FCommitPackageInfo&& CommitInfo)
 
 				CommitEventArgs.AdditionalFiles.Add(FAdditionalFileInfo
 				{ 
+					CommitInfo.PackageName,
 					CommitInfo.PackageName,
 					File.ZenManifestClientPath,
 					File.Info.ChunkId
diff --git a/Engine/Source/Developer/IoStoreUtilities/Public/PackageStoreManifest.h b/Engine/Source/Developer/IoStoreUtilities/Public/PackageStoreManifest.h
--- a/Engine/Source/Developer/IoStoreUtilities/Public/PackageStoreManifest.h
+++ b/Engine/Source/Developer/IoStoreUtilities/Public/PackageStoreManifest.h
@@ -31,9 +31,9 @@ public:
 	IOSTOREUTILITIES_API FPackageStoreManifest(const FString& CookedOutputPath);
 	IOSTOREUTILITIES_API ~FPackageStoreManifest() = default;
 
-	IOSTOREUTILITIES_API void BeginPackage(FName PackageName);
-	IOSTOREUTILITIES_API void AddPackageData(FName PackageName, const FString& FileName, const FIoChunkId& ChunkId);
-	IOSTOREUTILITIES_API void AddBulkData(FName PackageName, const FString& FileName, const FIoChunkId& ChunkId);
+	IOSTOREUTILITIES_API void BeginPackage(FName InputPackageName);
+	IOSTOREUTILITIES_API void AddPackageData(FName InputPackageName, FName OutputPackageName, const FString& FileName, const FIoChunkId& ChunkId);
+	IOSTOREUTILITIES_API void AddBulkData(FName InputPackageName, FName OutputPackageName, const FString& FileName, const FIoChunkId& ChunkId);
 
 	IOSTOREUTILITIES_API FIoStatus Save(const TCHAR* Filename) const;
 	IOSTOREUTILITIES_API FIoStatus Load(const TCHAR* Filename);
@@ -45,9 +45,15 @@ public:
 	IOSTOREUTILITIES_API const FZenServerInfo* ReadZenServerInfo() const;
 
 private:
+	FPackageInfo* GetPackageInfo_NoLock(FName InputPackageName, FName OutputPackageName);
+
 	mutable FCriticalSection CriticalSection;
 	FString CookedOutputPath;
-	TMap<FName, FPackageInfo> PackageInfoByNameMap;
+	/** Main Package Output Info map */
+	TMap<FName, FPackageInfo> MainPackageInfoByNameMap;
+	/** Additionally Generated Package Info Output map */
+	TMultiMap<FName, FPackageInfo> OptionalPackageInfoByNameMap;
+
 	TMap<FIoChunkId, FString> FileNameByChunkIdMap;
 	TUniquePtr<FZenServerInfo> ZenServerInfo;
 };
diff --git a/Engine/Source/Editor/UnrealEd/Private/Commandlets/AssetRegistryGenerator.cpp b/Engine/Source/Editor/UnrealEd/Private/Commandlets/AssetRegistryGenerator.cpp
--- a/Engine/Source/Editor/UnrealEd/Private/Commandlets/AssetRegistryGenerator.cpp
+++ b/Engine/Source/Editor/UnrealEd/Private/Commandlets/AssetRegistryGenerator.cpp
@@ -24,6 +24,7 @@
 #include "Commandlets/ChunkDependencyInfo.h"
 #include "IPlatformFileSandboxWrapper.h"
 #include "Misc/ConfigCacheIni.h"
+#include "Misc/PathViews.h"
 #include "Stats/StatsMisc.h"
 #include "Templates/UniquePtr.h"
 #include "Engine/AssetManager.h"
@@ -1524,13 +1525,13 @@ class FPackageCookerOpenOrderVisitor : public IPlatformFile::FDirectoryVisitor
 {
 	const FSandboxPlatformFile& SandboxFile;
 	const FString& PlatformSandboxPath;
-	const TSet<FString>& ValidExtensions;
+	const TSet<FStringView>& ValidExtensions;
 	TMultiMap<FString, FString>& PackageExtensions;
 public:
 	FPackageCookerOpenOrderVisitor(
 		const FSandboxPlatformFile& InSandboxFile,
 		const FString& InPlatformSandboxPath,
-		const TSet<FString>& InValidExtensions,
+		const TSet<FStringView>& InValidExtensions,
 		TMultiMap<FString, FString>& OutPackageExtensions) :
 		SandboxFile(InSandboxFile),
 		PlatformSandboxPath(InPlatformSandboxPath),
@@ -1544,10 +1545,16 @@ public:
 			return true;
 
 		FString Filename = FilenameOrDirectory;
-		FString FileExtension = FPaths::GetExtension(Filename, true);
-
+		FStringView UnusedFilePath, FileBaseName, FileExtension;
+		FPathViews::Split(Filename, UnusedFilePath, FileBaseName, FileExtension);
 		if (ValidExtensions.Contains(FileExtension))
 		{
+			// if the file base name ends with an optional extension, ignore it. (i.e. .o.uasset/.o.uexp etc)
+			if (FileBaseName.EndsWith(FPackagePath::GetOptionalSegmentExtensionModifier()))
+			{
+				return true;
+			}
+
 			FString PackageName;
 			Filename.ReplaceInline(*PlatformSandboxPath, *SandboxFile.GetSandboxDirectory());
 			FString AssetSourcePath = SandboxFile.ConvertFromSandboxPath(*Filename);
@@ -1640,16 +1647,16 @@ bool FAssetRegistryGenerator::WriteCookerOpenOrder(FSandboxPlatformFile& InSandb
 		}
 		
 		// Iterate sandbox folder and generate a map from package name to cooked files
-		const TArray<FString> ValidExtensions =
-		{
-			TEXT(".uasset"),
-			TEXT(".uexp"),
-			TEXT(".ubulk"),
-			TEXT(".uptnl"),
-			TEXT(".umap"),
-			TEXT(".ufont")
+		const TArray<FStringView> ValidExtensions =
+		{
+			TEXT("uasset"),
+			TEXT("uexp"),
+			TEXT("ubulk"),
+			TEXT("uptnl"),
+			TEXT("umap"),
+			TEXT("ufont")
 		};
-		const TSet<FString> ValidExtensionSet(ValidExtensions);
+		const TSet<FStringView> ValidExtensionSet(ValidExtensions);
 
 		const FString SandboxPath = InSandboxFile.GetSandboxDirectory();
 		const FString Platform = TargetPlatform->PlatformName();
diff --git a/Engine/Source/Editor/UnrealEd/Private/Cooker/DiffPackageWriter.cpp b/Engine/Source/Editor/UnrealEd/Private/Cooker/DiffPackageWriter.cpp
--- a/Engine/Source/Editor/UnrealEd/Private/Cooker/DiffPackageWriter.cpp
+++ b/Engine/Source/Editor/UnrealEd/Private/Cooker/DiffPackageWriter.cpp
@@ -141,11 +141,11 @@ TFuture<FMD5Hash> FDiffPackageWriter::CommitPackage(FCommitPackageInfo&& Info)
 void FDiffPackageWriter::WritePackageData(const FPackageInfo& Info, FLargeMemoryWriter& ExportsArchive,
 	const TArray<FFileRegion>& FileRegions)
 {
-	Inner->CompleteExportsArchiveForDiff(Info.PackageName, ExportsArchive);
+	Inner->CompleteExportsArchiveForDiff(Info, ExportsArchive);
 
 	FArchiveStackTrace& Writer = static_cast<FArchiveStackTrace&>(ExportsArchive);
 	ICookedPackageWriter::FPreviousCookedBytesData PreviousInnerData;
-	Inner->GetPreviousCookedBytes(Info.PackageName, PreviousInnerData);
+	Inner->GetPreviousCookedBytes(Info, PreviousInnerData);
 
 	FArchiveStackTrace::FPackageData PreviousPackageData;
 	PreviousPackageData.Data = PreviousInnerData.Data.Get();
@@ -157,7 +157,7 @@ void FDiffPackageWriter::WritePackageData(const FPackageInfo& Info, FLargeMemory
 	{
 		TMap<FName, FArchiveDiffStats> PackageDiffStats;
 		const TCHAR* CutoffString = TEXT("UEditorEngine::Save()");
-		Writer.CompareWith(PreviousPackageData, *BeginInfo.LooseFilePath, Info.HeaderSize, CutoffString,
+		Writer.CompareWith(PreviousPackageData, *Info.LooseFilePath, Info.HeaderSize, CutoffString,
 			MaxDiffsToLog, PackageDiffStats);
 
 		//COOK_STAT(FSavePackageStats::NumberOfDifferentPackages++);
diff --git a/Engine/Source/Editor/UnrealEd/Private/Cooker/LooseCookedPackageWriter.cpp b/Engine/Source/Editor/UnrealEd/Private/Cooker/LooseCookedPackageWriter.cpp
--- a/Engine/Source/Editor/UnrealEd/Private/Cooker/LooseCookedPackageWriter.cpp
+++ b/Engine/Source/Editor/UnrealEd/Private/Cooker/LooseCookedPackageWriter.cpp
@@ -82,16 +82,16 @@ TFuture<FMD5Hash> FLooseCookedPackageWriter::AsyncSave(FRecord& Record, const FC
 	return AsyncSaveOutputFiles(Record, Context);
 }
 
-void FLooseCookedPackageWriter::CompleteExportsArchiveForDiff(FName PackageName, FLargeMemoryWriter& ExportsArchive)
+void FLooseCookedPackageWriter::CompleteExportsArchiveForDiff(const FPackageInfo& Info, FLargeMemoryWriter& ExportsArchive)
 {
-	FPackageWriterRecords::FPackage& BaseRecord = Records.FindRecordChecked(PackageName);
+	FPackageWriterRecords::FPackage& BaseRecord = Records.FindRecordChecked(Info.InputPackageName);
 	FRecord& Record = static_cast<FRecord&>(BaseRecord);
 	Record.bCompletedExportsArchiveForDiff = true;
 
 	// Add on all the attachments which are usually added on during Commit. The order must match AsyncSave.
 	for (FBulkDataRecord& BulkRecord : Record.BulkDatas)
 	{
-		if (BulkRecord.Info.BulkDataType == FBulkDataInfo::AppendToExports)
+		if (BulkRecord.Info.BulkDataType == FBulkDataInfo::AppendToExports && BulkRecord.Info.MultiOutputIndex == Info.MultiOutputIndex)
 		{
 			ExportsArchive.Serialize(const_cast<void*>(BulkRecord.Buffer.GetData()),
 				BulkRecord.Buffer.GetSize());
@@ -99,8 +99,11 @@ void FLooseCookedPackageWriter::CompleteExportsArchiveForDiff(FName PackageName,
 	}
 	for (FLinkerAdditionalDataRecord& AdditionalRecord : Record.LinkerAdditionalDatas)
 	{
-		ExportsArchive.Serialize(const_cast<void*>(AdditionalRecord.Buffer.GetData()),
-			AdditionalRecord.Buffer.GetSize());
+		if (AdditionalRecord.Info.MultiOutputIndex == Info.MultiOutputIndex)
+		{
+			ExportsArchive.Serialize(const_cast<void*>(AdditionalRecord.Buffer.GetData()),
+				AdditionalRecord.Buffer.GetSize());
+		}
 	}
 
 	uint32 FooterData = PACKAGE_FILE_TAG;
@@ -110,7 +113,11 @@ void FLooseCookedPackageWriter::CompleteExportsArchiveForDiff(FName PackageName,
 
 void FLooseCookedPackageWriter::CollectForSavePackageData(FRecord& Record, FCommitContext& Context)
 {
-	Context.ExportsBuffers.Add(FExportBuffer{ Record.Package->Buffer, MoveTemp(Record.Package->Regions) });
+	Context.ExportsBuffers.AddDefaulted(Record.Packages.Num());
+	for (FPackageWriterRecords::FWritePackage& Package : Record.Packages)
+	{
+		Context.ExportsBuffers[Package.Info.MultiOutputIndex].Add(FExportBuffer{ Package.Buffer, MoveTemp(Package.Regions) });
+	}
 }
 
 void FLooseCookedPackageWriter::CollectForSaveBulkData(FRecord& Record, FCommitContext& Context)
@@ -124,16 +131,16 @@ void FLooseCookedPackageWriter::CollectForSaveBulkData(FRecord& Record, FCommitC
 				// Already Added in CompleteExportsArchiveForDiff
 				continue;
 			}
-			Context.ExportsBuffers.Add(FExportBuffer{ BulkRecord.Buffer, MoveTemp(BulkRecord.Regions) });
+			Context.ExportsBuffers[BulkRecord.Info.MultiOutputIndex].Add(FExportBuffer{ BulkRecord.Buffer, MoveTemp(BulkRecord.Regions) });
 		}
 		else
 		{
 			FWriteFileData& OutputFile = Context.OutputFiles.Emplace_GetRef();
-			OutputFile.Filename = FPaths::ChangeExtension(Record.Begin.LooseFilePath,
-				LexToString(BulkDataTypeToExtension(BulkRecord.Info.BulkDataType)));
+			OutputFile.Filename = BulkRecord.Info.LooseFilePath;
 			OutputFile.Buffer = FCompositeBuffer(BulkRecord.Buffer);
 			OutputFile.Regions = MoveTemp(BulkRecord.Regions);
 			OutputFile.bIsSidecar = true;
+			OutputFile.bContributeToHash = BulkRecord.Info.MultiOutputIndex == 0; // Only caculate the main package output hash
 		}
 	}
 }
@@ -147,7 +154,7 @@ void FLooseCookedPackageWriter::CollectForSaveLinkerAdditionalDataRecords(FRecor
 
 	for (FLinkerAdditionalDataRecord& AdditionalRecord : Record.LinkerAdditionalDatas)
 	{
-		Context.ExportsBuffers.Add(FExportBuffer{ AdditionalRecord.Buffer, MoveTemp(AdditionalRecord.Regions) });
+		Context.ExportsBuffers[AdditionalRecord.Info.MultiOutputIndex].Add(FExportBuffer{ AdditionalRecord.Buffer, MoveTemp(AdditionalRecord.Regions) });
 	}
 }
 
@@ -159,6 +166,7 @@ void FLooseCookedPackageWriter::CollectForSaveAdditionalFileRecords(FRecord& Rec
 		OutputFile.Filename = AdditionalRecord.Info.Filename;
 		OutputFile.Buffer = FCompositeBuffer(AdditionalRecord.Buffer);
 		OutputFile.bIsSidecar = true;
+		OutputFile.bContributeToHash = AdditionalRecord.Info.MultiOutputIndex == 0; // Only calculate the main package output hash
 	}
 }
 
@@ -172,7 +180,10 @@ void FLooseCookedPackageWriter::CollectForSaveExportsFooter(FRecord& Record, FCo
 
 	uint32 FooterData = PACKAGE_FILE_TAG;
 	FSharedBuffer Buffer = FSharedBuffer::Clone(&FooterData, sizeof(FooterData));
-	Context.ExportsBuffers.Add(FExportBuffer{ Buffer, TArray<FFileRegion>() });
+	for (FPackageWriterRecords::FWritePackage& Package : Record.Packages)
+	{
+		Context.ExportsBuffers[Package.Info.MultiOutputIndex].Add(FExportBuffer{ Buffer, TArray<FFileRegion>() });
+	}
 }
 
 void FLooseCookedPackageWriter::AddToExportsSize(int64& ExportsSize)
@@ -182,47 +193,55 @@ void FLooseCookedPackageWriter::AddToExportsSize(int64& ExportsSize)
 
 void FLooseCookedPackageWriter::CollectForSaveExportsBuffers(FRecord& Record, FCommitContext& Context)
 {
-	// Split the ExportsBuffer into (1) Header and (2) Exports + AllAppendedData
-	int64 HeaderSize = Record.Package->Info.HeaderSize;
-	check(Context.ExportsBuffers.Num() > 0);
-	FExportBuffer& HeaderAndExportsBuffer = Context.ExportsBuffers[0];
-	FSharedBuffer& HeaderAndExportsData = HeaderAndExportsBuffer.Buffer;
-
-	// Header (.uasset/.umap)
+	check(Context.ExportsBuffers.Num() == Record.Packages.Num());
+	for (FPackageWriterRecords::FWritePackage& Package : Record.Packages)
 	{
-		FWriteFileData& OutputFile = Context.OutputFiles.Emplace_GetRef();
-		OutputFile.Filename = Record.Begin.LooseFilePath;
-		OutputFile.Buffer = FCompositeBuffer(
-			FSharedBuffer::MakeView(HeaderAndExportsData.GetData(), HeaderSize, HeaderAndExportsData));
-		OutputFile.bIsSidecar = false;
-	}
-
-	// Exports + AllAppendedData (.uexp)
-	{
-		FWriteFileData& OutputFile = Context.OutputFiles.Emplace_GetRef();
-		OutputFile.Filename = FPaths::ChangeExtension(Record.Begin.LooseFilePath, LexToString(EPackageExtension::Exports));
-		OutputFile.bIsSidecar = false;
-
-		int32 NumBuffers = Context.ExportsBuffers.Num();
-		TArray<FSharedBuffer> BuffersForComposition;
-		BuffersForComposition.Reserve(NumBuffers);
+		TArray<FExportBuffer>& ExportsBuffers = Context.ExportsBuffers[Package.Info.MultiOutputIndex];
+		check(ExportsBuffers.Num() > 0);
 
-		const uint8* ExportsStart = static_cast<const uint8*>(HeaderAndExportsData.GetData()) + HeaderSize;
-		BuffersForComposition.Add(FSharedBuffer::MakeView(ExportsStart, HeaderAndExportsData.GetSize() - HeaderSize,
-			HeaderAndExportsData));
-		OutputFile.Regions.Append(MoveTemp(HeaderAndExportsBuffer.Regions));
+		// Split the ExportsBuffer into (1) Header and (2) Exports + AllAppendedData
+		int64 HeaderSize = Package.Info.HeaderSize;
+		FExportBuffer& HeaderAndExportsBuffer = ExportsBuffers[0];
+		FSharedBuffer& HeaderAndExportsData = HeaderAndExportsBuffer.Buffer;
 
-		for (FExportBuffer& ExportsBuffer : TArrayView<FExportBuffer>(Context.ExportsBuffers).Slice(1, NumBuffers - 1))
+		// Header (.uasset/.umap)
 		{
-			BuffersForComposition.Add(ExportsBuffer.Buffer);
-			OutputFile.Regions.Append(MoveTemp(ExportsBuffer.Regions));
+			FWriteFileData& OutputFile = Context.OutputFiles.Emplace_GetRef();
+			OutputFile.Filename = Package.Info.LooseFilePath;
+			OutputFile.Buffer = FCompositeBuffer(
+				FSharedBuffer::MakeView(HeaderAndExportsData.GetData(), HeaderSize, HeaderAndExportsData));
+			OutputFile.bIsSidecar = false;
+			OutputFile.bContributeToHash = Package.Info.MultiOutputIndex == 0; // Only calculate the main package output hash
 		}
-		OutputFile.Buffer = FCompositeBuffer(BuffersForComposition);
 
-		// Adjust regions so they are relative to the start of the uexp file
-		for (FFileRegion& Region : OutputFile.Regions)
+		// Exports + AllAppendedData (.uexp)
 		{
-			Region.Offset -= HeaderSize;
+			FWriteFileData& OutputFile = Context.OutputFiles.Emplace_GetRef();
+			OutputFile.Filename = FPaths::ChangeExtension(Package.Info.LooseFilePath, LexToString(EPackageExtension::Exports));
+			OutputFile.bIsSidecar = false;
+			OutputFile.bContributeToHash = Package.Info.MultiOutputIndex == 0; // Only caculate the main package output hash
+
+			int32 NumBuffers = ExportsBuffers.Num();
+			TArray<FSharedBuffer> BuffersForComposition;
+			BuffersForComposition.Reserve(NumBuffers);
+
+			const uint8* ExportsStart = static_cast<const uint8*>(HeaderAndExportsData.GetData()) + HeaderSize;
+			BuffersForComposition.Add(FSharedBuffer::MakeView(ExportsStart, HeaderAndExportsData.GetSize() - HeaderSize,
+				HeaderAndExportsData));
+			OutputFile.Regions.Append(MoveTemp(HeaderAndExportsBuffer.Regions));
+
+			for (FExportBuffer& ExportsBuffer : TArrayView<FExportBuffer>(ExportsBuffers).Slice(1, NumBuffers - 1))
+			{
+				BuffersForComposition.Add(ExportsBuffer.Buffer);
+				OutputFile.Regions.Append(MoveTemp(ExportsBuffer.Regions));
+			}
+			OutputFile.Buffer = FCompositeBuffer(BuffersForComposition);
+
+			// Adjust regions so they are relative to the start of the uexp file
+			for (FFileRegion& Region : OutputFile.Regions)
+			{
+				Region.Offset -= HeaderSize;
+			}
 		}
 	}
 }
@@ -288,7 +307,8 @@ static void WriteToFile(const FString& Filename, const FCompositeBuffer& Buffer)
 
 void FLooseCookedPackageWriter::FWriteFileData::Write(FMD5& AccumulatedHash, EWriteOptions WriteOptions) const
 {
-	if (EnumHasAnyFlags(WriteOptions, EWriteOptions::ComputeHash))
+	//@todo: FH: Should we calculate the hash of both output, currently only the main package output hash is calculated
+	if (EnumHasAnyFlags(WriteOptions, EWriteOptions::ComputeHash) && bContributeToHash)
 	{
 		for (const FSharedBuffer& Segment : Buffer.GetSegments())
 		{
@@ -323,17 +343,22 @@ void FLooseCookedPackageWriter::FWriteFileData::Write(FMD5& AccumulatedHash, EWr
 
 void FLooseCookedPackageWriter::UpdateManifest(FRecord& Record)
 {
-	FName PackageName = Record.Begin.PackageName;
-	FIoChunkId ChunkId = CreateIoChunkId(FPackageId::FromName(PackageName).Value(), 0, EIoChunkType::ExportBundleData);
-	PackageStoreManifest.AddPackageData(PackageName, Record.Begin.LooseFilePath, ChunkId);
+	for (const FPackageWriterRecords::FWritePackage& Package : Record.Packages)
+	{
+		PackageStoreManifest.AddPackageData(Package.Info.InputPackageName, Package.Info.OutputPackageName, Package.Info.LooseFilePath, Package.Info.ChunkId);
+	}
+	for (const FPackageWriterRecords::FBulkData& BulkData : Record.BulkDatas)
+	{
+		PackageStoreManifest.AddBulkData(BulkData.Info.InputPackageName, BulkData.Info.OutputPackageName, BulkData.Info.LooseFilePath, BulkData.Info.ChunkId);
+	}
 }
 
-bool FLooseCookedPackageWriter::GetPreviousCookedBytes(FName PackageName, FPreviousCookedBytesData& OutData)
+bool FLooseCookedPackageWriter::GetPreviousCookedBytes(const FPackageInfo& Info, FPreviousCookedBytesData& OutData)
 {
-	FPackageWriterRecords::FPackage& BaseRecord = Records.FindRecordChecked(PackageName);
+	FPackageWriterRecords::FPackage& BaseRecord = Records.FindRecordChecked(Info.InputPackageName);
 	FRecord& Record = static_cast<FRecord&>(BaseRecord);
 	FArchiveStackTrace::FPackageData ExistingPackageData;
-	FArchiveStackTrace::LoadPackageIntoMemory(*Record.Begin.LooseFilePath, ExistingPackageData, OutData.Data);
+	FArchiveStackTrace::LoadPackageIntoMemory(*Record.Packages[Info.MultiOutputIndex].Info.LooseFilePath, ExistingPackageData, OutData.Data);
 	OutData.Size = ExistingPackageData.Size;
 	OutData.HeaderSize = ExistingPackageData.HeaderSize;
 	OutData.StartOffset = ExistingPackageData.StartOffset;
diff --git a/Engine/Source/Editor/UnrealEd/Private/Cooker/LooseCookedPackageWriter.h b/Engine/Source/Editor/UnrealEd/Private/Cooker/LooseCookedPackageWriter.h
--- a/Engine/Source/Editor/UnrealEd/Private/Cooker/LooseCookedPackageWriter.h
+++ b/Engine/Source/Editor/UnrealEd/Private/Cooker/LooseCookedPackageWriter.h
@@ -46,8 +46,8 @@ public:
 	virtual void RemoveCookedPackages(TArrayView<const FName> PackageNamesToRemove) override;
 	virtual void RemoveCookedPackages() override;
 	virtual void MarkPackagesUpToDate(TArrayView<const FName> UpToDatePackages) override;
-	virtual bool GetPreviousCookedBytes(FName PackageName, FPreviousCookedBytesData& OutData) override;
-	virtual void CompleteExportsArchiveForDiff(FName PackageName, FLargeMemoryWriter& ExportsArchive) override;
+	virtual bool GetPreviousCookedBytes(const FPackageInfo& Info, FPreviousCookedBytesData& OutData) override;
+	virtual void CompleteExportsArchiveForDiff(const FPackageInfo& Info, FLargeMemoryWriter& ExportsArchive) override;
 	virtual TFuture<FMD5Hash> CommitPackageInternal(FPackageWriterRecords::FPackage&& BaseRecord,
 		const FCommitPackageInfo& Info) override;
 	virtual FPackageWriterRecords::FPackage* ConstructRecord() override;
@@ -61,7 +61,7 @@ private:
 		bool bCompletedExportsArchiveForDiff = false;
 	};
 
-	/** Buffers that are combined into the HeaderAndExports file (which is then split into .uasset + .uexp). */
+	/** Buffers that are combined into the HeaderAndExports file (which is then split into .uasset + .uexp or .uoasset + .uoexp). */
 	struct FExportBuffer
 	{
 		FSharedBuffer Buffer;
@@ -69,7 +69,7 @@ private:
 	};
 
 	/**
-	 * The data needed to asynchronously write one of the files (.uasset, .uexp, .ubulk, and any additional),
+	 * The data needed to asynchronously write one of the files (.uasset, .uexp, .ubulk, any optional and any additional),
 	 * without reference back to other data on this writer.
 	 */
 	struct FWriteFileData
@@ -78,6 +78,7 @@ private:
 		FCompositeBuffer Buffer;
 		TArray<FFileRegion> Regions;
 		bool bIsSidecar;
+		bool bContributeToHash = true;
 
 		void Write(FMD5& AccumulatedHash, EWriteOptions WriteOptions) const;
 	};
@@ -86,7 +87,7 @@ private:
 	struct FCommitContext
 	{
 		const FCommitPackageInfo& Info;
-		TArray<FExportBuffer> ExportsBuffers;
+		TArray<TArray<FExportBuffer>> ExportsBuffers;
 		TArray<FWriteFileData> OutputFiles;
 	};
 
diff --git a/Engine/Source/Editor/UnrealEd/Private/EditorDomain/EditorDomainUtils.cpp b/Engine/Source/Editor/UnrealEd/Private/EditorDomain/EditorDomainUtils.cpp
--- a/Engine/Source/Editor/UnrealEd/Private/EditorDomain/EditorDomainUtils.cpp
+++ b/Engine/Source/Editor/UnrealEd/Private/EditorDomain/EditorDomainUtils.cpp
@@ -1007,16 +1007,19 @@ protected:
 			// 6) Handle EPackageSegment::PayloadSidecar in FEditorDomain::OpenReadPackage by returning an archive configured to deserialize the sidecar segment.
 			unimplemented();
 		}
-		WritePackageRecord = *Record.Package;
+		
+		// EditorDomain save does not support multi package outputs
+		check(Record.Packages.Num() == 1);
+		FPackageWriterRecords::FWritePackage& Package = Record.Packages[0];
+		WritePackageRecord = Package;
 
 		TArray<FSharedBuffer> AttachmentBuffers;
-
-		for (const FFileRegion& FileRegion : Record.Package->Regions)
+		for (const FFileRegion& FileRegion : Package.Regions)
 		{
 			checkf(FileRegion.Type == EFileRegionType::None, TEXT("Does not support FileRegion types other than None."));
 		}
-		check(Record.Package->Buffer.GetSize() > 0); // Header+Exports segment is non-zero in length
-		AttachmentBuffers.Add(Record.Package->Buffer);
+		check(Package.Buffer.GetSize() > 0); // Header+Exports segment is non-zero in length
+		AttachmentBuffers.Add(Package.Buffer);
 
 		BulkDataSize = 0;
 		for (const FBulkDataRecord& BulkRecord : Record.BulkDatas)
@@ -1087,7 +1090,7 @@ protected:
 			Attachments.Add(FAttachment{ Buffer, IntToValueId(AttachmentIndex++) });
 			FileSize += Buffer.GetSize();
 		}
-		WritePackageRecord = MoveTemp(*Record.Package);
+		WritePackageRecord = MoveTemp(Package);
 
 		return TFuture<FMD5Hash>();
 	}
diff --git a/Engine/Source/Runtime/Core/Public/Serialization/PackageWriter.h b/Engine/Source/Runtime/Core/Public/Serialization/PackageWriter.h
--- a/Engine/Source/Runtime/Core/Public/Serialization/PackageWriter.h
+++ b/Engine/Source/Runtime/Core/Public/Serialization/PackageWriter.h
@@ -88,10 +88,14 @@ public:
 
 	struct FPackageInfo
 	{
-		FName		PackageName;
+		/** Associated Package Name Entry from BeginPackage */
+		FName		InputPackageName;
+		/** Output Package Name (Input Package can produce multiple output) */
+		FName		OutputPackageName;
 		FString		LooseFilePath;
-		uint64		HeaderSize;
+		uint64		HeaderSize = 0;
 		FIoChunkId	ChunkId = FIoChunkId::InvalidChunkId;
+		uint32		MultiOutputIndex = 0;
 	};
 
 	/** Write package data (exports and serialized header)
@@ -112,10 +116,14 @@ public:
 			NumTypes,
 		};
 
-		FName		PackageName;
+		/** Associated Package Name Entry */
+		FName		InputPackageName;
+		/** Output Package Name (Input Package can produce multiple output) */
+		FName		OutputPackageName;
 		EType		BulkDataType = BulkSegment;
 		FString		LooseFilePath;
 		FIoChunkId	ChunkId = FIoChunkId::InvalidChunkId;
+		uint32		MultiOutputIndex = 0;
 	};
 
 	/** Write bulk data for the current package
@@ -124,9 +132,13 @@ public:
 
 	struct FAdditionalFileInfo
 	{
-		FName	PackageName;
-		FString	Filename;
-		FIoChunkId ChunkId;
+		/** Associated Package Name Entry */
+		FName		InputPackageName;
+		/** Output Package Name (Input Package can produce multiple output) */
+		FName		OutputPackageName;
+		FString		Filename;
+		FIoChunkId	ChunkId;
+		uint32		MultiOutputIndex = 0;
 	};
 
 	/** Write separate files written by UObjects during cooking via UObject::CookAdditionalFiles. */
@@ -134,7 +146,11 @@ public:
 
 	struct FLinkerAdditionalDataInfo
 	{
-		FName PackageName;
+		/** Associated Package Name Entry */
+		FName	InputPackageName;
+		/** Output Package Name (Input Package can produce multiple output) */
+		FName	OutputPackageName;
+		uint32	MultiOutputIndex = 0;
 	};
 	/** Write separate data written by UObjects via FLinkerSave::AdditionalDataToAppend. */
 	virtual void WriteLinkerAdditionalData(const FLinkerAdditionalDataInfo& Info, const FIoBuffer& Data, const TArray<FFileRegion>& FileRegions) = 0;
@@ -269,14 +285,14 @@ public:
 		int64 StartOffset;
 	};
 	/** Load the bytes of the previously-cooked package, used for diffing */
-	virtual bool GetPreviousCookedBytes(FName PackageName, FPreviousCookedBytesData& OutData)
+	virtual bool GetPreviousCookedBytes(const FPackageInfo& Info, FPreviousCookedBytesData& OutData)
 	{
 		// The subclass must override this method if it returns bDiffModeSupported=true in GetCookCapabilities
 		unimplemented();
 		return false;
 	}
 	/** Append all data to the Exports archive that would normally be done in CommitPackage, used for diffing. */
-	virtual void CompleteExportsArchiveForDiff(FName PackageName, FLargeMemoryWriter& ExportsArchive)
+	virtual void CompleteExportsArchiveForDiff(const FPackageInfo& Info, FLargeMemoryWriter& ExportsArchive)
 	{
 		// The subclass must override this method if it returns bDiffModeSupported=true in GetCookCapabilities
 		unimplemented();
diff --git a/Engine/Source/Runtime/Core/Public/UObject/ObjectVersion.h b/Engine/Source/Runtime/Core/Public/UObject/ObjectVersion.h
--- a/Engine/Source/Runtime/Core/Public/UObject/ObjectVersion.h
+++ b/Engine/Source/Runtime/Core/Public/UObject/ObjectVersion.h
@@ -50,6 +50,9 @@ enum class EUnrealEngineObjectUE5Version : uint32
 
 	// Added a payload table of contents to the package summary 
 	PAYLOAD_TOC,
+
+	// Added data to identify references from and to optional package
+	OPTIONAL_RESOURCES,
 	
 	// -----<new versions can be added before this line>-------------------------------------------------
 	// - this needs to be the last line (see note below)
diff --git a/Engine/Source/Runtime/CoreUObject/Internal/Serialization/PackageWriterToSharedBuffer.h b/Engine/Source/Runtime/CoreUObject/Internal/Serialization/PackageWriterToSharedBuffer.h
--- a/Engine/Source/Runtime/CoreUObject/Internal/Serialization/PackageWriterToSharedBuffer.h
+++ b/Engine/Source/Runtime/CoreUObject/Internal/Serialization/PackageWriterToSharedBuffer.h
@@ -58,7 +58,7 @@ public:
 		/** Always valid during CommitPackageInternal */
 		IPackageWriter::FBeginPackageInfo Begin;
 		/** Always valid during CommitPackageInternal if Info.bSucceeded */
-		TOptional<FWritePackage> Package;
+		TArray<FWritePackage> Packages;
 		TArray<FBulkData> BulkDatas;
 		TArray<FAdditionalFile> AdditionalFiles;
 		TArray<FLinkerAdditionalData>  LinkerAdditionalDatas;
diff --git a/Engine/Source/Runtime/CoreUObject/Private/Misc/PackagePath.cpp b/Engine/Source/Runtime/CoreUObject/Private/Misc/PackagePath.cpp
--- a/Engine/Source/Runtime/CoreUObject/Private/Misc/PackagePath.cpp
+++ b/Engine/Source/Runtime/CoreUObject/Private/Misc/PackagePath.cpp
@@ -186,6 +186,11 @@ EPackageExtension FPackagePath::ParseExtension(FStringView Filename, int32* OutE
 	return EPackageExtension::Custom;
 }
 
+const TCHAR* FPackagePath::GetOptionalSegmentExtensionModifier()
+{
+	return TEXT(".o");
+}
+
 const TCHAR* FPackagePath::GetExternalActorsFolderName()
 {
 	return TEXT("__ExternalActors__");
diff --git a/Engine/Source/Runtime/CoreUObject/Private/Serialization/PackageWriterToSharedBuffer.cpp b/Engine/Source/Runtime/CoreUObject/Private/Serialization/PackageWriterToSharedBuffer.cpp
--- a/Engine/Source/Runtime/CoreUObject/Private/Serialization/PackageWriterToSharedBuffer.cpp
+++ b/Engine/Source/Runtime/CoreUObject/Private/Serialization/PackageWriterToSharedBuffer.cpp
@@ -36,34 +36,34 @@ void FPackageWriterRecords::BeginPackage(FPackage* Record, const IPackageWriter:
 void FPackageWriterRecords::WritePackageData(const IPackageWriter::FPackageInfo& Info,
 	FLargeMemoryWriter& ExportsArchive, const TArray<FFileRegion>& FileRegions)
 {
-	FPackage& Record = FindRecordChecked(Info.PackageName);
-	int64 DataSize= ExportsArchive.TotalSize();
+	FPackage& Record = FindRecordChecked(Info.InputPackageName);
+	int64 DataSize = ExportsArchive.TotalSize();
 	checkf(DataSize > 0, TEXT("IPackageWriter->WritePackageData must not be called with an empty ExportsArchive"));
 	checkf(static_cast<uint64>(DataSize) >= Info.HeaderSize,
 		TEXT("IPackageWriter->WritePackageData must not be called with HeaderSize > ExportsArchive.TotalSize"));
 	FSharedBuffer Buffer = FSharedBuffer::TakeOwnership(ExportsArchive.ReleaseOwnership(), DataSize,
 		FMemory::Free);
-	Record.Package = FWritePackage{ Info, MoveTemp(Buffer), FileRegions };
+	Record.Packages.Insert(FWritePackage{ Info, MoveTemp(Buffer), FileRegions }, Info.MultiOutputIndex);
 }
 
 void FPackageWriterRecords::WriteBulkData(const IPackageWriter::FBulkDataInfo& Info, const FIoBuffer& BulkData,
 	const TArray<FFileRegion>& FileRegions)
 {
-	FPackage& Record = FindRecordChecked(Info.PackageName);
+	FPackage& Record = FindRecordChecked(Info.InputPackageName);
 	Record.BulkDatas.Add(FBulkData{ Info, IoBufferToSharedBuffer(BulkData), FileRegions });
 }
 
 void FPackageWriterRecords::WriteAdditionalFile(const IPackageWriter::FAdditionalFileInfo& Info,
 	const FIoBuffer& FileData)
 {
-	FPackage& Record = FindRecordChecked(Info.PackageName);
+	FPackage& Record = FindRecordChecked(Info.InputPackageName);
 	Record.AdditionalFiles.Add(FAdditionalFile{ Info, IoBufferToSharedBuffer(FileData) });
 }
 
 void FPackageWriterRecords::WriteLinkerAdditionalData(const IPackageWriter::FLinkerAdditionalDataInfo& Info,
 	const FIoBuffer& Data, const TArray<FFileRegion>& FileRegions)
 {
-	FPackage& Record = FindRecordChecked(Info.PackageName);
+	FPackage& Record = FindRecordChecked(Info.InputPackageName);
 	Record.LinkerAdditionalDatas.Add(
 		FLinkerAdditionalData{ Info, IoBufferToSharedBuffer(Data), FileRegions });
 }
@@ -93,13 +93,15 @@ TUniquePtr<FPackageWriterRecords::FPackage> FPackageWriterRecords::FindAndRemove
 
 void FPackageWriterRecords::ValidateCommit(FPackage& Record, const IPackageWriter::FCommitPackageInfo& Info) const
 {
-	checkf(Info.bSucceeded == false || Record.Package.IsSet(),
+	checkf(Info.bSucceeded == false || Record.Packages.Num() > 0,
 		TEXT("IPackageWriter->WritePackageData must be called before Commit if the Package save was successful."));
-	bool HasBulkDataType[IPackageWriter::FBulkDataInfo::NumTypes]{};
+	checkf(Info.bSucceeded == false || Record.Packages.FindByPredicate([](const FPackageWriterRecords::FWritePackage& Package) { return Package.Info.MultiOutputIndex == 0; }),
+		TEXT("SavePackage must provide output 0 when saving multioutput packages."));
+	uint8 HasBulkDataType[IPackageWriter::FBulkDataInfo::NumTypes]{};
 	for (FBulkData& BulkRecord : Record.BulkDatas)
 	{
-		checkf(!HasBulkDataType[(int32)BulkRecord.Info.BulkDataType],
+		checkf((HasBulkDataType[(int32)BulkRecord.Info.BulkDataType] & (1 << BulkRecord.Info.MultiOutputIndex)) == 0,
 			TEXT("IPackageWriter->WriteBulkData must not be called with more than one BulkData of the same type."));
-		HasBulkDataType[(int32)BulkRecord.Info.BulkDataType] = true;
+		HasBulkDataType[(int32)BulkRecord.Info.BulkDataType] |= 1 << BulkRecord.Info.MultiOutputIndex;
 	}
 }
\ No newline at end of file
diff --git a/Engine/Source/Runtime/CoreUObject/Private/UObject/ObjectResource.cpp b/Engine/Source/Runtime/CoreUObject/Private/UObject/ObjectResource.cpp
--- a/Engine/Source/Runtime/CoreUObject/Private/UObject/ObjectResource.cpp
+++ b/Engine/Source/Runtime/CoreUObject/Private/UObject/ObjectResource.cpp
@@ -45,6 +45,7 @@ FObjectExport::FObjectExport()
 , bNotForServer(false)
 , bNotAlwaysLoadedForEditorGame(true)
 , bIsAsset(false)
+, bGeneratePublicHash(false)
 , bExportLoadFailed(false)
 // @todo: BP2CPP_remove
 PRAGMA_DISABLE_DEPRECATION_WARNINGS
@@ -74,6 +75,7 @@ FObjectExport::FObjectExport( UObject* InObject, bool bInNotAlwaysLoadedForEdito
 , bNotForServer(false)
 , bNotAlwaysLoadedForEditorGame(bInNotAlwaysLoadedForEditorGame)
 , bIsAsset(false)
+, bGeneratePublicHash(false)
 , bExportLoadFailed(false)
 // @todo: BP2CPP_remove
 PRAGMA_DISABLE_DEPRECATION_WARNINGS
@@ -167,6 +169,11 @@ void operator<<(FStructuredArchive::FSlot Slot, FObjectExport& E)
 		Record << SA_VALUE(TEXT("bIsAsset"), E.bIsAsset);
 	}
 
+	if (BaseArchive.UEVer() >= EUnrealEngineObjectUE5Version::OPTIONAL_RESOURCES)
+	{
+		Record << SA_VALUE(TEXT("bGeneratePublicHash"), E.bGeneratePublicHash);
+	}
+
 	if (BaseArchive.UEVer() >= VER_UE4_PRELOAD_DEPENDENCIES_IN_COOKED_EXPORTS)
 	{
 		Record << SA_VALUE(TEXT("FirstExportDependency"), E.FirstExportDependency);
@@ -249,6 +256,7 @@ void operator<<(FStructuredArchive::FSlot Slot, FObjectTextExport& E)
 FObjectImport::FObjectImport()
 	: FObjectResource()
 	, SourceIndex(INDEX_NONE)
+	, bImportOptional(false)
 	, bImportPackageHandled(false)
 	, bImportSearchedFor(false)
 	, bImportFailed(false)
@@ -262,6 +270,7 @@ FObjectImport::FObjectImport(UObject* InObject)
 	, ClassPackage(InObject ? InObject->GetClass()->GetOuter()->GetFName() : NAME_None)
 	, ClassName(InObject ? InObject->GetClass()->GetFName() : NAME_None)
 	, SourceIndex(INDEX_NONE)
+	, bImportOptional(false)
 	, bImportPackageHandled(false)
 	, bImportSearchedFor(false)
 	, bImportFailed(false)
@@ -275,6 +284,7 @@ FObjectImport::FObjectImport(UObject* InObject, UClass* InClass)
 	, ClassPackage((InObject && InClass) ? InClass->GetOuter()->GetFName() : NAME_None)
 	, ClassName((InObject && InClass) ? InClass->GetFName() : NAME_None)
 	, SourceIndex(INDEX_NONE)
+	, bImportOptional(false)
 	, bImportPackageHandled(false)
 	, bImportSearchedFor(false)
 	, bImportFailed(false)
@@ -298,7 +308,6 @@ void operator<<(FStructuredArchive::FSlot Slot, FObjectImport& I)
 	Record << SA_VALUE(TEXT("OuterIndex"), I.OuterIndex);
 	Record << SA_VALUE(TEXT("ObjectName"), I.ObjectName);
 
-	//@todo: re-enable package override at runtime when ready
 #if WITH_EDITORONLY_DATA
 	if (Slot.GetUnderlyingArchive().UEVer() >= VER_UE4_NON_OUTER_PACKAGE_IMPORT && !Slot.GetUnderlyingArchive().IsFilterEditorOnly())
 	{
@@ -306,6 +315,11 @@ void operator<<(FStructuredArchive::FSlot Slot, FObjectImport& I)
 	}
 #endif
 
+	if (Slot.GetUnderlyingArchive().UEVer() >= EUnrealEngineObjectUE5Version::OPTIONAL_RESOURCES)
+	{
+		Record << SA_VALUE(TEXT("bImportOptional"), I.bImportOptional);
+	}
+
 	if (Slot.GetUnderlyingArchive().IsLoading())
 	{
 		I.SourceLinker = NULL;
diff --git a/Engine/Source/Runtime/CoreUObject/Private/UObject/SavePackage.cpp b/Engine/Source/Runtime/CoreUObject/Private/UObject/SavePackage.cpp
--- a/Engine/Source/Runtime/CoreUObject/Private/UObject/SavePackage.cpp
+++ b/Engine/Source/Runtime/CoreUObject/Private/UObject/SavePackage.cpp
@@ -3839,7 +3839,7 @@ FSavePackageResultStruct UPackage::Save(UPackage* InOuter, UObject* InAsset, con
 							TotalPackageSizeUncompressed += Size;
 
 							IPackageWriter::FAdditionalFileInfo FileInfo;
-							FileInfo.PackageName = InOuter->GetFName();
+							FileInfo.OutputPackageName = FileInfo.InputPackageName = InOuter->GetFName();
 							FileInfo.Filename = *Writer.GetArchiveName();
 
 							FIoBuffer FileData(FIoBuffer::AssumeOwnership, Writer.ReleaseOwnership(), Size);
@@ -4119,11 +4119,11 @@ FSavePackageResultStruct UPackage::Save(UPackage* InOuter, UObject* InAsset, con
 						if (PackageWriter)
 						{
 							IPackageWriter::FPackageInfo PackageInfo;
-							PackageInfo.PackageName = InOuter->GetFName();
+							PackageInfo.OutputPackageName = PackageInfo.InputPackageName = InOuter->GetFName();
 							PackageInfo.LooseFilePath = Filename;
 							PackageInfo.HeaderSize = Linker->Summary.TotalHeaderSize;
 
-							FPackageId PackageId = FPackageId::FromName(PackageInfo.PackageName);
+							FPackageId PackageId = FPackageId::FromName(PackageInfo.OutputPackageName);
 							PackageInfo.ChunkId = CreateIoChunkId(PackageId.Value(), 0, EIoChunkType::ExportBundleData);
 
 							PackageWriter->WritePackageData(PackageInfo, *Writer, Linker->FileRegions);
diff --git a/Engine/Source/Runtime/CoreUObject/Private/UObject/SavePackage/SaveContext.h b/Engine/Source/Runtime/CoreUObject/Private/UObject/SavePackage/SaveContext.h
--- a/Engine/Source/Runtime/CoreUObject/Private/UObject/SavePackage/SaveContext.h
+++ b/Engine/Source/Runtime/CoreUObject/Private/UObject/SavePackage/SaveContext.h
@@ -93,6 +93,21 @@ struct FIllegalReference
 /** Hold the harvested exports and imports for a realm */
 struct FHarvestedRealm
 {
+
+	~FHarvestedRealm()
+	{
+		CloseLinkerArchives();
+
+		if (TempFilename.IsSet())
+		{
+			IFileManager::Get().Delete(*TempFilename.GetValue());
+		}
+		if (TextFormatTempFilename.IsSet())
+		{
+			IFileManager::Get().Delete(*TextFormatTempFilename.GetValue());
+		}
+	}
+
 	void AddImport(UObject* InObject)
 	{
 		Imports.Add(InObject);
@@ -217,7 +232,94 @@ struct FHarvestedRealm
 		return false;
 	}
 
+	FLinkerSave* GetLinker() const
+	{
+		return Linker.Get();
+	}
+
+	void SetLinker(TPimplPtr<FLinkerSave> InLinker)
+	{
+		Linker = MoveTemp(InLinker);
+	}
+
+	bool CloseLinkerArchives()
+	{
+		bool bSuccess = true;
+		if (Linker)
+		{
+			bSuccess = Linker->CloseAndDestroySaver();
+		}
+		StructuredArchive.Reset();
+		Formatter.Reset();
+		TextFormatArchive.Reset();
+		return bSuccess;
+	}
+
+	FArchive* GetTextFormatArchive() const
+	{
+		return TextFormatArchive.Get();
+	}
+
+	void SetTextFormatArchive(TUniquePtr<FArchive> InTextArchive)
+	{
+		TextFormatArchive = MoveTemp(InTextArchive);
+	}
+
+	FArchiveFormatterType* GetFormatter() const
+	{
+		return Formatter.Get();
+	}
+
+	void SetFormatter(TUniquePtr<FArchiveFormatterType> InFormatter)
+	{
+		Formatter = MoveTemp(InFormatter);
+	}
+
+	FStructuredArchive* GetStructuredArchive() const
+	{
+		return StructuredArchive.Get();
+	}
+
+	void SetStructuredArchive(TUniquePtr<FStructuredArchive> InArchive)
+	{
+		StructuredArchive = MoveTemp(InArchive);
+	}
+
+	const TOptional<FString>& GetTempFilename() const
+	{
+		return TempFilename;
+	}
+
+	void SetTempFilename(TOptional<FString> InTemp)
+	{
+		TempFilename = MoveTemp(InTemp);
+	}
+
+	const TOptional<FString>& GetTextFormatTempFilename() const
+	{
+		return TextFormatTempFilename;
+	}
+
+	void SetTextFormatTempFilename(TOptional<FString> InTemp)
+	{
+		TextFormatTempFilename = MoveTemp(InTemp);
+	}
+
 private:
+	friend class FSaveContext;
+
+	/** Linker associated with this realm. */
+	TPimplPtr<FLinkerSave> Linker;
+
+	/** Archives associated with this linker and realm. */
+	TUniquePtr<FArchive> TextFormatArchive;
+	TUniquePtr<FArchiveFormatterType> Formatter;
+	TUniquePtr<FStructuredArchive> StructuredArchive;
+
+	/** Temp Filename for the archive. */
+	TOptional<FString> TempFilename;
+	TOptional<FString> TextFormatTempFilename;
+
 	// Set of objects excluded (import or exports) through marks or otherwise (i.e. transient flags, etc)
 	TSet<UObject*> Excluded;
 	// Set of objects marked as export
@@ -320,17 +422,6 @@ public:
 
 	~FSaveContext()
 	{
-		CloseLinkerArchives();
-
-		if (TempFilename.IsSet())
-		{
-			IFileManager::Get().Delete(*TempFilename.GetValue());
-		}
-		if (TextFormatTempFilename.IsSet())
-		{
-			IFileManager::Get().Delete(*TextFormatTempFilename.GetValue());
-		}
-
 		if (bNeedPreSaveCleanup && Asset)
 		{
 			UE::SavePackageUtilities::CallPostSaveRoot(Asset, ObjectSaveContext, bNeedPreSaveCleanup);
@@ -713,22 +804,82 @@ public:
 		CustomVersions = MoveTemp(InCustomVersions);
 	}
 
+	TArray<FLinkerSave*> GetLinkers() const
+	{
+		TArray<FLinkerSave*> Linkers;
+		for (const FHarvestedRealm& Realm : HarvestedRealms)
+		{
+			if (FLinkerSave* Linker = Realm.GetLinker())
+			{
+				Linkers.Add(Realm.GetLinker());
+			}
+		}
+		return Linkers;
+	}
+
 	FLinkerSave* GetLinker() const
 	{
-		return Linker.Get();
+		return GetHarvestedRealm().GetLinker();
+	}
+
+	void SetLinker(TPimplPtr<FLinkerSave> InLinker)
+	{
+		GetHarvestedRealm().SetLinker(MoveTemp(InLinker));
 	}
 
 	bool CloseLinkerArchives()
 	{
-		bool bSuccess = true;
-		if (Linker)
-		{
-			bSuccess = Linker->CloseAndDestroySaver();
-		}
-		StructuredArchive.Reset();
-		Formatter.Reset();
-		TextFormatArchive.Reset();
-		return bSuccess;
+		return GetHarvestedRealm().CloseLinkerArchives();
+	}
+
+	FArchive* GetTextFormatArchive() const
+	{
+		return GetHarvestedRealm().GetTextFormatArchive();
+	}
+
+	void SetTextFormatArchive(TUniquePtr<FArchive> InTextArchive)
+	{
+		GetHarvestedRealm().SetTextFormatArchive(MoveTemp(InTextArchive));
+	}
+
+	FArchiveFormatterType* GetFormatter() const
+	{
+		return GetHarvestedRealm().GetFormatter();
+	}
+
+	void SetFormatter(TUniquePtr<FArchiveFormatterType> InFormatter)
+	{
+		GetHarvestedRealm().SetFormatter(MoveTemp(InFormatter));
+	}
+
+	FStructuredArchive* GetStructuredArchive() const
+	{
+		return GetHarvestedRealm().GetStructuredArchive();
+	}
+
+	void SetStructuredArchive(TUniquePtr<FStructuredArchive> InArchive)
+	{
+		GetHarvestedRealm().SetStructuredArchive(MoveTemp(InArchive));
+	}
+
+	const TOptional<FString>& GetTempFilename() const
+	{
+		return GetHarvestedRealm().GetTempFilename();
+	}
+
+	void SetTempFilename(TOptional<FString> InTemp)
+	{
+		GetHarvestedRealm().SetTempFilename(MoveTemp(InTemp));
+	}
+
+	const TOptional<FString>& GetTextFormatTempFilename() const
+	{
+		return GetHarvestedRealm().GetTextFormatTempFilename();
+	}
+
+	void SetTextFormatTempFilename(TOptional<FString> InTemp)
+	{
+		GetHarvestedRealm().SetTextFormatTempFilename(MoveTemp(InTemp));
 	}
 
 	FSavePackageResultStruct GetFinalResult()
@@ -748,7 +899,7 @@ public:
 		ESavePackageResult FinalResult = IsStubRequested() ? ESavePackageResult::GenerateStub : ESavePackageResult::Success;
 		return FSavePackageResultStruct(FinalResult, TotalPackageSizeUncompressed,
 			AsyncWriteAndHashSequence.Finalize(EAsyncExecution::TaskGraph, MoveTemp(HashCompletionFunc)),
-			SerializedPackageFlags, IsCompareLinker() ? MoveTemp(Linker) : nullptr);
+			SerializedPackageFlags, IsCompareLinker() ? MoveTemp(GetHarvestedRealm().Linker) : nullptr);
 	}
 
 	FObjectSaveContextData& GetObjectSaveContext()
@@ -778,16 +929,9 @@ public:
 public:
 	ESavePackageResult Result;
 
-	TPimplPtr<FLinkerSave> Linker;
-	TUniquePtr<FArchive> TextFormatArchive;
-	TUniquePtr<FArchiveFormatterType> Formatter;
-	TUniquePtr<FStructuredArchive> StructuredArchive;
-
-	TOptional<FString> TempFilename;
-	TOptional<FString> TextFormatTempFilename;
-
 	EPropertyLocalizationGathererResultFlags GatherableTextResultFlags = EPropertyLocalizationGathererResultFlags::Empty;
 
+	//@note FH: Most of these public members should be moved to harvested realm class 
 	int64 PackageHeaderAndExportSize = 0;
 	int64 TotalPackageSizeUncompressed = 0;
 	int32 OffsetAfterPackageFileSummary = 0;
diff --git a/Engine/Source/Runtime/CoreUObject/Private/UObject/SavePackage/SavePackageUtilities.cpp b/Engine/Source/Runtime/CoreUObject/Private/UObject/SavePackage/SavePackageUtilities.cpp
--- a/Engine/Source/Runtime/CoreUObject/Private/UObject/SavePackage/SavePackageUtilities.cpp
+++ b/Engine/Source/Runtime/CoreUObject/Private/UObject/SavePackage/SavePackageUtilities.cpp
@@ -12,6 +12,7 @@
 #include "Misc/ConfigCacheIni.h"
 #include "Misc/CommandLine.h"
 #include "Misc/Paths.h"
+#include "Misc/PathViews.h"
 #include "Misc/ScopedSlowTask.h"
 #include "SaveContext.h"
 #include "Serialization/BulkData.h"
@@ -1990,7 +1991,7 @@ ESavePackageResult CreatePayloadSidecarFile(FLinkerSave& Linker, const FPackageP
 	if (PackageWriter)
 	{
 		IPackageWriter::FAdditionalFileInfo SidecarSegmentInfo;
-		SidecarSegmentInfo.PackageName = PackagePath.GetPackageFName();
+		SidecarSegmentInfo.OutputPackageName = SidecarSegmentInfo.InputPackageName = PackagePath.GetPackageFName();
 		SidecarSegmentInfo.Filename = TargetFilePath;
 		FIoBuffer FileData(FIoBuffer::AssumeOwnership, Ar.ReleaseOwnership(), DataSize);
 		PackageWriter->WriteAdditionalFile(SidecarSegmentInfo, FileData);
@@ -2017,7 +2018,7 @@ ESavePackageResult CreatePayloadSidecarFile(FLinkerSave& Linker, const FPackageP
 
 ESavePackageResult SaveBulkData(FLinkerSave* Linker, int64& InOutStartOffset, const UPackage* InOuter, const TCHAR* Filename, const ITargetPlatform* TargetPlatform,
 				  FSavePackageContext* SavePackageContext, uint32 SaveFlags, const bool bTextFormat, const bool bComputeHash,
-				  TAsyncWorkSequence<FMD5>& AsyncWriteAndHashSequence, int64& TotalPackageSizeUncompressed)
+				  TAsyncWorkSequence<FMD5>& AsyncWriteAndHashSequence, int64& TotalPackageSizeUncompressed, bool bIsOptionalRealm)
 {
 	// Now we write all the bulkdata that is supposed to be at the end of the package
 	// and fix up the offset
@@ -2307,35 +2308,52 @@ ESavePackageResult SaveBulkData(FLinkerSave* Linker, int64& InOutStartOffset, co
 			};
 
 			IPackageWriter::FBulkDataInfo BulkInfo;
-			BulkInfo.PackageName = InOuter->GetFName();
-
-			FPackageId PackageId = FPackageId::FromName(BulkInfo.PackageName);
+			BulkInfo.InputPackageName = InOuter->GetFName();
+			// Adjust the OutputPackageName and LooseFilePath if needed
+			if (bIsOptionalRealm)
+			{
+				// Optional output have the form PackagePath.o.ext
+				FString OutputPackageName = BulkInfo.InputPackageName.ToString() + FPackagePath::GetOptionalSegmentExtensionModifier();
+				BulkInfo.OutputPackageName = *OutputPackageName;
+				BulkInfo.LooseFilePath = FPathViews::ChangeExtension(Filename, TEXT("o.") + FPaths::GetExtension(Filename));
+				BulkInfo.MultiOutputIndex = 1;
+			}
+			else
+			{
+				BulkInfo.OutputPackageName = BulkInfo.InputPackageName;
+				BulkInfo.LooseFilePath = Filename;
+			}
+			FPackageId PackageId = FPackageId::FromName(BulkInfo.OutputPackageName);
 				
 			if (BulkArchive->TotalSize())
 			{
 				BulkInfo.ChunkId = CreateIoChunkId(PackageId.Value(), 0, EIoChunkType::BulkData);
 				BulkInfo.BulkDataType = bSeparateSegmentsEnabled ?
 					IPackageWriter::FBulkDataInfo::BulkSegment : IPackageWriter::FBulkDataInfo::AppendToExports;
-				BulkInfo.LooseFilePath = FPaths::ChangeExtension(Filename, LexToString(EPackageExtension::BulkDataDefault));
+				BulkInfo.LooseFilePath = FPathViews::ChangeExtension(BulkInfo.LooseFilePath, LexToString(EPackageExtension::BulkDataDefault));
+					
 				PackageWriter->WriteBulkData(BulkInfo, AddSizeAndConvertToIoBuffer(BulkArchive.Get()), BulkArchive->FileRegions);
 			}
 			if (OptionalBulkArchive && OptionalBulkArchive->TotalSize())
 			{
+				checkf(!bIsOptionalRealm, TEXT("OptionalBulkData is currently unsupported with optional package multi output"));
 				BulkInfo.ChunkId = CreateIoChunkId(PackageId.Value(), 0, EIoChunkType::OptionalBulkData);
 				BulkInfo.BulkDataType = IPackageWriter::FBulkDataInfo::Optional;
-				BulkInfo.LooseFilePath = FPaths::ChangeExtension(Filename, LexToString(EPackageExtension::BulkDataOptional));
+				BulkInfo.LooseFilePath = FPathViews::ChangeExtension(BulkInfo.LooseFilePath, LexToString(EPackageExtension::BulkDataOptional));
 				PackageWriter->WriteBulkData(BulkInfo, AddSizeAndConvertToIoBuffer(OptionalBulkArchive.Get()), OptionalBulkArchive->FileRegions);
 			}
 			if (MappedBulkArchive && MappedBulkArchive->TotalSize())
 			{
+				checkf(!bIsOptionalRealm, TEXT("MemoryMappedBulkData is currently unsupported with optional package multi output"));
 				BulkInfo.ChunkId = CreateIoChunkId(PackageId.Value(), 0, EIoChunkType::MemoryMappedBulkData);
 				BulkInfo.BulkDataType = IPackageWriter::FBulkDataInfo::Mmap;
-				BulkInfo.LooseFilePath = FPaths::ChangeExtension(Filename, LexToString(EPackageExtension::BulkDataMemoryMapped));
+				BulkInfo.LooseFilePath = FPathViews::ChangeExtension(BulkInfo.LooseFilePath, LexToString(EPackageExtension::BulkDataMemoryMapped));
 				PackageWriter->WriteBulkData(BulkInfo, AddSizeAndConvertToIoBuffer(MappedBulkArchive.Get()), MappedBulkArchive->FileRegions);
 			}
 		}
 		else
 		{
+			checkf(!bIsOptionalRealm, TEXT("Package optional package multi output is unsupported without a PackageWriter"));
 			auto WriteBulkData = [&](FLargeMemoryWriterWithRegions* Archive, const TCHAR* BulkFileExtension)
 			{
 				if (const int64 DataSize = Archive ? Archive->TotalSize() : 0)
diff --git a/Engine/Source/Runtime/CoreUObject/Private/UObject/SavePackage/SavePackageUtilities.h b/Engine/Source/Runtime/CoreUObject/Private/UObject/SavePackage/SavePackageUtilities.h
--- a/Engine/Source/Runtime/CoreUObject/Private/UObject/SavePackage/SavePackageUtilities.h
+++ b/Engine/Source/Runtime/CoreUObject/Private/UObject/SavePackage/SavePackageUtilities.h
@@ -379,7 +379,7 @@ namespace SavePackageUtilities
 	ESavePackageResult SaveBulkData(FLinkerSave* Linker, int64& InOutStartOffset, const UPackage* InOuter,
 		const TCHAR* Filename, const ITargetPlatform* TargetPlatform, FSavePackageContext* SavePackageContext,
 		uint32 SaveFlags, const bool bTextFormat, const bool bComputeHash,
-		TAsyncWorkSequence<FMD5>& AsyncWriteAndHashSequence, int64& TotalPackageSizeUncompressed);
+		TAsyncWorkSequence<FMD5>& AsyncWriteAndHashSequence, int64& TotalPackageSizeUncompressed, bool bIsOptionalRealm = false);
 	
 	/**
 	 * Used to append additional data to the end of the package file by invoking callbacks stored in the linker.
diff --git a/Engine/Source/Runtime/CoreUObject/Private/UObject/SavePackage2.cpp b/Engine/Source/Runtime/CoreUObject/Private/UObject/SavePackage2.cpp
--- a/Engine/Source/Runtime/CoreUObject/Private/UObject/SavePackage2.cpp
+++ b/Engine/Source/Runtime/CoreUObject/Private/UObject/SavePackage2.cpp
@@ -16,6 +16,7 @@
 #include "Misc/FeedbackContext.h"
 #include "Misc/MessageDialog.h"
 #include "Misc/PackageName.h"
+#include "Misc/PathViews.h"
 #include "Misc/RedirectCollector.h"
 #include "Misc/ScopedSlowTask.h"
 #include "Misc/SecureHash.h"
@@ -340,10 +341,12 @@ ESavePackageResult HarvestPackage(FSaveContext& SaveContext)
 
 	// If we have a valid optional context and we are saving it,
 	// transform any harvested non optional export into imports
+	// Mark other optional import package as well
 	if (SaveContext.IsSaveOptional() &&
 		SaveContext.IsCooking() &&
 		SaveContext.GetHarvestedRealm(ESaveRealm::Optional).GetExports().Num())
 	{
+		bool bHasNonOptionalSelfReference = false;
 		FHarvestedRealm& OptionalContext = SaveContext.GetHarvestedRealm(ESaveRealm::Optional);
 		for (auto It = OptionalContext.GetExports().CreateIterator(); It; ++It)
 		{
@@ -356,6 +359,8 @@ ESavePackageResult HarvestPackage(FSaveContext& SaveContext)
 					GameExport->bGeneratePublicHash = true;
 					// Transform the export as an import
 					OptionalContext.AddImport(It->Obj);
+					// Flag the package itself to be an import
+					bHasNonOptionalSelfReference = true;
 				}
 				// if not found in the game context, record an illegal reference
 				else
@@ -365,8 +370,11 @@ ESavePackageResult HarvestPackage(FSaveContext& SaveContext)
 				It.RemoveCurrent();
 			}
 		}
-		// Also add the current package itself as an import, if the optional context isn't empty
-		OptionalContext.AddImport(SaveContext.GetPackage());
+		// Also add the current package itself as an import if we are referencing any non optional export
+		if (bHasNonOptionalSelfReference)
+		{
+			OptionalContext.AddImport(SaveContext.GetPackage());
+		}
 	}
 
 	// Trim PrestreamPackage list 
@@ -457,12 +465,18 @@ ESavePackageResult ValidateExports(FSaveContext& SaveContext)
 {
 	SCOPED_SAVETIMER(UPackage_Save_ValidateExports);
 
-	//@todo: add validation on optional context
+	/// Export Validation for optional realm
 	if (SaveContext.GetCurrentHarvestingRealm() == ESaveRealm::Optional)
 	{
+		// return empty realm to skip processing if no export are found
+		if (SaveContext.GetExports().Num() == 0)
+		{
+			return ESavePackageResult::EmptyRealm;
+		}
 		return ReturnSuccessOrCancel();
 	}
 
+	/// Export validation for game/editor realm
 	// Check if we gathered any exports
 	if (SaveContext.GetExports().Num() == 0)
 	{
@@ -788,29 +802,29 @@ ESavePackageResult CreateLinker(FSaveContext& SaveContext)
 				PackageWriter->CreateLinkerArchive(SaveContext.GetPackage()->GetFName(), SaveContext.GetAsset()) :
 				// The LargeMemoryWriter does not need to be persistent; the LinkerSave wraps it and reports Persistent=true
 				TUniquePtr<FLargeMemoryWriter>(new FLargeMemoryWriter(0, false /* bIsPersistent */, *SaveContext.GetPackage()->GetName()));
-			SaveContext.Linker = MakePimpl<FLinkerSave>(SaveContext.GetPackage(), ExportsArchive.Release(),
-				SaveContext.IsForceByteSwapping(), SaveContext.IsSaveUnversionedNative());
+			SaveContext.SetLinker(MakePimpl<FLinkerSave>(SaveContext.GetPackage(), ExportsArchive.Release(),
+				SaveContext.IsForceByteSwapping(), SaveContext.IsSaveUnversionedNative()));
 		}
 		else
 		{
 			// Allocate the linker with a tempfile, forcing byte swapping if wanted.
-			SaveContext.TempFilename = FPaths::CreateTempFilename(*FPaths::ProjectSavedDir(), *BaseFilename.Left(32));
-			SaveContext.Linker = MakePimpl<FLinkerSave>(SaveContext.GetPackage(), *SaveContext.TempFilename.GetValue(), SaveContext.IsForceByteSwapping(), SaveContext.IsSaveUnversionedNative());
+			SaveContext.SetTempFilename(FPaths::CreateTempFilename(*FPaths::ProjectSavedDir(), *BaseFilename.Left(32)));
+			SaveContext.SetLinker(MakePimpl<FLinkerSave>(SaveContext.GetPackage(), *SaveContext.GetTempFilename().GetValue(), SaveContext.IsForceByteSwapping(), SaveContext.IsSaveUnversionedNative()));
 		}
 
 		if (SaveContext.IsGenerateSaveError())
 		{
-			SaveContext.Linker->SetOutputDevice(SaveContext.GetError());
+			SaveContext.GetLinker()->SetOutputDevice(SaveContext.GetError());
 		}
-		SaveContext.Linker->bUpdatingLoadedPath = SaveContext.IsUpdatingLoadedPath();
-		SaveContext.Linker->bProceduralSave = SaveContext.IsProceduralSave();
+		SaveContext.GetLinker()->bUpdatingLoadedPath = SaveContext.IsUpdatingLoadedPath();
+		SaveContext.GetLinker()->bProceduralSave = SaveContext.IsProceduralSave();
 
 		if (UE::FPackageTrailer::IsEnabled())
 		{
 			// The package trailer is not supported for text based assets yet
 			if (!SaveContext.IsTextFormat() && !SaveContext.IsProceduralSave())
 			{
-				SaveContext.Linker->PackageTrailerBuilder = MakeUnique<UE::FPackageTrailerBuilder>(SaveContext.GetPackage()->GetFName());
+				SaveContext.GetLinker()->PackageTrailerBuilder = MakeUnique<UE::FPackageTrailerBuilder>(SaveContext.GetPackage()->GetFName());
 			}
 			else if ((SaveContext.GetSaveArgs().SaveFlags & SAVE_BulkDataByReference) != 0)
 			{
@@ -818,7 +832,7 @@ ESavePackageResult CreateLinker(FSaveContext& SaveContext)
 				{
 					if (const UE::FPackageTrailer* Trailer = LinkerLoad->GetPackageTrailer())
 					{
-						SaveContext.Linker->PackageTrailerBuilder = UE::FPackageTrailerBuilder::CreateReferenceToTrailer(*Trailer, SaveContext.GetPackage()->GetFName());
+						SaveContext.GetLinker()->PackageTrailerBuilder = UE::FPackageTrailerBuilder::CreateReferenceToTrailer(*Trailer, SaveContext.GetPackage()->GetFName());
 					}
 				}
 			}
@@ -827,27 +841,27 @@ ESavePackageResult CreateLinker(FSaveContext& SaveContext)
 #if WITH_TEXT_ARCHIVE_SUPPORT
 		if (SaveContext.IsTextFormat())
 		{
-			if (SaveContext.TempFilename.IsSet())
+			if (SaveContext.GetTempFilename().IsSet())
 			{
-				SaveContext.TextFormatTempFilename = SaveContext.TempFilename.GetValue() + FPackageName::GetTextAssetPackageExtension();
+				SaveContext.SetTextFormatTempFilename(SaveContext.GetTempFilename().GetValue() + FPackageName::GetTextAssetPackageExtension());
 			}
 			else
 			{
-				SaveContext.TextFormatTempFilename = FPaths::CreateTempFilename(*FPaths::ProjectSavedDir(), *BaseFilename.Left(32)) + FPackageName::GetTextAssetPackageExtension();
+				SaveContext.SetTextFormatTempFilename(FPaths::CreateTempFilename(*FPaths::ProjectSavedDir(), *BaseFilename.Left(32)) + FPackageName::GetTextAssetPackageExtension());
 			}
-			SaveContext.TextFormatArchive = TUniquePtr<FArchive>(IFileManager::Get().CreateFileWriter(*SaveContext.TextFormatTempFilename.GetValue()));
-			TUniquePtr<FJsonArchiveOutputFormatter> OutputFormatter = MakeUnique<FJsonArchiveOutputFormatter>(*SaveContext.TextFormatArchive);
-			OutputFormatter->SetObjectIndicesMap(&SaveContext.Linker->ObjectIndicesMap);
-			SaveContext.Formatter = MoveTemp(OutputFormatter);
+			SaveContext.SetTextFormatArchive(TUniquePtr<FArchive>(IFileManager::Get().CreateFileWriter(*SaveContext.GetTextFormatTempFilename().GetValue())));
+			TUniquePtr<FJsonArchiveOutputFormatter> OutputFormatter = MakeUnique<FJsonArchiveOutputFormatter>(*SaveContext.GetTextFormatArchive());
+			OutputFormatter->SetObjectIndicesMap(&SaveContext.GetLinker()->ObjectIndicesMap);
+			SaveContext.SetFormatter(MoveTemp(OutputFormatter));
 		}
 		else
 #endif
 		{
-			SaveContext.Formatter = MakeUnique<FBinaryArchiveFormatter>(*(FArchive*)SaveContext.Linker.Get());
+			SaveContext.SetFormatter(MakeUnique<FBinaryArchiveFormatter>(*(FArchive*)SaveContext.GetLinker()));
 		}
 	}
 
-	SaveContext.StructuredArchive = MakeUnique<FStructuredArchive>(*SaveContext.Formatter);
+	SaveContext.SetStructuredArchive(MakeUnique<FStructuredArchive>(*SaveContext.GetFormatter()));
 	return ReturnSuccessOrCancel();
 }
 
@@ -889,36 +903,36 @@ ESavePackageResult BuildLinker(FSaveContext& SaveContext)
 	// Setup Linker 
 	{
 		// Use the custom versions we harvested from the dependency harvesting pass
-		SaveContext.Linker->Summary.SetCustomVersionContainer(SaveContext.GetCustomVersions());
+		SaveContext.GetLinker()->Summary.SetCustomVersionContainer(SaveContext.GetCustomVersions());
 
-		SaveContext.Linker->SetPortFlags(SaveContext.GetPortFlags());
-		SaveContext.Linker->SetFilterEditorOnly(SaveContext.IsFilterEditorOnly());
-		SaveContext.Linker->SetCookingTarget(SaveContext.GetTargetPlatform());
+		SaveContext.GetLinker()->SetPortFlags(SaveContext.GetPortFlags());
+		SaveContext.GetLinker()->SetFilterEditorOnly(SaveContext.IsFilterEditorOnly());
+		SaveContext.GetLinker()->SetCookingTarget(SaveContext.GetTargetPlatform());
 
 		bool bUseUnversionedProperties = SaveContext.IsSaveUnversionedProperties();
-		SaveContext.Linker->SetUseUnversionedPropertySerialization(bUseUnversionedProperties);
+		SaveContext.GetLinker()->SetUseUnversionedPropertySerialization(bUseUnversionedProperties);
 
 #if WITH_EDITOR
 		if (SaveContext.IsCooking())
 		{
-			SaveContext.Linker->SetDebugSerializationFlags(DSF_EnableCookerWarnings | SaveContext.Linker->GetDebugSerializationFlags());
+			SaveContext.GetLinker()->SetDebugSerializationFlags(DSF_EnableCookerWarnings | SaveContext.GetLinker()->GetDebugSerializationFlags());
 		}
 #endif
 		// Make sure the package has the same version as the linker
-		SaveContext.GetPackage()->LinkerPackageVersion = SaveContext.Linker->UEVer();
-		SaveContext.GetPackage()->LinkerLicenseeVersion = SaveContext.Linker->LicenseeUEVer();
-		SaveContext.GetPackage()->LinkerCustomVersion = SaveContext.Linker->GetCustomVersions();
+		SaveContext.GetPackage()->LinkerPackageVersion = SaveContext.GetLinker()->UEVer();
+		SaveContext.GetPackage()->LinkerLicenseeVersion = SaveContext.GetLinker()->LicenseeUEVer();
+		SaveContext.GetPackage()->LinkerCustomVersion = SaveContext.GetLinker()->GetCustomVersions();
 	}
-	
+
 	PRAGMA_DISABLE_DEPRECATION_WARNINGS
-	SaveContext.Linker->Summary.Guid = SaveContext.IsKeepGuid() ? SaveContext.GetPackage()->GetGuid() : SaveContext.GetPackage()->MakeNewGuid();
+	SaveContext.GetLinker()->Summary.Guid = SaveContext.IsKeepGuid() ? SaveContext.GetPackage()->GetGuid() : SaveContext.GetPackage()->MakeNewGuid();
 	PRAGMA_ENABLE_DEPRECATION_WARNINGS
 #if WITH_EDITORONLY_DATA
-	SaveContext.Linker->Summary.PersistentGuid = SaveContext.GetPackage()->GetPersistentGuid();
+	SaveContext.GetLinker()->Summary.PersistentGuid = SaveContext.GetPackage()->GetPersistentGuid();
 #endif
-	SaveContext.Linker->Summary.Generations = TArray<FGenerationInfo>{ FGenerationInfo(0, 0) };
+	SaveContext.GetLinker()->Summary.Generations = TArray<FGenerationInfo>{ FGenerationInfo(0, 0) };
 
-	FLinkerSave* Linker = SaveContext.Linker.Get();
+	FLinkerSave* Linker = SaveContext.GetLinker();
 
 	// Build Name Map
 	{
@@ -943,8 +957,8 @@ ESavePackageResult BuildLinker(FSaveContext& SaveContext)
 
 		Sort(Linker->NameMap.GetData(), Linker->Summary.NamesReferencedFromExportDataCount, FNameEntryIdSortHelper());
 		Sort(Linker->NameMap.GetData() + Linker->Summary.NamesReferencedFromExportDataCount,
-			 Linker->Summary.NameCount - Linker->Summary.NamesReferencedFromExportDataCount,
-			 FNameEntryIdSortHelper());
+			Linker->Summary.NameCount - Linker->Summary.NamesReferencedFromExportDataCount,
+			FNameEntryIdSortHelper());
 
 		if (!SaveContext.IsTextFormat())
 		{
@@ -968,16 +982,23 @@ ESavePackageResult BuildLinker(FSaveContext& SaveContext)
 			FPropertyLocalizationDataGatherer(Linker->GatherableTextDataMap, SaveContext.GetPackage(), SaveContext.GatherableTextResultFlags);
 		}
 	}
-	
+
 	// Build ImportMap
 	{
 		SCOPED_SAVETIMER(UPackage_Save_BuildImportMap);
+
 		for (UObject* Import : SaveContext.GetImports())
 		{
 			UClass* ImportClass = Import->GetClass();
 			FName ReplacedName = NAME_None;
 			FObjectImport& ObjectImport = Linker->ImportMap.Add_GetRef(FObjectImport(Import, ImportClass));
 
+			// flag the import as optional
+			if (Import->GetClass()->HasAnyClassFlags(CLASS_Optional))
+			{
+				ObjectImport.bImportOptional = true;
+			}
+
 			// If the package import is a prestream package, mark it as such by hacking its class name
 			if (SaveContext.IsPrestreamPackage(Cast<UPackage>(Import)))
 			{
@@ -1007,13 +1028,7 @@ ESavePackageResult BuildLinker(FSaveContext& SaveContext)
 		for (const FTaggedExport& TaggedExport : SaveContext.GetExports())
 		{
 			FObjectExport& Export = Linker->ExportMap.Add_GetRef(FObjectExport(TaggedExport.Obj, TaggedExport.bNotAlwaysLoadedForEditorGame));
-			// @note: Artificially flag export as public if needed
-			// This is temporary for testing the optional realm cooked content with the iostore
-			// It should be changed to a separate member on `FObjectExport`
-			if (TaggedExport.bGeneratePublicHash)
-			{
-				Export.ObjectFlags |= RF_Public;
-			}
+			Export.bGeneratePublicHash = TaggedExport.bGeneratePublicHash;
 
 			if (UPackage* Package = Cast<UPackage>(Export.Object))
 			{
@@ -1188,7 +1203,7 @@ ESavePackageResult BuildLinker(FSaveContext& SaveContext)
 
 void SavePreloadDependencies(FStructuredArchive::FRecord& StructuredArchiveRoot, FSaveContext& SaveContext)
 {
-	FLinkerSave* Linker = SaveContext.Linker.Get();
+	FLinkerSave* Linker = SaveContext.GetLinker();
 	auto IncludeObjectAsDependency = [Linker, &SaveContext](int32 CallSite, TSet<FPackageIndex>& AddTo, UObject* ToTest, UObject* ForObj, bool bMandatory, bool bOnlyIfInLinkerTable)
 	{
 		// Skip transient, editor only, and excluded client/server objects
@@ -1488,7 +1503,7 @@ void WriteGatherableText(FStructuredArchive::FRecord& StructuredArchiveRoot, FSa
 
 ESavePackageResult WritePackageHeader(FStructuredArchive::FRecord& StructuredArchiveRoot, FSaveContext& SaveContext)
 {
-	FLinkerSave* Linker = SaveContext.Linker.Get();
+	FLinkerSave* Linker = SaveContext.GetLinker();
 #if WITH_EDITOR
 	FArchiveStackTraceIgnoreScope IgnoreDiffScope(SaveContext.IsIgnoringHeaderDiff());
 #endif
@@ -1644,7 +1659,7 @@ ESavePackageResult WriteExports(FStructuredArchive::FRecord& StructuredArchiveRo
 {
 	//COOK_STAT(FScopedDurationTimer SaveTimer(FSavePackageStats::SerializeExportsTimeSec));
 	SCOPED_SAVETIMER(UPackage_Save_SaveExports);
-	FLinkerSave* Linker = SaveContext.Linker.Get();
+	FLinkerSave* Linker = SaveContext.GetLinker();
 	FScopedSlowTask SlowTask(Linker->ExportMap.Num(), FText(), SaveContext.IsUsingSlowTask());
 
 	FStructuredArchive::FRecord ExportsRecord = StructuredArchiveRoot.EnterRecord(SA_FIELD_NAME(TEXT("Exports")));
@@ -1734,27 +1749,27 @@ ESavePackageResult WriteExports(FStructuredArchive::FRecord& StructuredArchiveRo
 
 [[nodiscard]] ESavePackageResult BuildAndWriteTrailer(FStructuredArchive::FRecord& StructuredArchiveRoot, FSaveContext& SaveContext)
 {
-	SaveContext.Linker->Summary.PayloadTocOffset = INDEX_NONE;
+	SaveContext.GetLinker()->Summary.PayloadTocOffset = INDEX_NONE;
 
-	if (SaveContext.Linker->PackageTrailerBuilder.IsValid())
+	if (SaveContext.GetLinker()->PackageTrailerBuilder.IsValid())
 	{
 		// At the moment we assume that we cannot have reference payloads in the trailer if SAVE_BulkDataByReference is not set and we
 		// cannot have locally stored payloads if SAVE_BulkDataByReference is set.
-		checkf((SaveContext.GetSaveArgs().SaveFlags & SAVE_BulkDataByReference) != 0 || SaveContext.Linker->PackageTrailerBuilder->GetNumReferencedPayloads() == 0,
+		checkf((SaveContext.GetSaveArgs().SaveFlags & SAVE_BulkDataByReference) != 0 || SaveContext.GetLinker()->PackageTrailerBuilder->GetNumReferencedPayloads() == 0,
 			TEXT("Attempting to build a package trailer with referenced payloads but the SAVE_BulkDataByReference flag is not set. '%s'"), *SaveContext.GetPackage()->GetName());
 
-		checkf(	(SaveContext.GetSaveArgs().SaveFlags & SAVE_BulkDataByReference) == 0 || SaveContext.Linker->PackageTrailerBuilder->GetNumLocalPayloads() == 0, 
+		checkf(	(SaveContext.GetSaveArgs().SaveFlags & SAVE_BulkDataByReference) == 0 || SaveContext.GetLinker()->PackageTrailerBuilder->GetNumLocalPayloads() == 0,
 				TEXT("Attempting to build a package trailer with local payloads but the SAVE_BulkDataByReference flag is set. '%s'"), *SaveContext.GetPackage()->GetName());
 
 		checkf(SaveContext.IsTextFormat() == false, TEXT("Attempting to build a package trailer for text based asset '%s', this is not supported!"), *SaveContext.GetPackage()->GetName());
 
-		SaveContext.Linker->Summary.PayloadTocOffset = SaveContext.Linker->Tell();
-		if (!SaveContext.Linker->PackageTrailerBuilder->BuildAndAppendTrailer(SaveContext.Linker.Get(), *SaveContext.Linker.Get()))
+		SaveContext.GetLinker()->Summary.PayloadTocOffset = SaveContext.GetLinker()->Tell();
+		if (!SaveContext.GetLinker()->PackageTrailerBuilder->BuildAndAppendTrailer(SaveContext.GetLinker(), *SaveContext.GetLinker()))
 		{
 			return ESavePackageResult::Error;
 		}
 
-		SaveContext.Linker->PackageTrailerBuilder.Reset();
+		SaveContext.GetLinker()->PackageTrailerBuilder.Reset();
 	}
 
 	return ESavePackageResult::Success;
@@ -1766,6 +1781,7 @@ ESavePackageResult WriteAdditionalExportFiles(FSaveContext& SaveContext)
 
 	if (SaveContext.IsCooking() && SaveContext.AdditionalFilesFromExports.Num() > 0)
 	{
+		checkf(SaveContext.GetCurrentHarvestingRealm() != ESaveRealm::Optional , TEXT("Addtional export files is currently unsupported with optional package multi output"));
 		IPackageWriter* PackageWriter = SavePackageContext ? SavePackageContext->PackageWriter : nullptr;
 		checkf(PackageWriter, TEXT("Cooking requires a PackageWriter"));
 		const bool bComputeHash = SaveContext.IsComputeHash();
@@ -1775,7 +1791,7 @@ ESavePackageResult WriteAdditionalExportFiles(FSaveContext& SaveContext)
 			SaveContext.TotalPackageSizeUncompressed += Size;
 
 			IPackageWriter::FAdditionalFileInfo FileInfo;
-			FileInfo.PackageName = SaveContext.GetPackage()->GetFName();
+			FileInfo.OutputPackageName = FileInfo.InputPackageName = SaveContext.GetPackage()->GetFName();
 			FileInfo.Filename = *Writer.GetArchiveName();
 
 			FIoBuffer FileData(FIoBuffer::AssumeOwnership, Writer.ReleaseOwnership(), Size);
@@ -1791,7 +1807,7 @@ ESavePackageResult UpdatePackageHeader(FStructuredArchive::FRecord& StructuredAr
 {
 	SCOPED_SAVETIMER(UPackage_Save_UpdatePackageHeader);
 
-	FLinkerSave* Linker = SaveContext.Linker.Get();
+	FLinkerSave* Linker = SaveContext.GetLinker();
 	int64 OffsetBeforeUpdates = Linker->Tell();
 #if WITH_EDITOR
 	FArchiveStackTraceIgnoreScope IgnoreDiffScope(SaveContext.IsIgnoringHeaderDiff());
@@ -1897,6 +1913,7 @@ ESavePackageResult FinalizeFile(FStructuredArchive::FRecord& StructuredArchiveRo
 	IPackageWriter* PackageWriter = SavePackageContext ? SavePackageContext->PackageWriter : nullptr;
 	if (PackageWriter || SaveContext.IsSaveToMemory())
 	{
+		bool bIsOptionalRealm = SaveContext.GetCurrentHarvestingRealm() == ESaveRealm::Optional;
 		FLinkerSave* Linker = SaveContext.GetLinker();
 		UE_LOG(LogSavePackage, Verbose, TEXT("Async saving from memory to '%s'"), *SaveContext.GetFilename());
 		FLargeMemoryWriter* Writer = static_cast<FLargeMemoryWriter*>(Linker->Saver);
@@ -1904,11 +1921,24 @@ ESavePackageResult FinalizeFile(FStructuredArchive::FRecord& StructuredArchiveRo
 		if (PackageWriter)
 		{
 			IPackageWriter::FPackageInfo PackageInfo;
-			PackageInfo.PackageName = SaveContext.GetPackage()->GetFName();
-			PackageInfo.LooseFilePath = SaveContext.GetFilename();
+			PackageInfo.InputPackageName = SaveContext.GetPackage()->GetFName();
+			// Adjust the OutputPackageName and LooseFilePath if needed
+			if (bIsOptionalRealm)
+			{
+				// Optional output have the form PackagePath.o.ext
+				FString OutputPackageName = PackageInfo.InputPackageName.ToString() + FPackagePath::GetOptionalSegmentExtensionModifier();
+				PackageInfo.OutputPackageName = *OutputPackageName;
+				PackageInfo.LooseFilePath = FPathViews::ChangeExtension(SaveContext.GetFilename(), TEXT("o.") + FPaths::GetExtension(SaveContext.GetFilename()));
+				PackageInfo.MultiOutputIndex = 1;
+			}
+			else
+			{
+				PackageInfo.OutputPackageName = PackageInfo.InputPackageName;
+				PackageInfo.LooseFilePath = SaveContext.GetFilename();
+			}
 			PackageInfo.HeaderSize = Linker->Summary.TotalHeaderSize;
 
-			FPackageId PackageId = FPackageId::FromName(PackageInfo.PackageName);
+			FPackageId PackageId = FPackageId::FromName(PackageInfo.OutputPackageName);
 			PackageInfo.ChunkId = CreateIoChunkId(PackageId.Value(), 0, EIoChunkType::ExportBundleData);
 
 			PackageWriter->WritePackageData(PackageInfo, *Writer, Linker->FileRegions);
@@ -1940,29 +1970,29 @@ ESavePackageResult FinalizeFile(FStructuredArchive::FRecord& StructuredArchiveRo
 		if (SaveContext.CloseLinkerArchives() == false)
 		{
 			UE_LOG(LogSavePackage, Error, TEXT("Error writing temp file '%s' for '%s'"),
-				SaveContext.TempFilename.IsSet() ? *SaveContext.TempFilename.GetValue() : TEXT("UNKNOWN"), SaveContext.GetFilename());
+				SaveContext.GetTempFilename().IsSet() ? *SaveContext.GetTempFilename().GetValue() : TEXT("UNKNOWN"), SaveContext.GetFilename());
 			return ESavePackageResult::Error;
 		}
 
 		// Move file to its real destination
-		checkf(SaveContext.TempFilename.IsSet(), TEXT("The package should've been saved to a tmp file first! (%s)"), *SaveContext.GetFilename());
+		checkf(SaveContext.GetTempFilename().IsSet(), TEXT("The package should've been saved to a tmp file first! (%s)"), *SaveContext.GetFilename());
 
 		// When saving in text format we will have two temp files, so we need to manually delete the non-textbased one
 		if (SaveContext.IsTextFormat())
 		{
-			check(SaveContext.TextFormatTempFilename.IsSet());
-			IFileManager::Get().Delete(*SaveContext.TempFilename.GetValue());
-			SaveContext.TempFilename = SaveContext.TextFormatTempFilename;
-			SaveContext.TextFormatTempFilename.Reset();
+			check(SaveContext.GetTextFormatTempFilename().IsSet());
+			IFileManager::Get().Delete(*SaveContext.GetTempFilename().GetValue());
+			SaveContext.SetTempFilename(SaveContext.GetTextFormatTempFilename());
+			SaveContext.SetTextFormatTempFilename(TOptional<FString>());
 		}
 
 		// Add the .uasset file to the list of output files (TODO: Fix the 0 size, it isn't used after this point but needs to be cleaned up)
-		SaveContext.AdditionalPackageFiles.Emplace(SaveContext.GetFilename(), SaveContext.TempFilename.GetValue(), 0);
+		SaveContext.AdditionalPackageFiles.Emplace(SaveContext.GetFilename(), SaveContext.GetTempFilename().GetValue(), 0);
 
 		ESavePackageResult FinalizeResult = SavePackageUtilities::FinalizeTempOutputFiles(SaveContext.GetTargetPackagePath(), SaveContext.AdditionalPackageFiles, SaveContext.IsComputeHash(), 
 			SaveContext.GetFinalTimestamp(), SaveContext.AsyncWriteAndHashSequence);
 		
-		SaveContext.TempFilename.Reset();
+		SaveContext.SetTempFilename(TOptional<FString>());
 
 		if (FinalizeResult != ESavePackageResult::Success)
 		{
@@ -2025,13 +2055,15 @@ void PostSavePackage(FSaveContext& SaveContext)
 	// Package has been saved, so unmark the NewlyCreated flag.
 	Package->ClearPackageFlags(PKG_NewlyCreated);
 
-	if (SaveContext.Linker != nullptr)
+	// Copy and modify the output SerializedPackageFlags from the PackageFlags written into the file
+	SaveContext.SerializedPackageFlags = SaveContext.GetLinker()->Summary.GetPackageFlags();
+	// Currently the PKG_ContainsNoAsset flag is not serialized as part of the summary
+	bool bContainsAsset = false;
+	for (FLinkerSave* Linker : SaveContext.GetLinkers())
 	{
-		// Copy and modify the output SerializedPackageFlags from the PackageFlags written into the file
-		SaveContext.SerializedPackageFlags = SaveContext.Linker->Summary.GetPackageFlags();
-		// Currently the PKG_ContainsNoAsset flag is not serialized as part of the summary
-		bool bContainsAsset = false;
-		for (FObjectExport& Export : SaveContext.Linker->ExportMap)
+		// GetLinkers shouldn't return null linker
+		check(Linker);
+		for (FObjectExport& Export : Linker->ExportMap)
 		{
 			if (Export.bIsAsset)
 			{
@@ -2039,10 +2071,14 @@ void PostSavePackage(FSaveContext& SaveContext)
 				break;
 			}
 		}
-		if (!bContainsAsset)
-		{
-			SaveContext.SerializedPackageFlags |= PKG_ContainsNoAsset;
-		}
+
+		// Call the linker post save callbacks
+		Linker->OnPostSave(SaveContext.GetTargetPackagePath(), FObjectPostSaveContext(SaveContext.GetObjectSaveContext()));
+	}
+
+	if (!bContainsAsset)
+	{
+		SaveContext.SerializedPackageFlags |= PKG_ContainsNoAsset;
 	}
 
 	// Notify the soft reference collector about our harvested soft references during save. 
@@ -2054,11 +2090,6 @@ void PostSavePackage(FSaveContext& SaveContext)
 	}
 #endif
 
-	// Call the linker post save callbacks
-	if (SaveContext.Linker != nullptr)
-	{
-		SaveContext.Linker->OnPostSave(SaveContext.GetTargetPackagePath(), FObjectPostSaveContext(SaveContext.GetObjectSaveContext()));
-	}
 
 	// Send a message that the package was saved
 	PRAGMA_DISABLE_DEPRECATION_WARNINGS;
@@ -2090,6 +2121,8 @@ ESavePackageResult SaveHarvestedRealms(FSaveContext& SaveContext, ESaveRealm Har
 	SaveContext.Result = ValidateExports(SaveContext);
 	if (SaveContext.Result != ESavePackageResult::Success)
 	{
+		// if we are skipping processing due to an empty realm, consider the save succesful since that only used internally for optional realm
+		SaveContext.Result = SaveContext.Result == ESavePackageResult::EmptyRealm ? ESavePackageResult::Success : SaveContext.Result;
 		return SaveContext.Result;
 	}
 
@@ -2117,7 +2150,7 @@ ESavePackageResult SaveHarvestedRealms(FSaveContext& SaveContext, ESaveRealm Har
 		return SaveContext.Result;
 	}
 
-	FStructuredArchive::FRecord StructuredArchiveRoot = SaveContext.StructuredArchive->Open().EnterRecord();
+	FStructuredArchive::FRecord StructuredArchiveRoot = SaveContext.GetStructuredArchive()->Open().EnterRecord();
 	StructuredArchiveRoot.GetUnderlyingArchive().SetSerializeContext(SaveContext.GetSerializeContext());
 
 	// Write Header
@@ -2136,7 +2169,7 @@ ESavePackageResult SaveHarvestedRealms(FSaveContext& SaveContext, ESaveRealm Har
 		// if we want to generate the SHA key, start tracking script writes
 		if (ScriptSHABytes)
 		{
-			SaveContext.Linker->StartScriptSHAGeneration();
+			SaveContext.GetLinker()->StartScriptSHAGeneration();
 		}
 	}
 
@@ -2150,23 +2183,23 @@ ESavePackageResult SaveHarvestedRealms(FSaveContext& SaveContext, ESaveRealm Har
 	// Get SHA Key
 	{
 		// if we want to generate the SHA key, get it out now that the package has finished saving
-		if (ScriptSHABytes && SaveContext.Linker->ContainsCode())
+		if (ScriptSHABytes && SaveContext.GetLinker()->ContainsCode())
 		{
 			// make space for the 20 byte key
 			ScriptSHABytes->Empty(20);
 			ScriptSHABytes->AddUninitialized(20);
 
 			// retrieve it
-			SaveContext.Linker->GetScriptSHAKey(ScriptSHABytes->GetData());
+			SaveContext.GetLinker()->GetScriptSHAKey(ScriptSHABytes->GetData());
 		}
 	}
 
 	IPackageWriter* PackageWriter = SaveContext.GetPackageWriter();
 	if (PackageWriter)
 	{
-		const int64 ExportsSize = SaveContext.Linker->Tell();
+		const int64 ExportsSize = SaveContext.GetLinker()->Tell();
 		SaveContext.Result = WriteAdditionalFiles(SaveContext, SlowTask, ExportsSize);
-		checkf(SaveContext.Linker->Tell() == ExportsSize, TEXT("The writing of additional files is not allowed to append to the LinkerSave when using a PackageWriter."));
+		checkf(SaveContext.GetLinker()->Tell() == ExportsSize, TEXT("The writing of additional files is not allowed to append to the LinkerSave when using a PackageWriter."));
 		if (SaveContext.Result != ESavePackageResult::Success)
 		{
 			return SaveContext.Result;
@@ -2198,7 +2231,7 @@ ESavePackageResult SaveHarvestedRealms(FSaveContext& SaveContext, ESaveRealm Har
 		return SaveContext.Result;
 	}	
 
-	int64 ExportsSize = SaveContext.Linker->Tell();
+	int64 ExportsSize = SaveContext.GetLinker()->Tell();
 	if (PackageWriter)
 	{
 		PackageWriter->AddToExportsSize(ExportsSize);
@@ -2233,7 +2266,7 @@ ESavePackageResult SaveHarvestedRealms(FSaveContext& SaveContext, ESaveRealm Har
 	//COOK_STAT(FSavePackageStats::MBWritten += ((double)SaveContext.TotalPackageSizeUncompressed) / 1024.0 / 1024.0);
 	{
 		SCOPED_SAVETIMER(UPackage_Save_MarkExportLoaded);
-		FLinkerSave* Linker = SaveContext.Linker.Get();
+		FLinkerSave* Linker = SaveContext.GetLinker();
 		// Mark exports and the package as RF_Loaded after they've been serialized
 		// This is to ensue that newly created packages are properly marked as loaded (since they now exist on disk and 
 		// in memory in the exact same state).
@@ -2337,21 +2370,22 @@ ESavePackageResult WriteAdditionalFiles(FSaveContext& SaveContext, FScopedSlowTa
 {
 	// Save Bulk Data
 	SlowTask.EnterProgressFrame();
-	int64 DataStartOffset = LinkerSize >= 0 ? LinkerSize : SaveContext.Linker.Get()->Tell();
-	SaveContext.Result = SavePackageUtilities::SaveBulkData(SaveContext.Linker.Get(), DataStartOffset,
+	int64 DataStartOffset = LinkerSize >= 0 ? LinkerSize : SaveContext.GetLinker()->Tell();
+	SaveContext.Result = SavePackageUtilities::SaveBulkData(SaveContext.GetLinker(), DataStartOffset,
 		SaveContext.GetPackage(), SaveContext.GetFilename(), SaveContext.GetTargetPlatform(),
 		SaveContext.GetSavePackageContext(), SaveContext.GetSaveArgs().SaveFlags, SaveContext.IsTextFormat(),
-		SaveContext.IsComputeHash(), SaveContext.AsyncWriteAndHashSequence, SaveContext.TotalPackageSizeUncompressed);
+		SaveContext.IsComputeHash(), SaveContext.AsyncWriteAndHashSequence, SaveContext.TotalPackageSizeUncompressed,
+		SaveContext.GetCurrentHarvestingRealm() == ESaveRealm::Optional);
 
 	// Add any pending data blobs to the end of the file by invoking the callbacks
-	ESavePackageResult Result = SavePackageUtilities::AppendAdditionalData(*SaveContext.Linker.Get(), DataStartOffset, SaveContext.GetSavePackageContext());
+	ESavePackageResult Result = SavePackageUtilities::AppendAdditionalData(*SaveContext.GetLinker(), DataStartOffset, SaveContext.GetSavePackageContext());
 	if (Result != ESavePackageResult::Success)
 	{
 		return Result;
 	}
 
 	// Create the payload side car file (if needed)
-	Result = SavePackageUtilities::CreatePayloadSidecarFile(*SaveContext.Linker.Get(), SaveContext.GetTargetPackagePath(), SaveContext.IsSaveToMemory(),
+	Result = SavePackageUtilities::CreatePayloadSidecarFile(*SaveContext.GetLinker(), SaveContext.GetTargetPackagePath(), SaveContext.IsSaveToMemory(),
 		SaveContext.AdditionalPackageFiles, SaveContext.GetSavePackageContext());
 	if (Result != ESavePackageResult::Success)
 	{
diff --git a/Engine/Source/Runtime/CoreUObject/Public/AssetRegistry/AssetData.h b/Engine/Source/Runtime/CoreUObject/Public/AssetRegistry/AssetData.h
--- a/Engine/Source/Runtime/CoreUObject/Public/AssetRegistry/AssetData.h
+++ b/Engine/Source/Runtime/CoreUObject/Public/AssetRegistry/AssetData.h
@@ -52,6 +52,7 @@ struct COREUOBJECT_API FAssetRegistryVersion
 		WorkspaceDomain,					// Added Version information to AssetPackageData
 		PackageImportedClasses,				// Added ImportedClasses to AssetPackageData
 		PackageFileSummaryVersionChange,	// A new version number of UE5 was added to FPackageFileSummary
+		ObjectResourceOptionalVersionChange,// Change to linker export/import resource serializationn
 
 		// -----<new versions can be added above this line>-------------------------------------------------
 		VersionPlusOne,
diff --git a/Engine/Source/Runtime/CoreUObject/Public/Misc/PackagePath.h b/Engine/Source/Runtime/CoreUObject/Public/Misc/PackagePath.h
--- a/Engine/Source/Runtime/CoreUObject/Public/Misc/PackagePath.h
+++ b/Engine/Source/Runtime/CoreUObject/Public/Misc/PackagePath.h
@@ -360,6 +360,12 @@ public:
 	 */
 	COREUOBJECT_API static EPackageExtension ParseExtension(FStringView Filename, int32* OutExtensionStart = nullptr);
 
+	/**
+	 * Get the string identifying optional segments
+	 * @return the string modifier added prior to the file extension that identify optional segments (i.e. ".o")
+	 */
+	COREUOBJECT_API static const TCHAR* GetOptionalSegmentExtensionModifier();
+
 	/**
 	 * Get the folder name from which all external actors paths are created
 	 * @return folder name
diff --git a/Engine/Source/Runtime/CoreUObject/Public/UObject/ObjectResource.h b/Engine/Source/Runtime/CoreUObject/Public/UObject/ObjectResource.h
--- a/Engine/Source/Runtime/CoreUObject/Public/UObject/ObjectResource.h
+++ b/Engine/Source/Runtime/CoreUObject/Public/UObject/ObjectResource.h
@@ -276,7 +276,7 @@ struct FObjectExport : public FObjectResource
 	 * The index into the FLinker's ExportMap for the next export in the linker's export hash table.
 	 * Transient
 	 */
-	int32				HashNext;
+	int32			HashNext;
 
 	/**
 	 * Whether the export was forced into the export table via OBJECTMARK_ForceTagExp.
@@ -309,6 +309,11 @@ struct FObjectExport : public FObjectResource
 	 */
 	bool			bIsAsset;
 
+	/**
+	 * True if this export should have its iostore public hash generated even if not RF_Public.
+	 */
+	bool			bGeneratePublicHash;
+
 	/**
 	 * Force this export to not load, it failed because the outer didn't exist.
 	 */
@@ -455,6 +460,11 @@ struct FObjectImport : public FObjectResource
 	 */
 	int32             SourceIndex;
 
+	/** 
+	 * Indicate if the import comes from an optional package, used to generate the proper chunk id in the io store
+	 */
+	bool			bImportOptional;
+
 	bool			bImportPackageHandled;
 	bool			bImportSearchedFor;
 	bool			bImportFailed;
diff --git a/Engine/Source/Runtime/CoreUObject/Public/UObject/Package.h b/Engine/Source/Runtime/CoreUObject/Public/UObject/Package.h
--- a/Engine/Source/Runtime/CoreUObject/Public/UObject/Package.h
+++ b/Engine/Source/Runtime/CoreUObject/Public/UObject/Package.h
@@ -55,6 +55,8 @@ enum class ESavePackageResult
 	ValidatorError,
 	/** Result from ISavePackageValidator that suppresses the save but is not an error. */
 	ValidatorSuppress,
+	/** Internal save result used to identify a valid empty internal save realm to skip over. @see ESaveRealm */
+	EmptyRealm,
 };
 
 /**
