Fix for static analysis warning

#rnx

- Simplified the assert checking that the trailer exists by moving inside the code branch that has already checked that the bulkdata object should have a trailer.
- This makes the code easier to read and shows that it is okay to use the Trailer pointer as the assert will have validated it.

diff --git a/Engine/Source/Runtime/CoreUObject/Private/Serialization/VirtualizedBulkData.cpp b/Engine/Source/Runtime/CoreUObject/Private/Serialization/VirtualizedBulkData.cpp
--- a/Engine/Source/Runtime/CoreUObject/Private/Serialization/VirtualizedBulkData.cpp
+++ b/Engine/Source/Runtime/CoreUObject/Private/Serialization/VirtualizedBulkData.cpp
@@ -747,10 +747,10 @@ void FVirtualizedUntypedBulkData::Serialize(FArchive& Ar, UObject* Owner, bool b
 			PackageSegment = EPackageSegment::Header;
 
 			const FPackageTrailer* Trailer = GetTrailerFromOwner(Owner);
-			checkf(Trailer != nullptr || !IsStoredInPackageTrailer(), TEXT("Payload was stored in a package trailer, but there no trailer loaded"));
-		
+				
 			if (IsStoredInPackageTrailer())
 			{
+				checkf(Trailer != nullptr, TEXT("Payload was stored in a package trailer, but there no trailer loaded"));
 				// Cache the offset from the trailer (if we move the loading of the payload to the trailer 
 				// at a later point then we can skip this)
 				OffsetInFile = Trailer->FindPayloadOffsetInFile(PayloadContentId);
