TVariantStructure - Alternative to TBaseStructure providing access to script structs of LWC variant types.

diff --git a/Engine/Source/Runtime/CoreUObject/Private/UObject/Class.cpp b/Engine/Source/Runtime/CoreUObject/Private/UObject/Class.cpp
--- a/Engine/Source/Runtime/CoreUObject/Private/UObject/Class.cpp
+++ b/Engine/Source/Runtime/CoreUObject/Private/UObject/Class.cpp
@@ -6166,24 +6166,6 @@ static UScriptStruct* StaticGetBaseStructureInternal(FName Name)
 	return Result;
 }
 
-UScriptStruct* TBaseStructure<FRotator>::Get()
-{
-	static UScriptStruct* ScriptStruct = StaticGetBaseStructureInternal(TEXT("Rotator"));
-	return ScriptStruct;
-}
-
-UScriptStruct* TBaseStructure<FQuat>::Get()
-{
-	static UScriptStruct* ScriptStruct = StaticGetBaseStructureInternal(TEXT("Quat"));
-	return ScriptStruct;
-}
-
-UScriptStruct* TBaseStructure<FTransform>::Get()
-{
-	static UScriptStruct* ScriptStruct = StaticGetBaseStructureInternal(TEXT("Transform"));
-	return ScriptStruct;
-}
-
 UScriptStruct* TBaseStructure<FLinearColor>::Get()
 {
 	static UScriptStruct* ScriptStruct = StaticGetBaseStructureInternal(TEXT("LinearColor"));
@@ -6196,30 +6178,6 @@ UScriptStruct* TBaseStructure<FColor>::Get()
 	return ScriptStruct;
 }
 
-UScriptStruct* TBaseStructure<FPlane>::Get()
-{
-	static UScriptStruct* ScriptStruct = StaticGetBaseStructureInternal(TEXT("Plane"));
-	return ScriptStruct;
-}
-
-UScriptStruct* TBaseStructure<FVector>::Get()
-{
-	static UScriptStruct* ScriptStruct = StaticGetBaseStructureInternal(TEXT("Vector"));
-	return ScriptStruct;
-}
-
-UScriptStruct* TBaseStructure<FVector2D>::Get()
-{
-	static UScriptStruct* ScriptStruct = StaticGetBaseStructureInternal(TEXT("Vector2D"));
-	return ScriptStruct;
-}
-
-UScriptStruct* TBaseStructure<FVector4>::Get()
-{
-	static UScriptStruct* ScriptStruct = StaticGetBaseStructureInternal(TEXT("Vector4"));
-	return ScriptStruct;
-}
-
 UScriptStruct* TBaseStructure<FRandomStream>::Get()
 {
 	static UScriptStruct* ScriptStruct = StaticGetBaseStructureInternal(TEXT("RandomStream"));
@@ -6232,12 +6190,6 @@ UScriptStruct* TBaseStructure<FGuid>::Get()
 	return ScriptStruct;
 }
 
-UScriptStruct* TBaseStructure<FBox2D>::Get()
-{
-	static UScriptStruct* ScriptStruct = StaticGetBaseStructureInternal(TEXT("Box2D"));
-	return ScriptStruct;
-}
-
 UScriptStruct* TBaseStructure<FFallbackStruct>::Get()
 {
 	static UScriptStruct* ScriptStruct = StaticGetBaseStructureInternal(TEXT("FallbackStruct"));
@@ -6376,6 +6328,37 @@ UScriptStruct* TBaseStructure<FTestUninitializedScriptStructMembersTest>::Get()
 	return ScriptStruct;
 }
 
+
+
+#define UE_DEFINE_CORE_VARIANT_TYPE(VARIANT, CORE)												\
+UScriptStruct* TBaseStructure<F##CORE>::Get()												\
+{																								\
+	static UScriptStruct* ScriptStruct = StaticGetBaseStructureInternal(NAME_##CORE);			\
+	return ScriptStruct;																		\
+}																								\
+UScriptStruct* TVariantStructure<F##VARIANT##f>::Get()											\
+{																								\
+	static UScriptStruct* ScriptStruct = StaticGetBaseStructureInternal(NAME_##VARIANT##f);		\
+	return ScriptStruct;																		\
+}																								\
+UScriptStruct* TVariantStructure<F##VARIANT##d>::Get()											\
+{																								\
+	static UScriptStruct* ScriptStruct = StaticGetBaseStructureInternal(NAME_##VARIANT##d);		\
+	return ScriptStruct;																		\
+}
+
+UE_DEFINE_CORE_VARIANT_TYPE(Vector2,	Vector2D);
+UE_DEFINE_CORE_VARIANT_TYPE(Vector3,	Vector);
+UE_DEFINE_CORE_VARIANT_TYPE(Vector4,	Vector4);
+UE_DEFINE_CORE_VARIANT_TYPE(Plane4,		Plane);
+UE_DEFINE_CORE_VARIANT_TYPE(Quat4,		Quat);
+UE_DEFINE_CORE_VARIANT_TYPE(Rotator3,	Rotator);
+UE_DEFINE_CORE_VARIANT_TYPE(Transform3,	Transform);
+UE_DEFINE_CORE_VARIANT_TYPE(Matrix44,	Matrix);
+UE_DEFINE_CORE_VARIANT_TYPE(Box2,		Box2D);
+
+#undef UE_DEFINE_CORE_VARIANT_TYPE
+
 IMPLEMENT_CORE_INTRINSIC_CLASS(UFunction, UStruct,
 	{
 	}
@@ -6442,7 +6425,6 @@ IMPLEMENT_CORE_INTRINSIC_CLASS(UDynamicClass, UClass,
 }
 );
 PRAGMA_ENABLE_DEPRECATION_WARNINGS
-
 #if defined(_MSC_VER) && _MSC_VER == 1900
 	#ifdef PRAGMA_ENABLE_SHADOW_VARIABLE_WARNINGS
 		PRAGMA_ENABLE_SHADOW_VARIABLE_WARNINGS
diff --git a/Engine/Source/Runtime/CoreUObject/Public/UObject/Class.h b/Engine/Source/Runtime/CoreUObject/Public/UObject/Class.h
--- a/Engine/Source/Runtime/CoreUObject/Public/UObject/Class.h
+++ b/Engine/Source/Runtime/CoreUObject/Public/UObject/Class.h
@@ -3856,21 +3856,6 @@ template< class T > struct TBaseStructure
 	}
 };
 
-template<> struct TBaseStructure<FRotator>
-{
-	COREUOBJECT_API static UScriptStruct* Get();
-};
-
-template<> struct TBaseStructure<FQuat>
-{
-	COREUOBJECT_API static UScriptStruct* Get();
-};
-
-template<> struct TBaseStructure<FTransform>
-{
-	COREUOBJECT_API static UScriptStruct* Get();
-};
-
 template<> struct TBaseStructure<FLinearColor>
 {
 	COREUOBJECT_API static UScriptStruct* Get();
@@ -3881,26 +3866,6 @@ template<> struct TBaseStructure<FColor>
 	COREUOBJECT_API static UScriptStruct* Get();
 };
 
-template<> struct  TBaseStructure<FPlane>
-{
-	COREUOBJECT_API static UScriptStruct* Get();
-};
-
-template<> struct  TBaseStructure<FVector>
-{
-	COREUOBJECT_API static UScriptStruct* Get();
-};
-
-template<> struct TBaseStructure<FVector2D>
-{
-	COREUOBJECT_API static UScriptStruct* Get();
-};
-
-template<> struct TBaseStructure<FVector4>
-{
-	COREUOBJECT_API static UScriptStruct* Get();
-};
-
 template<> struct TBaseStructure<FRandomStream>
 {
 	COREUOBJECT_API static UScriptStruct* Get();
@@ -3911,11 +3876,6 @@ template<> struct TBaseStructure<FGuid>
 	COREUOBJECT_API static UScriptStruct* Get();
 };
 
-template<> struct TBaseStructure<FBox2D>
-{
-	COREUOBJECT_API static UScriptStruct* Get();
-};	
-
 template<> struct TBaseStructure<FFallbackStruct>
 {
 	COREUOBJECT_API static UScriptStruct* Get();
@@ -4039,3 +3999,28 @@ template<> struct TBaseStructure<FTestUninitializedScriptStructMembersTest>
 {
 	COREUOBJECT_API static UScriptStruct* Get();
 };
+
+
+
+// TBaseStructure for explicit core variant types only. e.g. FVector3d returns "Vector3d" struct. 
+template< class T > struct TVariantStructure
+{
+	static_assert(sizeof(T) == 0, "Unsupported for this type. Did you mean to use TBaseStructure?");
+};
+
+#define UE_DECLARE_CORE_VARIANT_TYPE(VARIANT, CORE)														\
+template<> struct TBaseStructure<F##CORE> { COREUOBJECT_API static UScriptStruct* Get(); };		\
+template<> struct TVariantStructure<F##VARIANT##f> { COREUOBJECT_API static UScriptStruct* Get(); };	\
+template<> struct TVariantStructure<F##VARIANT##d> { COREUOBJECT_API static UScriptStruct* Get(); };
+
+UE_DECLARE_CORE_VARIANT_TYPE(Vector2,	Vector2D);
+UE_DECLARE_CORE_VARIANT_TYPE(Vector3,	Vector);
+UE_DECLARE_CORE_VARIANT_TYPE(Vector4,	Vector4);
+UE_DECLARE_CORE_VARIANT_TYPE(Plane4,	Plane);
+UE_DECLARE_CORE_VARIANT_TYPE(Quat4,		Quat);
+UE_DECLARE_CORE_VARIANT_TYPE(Rotator3,	Rotator);
+UE_DECLARE_CORE_VARIANT_TYPE(Transform3,Transform);
+UE_DECLARE_CORE_VARIANT_TYPE(Matrix44,	Matrix);
+UE_DECLARE_CORE_VARIANT_TYPE(Box2,		Box2D);
+
+#undef UE_DECLARE_CORE_VARIANT_TYPE
\ No newline at end of file
