Removed USE_BULKDATA_STREAMING_TOKEN from the bulkdata API

#rnx

### Token Removal
- Now that we support the IoStore in the editor there is no reason to keep the bulkdata streaming token as both the editor and runtime systems can use the same code paths.
- Existing use of USE_BULKDATA_STREAMING_TOKEN will now give deprecation warnings however code should continue to work as long as STREAMINGTOKEN_PARAM was correctly used.
- ::CreateStreamingToken has been removed from the bulkdata api, but any calls to this should've been wrapped with #if USE_BULKDATA_STREAMING_TOKEN

### API Fixes
- We can remove all code that relied on USE_BULKDATA_STREAMING_TOKEN being == 1
- LANDSCAPE_LOD_STREAMING_USE_TOKEN required USE_BULKDATA_STREAMING_TOKEN to be enabled, so we can now assume that it is always 0 and remove the associated code paths.

diff --git a/Engine/Source/Runtime/CoreUObject/Private/Serialization/BulkData.cpp b/Engine/Source/Runtime/CoreUObject/Private/Serialization/BulkData.cpp
--- a/Engine/Source/Runtime/CoreUObject/Private/Serialization/BulkData.cpp
+++ b/Engine/Source/Runtime/CoreUObject/Private/Serialization/BulkData.cpp
@@ -1952,36 +1952,16 @@ IBulkDataIORequest* FUntypedBulkData::CreateStreamingRequest(int64 OffsetInBulkD
 	}
 }
 
-#if USE_BULKDATA_STREAMING_TOKEN 
-FBulkDataStreamingToken FUntypedBulkData::CreateStreamingToken() const
-{
-	// Checks since we are casting signed 64bit values to unsigned 32bit
-	check(GetBulkDataOffsetInFile() >= 0);
-	check(GetBulkDataOffsetInFile() < TNumericLimits<uint32>::Max());
-	check(GetBulkDataSize() >= 0);
-	check(GetBulkDataSize() < TNumericLimits<uint32>::Max());
-
-	return FBulkDataStreamingToken((uint32)GetBulkDataOffsetInFile(), (uint32)GetBulkDataSize());
-}
-
-IBulkDataIORequest* FUntypedBulkData::CreateStreamingRequestForRange(const FString& Filename, const BulkDataRangeArray& RangeArray, EAsyncIOPriorityAndFlags Priority, FBulkDataIORequestCallBack* CompleteCallback)
-{
-	FPackagePath PackagePath;
-	EPackageSegment PackageSegment;
-	PackagePath = FPackagePath::FromLocalPath(Filename, PackageSegment);
-	return CreateStreamingRequestForRange(PackagePath, PackageSegment, RangeArray, Priority, CompleteCallback);
-}
-
-IBulkDataIORequest* FUntypedBulkData::CreateStreamingRequestForRange(const FPackagePath& PackagePath, EPackageSegment PackageSegment, const BulkDataRangeArray& RangeArray, EAsyncIOPriorityAndFlags Priority, FBulkDataIORequestCallBack* CompleteCallback)
-{
-	check(PackagePath.IsEmpty() == false);
+IBulkDataIORequest* FUntypedBulkData::CreateStreamingRequestForRange(const BulkDataRangeArray& RangeArray, EAsyncIOPriorityAndFlags Priority, FBulkDataIORequestCallBack* CompleteCallback)
+{	
 	check(RangeArray.Num() > 0);
 
-	const FBulkDataStreamingToken& Start = *(RangeArray[0]);
-	const FBulkDataStreamingToken& End = *(RangeArray[RangeArray.Num()-1]);
+	const FUntypedBulkData& Start = *(RangeArray[0]);
+	const FUntypedBulkData& End = *(RangeArray[RangeArray.Num()-1]);
 
-	check(Start.IsValid());
-	check(End.IsValid());
+	const FPackagePath& PackagePath = Start.GetPackagePath();
+	const EPackageSegment PackageSegment = Start.GetPackageSegment();
+	check(PackagePath.IsEmpty() == false);
 
 	// TODO: The caller is assuming that their specified PackageSegments are the way in which bulkdata is stored for their package
 	// To allow more flexible bulkdata storage, we will need to eliminate this interface and have them provide an array of bulkdatas.
@@ -1994,8 +1974,8 @@ IBulkDataIORequest* FUntypedBulkData::CreateStreamingRequestForRange(const FPack
 		return nullptr;
 	}
 
-	const int64 ReadOffset = Start.GetOffset();
-	const int64 ReadSize = (End.GetOffset() + End.GetBulkDataSize()) - ReadOffset;
+	const int64 ReadOffset = Start.GetBulkDataOffsetInFile();
+	const int64 ReadSize = (End.GetBulkDataOffsetInFile() + End.GetBulkDataSize()) - ReadOffset;
 
 	check(ReadSize > 0);
 
@@ -2012,51 +1992,6 @@ IBulkDataIORequest* FUntypedBulkData::CreateStreamingRequestForRange(const FPack
 	}
 }
 
-#else
-IBulkDataIORequest* FUntypedBulkData::CreateStreamingRequestForRange(const BulkDataRangeArray& RangeArray, EAsyncIOPriorityAndFlags Priority, FBulkDataIORequestCallBack* CompleteCallback)
-{
-	checkf(RangeArray.Num() > 0, TEXT("RangeArray cannot be empty"));
-
-	const FUntypedBulkData& Start = *(RangeArray[0]);
-
-	checkf(!Start.IsInlined(), TEXT("Cannot stream inlined BulkData"));
-
-	if (Start.IsUsingIODispatcher())
-	{
-		const FUntypedBulkData& End = *RangeArray[RangeArray.Num() - 1];
-
-		checkf(Start.IsInSeparateFile(),
-			TEXT("Attempting to CreateStreamingRequestForRange on %s when the IoDispatcher is enabled, this operation is not supported!"),
-			Start.IsInlined() ? TEXT("inline BulkData") : TEXT("BulkData in end-of-package-file section"));
-		checkf(End.IsInSeparateFile() && Start.CreateChunkId() == End.CreateChunkId(), TEXT("BulkData range does not come from the same file (%s vs %s)"),
-			*Start.GetPackagePath().GetDebugName(Start.GetPackageSegment()), *End.GetPackagePath().GetDebugName(End.GetPackageSegment()));
-
-		const int64 ReadOffset = Start.GetBulkDataOffsetInFile();
-		const int64 ReadLength = (End.GetBulkDataOffsetInFile() + End.GetBulkDataSize()) - ReadOffset;
-
-		checkf(ReadLength > 0, TEXT("Read length is 0"));
-
-		TUniquePtr<IBulkDataIORequest> IoRequest = CreateBulkDataIoDispatcherRequest(Start.CreateChunkId(), ReadOffset, ReadLength, CompleteCallback, nullptr, ConvertToIoDispatcherPriority(Priority));
-
-		return IoRequest.Release();
-	}
-	else
-	{
-		const FUntypedBulkData& End = *RangeArray[RangeArray.Num() - 1];
-
-		checkf(Start.GetPackagePath() == End.GetPackagePath(), TEXT("BulkData range does not come from the same file (%s vs %s)"),
-			*Start.GetPackagePath().GetDebugName(), *End.GetPackagePath().GetDebugName());
-
-		const int64 ReadOffset = Start.GetBulkDataOffsetInFile();
-		const int64 ReadLength = (End.GetBulkDataOffsetInFile() + End.GetBulkDataSize()) - ReadOffset;
-
-		checkf(ReadLength > 0, TEXT("Read length is 0"));
-
-		return Start.CreateStreamingRequest(0, ReadLength, Priority, CompleteCallback, nullptr);
-	}
-}
-#endif //USE_BULKDATA_STREAMING_TOKEN
-
 FBulkDataIORequest::FBulkDataIORequest(IAsyncReadFileHandle* InFileHandle)
 	: FileHandle(InFileHandle)
 	, ReadRequest(nullptr)
diff --git a/Engine/Source/Runtime/CoreUObject/Public/Serialization/BulkData.h b/Engine/Source/Runtime/CoreUObject/Public/Serialization/BulkData.h
--- a/Engine/Source/Runtime/CoreUObject/Public/Serialization/BulkData.h
+++ b/Engine/Source/Runtime/CoreUObject/Public/Serialization/BulkData.h
@@ -23,13 +23,8 @@
 #endif
 
 // Enable the following to use the more compact FBulkDataStreamingToken in places where it is implemented
-#define USE_BULKDATA_STREAMING_TOKEN !USE_NEW_BULKDATA && !WITH_IOSTORE_IN_EDITOR
-
-#if USE_BULKDATA_STREAMING_TOKEN
-	#define STREAMINGTOKEN_PARAM(param) param,
-#else
-	#define STREAMINGTOKEN_PARAM(param) 
-#endif
+#define USE_BULKDATA_STREAMING_TOKEN DEPRECATED_MACRO(5.0, "USE_BULKDATA_STREAMING_TOKEN now always evaluates to 0 and will be removed") 0
+#define STREAMINGTOKEN_PARAM(param) DEPRECATED_MACRO(5.0, "STREAMINGTOKEN_PARAM now always evaluates to a NOP")
 
 class IMappedFileHandle;
 class IMappedFileRegion;
@@ -117,52 +112,7 @@ private:
 	int64 Size;
 };
 
-#if USE_BULKDATA_STREAMING_TOKEN 
-/**
- * Some areas of code currently find FUntypedBulkData too bloated for use so were storing smaller parts of it
- * and performing the file IO manually. This is an attempt to wrap that functionality in a single place to make
- * it easier to replace in the future.
- * This structure will represent the area of a file that a FUntypedBulkData points to.
- * The structure does not contain info on which file it is however (as that would increase the data
- * size by an unacceptable amount) and it is assumed that this knowledge is being kept at a higher level.
- *
- * This is intended only to exist for a short time to aid the transition to a new system.
- */
-struct COREUOBJECT_API FBulkDataStreamingToken
-{
-public:
-	static constexpr uint32 InvalidOffset = TNumericLimits<uint32>::Max();
-
-public:
-	FBulkDataStreamingToken() {}
-	FBulkDataStreamingToken(uint32 InOffsetInFile, uint32 InBulkDataSize)
-		: OffsetInFile(InOffsetInFile)
-		, BulkDataSize(InBulkDataSize)
-	{
-	}
-
-	bool IsValid() const
-	{ 
-		return OffsetInFile != InvalidOffset && BulkDataSize > 0;
-	}
-
-	uint32 GetOffset() const
-	{
-		return OffsetInFile;
-	}
-
-	uint32 GetBulkDataSize() const
-	{
-		return BulkDataSize;
-	}
-
-private:
-	uint32 OffsetInFile{ InvalidOffset };
-	uint32 BulkDataSize{ 0 };
-};
-#else
 class FBulkDataStreamingToken; // Forward declared but not implemented
-#endif
 
 /**
  * @documentation @todo documentation
@@ -296,11 +246,7 @@ public:
 	friend class FExportArchive;
 	friend class UE::Virtualization::FVirtualizedUntypedBulkData; // To allow access to AttachedAr
 
-#if USE_BULKDATA_STREAMING_TOKEN
-	using BulkDataRangeArray = TArray<FBulkDataStreamingToken*, TInlineAllocator<8>>;	
-#else
 	using BulkDataRangeArray = TArray<FUntypedBulkData*, TInlineAllocator<8>>;
-#endif //USE_BULKDATA_STREAMING_TOKEN
 
 	/*-----------------------------------------------------------------------------
 		Constructors and operators
@@ -728,36 +674,19 @@ public:
 	 */
 	IBulkDataIORequest* CreateStreamingRequest(int64 OffsetInBulkData, int64 BytesToRead, EAsyncIOPriorityAndFlags Priority, FBulkDataIORequestCallBack* CompleteCallback, uint8* UserSuppliedMemory) const;
 
-#if USE_BULKDATA_STREAMING_TOKEN 
-
-	/**
-	 * Creates a FBulkDataStreamingToken representing the area of the file that the FUntypedBulkData represents. See the declaration of
-	 * FBulkDataStreamingToken for further details.
-	 *
-	 * @return	A FBulkDataStreamingToken valid for the FUntypedBulkData.
-	 */
-	FBulkDataStreamingToken CreateStreamingToken() const;
-
 	/**
 	 * Create an async read request for a range of bulk data streaming tokens
 	 * The request will read all data between the two given streaming tokens objects. They must both represent areas of data in the file!
 	 * There is no way to validate this and it is up to the caller to make sure that it is correct.
 	 * The memory to be read into will be automatically allocated the size of which can be retrieved by calling IBulkDataIORequest::GetSize()
 	 *
-	 * @param PackagePath		The PackagePath to read from.
-	 * @param PackageSegment	The PackageSegment to read from.
 	 * @param Start				The bulk data to start reading from.
 	 * @param End				The bulk data to finish reading from.
 	 * @param Priority			Priority and flags of the request. If this includes AIOP_FLAG_PRECACHE, then memory will never be returned. The request should always be canceled and waited for, even for a precache request.
 	 * @param CompleteCallback	Called from an arbitrary thread when the request is complete. Can be nullptr, if non-null, must remain valid until it is called. It will always be called.
 	 * @return					A request for the read. This is owned by the caller and must be deleted by the caller.
 	**/
-	static IBulkDataIORequest* CreateStreamingRequestForRange(const FPackagePath& PackagePath, EPackageSegment PackageSegment, const BulkDataRangeArray& RangeArray, EAsyncIOPriorityAndFlags Priority, FBulkDataIORequestCallBack* CompleteCallback);
-	UE_DEPRECATED(4.26, "Use version that takes a FPackagePath instead")
-	static IBulkDataIORequest* CreateStreamingRequestForRange(const FString& Filename, const BulkDataRangeArray& RangeArray, EAsyncIOPriorityAndFlags Priority, FBulkDataIORequestCallBack* CompleteCallback);
-#else
 	static IBulkDataIORequest* CreateStreamingRequestForRange(const BulkDataRangeArray& RangeArray, EAsyncIOPriorityAndFlags Priority, FBulkDataIORequestCallBack* CompleteCallback);
-#endif //USE_BULKDATA_STREAMING_TOKEN
 
 	/** Enable the given flags in the given accumulator variable. */
 	static void SetBulkDataFlagsOn(EBulkDataFlags& InOutAccumulator, EBulkDataFlags FlagsToSet);
diff --git a/Engine/Source/Runtime/Engine/Classes/Engine/SkeletalMesh.h b/Engine/Source/Runtime/Engine/Classes/Engine/SkeletalMesh.h
--- a/Engine/Source/Runtime/Engine/Classes/Engine/SkeletalMesh.h
+++ b/Engine/Source/Runtime/Engine/Classes/Engine/SkeletalMesh.h
@@ -2716,12 +2716,6 @@ public:
 	virtual EStreamableRenderAssetType GetRenderAssetType() const final override { return EStreamableRenderAssetType::SkeletalMesh; }
 	//~ End UStreamableRenderAsset Interface.
 
-#if USE_BULKDATA_STREAMING_TOKEN
-	UE_DEPRECATED(5.0, "Use GetMipDataPackagePath instead")
-	bool GetMipDataFilename(const int32 MipIndex, FString& OutBulkDataFilename) const;
-	bool GetMipDataPackagePath(const int32 MipIndex, FPackagePath& OutPackagePath, EPackageSegment& OutPackageSegment) const;
-#endif
-
 	/**
 	* Cancels any pending static mesh streaming actions if possible.
 	* Returns when no more async loading requests are in flight.
diff --git a/Engine/Source/Runtime/Engine/Classes/Engine/StaticMesh.h b/Engine/Source/Runtime/Engine/Classes/Engine/StaticMesh.h
--- a/Engine/Source/Runtime/Engine/Classes/Engine/StaticMesh.h
+++ b/Engine/Source/Runtime/Engine/Classes/Engine/StaticMesh.h
@@ -1575,12 +1575,6 @@ public:
 	ENGINE_API virtual EStreamableRenderAssetType GetRenderAssetType() const final override;
 	//~ End UStreamableRenderAsset Interface
 
-#if USE_BULKDATA_STREAMING_TOKEN
-	UE_DEPRECATED(5.0, "Use GetMipDataPackagePath instead")
-	bool GetMipDataFilename(const int32 MipIndex, FString& OutBulkDataFilename) const;
-	bool GetMipDataPackagePath(const int32 MipIndex, FPackagePath& OutPackagePath, EPackageSegment& OutPackageSegment) const;
-#endif
-
 	/**
 	* Cancels any pending static mesh streaming actions if possible.
 	* Returns when no more async loading requests are in flight.
diff --git a/Engine/Source/Runtime/Engine/Private/SkeletalMesh.cpp b/Engine/Source/Runtime/Engine/Private/SkeletalMesh.cpp
--- a/Engine/Source/Runtime/Engine/Private/SkeletalMesh.cpp
+++ b/Engine/Source/Runtime/Engine/Private/SkeletalMesh.cpp
@@ -1217,63 +1217,12 @@ int32 USkeletalMesh::CalcCumulativeLODSize(int32 NumLODs) const
 	return Accum;
 }
 
-#if USE_BULKDATA_STREAMING_TOKEN
-bool USkeletalMesh::GetMipDataFilename(const int32 MipIndex, FString& OutBulkDataFilename) const
-{
-	FPackagePath PackagePath;
-	EPackageSegment PackageSegment;
-	if (GetMipDataPackagePath(MipIndex, PackagePath, PackageSegment))
-	{
-		OutBulkDataFilename = PackagePath.GetLocalFullPath(PackageSegment);
-		return true;
-	}
-	return false;
-}
-
-bool USkeletalMesh::GetMipDataPackagePath(const int32 MipIndex, FPackagePath& OutPackagePath, EPackageSegment& OutPackageSegment) const
-{
-	// TODO: this is slow. Should cache the name once per mesh
-	FPackagePath PackagePath = GetOutermost()->GetLoadedPath();
-	// Handle name redirection and localization
-	{
-		FString PackageName = PackagePath.GetPackageName();
-		const FCoreRedirectObjectName RedirectedName =
-			FCoreRedirects::GetRedirectedName(
-				ECoreRedirectFlags::Type_Package,
-				FCoreRedirectObjectName(NAME_None, NAME_None, *PackageName));
-		FString LocalizedName;
-		LocalizedName = FPackageName::GetDelegateResolvedPackagePath(RedirectedName.PackageName.ToString());
-		LocalizedName = FPackageName::GetLocalizedPackagePath(LocalizedName);
-		if (LocalizedName != PackageName)
-		{
-			PackagePath = FPackagePath::FromPackageNameChecked(LocalizedName);
-		}
-	}
-	bool bSucceed = FPackageName::DoesPackageExist(PackagePath, &PackagePath);
-	check(bSucceed);
-	const FSkeletalMeshRenderData* SkelMeshRenderData = GetResourceForRendering();
-	const bool bLODIsOptional = SkelMeshRenderData ? (MipIndex < SkelMeshRenderData->LODRenderData.Num() - SkelMeshRenderData->NumNonOptionalLODs) : false;
-	OutPackagePath = MoveTemp(PackagePath);
-	OutPackageSegment = bLODIsOptional ? EPackageSegment::BulkDataOptional : EPackageSegment::BulkDataDefault;
-	return true;
-}
-#endif // USE_BULKDATA_STREAMING_TOKEN
-
 FIoFilenameHash USkeletalMesh::GetMipIoFilenameHash(const int32 MipIndex) const
 {
-#if USE_BULKDATA_STREAMING_TOKEN
-	FPackagePath PackagePath;
-	EPackageSegment PackageSegment;
-	if (GetMipDataPackagePath(MipIndex, PackagePath, PackageSegment))
-	{
-		return MakeIoFilenameHash(PackagePath);
-	}
-#else
 	if (GetSkeletalMeshRenderData() && GetSkeletalMeshRenderData()->LODRenderData.IsValidIndex(MipIndex))
 	{
 		return GetSkeletalMeshRenderData()->LODRenderData[MipIndex].StreamingBulkData.GetIoFilenameHash();
 	}
-#endif
 	else
 	{
 		return INVALID_IO_FILENAME_HASH;
@@ -1282,13 +1231,7 @@ FIoFilenameHash USkeletalMesh::GetMipIoFilenameHash(const int32 MipIndex) const
 
 bool USkeletalMesh::DoesMipDataExist(const int32 MipIndex) const
 {
-#if USE_BULKDATA_STREAMING_TOKEN
-	FPackagePath PackagePath;
-	EPackageSegment PackageSegment;
-	return GetMipDataPackagePath(MipIndex, PackagePath, PackageSegment) && IPackageResourceManager::Get().DoesPackageExist(PackagePath, PackageSegment);
-#else
 	return GetSkeletalMeshRenderData() && GetSkeletalMeshRenderData()->LODRenderData.IsValidIndex(MipIndex) && GetSkeletalMeshRenderData()->LODRenderData[MipIndex].StreamingBulkData.DoesExist();
-#endif
 }
 
 bool USkeletalMesh::HasPendingRenderResourceInitialization() const
diff --git a/Engine/Source/Runtime/Engine/Private/SkeletalMeshLODRenderData.cpp b/Engine/Source/Runtime/Engine/Private/SkeletalMeshLODRenderData.cpp
--- a/Engine/Source/Runtime/Engine/Private/SkeletalMeshLODRenderData.cpp
+++ b/Engine/Source/Runtime/Engine/Private/SkeletalMeshLODRenderData.cpp
@@ -878,17 +878,9 @@ void FSkeletalMeshLODRenderData::Serialize(FArchive& Ar, UObject* Owner, int32 I
 			else
 #endif
 			{
-#if USE_BULKDATA_STREAMING_TOKEN
-				FByteBulkData TmpBulkData;
-				TmpBulkData.Serialize(Ar, Owner, Idx, false);
-				bIsLODOptional = TmpBulkData.IsOptional();
-
-				StreamingBulkData = TmpBulkData.CreateStreamingToken();
-#else
 				StreamingBulkData.Serialize(Ar, Owner, Idx, false);
 				bIsLODOptional = StreamingBulkData.IsOptional();
-#endif
-			
+
 				if (StreamingBulkData.GetBulkDataSize() == 0)
 				{
 					bDiscardBulkData = true;
diff --git a/Engine/Source/Runtime/Engine/Private/StaticMesh.cpp b/Engine/Source/Runtime/Engine/Private/StaticMesh.cpp
--- a/Engine/Source/Runtime/Engine/Private/StaticMesh.cpp
+++ b/Engine/Source/Runtime/Engine/Private/StaticMesh.cpp
@@ -786,16 +786,8 @@ void FStaticMeshLODResources::Serialize(FArchive& Ar, UObject* Owner, int32 Inde
 			else
 #endif
 			{
-#if USE_BULKDATA_STREAMING_TOKEN
-				FByteBulkData TmpBulkData;
-				TmpBulkData.Serialize(Ar, Owner, Index, false);
-				bIsOptionalLOD = TmpBulkData.IsOptional();
-
-				StreamingBulkData = TmpBulkData.CreateStreamingToken();			
-#else
 				StreamingBulkData.Serialize(Ar, Owner, Index, false);
 				bIsOptionalLOD = StreamingBulkData.IsOptional();
-#endif
 
 #if WITH_EDITORONLY_DATA
 				BulkDataSize = (uint32)StreamingBulkData.GetBulkDataSize();
@@ -805,13 +797,8 @@ void FStaticMeshLODResources::Serialize(FArchive& Ar, UObject* Owner, int32 Inde
 				{
 					TmpBuff.Empty(BulkDataSize);
 					TmpBuff.AddUninitialized(BulkDataSize);
-					
 					void* Dest = TmpBuff.GetData();
-#if USE_BULKDATA_STREAMING_TOKEN	
-					TmpBulkData.GetCopy(&Dest);
-#else
 					StreamingBulkData.GetCopy(&Dest);
-#endif //#if USE_BULKDATA_STREAMING_TOKEN
 				}
 #endif
 			}
@@ -6254,61 +6241,12 @@ int32 UStaticMesh::CalcCumulativeLODSize(int32 NumLODs) const
 	return Accum;
 }
 
-#if USE_BULKDATA_STREAMING_TOKEN
-bool UStaticMesh::GetMipDataFilename(const int32 MipIndex, FString& OutBulkDataFilename) const
-{
-	FPackagePath PackagePath;
-	EPackageSegment PackageSegment;
-	if (GetMipDataPackagePath(MipIndex, PackagePath, PackageSegment))
-	{
-		OutBulkDataFilename = PackagePath.GetLocalFullPath(PackageSegment);
-		return true;
-	}
-	return false;
-}
-
-bool UStaticMesh::GetMipDataPackagePath(const int32 MipIndex, FPackagePath& OutPackagePath, EPackageSegment& OutPackageSegment) const
-{
-	// TODO: this is slow. Should cache the name once per mesh
-	FPackagePath PackagePath = GetOutermost()->GetLoadedPath();
-	// Handle name redirection and localization
-	{
-		FString PackageName = PackagePath.GetPackageName();
-		const FCoreRedirectObjectName RedirectedName =
-			FCoreRedirects::GetRedirectedName(
-				ECoreRedirectFlags::Type_Package,
-				FCoreRedirectObjectName(NAME_None, NAME_None, *PackageName));
-		FString LocalizedName;
-		LocalizedName = FPackageName::GetDelegateResolvedPackagePath(RedirectedName.PackageName.ToString());
-		LocalizedName = FPackageName::GetLocalizedPackagePath(LocalizedName);
-		if (LocalizedName != PackageName)
-		{
-			PackagePath = FPackagePath::FromPackageNameChecked(LocalizedName);
-		}
-	}
-	bool bSucceed = FPackageName::DoesPackageExist(PackagePath, &PackagePath);
-	check(bSucceed);
-	OutPackagePath = MoveTemp(PackagePath);
-	OutPackageSegment = MipIndex < GetDefaultMinLOD() ? EPackageSegment::BulkDataOptional : EPackageSegment::BulkDataDefault;
-	return true;
-}
-#endif // USE_BULKDATA_STREAMING_TOKEN
-
 FIoFilenameHash UStaticMesh::GetMipIoFilenameHash(const int32 MipIndex) const
 {
-#if USE_BULKDATA_STREAMING_TOKEN
-	FPackagePath PackagePath;
-	EPackageSegment PackageSegment;
-	if (GetMipDataPackagePath(MipIndex, PackagePath, PackageSegment))
-	{
-		return MakeIoFilenameHash(PackagePath);
-	}
-#else
 	if (GetRenderData() && GetRenderData()->LODResources.IsValidIndex(MipIndex))
 	{
 		return GetRenderData()->LODResources[MipIndex].StreamingBulkData.GetIoFilenameHash();
 	}
-#endif
 	else
 	{
 		return INVALID_IO_FILENAME_HASH;
@@ -6317,13 +6255,7 @@ FIoFilenameHash UStaticMesh::GetMipIoFilenameHash(const int32 MipIndex) const
 
 bool UStaticMesh::DoesMipDataExist(const int32 MipIndex) const
 {
-#if USE_BULKDATA_STREAMING_TOKEN
-	FPackagePath PackagePath;
-	EPackageSegment PackageSegment;
-	return GetMipDataPackagePath(MipIndex, PackagePath, PackageSegment) && IPackageResourceManager::Get().DoesPackageExist(PackagePath, PackageSegment);
-#else
 	return GetRenderData() && GetRenderData()->LODResources.IsValidIndex(MipIndex) && GetRenderData()->LODResources[MipIndex].StreamingBulkData.DoesExist();
-#endif
 }
 
 bool UStaticMesh::HasPendingRenderResourceInitialization() const
diff --git a/Engine/Source/Runtime/Engine/Private/Streaming/LandscapeMeshMobileUpdate.cpp b/Engine/Source/Runtime/Engine/Private/Streaming/LandscapeMeshMobileUpdate.cpp
--- a/Engine/Source/Runtime/Engine/Private/Streaming/LandscapeMeshMobileUpdate.cpp
+++ b/Engine/Source/Runtime/Engine/Private/Streaming/LandscapeMeshMobileUpdate.cpp
@@ -251,29 +251,17 @@ void FLandscapeMeshMobileStreamIn_IO::SetIORequest(const FContext& Context)
 	FLandscapeMobileRenderData* RenderData = Context.RenderData;
 	if (LandscapeProxy && RenderData)
 	{
-#if USE_BULKDATA_STREAMING_TOKEN
-		FPackagePath PackagePath;
-		EPackageSegment PackageSegment;
-		verify(LandscapeProxy->GetMipDataPackagePath(PendingFirstLODIdx, PackagePath, PackageSegment));
-#endif
-
 		SetAsyncFileCallback(Context);
 
 		for (int32 Index = PendingFirstLODIdx; Index < CurrentFirstLODIdx; ++Index)
 		{
 			FBulkDataInterface::BulkDataRangeArray BulkDataArray;
-#if !LANDSCAPE_LOD_STREAMING_USE_TOKEN && USE_BULKDATA_STREAMING_TOKEN
-			FBulkDataStreamingToken StreamingToken = LandscapeProxy->GetStreamingLODBulkData(Index).CreateStreamingToken();
-			BulkDataArray.Add(&StreamingToken);
-#else
 			BulkDataArray.Add(&LandscapeProxy->GetStreamingLODBulkData(Index));
-#endif
+
 			if (BulkDataArray[0]->GetBulkDataSize() > 0)
 			{
 				TaskSynchronization.Increment();
 				IORequests[Index] = FBulkDataInterface::CreateStreamingRequestForRange(
-					STREAMINGTOKEN_PARAM(PackagePath)
-					STREAMINGTOKEN_PARAM(PackageSegment)
 					BulkDataArray,
 					bHighPrioIORequest ? AIOP_BelowNormal : AIOP_Low,
 					&AsyncFileCallback);
@@ -406,7 +394,7 @@ void FLandscapeMeshMobileStreamIn_GPUDataOnly::GetStagingData(const FContext& Co
 	{
 		for (int32 Idx = PendingFirstLODIdx; Idx < CurrentFirstLODIdx; ++Idx)
 		{
-			ULandscapeLODStreamingProxy::BulkDataType& BulkData = LandscapeProxy->GetStreamingLODBulkData(Idx);
+			FByteBulkData& BulkData = LandscapeProxy->GetStreamingLODBulkData(Idx);
 			if (BulkData.GetBulkDataSize() > 0)
 			{
 				StagingLODDataSizes[Idx] = BulkData.GetBulkDataSize();
diff --git a/Engine/Source/Runtime/Engine/Private/Streaming/SkeletalMeshUpdate.cpp b/Engine/Source/Runtime/Engine/Private/Streaming/SkeletalMeshUpdate.cpp
--- a/Engine/Source/Runtime/Engine/Private/Streaming/SkeletalMeshUpdate.cpp
+++ b/Engine/Source/Runtime/Engine/Private/Streaming/SkeletalMeshUpdate.cpp
@@ -460,12 +460,6 @@ void FSkeletalMeshStreamIn_IO::SetIORequest(const FContext& Context)
 	FSkeletalMeshRenderData* RenderData = Context.RenderData;
 	if (Mesh && RenderData)
 	{
-#if USE_BULKDATA_STREAMING_TOKEN
-		FPackagePath PackagePath;
-		EPackageSegment PackageSegment;
-		verify(Mesh->GetMipDataPackagePath(ResourceState.AssetLODBias + PendingFirstLODIdx, PackagePath, PackageSegment));
-#endif	
-
 		SetAsyncFileCallback(Context);
 
 		FBulkDataInterface::BulkDataRangeArray BulkDataArray;
@@ -479,8 +473,6 @@ void FSkeletalMeshStreamIn_IO::SetIORequest(const FContext& Context)
 		TaskSynchronization.Increment();
 
 		IORequest = FBulkDataInterface::CreateStreamingRequestForRange(
-			STREAMINGTOKEN_PARAM(PackagePath)
-			STREAMINGTOKEN_PARAM(PackageSegment)
 			BulkDataArray,
 			bHighPrioIORequest ? AIOP_BelowNormal : AIOP_Low,
 			&AsyncFileCallback);
diff --git a/Engine/Source/Runtime/Engine/Private/Streaming/StaticMeshUpdate.cpp b/Engine/Source/Runtime/Engine/Private/Streaming/StaticMeshUpdate.cpp
--- a/Engine/Source/Runtime/Engine/Private/Streaming/StaticMeshUpdate.cpp
+++ b/Engine/Source/Runtime/Engine/Private/Streaming/StaticMeshUpdate.cpp
@@ -482,12 +482,6 @@ void FStaticMeshStreamIn_IO::SetIORequest(const FContext& Context)
 	FStaticMeshRenderData* RenderData = Context.RenderData;
 	if (Mesh && RenderData)
 	{
-#if USE_BULKDATA_STREAMING_TOKEN
-		FPackagePath PackagePath;
-		EPackageSegment PackageSegment;
-		verify(Mesh->GetMipDataPackagePath(PendingFirstLODIdx, PackagePath, PackageSegment));
-#endif	
-
 		SetAsyncFileCallback(Context);
 
 		FBulkDataInterface::BulkDataRangeArray BulkDataArray;
@@ -501,8 +495,6 @@ void FStaticMeshStreamIn_IO::SetIORequest(const FContext& Context)
 		TaskSynchronization.Increment();
 
 		IORequest = FBulkDataInterface::CreateStreamingRequestForRange(
-			STREAMINGTOKEN_PARAM(PackagePath)
-			STREAMINGTOKEN_PARAM(PackageSegment)
 			BulkDataArray,
 			bHighPrioIORequest ? AIOP_BelowNormal : AIOP_Low,
 			&AsyncFileCallback);
diff --git a/Engine/Source/Runtime/Engine/Public/Rendering/SkeletalMeshLODRenderData.h b/Engine/Source/Runtime/Engine/Public/Rendering/SkeletalMeshLODRenderData.h
--- a/Engine/Source/Runtime/Engine/Public/Rendering/SkeletalMeshLODRenderData.h
+++ b/Engine/Source/Runtime/Engine/Public/Rendering/SkeletalMeshLODRenderData.h
@@ -149,7 +149,7 @@ public:
 	/** Precooked ray tracing geometry. Used as source data to build skeletal mesh instance ray tracing geometry. */
 	FRayTracingGeometry SourceRayTracingGeometry;
 
-	typename TChooseClass<USE_BULKDATA_STREAMING_TOKEN, FBulkDataStreamingToken, FByteBulkData>::Result StreamingBulkData;
+	FByteBulkData StreamingBulkData;
 
 	/** Whether buffers of this LOD is inlined (i.e. stored in .uexp instead of .ubulk) */
 	uint32 bStreamedDataInlined : 1;
diff --git a/Engine/Source/Runtime/Engine/Public/StaticMeshResources.h b/Engine/Source/Runtime/Engine/Public/StaticMeshResources.h
--- a/Engine/Source/Runtime/Engine/Public/StaticMeshResources.h
+++ b/Engine/Source/Runtime/Engine/Public/StaticMeshResources.h
@@ -393,7 +393,7 @@ public:
 	/** Sum of all vertex and index buffer sizes. Calculated in SerializeBuffers */
 	uint32 BuffersSize;
 
-	typename TChooseClass<USE_BULKDATA_STREAMING_TOKEN, FBulkDataStreamingToken, FByteBulkData>::Result StreamingBulkData;
+	FByteBulkData StreamingBulkData;
 
 #if STATS
 	uint32 StaticMeshIndexMemory;
diff --git a/Engine/Source/Runtime/Landscape/Classes/LandscapeComponent.h b/Engine/Source/Runtime/Landscape/Classes/LandscapeComponent.h
--- a/Engine/Source/Runtime/Landscape/Classes/LandscapeComponent.h
+++ b/Engine/Source/Runtime/Landscape/Classes/LandscapeComponent.h
@@ -15,8 +15,6 @@
 
 #include "LandscapeComponent.generated.h"
 
-#define LANDSCAPE_LOD_STREAMING_USE_TOKEN (!WITH_EDITORONLY_DATA && USE_BULKDATA_STREAMING_TOKEN)
-
 class ALandscape;
 class ALandscapeProxy;
 class FLightingBuildOptions;
@@ -113,16 +111,14 @@ class FLandscapeComponentDerivedData
 	    On device, freed once it has been decompressed. */
 	TArray<uint8> CompressedLandscapeData;
 
-#if LANDSCAPE_LOD_STREAMING_USE_TOKEN
-	TArray<FBulkDataStreamingToken> StreamingLODDataArray;
-#else
 	TArray<FByteBulkData> StreamingLODDataArray;
-#endif
 	
 	/** Cached render data. Only valid on device. */
 	TSharedPtr<FLandscapeMobileRenderData, ESPMode::ThreadSafe > CachedRenderData;
 
+	UE_DEPRECATED(5.0, "The path will not be valid when staged via the IoStore, functionality is being removed")
 	FPackagePath CachedLODDataPackagePath;
+	UE_DEPRECATED(5.0, "The path will not be valid when staged via the IoStore, functionality is being removed")
 	EPackageSegment CachedLODDataPackageSegment;
 
 	friend class ULandscapeLODStreamingProxy;
@@ -434,16 +430,15 @@ class ULandscapeLODStreamingProxy : public UStreamableRenderAsset
 	virtual EStreamableRenderAssetType GetRenderAssetType() const final override { return EStreamableRenderAssetType::LandscapeMeshMobile; }
 	//~ End UStreamableRenderAsset Interface
 
-	UE_DEPRECATED(5.0, "Use GetMipDataPackagePath instead")
+	UE_DEPRECATED(5.0, "The mip will not be able to return a valid file name when staged via the IoStore, functionality is being removed")
 	LANDSCAPE_API bool GetMipDataFilename(const int32 MipIndex, FString& OutBulkDataFilename) const;
+	UE_DEPRECATED(5.0, "The mip will not be able to return a valid package path when staged via the IoStore, functionality is being removed")
 	LANDSCAPE_API bool GetMipDataPackagePath(const int32 MipIndex, FPackagePath& OutBulkDataPackagePath, EPackageSegment& OutPackageSegment) const;
 
-
 	LANDSCAPE_API TArray<float> GetLODScreenSizeArray() const;
 	LANDSCAPE_API TSharedPtr<FLandscapeMobileRenderData, ESPMode::ThreadSafe> GetRenderData() const;
 
-	typedef typename TChooseClass<LANDSCAPE_LOD_STREAMING_USE_TOKEN, FBulkDataStreamingToken, FByteBulkData>::Result BulkDataType;
-	LANDSCAPE_API BulkDataType& GetStreamingLODBulkData(int32 LODIdx) const;
+	LANDSCAPE_API FByteBulkData& GetStreamingLODBulkData(int32 LODIdx) const;
 
 	static LANDSCAPE_API void CancelAllPendingStreamingActions();
 
diff --git a/Engine/Source/Runtime/Landscape/Private/Landscape.cpp b/Engine/Source/Runtime/Landscape/Private/Landscape.cpp
--- a/Engine/Source/Runtime/Landscape/Private/Landscape.cpp
+++ b/Engine/Source/Runtime/Landscape/Private/Landscape.cpp
@@ -3735,7 +3735,6 @@ void FLandscapeComponentDerivedData::InitializeFromUncompressedData(const TArray
 	FinalArchive << CompressedSize;
 	FinalArchive.Serialize(TempCompressedMemory.GetData(), CompressedSize);
 
-#if !LANDSCAPE_LOD_STREAMING_USE_TOKEN
 	const int32 NumStreamingLODs = StreamingLODs.Num();
 	StreamingLODDataArray.Empty(NumStreamingLODs);
 	for (int32 Idx = 0; Idx < NumStreamingLODs; ++Idx)
@@ -3752,7 +3751,6 @@ void FLandscapeComponentDerivedData::InitializeFromUncompressedData(const TArray
 			LODData.Unlock();
 		}
 	}
-#endif
 }
 
 void FLandscapeComponentDerivedData::Serialize(FArchive& Ar, UObject* Owner)
@@ -3767,26 +3765,25 @@ void FLandscapeComponentDerivedData::Serialize(FArchive& Ar, UObject* Owner)
 		StreamingLODDataArray.AddDefaulted(NumStreamingLODs);
 	}
 
+	PRAGMA_DISABLE_DEPRECATION_WARNINGS
 	CachedLODDataPackagePath.Empty();
 	CachedLODDataPackageSegment = EPackageSegment::Header;
+	PRAGMA_ENABLE_DEPRECATION_WARNINGS
 
 	for (int32 Idx = 0; Idx < NumStreamingLODs; ++Idx)
 	{
-#if LANDSCAPE_LOD_STREAMING_USE_TOKEN
-		FByteBulkData LODData;
-		LODData.Serialize(Ar, Owner, Idx);
-		StreamingLODDataArray[Idx] = LODData.CreateStreamingToken();
-#else
 		FByteBulkData& LODData = StreamingLODDataArray[Idx];
 		LODData.Serialize(Ar, Owner, Idx);
-#endif
-#if USE_BULKDATA_STREAMING_TOKEN
-		if (CachedLODDataPackagePath.IsEmpty() && !!(LODData.GetBulkDataFlags() & BULKDATA_Force_NOT_InlinePayload))
+
+		PRAGMA_DISABLE_DEPRECATION_WARNINGS
+		if (CachedLODDataPackagePath.IsEmpty() && 
+			!!(LODData.GetBulkDataFlags() & BULKDATA_Force_NOT_InlinePayload) &&
+			LODData.IsUsingIODispatcher() == false)
 		{
 			CachedLODDataPackagePath = LODData.GetPackagePath();
 			CachedLODDataPackageSegment = LODData.GetPackageSegment();
 		}
-#endif
+		PRAGMA_ENABLE_DEPRECATION_WARNINGS
 	}
 }
 
@@ -4386,7 +4383,12 @@ bool ULandscapeLODStreamingProxy::GetMipDataFilename(const int32 MipIndex, FStri
 {
 	FPackagePath PackagePath;
 	EPackageSegment PackageSegment;
-	if (GetMipDataPackagePath(MipIndex, PackagePath, PackageSegment))
+
+	PRAGMA_DISABLE_DEPRECATION_WARNINGS
+	const bool bResult = GetMipDataPackagePath(MipIndex, PackagePath, PackageSegment);
+	PRAGMA_ENABLE_DEPRECATION_WARNINGS
+
+	if (bResult)
 	{
 		OutBulkDataFilename = PackagePath.GetLocalFullPath(PackageSegment);
 		return true;
@@ -4400,8 +4402,11 @@ bool ULandscapeLODStreamingProxy::GetMipDataPackagePath(const int32 MipIndex, FP
 	const int32 NumStreamingLODs = LandscapeComponent->PlatformData.StreamingLODDataArray.Num();
 	if (MipIndex >= 0 && MipIndex < NumStreamingLODs)
 	{
+		PRAGMA_DISABLE_DEPRECATION_WARNINGS
 		OutPackagePath = LandscapeComponent->PlatformData.CachedLODDataPackagePath;
 		OutPackageSegment = LandscapeComponent->PlatformData.CachedLODDataPackageSegment;
+		PRAGMA_ENABLE_DEPRECATION_WARNINGS
+
 		return true;
 	}
 	return false;
@@ -4409,19 +4414,10 @@ bool ULandscapeLODStreamingProxy::GetMipDataPackagePath(const int32 MipIndex, FP
 
 FIoFilenameHash ULandscapeLODStreamingProxy::GetMipIoFilenameHash(const int32 MipIndex) const
 {
-#if LANDSCAPE_LOD_STREAMING_USE_TOKEN
-	FPackagePath PackagePath;
-	EPackageSegment PackageSegment;
-	if (GetMipDataPackagePath(MipIndex, PackagePath, PackageSegment))
-	{
-		return MakeIoFilenameHash(PackagePath);
-	}
-#else
 	if (LandscapeComponent && LandscapeComponent->PlatformData.StreamingLODDataArray.IsValidIndex(MipIndex))
 	{
 		return LandscapeComponent->PlatformData.StreamingLODDataArray[MipIndex].GetIoFilenameHash();
 	}
-#endif
 	else
 	{
 		return INVALID_IO_FILENAME_HASH;
@@ -4485,7 +4481,7 @@ TSharedPtr<FLandscapeMobileRenderData, ESPMode::ThreadSafe> ULandscapeLODStreami
 	return LandscapeComponent->PlatformData.CachedRenderData;
 }
 
-typename ULandscapeLODStreamingProxy::BulkDataType& ULandscapeLODStreamingProxy::GetStreamingLODBulkData(int32 LODIdx) const
+FByteBulkData& ULandscapeLODStreamingProxy::GetStreamingLODBulkData(int32 LODIdx) const
 {
 	check(LandscapeComponent);
 	return LandscapeComponent->PlatformData.StreamingLODDataArray[LODIdx];
