Don't mark loaded packages as missing, prevents incorrect skipping of imports after requesting load of an invalid export in a valid package

diff --git a/Engine/Source/Runtime/CoreUObject/Private/UObject/UObjectGlobals.cpp b/Engine/Source/Runtime/CoreUObject/Private/UObject/UObjectGlobals.cpp
--- a/Engine/Source/Runtime/CoreUObject/Private/UObject/UObjectGlobals.cpp
+++ b/Engine/Source/Runtime/CoreUObject/Private/UObject/UObjectGlobals.cpp
@@ -972,7 +972,7 @@ UObject* StaticLoadObject(UClass* ObjectClass, UObject* InOuter, const TCHAR* In
 			const FString Error = FText::Format(NSLOCTEXT("Core", "ObjectNotFound", "Failed to find object '{ClassName} {OuterName}.{ObjectName}'"), Arguments).ToString();
 			SafeLoadError(InOuter, LoadFlags, *Error);
 
-			if (InOuter)
+			if (InOuter && !InOuter->HasAnyFlags(RF_WasLoaded))
 			{
 				// Stop future repeated warnings
 				FLinkerLoad::AddKnownMissingPackage(FName(*InOuter->GetPathName()));
