Add code comment to explain why FPayloadToc/FTocEntry have not yet been removed from the code base.

#rnx

diff --git a/Engine/Source/Runtime/CoreUObject/Public/Serialization/VirtualizedBulkData.h b/Engine/Source/Runtime/CoreUObject/Public/Serialization/VirtualizedBulkData.h
--- a/Engine/Source/Runtime/CoreUObject/Public/Serialization/VirtualizedBulkData.h
+++ b/Engine/Source/Runtime/CoreUObject/Public/Serialization/VirtualizedBulkData.h
@@ -528,6 +528,11 @@ using FWordVirtualizedBulkData	= TVirtualizedBulkData<uint16>;
 using FIntVirtualizedBulkData	= TVirtualizedBulkData<int32>;
 using FFloatVirtualizedBulkData	= TVirtualizedBulkData<float>;
 
+/** 
+ * NOTE: FPayloadToc/FTocEntry are now strictly legacy and are only used by the experimental sidecar system.
+ * They will be removed when the sidecar storage is changed to use FPackageTrailer instead.
+ */
+
 /** 
   * Represents an entry to the table of contents found at the start of a payload sidecar file.
   * This might be moved to it's own header and the table of contents made into a proper class 
