Add PackageName.DumpMointPoints, PackageName.RegisterMountPoint and PackageName.UnregisterMountPoint console commands

diff --git a/Engine/Source/Runtime/CoreUObject/Private/Misc/PackageName.cpp b/Engine/Source/Runtime/CoreUObject/Private/Misc/PackageName.cpp
--- a/Engine/Source/Runtime/CoreUObject/Private/Misc/PackageName.cpp
+++ b/Engine/Source/Runtime/CoreUObject/Private/Misc/PackageName.cpp
@@ -34,6 +34,8 @@
 
 DEFINE_LOG_CATEGORY(LogPackageName);
 
+#define LOCTEXT_NAMESPACE "PackageNames"
+
 static FRWLock ContentMountPointCriticalSection;
 
 /** Event that is triggered when a new content path is mounted */
@@ -291,6 +293,8 @@ struct FLongPackagePathsSingleton
 			UpdateValidLongPackageRoots();
 		}		
 
+		UE_LOG(LogPackageName, Display, TEXT("FPackageName: Mount point added: '%s' mounted to '%s'"), *RelativeContentPath, *RootPath);
+		
 		// Let subscribers know that a new content path was mounted
 		FPackageName::OnContentPathMounted().Broadcast( RootPath, RelativeContentPath);
 	}
@@ -339,7 +343,27 @@ struct FLongPackagePathsSingleton
 	}
 
 private:
+#if !UE_BUILD_SHIPPING
+	const FAutoConsoleCommand DumpMountPointsCommand;
+	const FAutoConsoleCommand RegisterMountPointCommand;
+	const FAutoConsoleCommand UnregisterMountPointCommand;
+#endif
+	
 	FLongPackagePathsSingleton()
+#if !UE_BUILD_SHIPPING
+	: DumpMountPointsCommand(
+	TEXT( "PackageName.DumpMountPoints" ),
+	*LOCTEXT("CommandText_DumpMountPoints", "Print registered LongPackagePath mount points").ToString(),
+		FConsoleCommandWithArgsDelegate::CreateRaw( this, &FLongPackagePathsSingleton::ExecDumpMountPoints ) )
+	, RegisterMountPointCommand(
+	TEXT( "PackageName.RegisterMountPoint" ),
+	*LOCTEXT("CommandText_RegisterMountPoint", "<RootPath> <ContentPath> // Register a LongPackagePath mount point").ToString(),
+		FConsoleCommandWithArgsDelegate::CreateRaw( this, &FLongPackagePathsSingleton::ExecInsertMountPoint ) )
+	, UnregisterMountPointCommand(
+	TEXT( "PackageName.UnregisterMountPoint" ),
+	*LOCTEXT("CommandText_UnregisterMountPoint", "<RootPath> <ContentPath> // Remove a LongPackagePath mount point").ToString(),
+		FConsoleCommandWithArgsDelegate::CreateRaw( this, &FLongPackagePathsSingleton::ExecRemoveMountPoint ) )
+#endif
 	{
 		SCOPED_BOOT_TIMING("FPackageName::FLongPackagePathsSingleton");
 
@@ -411,6 +435,146 @@ private:
 
 		UpdateValidLongPackageRoots();
 	}
+
+#if !UE_BUILD_SHIPPING
+	void ExecDumpMountPoints(const TArray<FString>& Args)
+	{
+		UE_LOG(LogPackageName, Log, TEXT("Valid mount points:"));
+
+		FReadScopeLock ScopeLock(ContentMountPointCriticalSection);
+		for (const auto& Pair : ContentRootToPath)
+		{
+			UE_LOG(LogPackageName, Log, TEXT("	'%s' -> '%s'"), *Pair.RootPath, *Pair.ContentPath);
+		}
+
+		UE_LOG(LogPackageName, Log, TEXT("Removable mount points:"));
+		for (const auto& Root : MountPointRootPaths)
+		{
+			UE_LOG(LogPackageName, Log, TEXT("	'%s'"), *Root);
+		}
+		
+	}
+
+	void ExecInsertMountPoint(const TArray<FString>& Args)
+	{
+		if ( Args.Num() < 2 )
+		{
+			UE_LOG(LogPackageName, Log, TEXT("Usage: PackageName.RegisterMountPoint <RootPath> <ContentPath>"));
+			UE_LOG(LogPackageName, Log, TEXT("Example ContentPath: '../../../ProjectName/Content/'"));
+			UE_LOG(LogPackageName, Log, TEXT("Example RootPath: '/Game/'"));
+			return;
+		}
+
+		const FString& RootPath = Args[0];
+		const FString& ContentPath = Args[1];
+
+		if (ContentPath[0] == TEXT('/'))
+		{
+			UE_LOG(LogPackageName, Error, TEXT("PackageName.RegisterMountPoint: Invalid ContentPath, should not start with '/'! Example: '../../../ProjectName/Content/'"));
+			return;
+		}
+
+		if (RootPath[0] != TEXT('/'))
+		{
+			UE_LOG(LogPackageName, Error, TEXT("PackageName.RegisterMountPoint: Invalid RootPath, should start with a '/'! Example: '/Game/'"));
+			return;
+		}
+		
+		Get().InsertMountPoint(RootPath, ContentPath);
+	}
+
+	void ExecRemoveMountPoint(const TArray<FString>& Args)
+	{
+		if ( Args.Num() < 1 || Args.Num() > 2 )
+		{
+			UE_LOG(LogPackageName, Log, TEXT("Usage: PackageName.UnregisterMountPoint <Path>"));
+			UE_LOG(LogPackageName, Log, TEXT("Removes either a Root or Content path if the path is unambiguous"));
+			UE_LOG(LogPackageName, Log, TEXT("Usage: PackageName.UnregisterMountPoint <RootPath> <ContentPath>"));
+			UE_LOG(LogPackageName, Log, TEXT("Removes a specific Root path to Content path mapping"));
+			UE_LOG(LogPackageName, Log, TEXT("Example ContentPath: '../../../ProjectName/Content/'"));
+			UE_LOG(LogPackageName, Log, TEXT("Example RootPath: '/Game/'"));
+			return;
+		}
+
+		if (Args.Num() == 1)
+		{
+			const FString& Path = Args[0];
+			const bool IsRootPath = Path[0] == TEXT('/');
+
+			FString RootPath;
+			FString ContentPath;
+			
+			if (IsRootPath)
+			{
+				FReadScopeLock ScopeLock(ContentMountPointCriticalSection);
+				
+				RootPath = Path;
+				auto Elements = Get().ContentRootToPath.FilterByPredicate([&RootPath](const FPathPair& elem){ return elem.RootPath.Equals(RootPath); });
+				if (Elements.Num() == 0)
+				{
+					UE_LOG(LogPackageName, Error, TEXT("PackageName.UnregisterMountPoint: Root path '%s' is not mounted!"), *RootPath);
+					return;
+				}
+
+				if (Elements.Num() > 1)
+				{
+					UE_LOG(LogPackageName, Error, TEXT("PackageName.UnregisterMountPoint: Root path '%s' is mounted to multiple content paths, specify content path to unmount explicitly!"), *RootPath);
+					for (const FPathPair& elem : Elements)
+					{
+						UE_LOG(LogPackageName, Error, TEXT("- %s"), *elem.ContentPath);
+					}
+					return;
+				}
+
+				ContentPath = Elements[0].ContentPath;
+			}
+			else
+			{
+				FReadScopeLock ScopeLock(ContentMountPointCriticalSection);
+			
+				ContentPath = Path;
+				auto Elements = Get().ContentPathToRoot.FilterByPredicate([&ContentPath](const FPathPair& elem){ return elem.ContentPath.Equals(ContentPath); });
+				if (Elements.Num() == 0)
+				{
+					UE_LOG(LogPackageName, Error, TEXT("PackageName.UnregisterMountPoint: Content path '%s' is not mounted!"), *ContentPath);
+					return;
+				}
+
+				if (Elements.Num() > 1)
+				{
+					UE_LOG(LogPackageName, Error, TEXT("PackageName.UnregisterMountPoint: Content path '%s' is mounted to multiple root paths, specify root path to unmount explicitly!"), *ContentPath);
+					for (const FPathPair& elem : Elements)
+					{
+						UE_LOG(LogPackageName, Error, TEXT("- %s"), *elem.RootPath);
+					}
+					return;
+				}
+
+				RootPath = Elements[0].RootPath;
+			}
+			
+			Get().RemoveMountPoint(RootPath, ContentPath);
+		}
+		else
+		{
+			const FString& RootPath = Args[0];
+			const FString& ContentPath = Args[1];
+			if (ContentPath[0] == TEXT('/'))
+			{
+				UE_LOG(LogPackageName, Error, TEXT("PackageName.UnregisterMountPoint: Invalid ContentPath, should not start with '/'! Example: '../../../ProjectName/Content/'"));
+				return;
+			}
+
+			if (RootPath[0] != TEXT('/'))
+			{
+				UE_LOG(LogPackageName, Error, TEXT("PackageName.UnregisterMountPoint: Invalid RootPath, should start with '/'! Example: '/Game/'"));
+				return;
+			}
+
+			Get().RemoveMountPoint(RootPath, ContentPath);
+		}
+	}
+#endif
 };
 
 void FPackageName::InternalFilenameToLongPackageName(FStringView InFilename, FStringBuilderBase& OutPackageName)
@@ -2442,3 +2606,4 @@ bool FPackageNameTests::RunTest(const FString& Parameters)
 }
 
 #endif //WITH_DEV_AUTOMATION_TESTS
+#undef LOCTEXT_NAMESPACE // "PackageNames"
\ No newline at end of file
