AsyncLoading2:
Add support for importing uncooked source packages from cooked packages
#rnx

diff --git a/Engine/Source/Runtime/Core/Public/IO/PackageStore.h b/Engine/Source/Runtime/Core/Public/IO/PackageStore.h
--- a/Engine/Source/Runtime/Core/Public/IO/PackageStore.h
+++ b/Engine/Source/Runtime/Core/Public/IO/PackageStore.h
@@ -45,6 +45,10 @@ struct FPackageStoreEntry
 	FPackageStoreExportInfo ExportInfo;
 	TArrayView<const FPackageId> ImportedPackageIds;
 	TArrayView<const FSHAHash> ShaderMapHashes;
+#if WITH_EDITOR
+	FName UncookedPackageName;
+	uint8 UncookedPackageHeaderExtension; // TODO: Can't include PackagePath.h
+#endif
 };
 
 /**
diff --git a/Engine/Source/Runtime/CoreUObject/Private/Serialization/AsyncLoading.cpp b/Engine/Source/Runtime/CoreUObject/Private/Serialization/AsyncLoading.cpp
--- a/Engine/Source/Runtime/CoreUObject/Private/Serialization/AsyncLoading.cpp
+++ b/Engine/Source/Runtime/CoreUObject/Private/Serialization/AsyncLoading.cpp
@@ -6872,7 +6872,7 @@ void FAsyncLoadingThread::FlushLoading(int32 PackageID)
 		GFlushAsyncLoadingTime += ElapsedTime;
 		GFlushAsyncLoadingCount++;
 
-		check(PackageID != INDEX_NONE || !IsAsyncLoading());
+		check(PackageID != INDEX_NONE || !IsAsyncLoadingPackages());
 
 	}
 }
diff --git a/Engine/Source/Runtime/CoreUObject/Private/Serialization/AsyncLoading2.cpp b/Engine/Source/Runtime/CoreUObject/Private/Serialization/AsyncLoading2.cpp
--- a/Engine/Source/Runtime/CoreUObject/Private/Serialization/AsyncLoading2.cpp
+++ b/Engine/Source/Runtime/CoreUObject/Private/Serialization/AsyncLoading2.cpp
@@ -835,16 +835,6 @@ public:
 				Package->GetFlags(), Package->GetInternalFlags(), RefCount);
 			checkf(false, TEXT("Package %s destroyed with RefCount"), *Package->GetName());
 		}
-		else if (RefCount < 0)
-		{
-			UE_LOG(LogStreaming, Error,
-				TEXT("RemovePackage: %s %s (0x%llX) - with (ObjectFlags=%x, InternalObjectFlags=%x) - ")
-				TEXT("Package not found!"),
-				*Package->GetName(),
-				*Package->GetLoadedPath().GetDebugName(), Package->GetPackageId().Value(),
-				Package->GetFlags(), Package->GetInternalFlags());
-			checkf(false, TEXT("Package %s not found"), *Package->GetName());
-		}
 	}
 
 	void RemovePackages(const FUnreachablePackages& UnreachablePackages)
@@ -966,7 +956,7 @@ struct FPackageImportStore
 	}
 
 private:
-	void AddAsyncFlags(UPackage* ImportedPackage)
+	void PinPublicExportsForGC(UPackage* ImportedPackage)
 	{
 		UE_ASYNC_UPACKAGE_DEBUG(ImportedPackage);
 		
@@ -978,15 +968,15 @@ private:
 		{
 			if (Object->HasAllFlags(RF_Public | RF_WasLoaded))
 			{
-				checkf(!Object->HasAnyInternalFlags(EInternalObjectFlags::Async), TEXT("%s"), *Object->GetFullName());
-				Object->SetInternalFlags(EInternalObjectFlags::Async);
+				checkf(!Object->HasAnyInternalFlags(EInternalObjectFlags::LoaderImport), TEXT("%s"), *Object->GetFullName());
+				Object->SetInternalFlags(EInternalObjectFlags::LoaderImport);
 			}
 		}, /* bIncludeNestedObjects*/ true);
-		checkf(!ImportedPackage->HasAnyInternalFlags(EInternalObjectFlags::Async), TEXT("%s"), *ImportedPackage->GetFullName());
-		ImportedPackage->SetInternalFlags(EInternalObjectFlags::Async);
+		checkf(!ImportedPackage->HasAnyInternalFlags(EInternalObjectFlags::LoaderImport), TEXT("%s"), *ImportedPackage->GetFullName());
+		ImportedPackage->SetInternalFlags(EInternalObjectFlags::LoaderImport);
 	}
 
-	void ClearAsyncFlags(UPackage* ImportedPackage)
+	void UnpinPublicExportsForGC(UPackage* ImportedPackage)
 	{
 		UE_ASYNC_UPACKAGE_DEBUG(ImportedPackage);
 
@@ -998,12 +988,12 @@ private:
 		{
 			if (Object->HasAllFlags(RF_Public | RF_WasLoaded))
 			{
-				checkf(Object->HasAnyInternalFlags(EInternalObjectFlags::Async), TEXT("%s"), *Object->GetFullName());
-				Object->AtomicallyClearInternalFlags(EInternalObjectFlags::Async);
+				checkf(Object->HasAnyInternalFlags(EInternalObjectFlags::LoaderImport), TEXT("%s"), *Object->GetFullName());
+				Object->AtomicallyClearInternalFlags(EInternalObjectFlags::LoaderImport);
 			}
 		}, /* bIncludeNestedObjects*/ true);
-		checkf(ImportedPackage->HasAnyInternalFlags(EInternalObjectFlags::Async), TEXT("%s"), *ImportedPackage->GetFullName());
-		ImportedPackage->AtomicallyClearInternalFlags(EInternalObjectFlags::Async);
+		checkf(ImportedPackage->HasAnyInternalFlags(EInternalObjectFlags::LoaderImport), TEXT("%s"), *ImportedPackage->GetFullName());
+		ImportedPackage->AtomicallyClearInternalFlags(EInternalObjectFlags::LoaderImport);
 	}
 
 public:
@@ -1014,7 +1004,7 @@ public:
 			FLoadedPackageRef& PackageRef = LoadedPackageStore.GetPackageRef(ImportedPackageId);
 			if (PackageRef.AddRef())
 			{
-				AddAsyncFlags(PackageRef.GetPackage());
+				PinPublicExportsForGC(PackageRef.GetPackage());
 			}
 		}
 		if (Desc.bCanBeImported)
@@ -1023,7 +1013,7 @@ public:
 			PackageRef.ClearErrorFlags();
 			if (PackageRef.AddRef())
 			{
-				AddAsyncFlags(PackageRef.GetPackage());
+				PinPublicExportsForGC(PackageRef.GetPackage());
 			}
 		}
 	}
@@ -1035,7 +1025,7 @@ public:
 			FLoadedPackageRef& PackageRef = LoadedPackageStore.GetPackageRef(ImportedPackageId);
 			if (PackageRef.ReleaseRef(Desc.UPackageId, ImportedPackageId))
 			{
-				ClearAsyncFlags(PackageRef.GetPackage());
+				UnpinPublicExportsForGC(PackageRef.GetPackage());
 			}
 		}
 		if (Desc.bCanBeImported)
@@ -1044,7 +1034,7 @@ public:
 			FLoadedPackageRef& PackageRef = LoadedPackageStore.GetPackageRef(Desc.UPackageId);
 			if (PackageRef.ReleaseRef(Desc.UPackageId, Desc.UPackageId))
 			{
-				ClearAsyncFlags(PackageRef.GetPackage());
+				UnpinPublicExportsForGC(PackageRef.GetPackage());
 			}
 		}
 	}
@@ -1631,18 +1621,6 @@ struct FAsyncPackage2
 		}
 	}
 
-	void PinObjectForGC(UObject* Object, bool bIsNewObject)
-	{
-		if (bIsNewObject && !IsInGameThread())
-		{
-			checkf(Object->HasAnyInternalFlags(EInternalObjectFlags::Async), TEXT("%s"), *Object->GetFullName());
-		}
-		else
-		{
-			Object->SetInternalFlags(EInternalObjectFlags::Async);
-		}
-	}
-
 	void ClearConstructedObjects();
 
 	/** Returns the UPackage wrapped by this, if it is valid */
@@ -1977,7 +1955,7 @@ class FAsyncLoadingThread2 final
 {
 	friend struct FAsyncPackage2;
 public:
-	FAsyncLoadingThread2(FIoDispatcher& IoDispatcher);
+	FAsyncLoadingThread2(FIoDispatcher& IoDispatcher, IAsyncPackageLoader* InUncookedPackageLoader);
 	virtual ~FAsyncLoadingThread2();
 
 private:
@@ -2059,6 +2037,8 @@ private:
 	/** I/O Dispatcher */
 	FIoDispatcher& IoDispatcher;
 
+	IAsyncPackageLoader* UncookedPackageLoader;
+
 	TSharedPtr<IPackageStore> PackageStore;
 	FLoadedPackageStore LoadedPackageStore;
 	FGlobalImportStore GlobalImportStore;
@@ -3263,6 +3243,73 @@ void FAsyncPackage2::ImportPackagesRecursive(FIoBatch& IoBatch, IPackageStore& P
 			Data.ImportedAsyncPackages[ImportedPackageIndex++] = nullptr;
 			continue;
 		}
+#if WITH_EDITOR
+		else if (!ImportedPackageEntry.UncookedPackageName.IsNone())
+		{
+			UPackage* UncookedPackage = nullptr;
+			if (!PackageRef.AreAllPublicExportsLoaded())
+			{
+				FPackagePath ImportedPackagePath = FPackagePath::FromPackageNameUnchecked(ImportedPackageEntry.UncookedPackageName);
+				ImportedPackagePath.SetHeaderExtension(static_cast<EPackageExtension>(ImportedPackageEntry.UncookedPackageHeaderExtension));
+				AsyncLoadingThread.UncookedPackageLoader->LoadPackage(ImportedPackagePath, NAME_None, FLoadPackageAsyncDelegate(), PKG_None, INDEX_NONE, 0, nullptr);
+				AsyncLoadingThread.UncookedPackageLoader->FlushLoading(INDEX_NONE);
+				UncookedPackage = FindObjectFast<UPackage>(nullptr, ImportedPackagePath.GetPackageFName());
+				if (UncookedPackage)
+				{
+					UncookedPackage->SetCanBeImportedFlag(true);
+					UncookedPackage->SetPackageId(ImportedPackageId);
+					checkf(!UncookedPackage->HasAnyInternalFlags(EInternalObjectFlags::LoaderImport), TEXT("%s"), *UncookedPackage->GetFullName());
+					UncookedPackage->SetInternalFlags(EInternalObjectFlags::LoaderImport);
+
+					ForEachObjectWithOuter(UncookedPackage, [this, ImportedPackageId](UObject* Object)
+					{
+						if (Object->HasAllFlags(RF_Public | RF_WasLoaded))
+						{
+							checkf(!Object->HasAnyInternalFlags(EInternalObjectFlags::LoaderImport), TEXT("%s"), *Object->GetFullName());
+							Object->SetInternalFlags(EInternalObjectFlags::LoaderImport);
+
+							TArray<FName, TInlineAllocator<64>> FullPath;
+							FullPath.Add(Object->GetFName());
+							UObject* Outer = Object->GetOuter();
+							while (Outer)
+							{
+								FullPath.Add(Outer->GetFName());
+								Outer = Outer->GetOuter();
+							}
+							TStringBuilder<256> PackageRelativeExportPath;
+							for (int32 PathIndex = FullPath.Num() - 2; PathIndex >= 0; --PathIndex)
+							{
+								TCHAR NameStr[FName::StringBufferSize];
+								uint32 NameLen = FullPath[PathIndex].ToString(NameStr);
+								for (uint32 I = 0; I < NameLen; ++I)
+								{
+									NameStr[I] = TChar<TCHAR>::ToLower(NameStr[I]);
+								}
+								PackageRelativeExportPath.Append('/');
+								PackageRelativeExportPath.Append(FStringView(NameStr, NameLen));
+							}
+							uint64 ExportHash = CityHash64(reinterpret_cast<const char*>(PackageRelativeExportPath.GetData() + 1), (PackageRelativeExportPath.Len() - 1) * sizeof(TCHAR));
+							ImportStore.StoreGlobalObject(ImportedPackageId, ExportHash, Object);
+						}
+					}, /* bIncludeNestedObjects*/ true);
+				}
+				PackageRef.SetPackage(UncookedPackage);
+				PackageRef.SetAllPublicExportsLoaded();
+			}
+			else
+			{
+				UncookedPackage = PackageRef.GetPackage();
+			}
+			if (!UncookedPackage)
+			{
+				PackageRef.SetHasFailed();
+				UE_ASYNC_PACKAGE_LOG(Warning, Desc, TEXT("ImportPackages: SkipPackage"),
+					TEXT("Failed to load uncooked imported package with id '0x%llX' ('%s')"), ImportedPackageId.Value(), *ImportedPackageEntry.UncookedPackageName.ToString());
+			}
+			Data.ImportedAsyncPackages[ImportedPackageIndex++] = nullptr;
+			continue;
+		}
+#endif
 
 		FAsyncPackage2* ImportedPackage = nullptr;
 		bool bInserted = false;
@@ -3874,11 +3921,12 @@ void FAsyncPackage2::EventDrivenCreateExport(int32 LocalExportIndex)
 	}
 
 	check(Object);
-	PinObjectForGC(Object, bIsNewObject);
-
+	EInternalObjectFlags FlagsToSet = EInternalObjectFlags::Async;
+	
 	if (Desc.bCanBeImported && Export.PublicExportHash)
 	{
 		check(Object->HasAnyFlags(RF_Public));
+		FlagsToSet |= EInternalObjectFlags::LoaderImport;
 		ImportStore.StoreGlobalObject(Desc.UPackageId, Export.PublicExportHash, Object);
 
 		UE_ASYNC_PACKAGE_LOG_VERBOSE(VeryVerbose, Desc, TEXT("CreateExport"),
@@ -3889,6 +3937,7 @@ void FAsyncPackage2::EventDrivenCreateExport(int32 LocalExportIndex)
 		UE_ASYNC_PACKAGE_LOG_VERBOSE(VeryVerbose, Desc, TEXT("CreateExport"), TEXT("Created %s export %s. Not tracked."),
 			Object->HasAnyFlags(RF_Public) ? TEXT("public") : TEXT("private"), *Object->GetPathName());
 	}
+	Object->SetInternalFlags(FlagsToSet);
 }
 
 bool FAsyncPackage2::EventDrivenSerializeExport(int32 LocalExportIndex, FExportArchive& Ar)
@@ -4858,9 +4907,10 @@ EAsyncPackageState::Type FAsyncLoadingThread2::TickAsyncLoadingFromGameThread(FA
 	return Result;
 }
 
-FAsyncLoadingThread2::FAsyncLoadingThread2(FIoDispatcher& InIoDispatcher)
+FAsyncLoadingThread2::FAsyncLoadingThread2(FIoDispatcher& InIoDispatcher, IAsyncPackageLoader* InUncookedPackageLoader)
 	: Thread(nullptr)
 	, IoDispatcher(InIoDispatcher)
+	, UncookedPackageLoader(InUncookedPackageLoader)
 {
 #if !WITH_IOSTORE_IN_EDITOR
 	GEventDrivenLoaderEnabled = true;
@@ -5261,6 +5311,7 @@ static void VerifyLoadFlagsWhenFinishedLoading()
 			const EObjectFlags Flags = Obj->GetFlags();
 			const bool bHasAnyAsyncFlags = !!(InternalFlags & AsyncFlags);
 			const bool bHasAnyLoadIntermediateFlags = !!(Flags & LoadIntermediateFlags);
+			const bool bHasLoaderImportFlag = !!(InternalFlags & EInternalObjectFlags::LoaderImport);
 			const bool bWasLoaded = !!(Flags & RF_WasLoaded);
 			const bool bLoadCompleted = !!(Flags & RF_LoadCompleted);
 
@@ -5271,6 +5322,13 @@ static void VerifyLoadFlagsWhenFinishedLoading()
 				Flags,
 				InternalFlags);
 
+			ensureMsgf(!bHasLoaderImportFlag || GUObjectArray.IsDisregardForGC(Obj),
+				TEXT("Object '%s' (ObjectFlags=%X, InternalObjectFlags=%x) should not have the LoaderImport flag now")
+				TEXT(", or this check is incorrectly reached during active loading."),
+				*Obj->GetFullName(),
+				Flags,
+				InternalFlags);
+
 			if (bWasLoaded)
 			{
 				const bool bIsPackage = Obj->IsA(UPackage::StaticClass());
@@ -5319,12 +5377,7 @@ FORCENOINLINE static void FilterUnreachableObjects(
 			else
 			{
 				UPackage* Package = static_cast<UPackage*>(Object);
-#if WITH_IOSTORE_IN_EDITOR
-				if (Package->HasAnyPackageFlags(PKG_Cooked))
-#endif
-				{
-					Packages.Emplace(Package);
-				}
+				Packages.Emplace(Package);
 			}
 		}
 	}
@@ -5403,7 +5456,7 @@ void FAsyncLoadingThread2::NotifyUnreachableObjects(const TArrayView<FUObjectIte
 	}
 
 #if ALT2_VERIFY_ASYNC_FLAGS
-	if (!IsAsyncLoadingPackages())
+	if (!IsAsyncLoading())
 	{
 		LoadedPackageStore.VerifyLoadedPackages();
 		VerifyLoadFlagsWhenFinishedLoading();
@@ -5543,9 +5596,6 @@ void FAsyncPackage2::ClearConstructedObjects()
 	}
 	ConstructedObjects.Empty();
 
-	// the async flag of all GC'able public export objects in non-temp packages are handled by FGlobalImportStore::ClearAsyncFlags
-	const bool bShouldClearAsyncFlagForPublicExports = GUObjectArray.IsDisregardForGC(LinkerRoot) || !Desc.bCanBeImported;
-
 	for (FExportObject& Export : Data.Exports)
 	{
 		if (Export.bFiltered | Export.bExportLoadFailed)
@@ -5557,26 +5607,12 @@ void FAsyncPackage2::ClearConstructedObjects()
 		check(Object);
 		checkf(Object->HasAnyFlags(RF_WasLoaded), TEXT("%s"), *Object->GetFullName());
 		checkf(Object->HasAnyInternalFlags(EInternalObjectFlags::Async), TEXT("%s"), *Object->GetFullName());
-		if (bShouldClearAsyncFlagForPublicExports || !Object->HasAnyFlags(RF_Public))
-		{
-			Object->AtomicallyClearInternalFlags(EInternalObjectFlags::AsyncLoading | EInternalObjectFlags::Async);
-		}
-		else
-		{
-			Object->AtomicallyClearInternalFlags(EInternalObjectFlags::AsyncLoading);
-		}
+		Object->AtomicallyClearInternalFlags(EInternalObjectFlags::AsyncLoading | EInternalObjectFlags::Async);
 	}
 
 	if (LinkerRoot)
 	{
-		if (bShouldClearAsyncFlagForPublicExports)
-		{
-			LinkerRoot->AtomicallyClearInternalFlags(EInternalObjectFlags::AsyncLoading | EInternalObjectFlags::Async);
-		}
-		else
-		{
-			LinkerRoot->AtomicallyClearInternalFlags(EInternalObjectFlags::AsyncLoading);
-		}
+		LinkerRoot->AtomicallyClearInternalFlags(EInternalObjectFlags::AsyncLoading | EInternalObjectFlags::Async);
 	}
 }
 
@@ -5740,7 +5776,12 @@ void FAsyncPackage2::CreateUPackage(const FZenPackageSummary* PackageSummary, co
 		check(LinkerRoot->HasAnyFlags(RF_WasLoaded));
 	}
 
-	PinObjectForGC(LinkerRoot, bCreatedLinkerRoot);
+	EInternalObjectFlags FlagsToSet = EInternalObjectFlags::Async;
+	if (Desc.bCanBeImported)
+	{
+		FlagsToSet |= EInternalObjectFlags::LoaderImport;
+	}
+	LinkerRoot->SetInternalFlags(FlagsToSet);
 
 	if (bCreatedLinkerRoot)
 	{
@@ -5943,7 +5984,7 @@ EAsyncPackageState::Type FAsyncLoadingThread2::ProcessLoadingFromGameThread(FAsy
 	CSV_CUSTOM_STAT(FileIO, ExistingQueuedPackagesQueueDepth, GetNumAsyncPackages(), ECsvCustomStatOp::Set);
 	
 	TickAsyncLoadingFromGameThread(ThreadState, bUseTimeLimit, bUseFullTimeLimit, TimeLimit);
-	return IsAsyncLoading() ? EAsyncPackageState::TimeOut : EAsyncPackageState::Complete;
+	return IsAsyncLoadingPackages() ? EAsyncPackageState::TimeOut : EAsyncPackageState::Complete;
 }
 
 void FAsyncLoadingThread2::FlushLoading(int32 RequestId)
@@ -5995,7 +6036,7 @@ void FAsyncLoadingThread2::FlushLoading(int32 RequestId)
 			}
 		}
 
-		check(RequestId != INDEX_NONE || !IsAsyncLoading());
+		check(RequestId != INDEX_NONE || !IsAsyncLoadingPackages());
 	}
 }
 
@@ -6042,9 +6083,9 @@ EAsyncPackageState::Type FAsyncLoadingThread2::ProcessLoadingUntilCompleteFromGa
 	return TimeLimit <= 0 ? EAsyncPackageState::TimeOut : EAsyncPackageState::Complete;
 }
 
-IAsyncPackageLoader* MakeAsyncPackageLoader2(FIoDispatcher& InIoDispatcher)
+IAsyncPackageLoader* MakeAsyncPackageLoader2(FIoDispatcher& InIoDispatcher, IAsyncPackageLoader* InUncookedPackageLoader)
 {
-	return new FAsyncLoadingThread2(InIoDispatcher);
+	return new FAsyncLoadingThread2(InIoDispatcher, InUncookedPackageLoader);
 }
 
 #endif //WITH_ASYNCLOADING2
diff --git a/Engine/Source/Runtime/CoreUObject/Private/Serialization/EditorPackageLoader.cpp b/Engine/Source/Runtime/CoreUObject/Private/Serialization/EditorPackageLoader.cpp
--- a/Engine/Source/Runtime/CoreUObject/Private/Serialization/EditorPackageLoader.cpp
+++ b/Engine/Source/Runtime/CoreUObject/Private/Serialization/EditorPackageLoader.cpp
@@ -18,8 +18,8 @@ class FEditorPackageLoader final
 public:
 	FEditorPackageLoader(FIoDispatcher& InIoDispatcher, IEDLBootNotificationManager& InEDLBootNotificationManager)
 	{
-		CookedPackageLoader.Reset(MakeAsyncPackageLoader2(InIoDispatcher));
 		UncookedPackageLoader.Reset(new FAsyncLoadingThread(/** ThreadIndex = */ 0, InEDLBootNotificationManager));
+		CookedPackageLoader.Reset(MakeAsyncPackageLoader2(InIoDispatcher, UncookedPackageLoader.Get()));
 	}
 
 	virtual ~FEditorPackageLoader () { }
diff --git a/Engine/Source/Runtime/CoreUObject/Public/Serialization/AsyncLoading2.h b/Engine/Source/Runtime/CoreUObject/Public/Serialization/AsyncLoading2.h
--- a/Engine/Source/Runtime/CoreUObject/Public/Serialization/AsyncLoading2.h
+++ b/Engine/Source/Runtime/CoreUObject/Public/Serialization/AsyncLoading2.h
@@ -335,6 +335,6 @@ COREUOBJECT_API void FindAllRuntimeScriptPackages(TArray<UPackage*>& OutPackages
  *
  * @return The async package loader.
  */
-IAsyncPackageLoader* MakeAsyncPackageLoader2(FIoDispatcher& InIoDispatcher);
+IAsyncPackageLoader* MakeAsyncPackageLoader2(FIoDispatcher& InIoDispatcher, IAsyncPackageLoader* UncookedPackageLoader = nullptr);
 
 #endif
diff --git a/Engine/Source/Runtime/CoreUObject/Public/UObject/ObjectMacros.h b/Engine/Source/Runtime/CoreUObject/Public/UObject/ObjectMacros.h
--- a/Engine/Source/Runtime/CoreUObject/Public/UObject/ObjectMacros.h
+++ b/Engine/Source/Runtime/CoreUObject/Public/UObject/ObjectMacros.h
@@ -560,6 +560,7 @@ enum class EInternalObjectFlags : int32
 {
 	None = 0,
 
+	LoaderImport = 1 << 20, ///< Object is ready to be imported by another package during loading
 	Garbage = 1 << 21, ///< Garbage from logical point of view and should not be referenced. This flag is mirrored in EObjectFlags as RF_Garbage for performance
 	PersistentGarbage = 1 << 22, ///< Same as above but referenced through a persistent reference so it can't be GC'd
 	ReachableInCluster = 1 << 23, ///< External reference to object in cluster exists
@@ -572,11 +573,11 @@ enum class EInternalObjectFlags : int32
 	RootSet = 1 << 30, ///< Object will not be garbage collected, even if unreferenced.
 	PendingConstruction = 1 << 31, ///< Object didn't have its class constructor called yet (only the UObjectBase one to initialize its most basic members)
 
-	GarbageCollectionKeepFlags = Native | Async | AsyncLoading,
+	GarbageCollectionKeepFlags = Native | Async | AsyncLoading | LoaderImport,
 	MirroredFlags = Garbage | PendingKill, /// Flags mirrored in EObjectFlags
 
 	//~ Make sure this is up to date!
-	AllFlags = Garbage | PersistentGarbage | ReachableInCluster | ClusterRoot | Native | Async | AsyncLoading | Unreachable | PendingKill | RootSet | PendingConstruction
+	AllFlags = LoaderImport | Garbage | PersistentGarbage | ReachableInCluster | ClusterRoot | Native | Async | AsyncLoading | Unreachable | PendingKill | RootSet | PendingConstruction
 };
 ENUM_CLASS_FLAGS(EInternalObjectFlags);
 
diff --git a/Engine/Source/Runtime/CoreUObject/Public/UObject/UObjectArray.h b/Engine/Source/Runtime/CoreUObject/Public/UObject/UObjectArray.h
--- a/Engine/Source/Runtime/CoreUObject/Public/UObject/UObjectArray.h
+++ b/Engine/Source/Runtime/CoreUObject/Public/UObject/UObjectArray.h
@@ -145,7 +145,7 @@ struct FUObjectItem
 		while (1)
 		{
 			int32 StartValue = int32(Flags);
-			if (StartValue & int32(FlagToSet))
+			if ((StartValue & int32(FlagToSet)) == int32(FlagToSet))
 			{
 				break;
 			}
diff --git a/Engine/Source/Runtime/PakFile/Private/FilePackageStore.cpp b/Engine/Source/Runtime/PakFile/Private/FilePackageStore.cpp
--- a/Engine/Source/Runtime/PakFile/Private/FilePackageStore.cpp
+++ b/Engine/Source/Runtime/PakFile/Private/FilePackageStore.cpp
@@ -7,47 +7,94 @@
 #include "ProfilingDebugging/CountersTrace.h"
 #include "Internationalization/PackageLocalizationManager.h"
 #include "Misc/ScopeLock.h"
+#include "Misc/PackageName.h"
 
 //PRAGMA_DISABLE_OPTIMIZATION
 
 DEFINE_LOG_CATEGORY_STATIC(LogFilePackageStore, Log, All);
 
-thread_local bool FFilePackageStore::bIsLockedOnThread = false;
+thread_local int32 FFilePackageStore::LockedOnThreadCount = 0;
 
 FFilePackageStore::FFilePackageStore()
 {
+#if WITH_EDITOR
+	OnContentPathMountedDelegateHandle = FPackageName::OnContentPathMounted().AddLambda([this](const FString& InAssetPath, const FString& InFilesystemPath)
+		{
+			{
+				FScopeLock _(&UncookedPackageRootsLock);
+				PendingUncookedPackageRoots.Add(InFilesystemPath);
+			}
+			{
+				FWriteScopeLock _(EntriesLock);
+				bNeedsUpdate = true;
+			}
+		});
+#endif
+}
+
+FFilePackageStore::~FFilePackageStore()
+{
+#if WITH_EDITOR
+	FPackageName::OnContentPathMounted().Remove(OnContentPathMountedDelegateHandle);
+#endif
 }
 
 void FFilePackageStore::Initialize()
 {
+#if WITH_EDITOR
+	TArray<FString> RootPaths;
+	FPackageName::QueryRootContentPaths(RootPaths);
+	{
+		FScopeLock _(&UncookedPackageRootsLock);
+		PendingUncookedPackageRoots.Append(RootPaths);
+	}
+	{
+		FWriteScopeLock _(EntriesLock);
+		bNeedsUpdate = true;
+	}
+#endif
 }
 
 void FFilePackageStore::Lock()
 {
-	EntriesLock.ReadLock();
-	if (bNeedsUpdate)
+	if (!LockedOnThreadCount)
 	{
-		Update();
+		EntriesLock.ReadLock();
+		if (bNeedsUpdate)
+		{
+			Update();
+		}
 	}
-	bIsLockedOnThread = true;
+	++LockedOnThreadCount;
 }
 
 void FFilePackageStore::Unlock()
 {
-	check(bIsLockedOnThread);
-	EntriesLock.ReadUnlock();
-	bIsLockedOnThread = false;
+	check(LockedOnThreadCount > 0);
+	if (--LockedOnThreadCount == 0)
+	{
+		EntriesLock.ReadUnlock();
+	}
 }
 
 bool FFilePackageStore::DoesPackageExist(FPackageId PackageId)
 {
-	check(bIsLockedOnThread);
+	check(LockedOnThreadCount);
 	return PackageId.IsValid() && StoreEntriesMap.Contains(PackageId);
 }
 
 EPackageStoreEntryStatus FFilePackageStore::GetPackageStoreEntry(FPackageId PackageId, FPackageStoreEntry& OutPackageStoreEntry)
 {
-	check(bIsLockedOnThread);
+	check(LockedOnThreadCount);
+#if WITH_EDITOR
+	const FUncookedPackage* FindUncookedPackage = UncookedPackagesMap.Find(PackageId);
+	if (FindUncookedPackage)
+	{
+		OutPackageStoreEntry.UncookedPackageName = FindUncookedPackage->PackageName;
+		OutPackageStoreEntry.UncookedPackageHeaderExtension = static_cast<uint8>(FindUncookedPackage->HeaderExtension);
+		return EPackageStoreEntryStatus::Ok;
+	}
+#endif
 	const FFilePackageStoreEntry* FindEntry = StoreEntriesMap.FindRef(PackageId);
 	if (FindEntry)
 	{
@@ -57,15 +104,12 @@ EPackageStoreEntryStatus FFilePackageStore::GetPackageStoreEntry(FPackageId Pack
 		OutPackageStoreEntry.ShaderMapHashes = MakeArrayView(FindEntry->ShaderMapHashes.Data(), FindEntry->ShaderMapHashes.Num());
 		return EPackageStoreEntryStatus::Ok;
 	}
-	else
-	{
-		return EPackageStoreEntryStatus::Missing;
-	}
+	return EPackageStoreEntryStatus::Missing;
 }
 
 bool FFilePackageStore::GetPackageRedirectInfo(FPackageId PackageId, FName& OutSourcePackageName, FPackageId& OutRedirectedToPackageId)
 {
-	check(bIsLockedOnThread);
+	check(LockedOnThreadCount);
 	TTuple<FName, FPackageId>* FindRedirect = RedirectsPackageMap.Find(PackageId);
 	if (FindRedirect)
 	{
@@ -172,7 +216,47 @@ void FFilePackageStore::Update()
 			RedirectsPackageMap.Emplace(Redirect.SourcePackageId, MakeTuple(SourcePackageName, Redirect.TargetPackageId));
 		}
 	}
+
+#if WITH_EDITOR
+	{
+		FScopeLock _(&UncookedPackageRootsLock);
+		uint64 TotalAddedCount = 0;
+		if (!PendingUncookedPackageRoots.IsEmpty())
+		{
+			UE_LOG(LogFilePackageStore, Display, TEXT("Searching for uncooked packages in %d new roots..."), PendingUncookedPackageRoots.Num());
+			for (const FString& RootPath : PendingUncookedPackageRoots)
+			{
+				TotalAddedCount += AddUncookedPackagesFromRoot(RootPath);
+			}
+			PendingUncookedPackageRoots.Empty();
+			UE_LOG(LogFilePackageStore, Display, TEXT("Found %lld uncooked packages"), TotalAddedCount);
+		}
+	}
+#endif
+
 	bNeedsUpdate = false;
 }
 
+#if WITH_EDITOR
+uint64 FFilePackageStore::AddUncookedPackagesFromRoot(const FString& RootPath)
+{
+	uint64 TotalAddedCount = 0;
+	FPackageName::IteratePackagesInDirectory(RootPath, [this, &RootPath, &TotalAddedCount](const TCHAR* InPackageFileName) -> bool
+		{
+			FPackagePath PackagePath = FPackagePath::FromLocalPath(InPackageFileName);
+			FName PackageName = PackagePath.GetPackageFName();
+			if (!PackageName.IsNone())
+			{
+				FPackageId PackageId = FPackageId::FromName(PackageName);
+				FUncookedPackage& UncookedPackage = UncookedPackagesMap.FindOrAdd(PackageId);
+				UncookedPackage.PackageName = PackageName;
+				UncookedPackage.HeaderExtension = PackagePath.GetHeaderExtension();
+				++TotalAddedCount;
+			}
+			return true;
+		});
+	return TotalAddedCount;
+}
+#endif
+
 //PRAGMA_ENABLE_OPTIMIZATION
\ No newline at end of file
diff --git a/Engine/Source/Runtime/PakFile/Private/FilePackageStore.h b/Engine/Source/Runtime/PakFile/Private/FilePackageStore.h
--- a/Engine/Source/Runtime/PakFile/Private/FilePackageStore.h
+++ b/Engine/Source/Runtime/PakFile/Private/FilePackageStore.h
@@ -5,6 +5,7 @@
 #include "CoreMinimal.h"
 #include "IO/PackageId.h"
 #include "IO/PackageStore.h"
+#include "Misc/PackagePath.h"
 
 struct FIoContainerHeader;
 struct FFilePackageStoreEntry;
@@ -17,7 +18,7 @@ class FFilePackageStore
 {
 public:
 	FFilePackageStore();
-	virtual ~FFilePackageStore() = default;
+	virtual ~FFilePackageStore();
 
 	virtual void Initialize() override;
 	virtual void Lock() override;
@@ -36,7 +37,18 @@ private:
 		uint32 Order;
 	};
 
+#if WITH_EDITOR
+	struct FUncookedPackage
+	{
+		FName PackageName;
+		EPackageExtension HeaderExtension;
+	};
+#endif
+
 	void Update();
+#if WITH_EDITOR
+	uint64 AddUncookedPackagesFromRoot(const FString& RootPath);
+#endif
 
 	FRWLock EntriesLock;
 	FCriticalSection UpdateLock;
@@ -46,5 +58,12 @@ private:
 	TMap<FPackageId, FName> LocalizedPackages;
 	bool bNeedsUpdate = false;
 
-	static thread_local bool bIsLockedOnThread;
+#if WITH_EDITOR
+	FDelegateHandle OnContentPathMountedDelegateHandle;
+	FCriticalSection UncookedPackageRootsLock;
+	TSet<FString> PendingUncookedPackageRoots;
+	TMap<FPackageId, FUncookedPackage> UncookedPackagesMap;
+#endif
+
+	static thread_local int32 LockedOnThreadCount;
 };
