Prevent propagation of the LoaderImport internal flag when duplicating objects
Restore propagation of the RF_Loaded flag when duplicating objects as it no longer has any negative effect on the loader
#rnx

diff --git a/Engine/Source/Runtime/CoreUObject/Private/UObject/UObjectGlobals.cpp b/Engine/Source/Runtime/CoreUObject/Private/UObject/UObjectGlobals.cpp
--- a/Engine/Source/Runtime/CoreUObject/Private/UObject/UObjectGlobals.cpp
+++ b/Engine/Source/Runtime/CoreUObject/Private/UObject/UObjectGlobals.cpp
@@ -2092,15 +2092,10 @@ UObject* StaticDuplicateObjectEx( FObjectDuplicationParameters& Parameters )
 		TEXT("Source and destination class sizes differ.  Source: %s (%i)   Destination: %s (%i)"),
 		*Parameters.SourceObject->GetClass()->GetName(), Parameters.SourceObject->GetClass()->GetPropertiesSize(),
 		*Parameters.DestClass->GetName(), Parameters.DestClass->GetPropertiesSize());
-	if (FPlatformProperties::RequiresCookedData())
-	{
-		checkf(!Parameters.SourceObject->HasAnyInternalFlags(EInternalObjectFlags::AsyncLoading),
-			TEXT("Duplicating object '%s' that's still being async loaded is not permitted."), *Parameters.SourceObject->GetFullName());
-		// Make sure we're not duplicating the Async internal flag, it will prevent the object from being gcd
-		Parameters.InternalFlagMask &= ~EInternalObjectFlags::Async;
-		// Make sure we're not duplicating the WasLoaded flag, it has special meaning to AsyncLoading code
-		Parameters.FlagMask &= ~RF_WasLoaded;
-	}
+	
+	UE_CLOG(FPlatformProperties::RequiresCookedData() && Parameters.SourceObject->HasAnyInternalFlags(EInternalObjectFlags::AsyncLoading), LogUObjectGlobals, Warning, TEXT("Duplicating object '%s' that's still being async loaded"), *Parameters.SourceObject->GetFullName());
+	// Make sure we're not duplicating the Async or LoaderImport internal flags, they will prevent the object from being gcd
+	Parameters.InternalFlagMask &= ~(EInternalObjectFlags::Async | EInternalObjectFlags::LoaderImport);
 
 	if (!IsAsyncLoading() && Parameters.SourceObject->HasAnyFlags(RF_ClassDefaultObject))
 	{
