Fix deprecation after merge

diff --git a/Engine/Source/Runtime/CoreUObject/Private/Serialization/AsyncLoading2.cpp b/Engine/Source/Runtime/CoreUObject/Private/Serialization/AsyncLoading2.cpp
--- a/Engine/Source/Runtime/CoreUObject/Private/Serialization/AsyncLoading2.cpp
+++ b/Engine/Source/Runtime/CoreUObject/Private/Serialization/AsyncLoading2.cpp
@@ -3285,7 +3285,7 @@ void FAsyncPackage2::ImportPackagesRecursive(FIoBatch& IoBatch, IPackageStore& P
 								{
 									NameStr[I] = TChar<TCHAR>::ToLower(NameStr[I]);
 								}
-								PackageRelativeExportPath.Append('/');
+								PackageRelativeExportPath.AppendChar('/');
 								PackageRelativeExportPath.Append(FStringView(NameStr, NameLen));
 							}
 							uint64 ExportHash = CityHash64(reinterpret_cast<const char*>(PackageRelativeExportPath.GetData() + 1), (PackageRelativeExportPath.Len() - 1) * sizeof(TCHAR));
