Delete the deprecated feature to disable EventDrivenLoader cooking.
Deprecate IsEventDrivenLoaderEnabledInCookedBuilds() function and always return true.
Delete s.DisableEDLDeprecationWarnings option and NOEDL command line argument.
Keep IsEventDrivenLoaderEnabled() but let it return true for FPlatformProperties::RequiresCookedData().
Delete UseSeperateBulkDataFiles option and always write bulkdata to different files/segments when cooking.

#rnx

diff --git a/Engine/Source/Editor/UnrealEd/Private/CookOnTheFlyServer.cpp b/Engine/Source/Editor/UnrealEd/Private/CookOnTheFlyServer.cpp
--- a/Engine/Source/Editor/UnrealEd/Private/CookOnTheFlyServer.cpp
+++ b/Engine/Source/Editor/UnrealEd/Private/CookOnTheFlyServer.cpp
@@ -4622,13 +4622,6 @@ void UCookOnTheFlyServer::Initialize( ECookMode::Type DesiredCookMode, ECookInit
 		}
 	}
 
-	bool bDisableEDLWarning = false;
-	GConfig->GetBool(TEXT("/Script/Engine.StreamingSettings"), TEXT("s.DisableEDLDeprecationWarnings"), /* out */ bDisableEDLWarning, GEngineIni);
-	if (!IsEventDrivenLoaderEnabledInCookedBuilds() && !bDisableEDLWarning)
-	{
-		UE_LOG(LogCook, Warning, TEXT("Cooking with Event Driven Loader disabled. Loading code will use deprecated path which will be removed in future release."));
-	}
-
 	// Prepare a map SplitDataClass to FRegisteredCookPackageSplitter* for TryGetRegisteredCookPackageSplitter to use
 	RegisteredSplitDataClasses.Reset();
 	UE::Cook::Private::FRegisteredCookPackageSplitter::ForEach([this](UE::Cook::Private::FRegisteredCookPackageSplitter* RegisteredCookPackageSplitter)
@@ -8980,4 +8973,4 @@ UE::Cook::FCookSavePackageContext& UCookOnTheFlyServer::FindOrCreateSaveContext(
 }
 
 
-#undef LOCTEXT_NAMESPACE
\ No newline at end of file
+#undef LOCTEXT_NAMESPACE
diff --git a/Engine/Source/Runtime/CoreUObject/Private/Serialization/ArchiveStackTrace.cpp b/Engine/Source/Runtime/CoreUObject/Private/Serialization/ArchiveStackTrace.cpp
--- a/Engine/Source/Runtime/CoreUObject/Private/Serialization/ArchiveStackTrace.cpp
+++ b/Engine/Source/Runtime/CoreUObject/Private/Serialization/ArchiveStackTrace.cpp
@@ -437,7 +437,6 @@ bool FArchiveStackTrace::LoadPackageIntoMemory(const TCHAR* InFilename, FPackage
 		// Handle EDL packages (uexp files)
 		TUniquePtr<FArchive> ExpFileArchive = nullptr;
 		OutPackageData.Size = UAssetFileArchive->TotalSize();
-		if (IsEventDrivenLoaderEnabledInCookedBuilds())
 		{
 			FString UExpFilename = FPaths::ChangeExtension(InFilename, TEXT("uexp"));
 			ExpFileArchive.Reset(IFileManager::Get().CreateFileReader(*UExpFilename));
diff --git a/Engine/Source/Runtime/CoreUObject/Private/Serialization/AsyncLoading.cpp b/Engine/Source/Runtime/CoreUObject/Private/Serialization/AsyncLoading.cpp
--- a/Engine/Source/Runtime/CoreUObject/Private/Serialization/AsyncLoading.cpp
+++ b/Engine/Source/Runtime/CoreUObject/Private/Serialization/AsyncLoading.cpp
@@ -4547,17 +4547,6 @@ FAsyncLoadingThread::FAsyncLoadingThread(int32 InThreadIndex, IEDLBootNotificati
 	// Current these two vars are always on or off together but can be made separate
 	GEventDrivenLoaderEnabled = IsEventDrivenLoaderEnabled();
 
-	if (IsEventDrivenLoaderEnabled())
-	{
-		UE_CLOG(!IsEventDrivenLoaderEnabledInCookedBuilds(), LogStreaming, Fatal,
-			TEXT("Event driven async loader is being used but it does NOT seem to be enabled in project settings."));
-	}
-	else if (FPlatformProperties::RequiresCookedData())
-	{
-		UE_CLOG(IsEventDrivenLoaderEnabledInCookedBuilds(), LogStreaming, Fatal,
-			TEXT("Event driven async loader is NOT being used but it seems to be enabled in project settings."));
-	}
-
 	PrecacheHandler = new FPrecacheCallbackHandler();
 	QueuedRequestsEvent = FPlatformProcess::GetSynchEventFromPool();
 	CancelLoadingEvent = FPlatformProcess::GetSynchEventFromPool();
@@ -4574,15 +4563,6 @@ FAsyncLoadingThread::FAsyncLoadingThread(int32 InThreadIndex, IEDLBootNotificati
 		FAsyncLoadingThreadSettings::Get().bAsyncLoadingThreadEnabled ? TEXT("true") : TEXT("false"),
 		FAsyncLoadingThreadSettings::Get().bAsyncPostLoadEnabled ? TEXT("true") : TEXT("false"));
 
-	bool bDisableEDLWarning = false;
-	if (GConfig)
-	{
-		GConfig->GetBool(TEXT("/Script/Engine.StreamingSettings"), TEXT("s.DisableEDLDeprecationWarnings"), /* out */ bDisableEDLWarning, GEngineIni);
-	}
-	if (!GEventDrivenLoaderEnabled && !bDisableEDLWarning)
-	{
-		UE_LOG(LogStreaming, Warning, TEXT("Event Driven Loader is disabled. Loading code will use deprecated path which will be removed in future release."));
-	}
 #endif
 }
 
@@ -8047,33 +8027,12 @@ static FAutoConsoleCommand DumpSerializeCmd(
 
 bool IsEventDrivenLoaderEnabledInCookedBuilds()
 {
-	static struct FEventDrivenLoaderEnabledInCookedBuildsInit
-	{
-		FEventDrivenLoaderEnabledInCookedBuildsInit()
-		{
-			check(GConfig || IsEngineExitRequested());
-			if (GConfig)
-			{
-				// Ensure that the streaming settings from the config have been applied 
-				ApplyCVarSettingsFromIni(TEXT("/Script/Engine.StreamingSettings"), *GEngineIni, ECVF_SetByProjectSetting);
-			}
-		}
-	} EventDrivenLoaderEnabledInCookedBuilds;
-
-	static const bool bNoEDL = !UE_BUILD_SHIPPING && FParse::Param(FCommandLine::Get(), TEXT("NOEDL"));
-	return !bNoEDL && (GEventDrivenLoaderEnabledInCookedBuilds != 0);
+	return true;
 }
 
 bool IsEventDrivenLoaderEnabled()
 {
-	static struct FEventDrivenLoaderEnabledInit
-	{
-		FEventDrivenLoaderEnabledInit()
-		{
-			GEventDrivenLoaderEnabled = IsEventDrivenLoaderEnabledInCookedBuilds() && FPlatformProperties::RequiresCookedData();
-		}
-	} EventDrivenLoaderEnabledInit;
-	return GEventDrivenLoaderEnabled;
+	return FPlatformProperties::RequiresCookedData();
 }
 
 //#pragma clang optimize on
diff --git a/Engine/Source/Runtime/CoreUObject/Private/Serialization/BulkData.cpp b/Engine/Source/Runtime/CoreUObject/Private/Serialization/BulkData.cpp
--- a/Engine/Source/Runtime/CoreUObject/Private/Serialization/BulkData.cpp
+++ b/Engine/Source/Runtime/CoreUObject/Private/Serialization/BulkData.cpp
@@ -1407,7 +1407,7 @@ void FUntypedBulkData::Serialize( FArchive& Ar, UObject* Owner, int32 Idx, bool
 
 			// determine whether we are going to store the payload inline or not.
 			bool bStoreInline = !!(LocalBulkDataFlags & BULKDATA_ForceInlinePayload) || !LinkerSave || Ar.IsTextFormat();
-			if (IsEventDrivenLoaderEnabledInCookedBuilds() && Ar.IsCooking() && !(LocalBulkDataFlags & BULKDATA_Force_NOT_InlinePayload))
+			if (Ar.IsCooking() && !(LocalBulkDataFlags & BULKDATA_Force_NOT_InlinePayload))
 			{
 				bStoreInline = true;
 			}
diff --git a/Engine/Source/Runtime/CoreUObject/Private/UObject/LinkerLoad.cpp b/Engine/Source/Runtime/CoreUObject/Private/UObject/LinkerLoad.cpp
--- a/Engine/Source/Runtime/CoreUObject/Private/UObject/LinkerLoad.cpp
+++ b/Engine/Source/Runtime/CoreUObject/Private/UObject/LinkerLoad.cpp
@@ -1306,14 +1306,6 @@ FLinkerLoad::ELinkerStatus FLinkerLoad::SerializePackageFileSummaryInternal()
 		}
 	}
 
-	if (FPlatformProperties::RequiresCookedData() &&
-		Summary.PreloadDependencyCount > 0 && Summary.PreloadDependencyOffset > 0 &&
-		!IsEventDrivenLoaderEnabledInCookedBuilds())
-	{
-		UE_LOG(LogLinker, Fatal, TEXT("Package %s contains preload dependency data but the current build does not support it. Make sure Event Driven Loader is enabled and rebuild the game executable."),
-			*GetArchiveName())
-	}
-
 #if PLATFORM_WINDOWS
 	if (!FPlatformProperties::RequiresCookedData() &&
 		// We can't check the post tag if the file is an EDL cooked package
diff --git a/Engine/Source/Runtime/CoreUObject/Private/UObject/LinkerSave.cpp b/Engine/Source/Runtime/CoreUObject/Private/UObject/LinkerSave.cpp
--- a/Engine/Source/Runtime/CoreUObject/Private/UObject/LinkerSave.cpp
+++ b/Engine/Source/Runtime/CoreUObject/Private/UObject/LinkerSave.cpp
@@ -216,8 +216,7 @@ FPackageIndex FLinkerSave::MapObject( const UObject* Object ) const
 		const FPackageIndex *Found = ObjectIndicesMap.Find(Object);
 		if (Found)
 		{
-			if (IsEventDrivenLoaderEnabledInCookedBuilds() &&
-				IsCooking() && CurrentlySavingExport.IsExport() &&
+			if (IsCooking() && CurrentlySavingExport.IsExport() &&
 				Object->GetOutermost()->GetFName() != GLongCoreUObjectPackageName && // We assume nothing in coreuobject ever loads assets in a constructor
 				*Found != CurrentlySavingExport) // would be weird, but I can't be a dependency on myself
 			{
diff --git a/Engine/Source/Runtime/CoreUObject/Private/UObject/SavePackage.cpp b/Engine/Source/Runtime/CoreUObject/Private/UObject/SavePackage.cpp
--- a/Engine/Source/Runtime/CoreUObject/Private/UObject/SavePackage.cpp
+++ b/Engine/Source/Runtime/CoreUObject/Private/UObject/SavePackage.cpp
@@ -182,7 +182,7 @@ FString FArchiveSaveTagExports::GetArchiveName() const
 
 FArchive& FArchiveSaveTagExports::operator<<(FWeakObjectPtr& Value)
 {
-	if (IsEventDrivenLoaderEnabledInCookedBuilds() && IsCooking())
+	if (IsCooking())
 	{
 		// Always serialize weak pointers for the purposes of object tagging
 		UObject* Object = static_cast<UObject*>(Value.Get(true));
@@ -229,7 +229,7 @@ FArchive& FArchiveSaveTagExports::operator<<(UObject*& Obj)
 		// If this is a CDO, gather it's subobjects and serialize them
 		if (Obj->HasAnyFlags(RF_ClassDefaultObject))
 		{
-			if (IsEventDrivenLoaderEnabledInCookedBuilds() && IsCooking())
+			if (IsCooking())
 			{
 				// Gets all subobjects defined in a class, including the CDO, CDO components and blueprint-created components
 				TArray<UObject*> ObjectTemplates;
@@ -409,7 +409,7 @@ FString FArchiveSaveTagImports::GetArchiveName() const
 
 FArchive& FArchiveSaveTagImports::operator<< (struct FWeakObjectPtr& Value)
 {
-	if (IsEventDrivenLoaderEnabledInCookedBuilds() && IsCooking())
+	if (IsCooking())
 	{
 		// Always serialize weak pointers for the purposes of object tagging
 		UObject* Object = static_cast<UObject*>(Value.Get(true));
@@ -492,7 +492,7 @@ FArchive& FArchiveSaveTagImports::operator<<( UObject*& Obj )
 				UClass* ClassObj = Cast<UClass>(Obj);
 
 				// Don't recurse into CDOs if we're already ignoring dependencies, we only want to recurse into our outer chain in that case
-				if (IsEventDrivenLoaderEnabledInCookedBuilds() && IsCooking() && !bIsNative && !bIgnoreDependencies && ClassObj)
+				if (IsCooking() && !bIsNative && !bIgnoreDependencies && ClassObj)
 				{
 					// We don't want to add this to Dependencies, we simply want it to be an import so that a serialization before creation dependency can be created to the CDO
 					FScopeIgnoreDependencies IgnoreDependencies(*this);
@@ -2316,8 +2316,7 @@ FSavePackageResultStruct UPackage::Save(UPackage* InOuter, UObject* InAsset, con
 						UObject* Template = Obj->GetArchetype();
 						if (Template)
 						{
-							// If we're not cooking for the event driven loader, exclude the CDO
-							if (Template != Class->GetDefaultObject() || (IsEventDrivenLoaderEnabledInCookedBuilds() && TargetPlatform))
+							if (Template != Class->GetDefaultObject() || TargetPlatform)
 							{
 								ImportTagger << Template;
 							}
@@ -2381,7 +2380,7 @@ FSavePackageResultStruct UPackage::Save(UPackage* InOuter, UObject* InAsset, con
 							}
 						}
 
-						if (IsEventDrivenLoaderEnabledInCookedBuilds() && TargetPlatform)
+						if (TargetPlatform)
 						{
 							TArray<UObject*> Deps;
 							Obj->GetPreloadDependencies(Deps);
@@ -2559,7 +2558,7 @@ FSavePackageResultStruct UPackage::Save(UPackage* InOuter, UObject* InAsset, con
 
 							if( !Obj->HasAnyFlags(RF_Public) && !Obj->HasAnyFlags(RF_Transient))
 							{
-								if (!IsEventDrivenLoaderEnabledInCookedBuilds() || !TargetPlatform || !ObjPackage->HasAnyPackageFlags(PKG_CompiledIn))
+								if (!TargetPlatform || !ObjPackage->HasAnyPackageFlags(PKG_CompiledIn))
 								{
 									PrivateObjects.Add(Obj);
 								}
@@ -3135,7 +3134,7 @@ FSavePackageResultStruct UPackage::Save(UPackage* InOuter, UObject* InAsset, con
 						checkf(false, TEXT("NULL XObject for import %i - Object: %s Class: %s"), i, *Linker->ImportMap[i].ObjectName.ToString(), *Linker->ImportMap[i].ClassName.ToString());
 					}
 				}
-				if (IsEventDrivenLoaderEnabledInCookedBuilds() && TargetPlatform)
+				if (TargetPlatform)
 				{
 					EDLCookChecker.AddExport(InOuter); // the package isn't actually in the export map, but that is ok, we add it as export anyway for error checking
 					for (int32 i = 0; i < Linker->ImportMap.Num(); i++)
@@ -3331,7 +3330,7 @@ FSavePackageResultStruct UPackage::Save(UPackage* InOuter, UObject* InAsset, con
 								Export.ClassIndex = FPackageIndex();
 							}
 
-							if (IsEventDrivenLoaderEnabledInCookedBuilds() && TargetPlatform)
+							if (TargetPlatform)
 							{
 								UObject* Archetype = Export.Object->GetArchetype();
 								check(Archetype);
@@ -3382,7 +3381,7 @@ FSavePackageResultStruct UPackage::Save(UPackage* InOuter, UObject* InAsset, con
 								//	*(Export.Object->GetPathName()),
 								//	*(Export.Object->GetOuter()->GetPathName()));
 
-								if (Linker->IsCooking() && IsEventDrivenLoaderEnabledInCookedBuilds())
+								if (Linker->IsCooking())
 								{
 									// Only packages are allowed to have no outer
 									ensureMsgf(Export.OuterIndex != FPackageIndex() || Export.Object->IsA(UPackage::StaticClass()), TEXT("Export %s has no valid outer when cooking!"), *Export.Object->GetPathName());
@@ -3400,7 +3399,7 @@ FSavePackageResultStruct UPackage::Save(UPackage* InOuter, UObject* InAsset, con
 				Linker->Summary.PreloadDependencyOffset = Linker->Tell();
 				Linker->Summary.PreloadDependencyCount = -1;
 
-				if (Linker->IsCooking() && IsEventDrivenLoaderEnabledInCookedBuilds())
+				if (Linker->IsCooking())
 				{
 #if WITH_EDITOR
 					FArchiveStackTraceIgnoreScope IgnoreSummaryDiffsScope(bIgnoreHeaderDiffs);
@@ -3953,7 +3952,7 @@ FSavePackageResultStruct UPackage::Save(UPackage* InOuter, UObject* InAsset, con
 										Import.SetPackageName(ImportPackage->GetFName());
 									}
 
-									if (Linker->IsCooking() && IsEventDrivenLoaderEnabledInCookedBuilds())
+									if (Linker->IsCooking())
 									{
 										// Only package imports are allowed to have no outer
 										ensureMsgf(Import.OuterIndex != FPackageIndex() || Import.ClassName == NAME_Package, TEXT("Import %s has no valid outer when cooking!"), *Import.XObject->GetPathName());
diff --git a/Engine/Source/Runtime/CoreUObject/Private/UObject/SavePackage/SavePackageUtilities.cpp b/Engine/Source/Runtime/CoreUObject/Private/UObject/SavePackage/SavePackageUtilities.cpp
--- a/Engine/Source/Runtime/CoreUObject/Private/UObject/SavePackage/SavePackageUtilities.cpp
+++ b/Engine/Source/Runtime/CoreUObject/Private/UObject/SavePackage/SavePackageUtilities.cpp
@@ -1365,7 +1365,7 @@ FEDLCookChecker::FEDLCookChecker()
 
 void FEDLCookChecker::SetActiveIfNeeded()
 {
-	bIsActive = IsEventDrivenLoaderEnabledInCookedBuilds() && !FParse::Param(FCommandLine::Get(), TEXT("DisableEDLCookChecker"));
+	bIsActive = !FParse::Param(FCommandLine::Get(), TEXT("DisableEDLCookChecker"));
 }
 
 void FEDLCookChecker::Reset()
@@ -2039,22 +2039,7 @@ ESavePackageResult SaveBulkData(FLinkerSave* Linker, int64& InOutStartOffset, co
 	TUniquePtr<FLargeMemoryWriterWithRegions> OptionalBulkArchive;
 	TUniquePtr<FLargeMemoryWriterWithRegions> MappedBulkArchive;
 
-	static const struct FCookerUseSeparateSegments
-	{
-		bool bEnable = false;
-
-		FCookerUseSeparateSegments()
-		{
-			GConfig->GetBool(TEXT("Core.System"), TEXT("UseSeperateBulkDataFiles"), /* out */ bEnable, GEngineIni);
-
-			if (IsEventDrivenLoaderEnabledInCookedBuilds())
-			{
-				// Always split bulk data when splitting cooked files
-				bEnable = true;
-			}
-		}
-	} CookerUseSeparateSegments;
-	const bool bSeparateSegmentsEnabled = CookerUseSeparateSegments.bEnable && Linker->IsCooking();
+	const bool bSeparateSegmentsEnabled = Linker->IsCooking();
 
 	int64 LinkerStart = 0;
 	if (PackageWriter || bSeparateSegmentsEnabled)
diff --git a/Engine/Source/Runtime/CoreUObject/Public/UObject/UObjectGlobals.h b/Engine/Source/Runtime/CoreUObject/Public/UObject/UObjectGlobals.h
--- a/Engine/Source/Runtime/CoreUObject/Public/UObject/UObjectGlobals.h
+++ b/Engine/Source/Runtime/CoreUObject/Public/UObject/UObjectGlobals.h
@@ -514,6 +514,8 @@ COREUOBJECT_API void CancelAsyncLoading();
 /**
 * Returns true if the event driven loader is enabled in cooked builds
 */
+UE_DEPRECATED(5.0, "Any call to IsEventDrivenLoaderEnabledInCookedBuilds can be removed. "
+		"Cooked packages are always split into different files/segments and headers will always contain preload dependencies.")
 COREUOBJECT_API bool IsEventDrivenLoaderEnabledInCookedBuilds();
 
 /**
diff --git a/Engine/Source/Runtime/Engine/Classes/Engine/CoreSettings.h b/Engine/Source/Runtime/Engine/Classes/Engine/CoreSettings.h
--- a/Engine/Source/Runtime/Engine/Classes/Engine/CoreSettings.h
+++ b/Engine/Source/Runtime/Engine/Classes/Engine/CoreSettings.h
@@ -118,9 +118,6 @@ protected:
 	//~ Begin UObject Interface
 	virtual void PostInitProperties() override;
 
-#if WITH_EDITOR
-	virtual void PostEditChangeProperty(FPropertyChangedEvent& PropertyChangedEvent) override;
-#endif
 	//~ End UObject Interface
 };
 
diff --git a/Engine/Source/Runtime/Engine/Private/Blueprint.cpp b/Engine/Source/Runtime/Engine/Private/Blueprint.cpp
--- a/Engine/Source/Runtime/Engine/Private/Blueprint.cpp
+++ b/Engine/Source/Runtime/Engine/Private/Blueprint.cpp
@@ -1269,20 +1269,6 @@ void UBlueprint::BeginCacheForCookedPlatformData(const ITargetPlatform *TargetPl
 				break;
 			}
 			
-			// EDL is required, because we need to enforce a preload dependency on the CDO (see UBlueprintGeneratedClass::GetDefaultObjectPreloadDependencies). This is
-			// difficult to support in the non-EDL case, because we have to enforce the dependency at runtime, which can lead to unpredictable results in a cooked build.
-			if(bResult && !IsEventDrivenLoaderEnabledInCookedBuilds())
-			{
-				bResult = false;
-
-				static bool bWarnOnEDLDisabled = true;
-				if (bWarnOnEDLDisabled)
-				{
-					UE_LOG(LogBlueprint, Warning, TEXT("Cannot cook Blueprint component data for faster instancing at runtime, because Event-Driven Loading (EDL) has been disabled. Re-enable EDL to support this feature, or disable the option to cook Blueprint component data."));
-					bWarnOnEDLDisabled = false;
-				}
-			}
-
 			return bResult;
 		};
 
diff --git a/Engine/Source/Runtime/Engine/Private/CoreSettings.cpp b/Engine/Source/Runtime/Engine/Private/CoreSettings.cpp
--- a/Engine/Source/Runtime/Engine/Private/CoreSettings.cpp
+++ b/Engine/Source/Runtime/Engine/Private/CoreSettings.cpp
@@ -184,34 +184,6 @@ void UStreamingSettings::PostInitProperties()
 #endif // #if WITH_EDITOR
 }
 
-#if WITH_EDITOR
-void UStreamingSettings::PostEditChangeProperty(FPropertyChangedEvent& PropertyChangedEvent)
-{
-	static FName NAME_EventDrivenLoaderEnabled(TEXT("EventDrivenLoaderEnabled"));
-
-	Super::PostEditChangeProperty(PropertyChangedEvent);
-	
-	if (PropertyChangedEvent.Property)
-	{
-		ExportValuesToConsoleVariables(PropertyChangedEvent.Property);
-
-		bool bDisableEDLWarning = false;
-		GConfig->GetBool(TEXT("/Script/Engine.StreamingSettings"), TEXT("s.DisableEDLDeprecationWarnings"), /* out */ bDisableEDLWarning, GEngineIni);
-
-		if (PropertyChangedEvent.Property->GetFName() == NAME_EventDrivenLoaderEnabled && !EventDrivenLoaderEnabled && !bDisableEDLWarning)
-		{
-			FNotificationInfo Info(NSLOCTEXT("CoreSettings", "EventDrivenLoaderDisabled", "Disabling the Event Driven Loader will result in using deprecated loading path"));
-			Info.bFireAndForget = true;
-			Info.bUseLargeFont = true;
-			Info.bUseThrobber = false;
-			Info.bUseSuccessFailIcons = false;
-			Info.ExpireDuration = 3.0f;
-			FSlateNotificationManager::Get().AddNotification(Info);
-		}
-	}
-}
-#endif // #if WITH_EDITOR
-
 UGarbageCollectionSettings::UGarbageCollectionSettings()
 : Super()
 {
