[VK] Add new PKG_ExternallyReferenceable flag to mark Public Assets

diff --git a/Engine/Source/Runtime/CoreUObject/Public/UObject/ObjectMacros.h b/Engine/Source/Runtime/CoreUObject/Public/UObject/ObjectMacros.h
--- a/Engine/Source/Runtime/CoreUObject/Public/UObject/ObjectMacros.h
+++ b/Engine/Source/Runtime/CoreUObject/Public/UObject/ObjectMacros.h
@@ -117,7 +117,7 @@ enum EPackageFlags
 	PKG_UncookedOnly				= 0x00000100,	///< Loaded only in uncooked builds (i.e. runtime in editor)
 	PKG_Cooked						= 0x00000200,	///< Package is cooked
 	PKG_ContainsNoAsset				= 0x00000400,	///< Package doesn't contain any asset object (although asset tags can be present)
-//	PKG_Unused						= 0x00000800,
+	PKG_ExternallyReferenceable		= 0x00000800,   ///< (Not Implemented) Objects in this package can be referenced in a different plugin or mount point (i.e /Game -> /Engine)
 //	PKG_Unused						= 0x00001000,
 	PKG_UnversionedProperties		= 0x00002000,   ///< Uses unversioned property serialization instead of versioned tagged property serialization
 	PKG_ContainsMapData				= 0x00004000,   ///< Contains map data (UObjects only referenced by a single ULevel) but is stored in a different package
