Add diagnostic logs when FilenameToLongPackageName fails. Update InternalFilenameToLongPackageName to also try comparing absolute paths.
#rnx

diff --git a/Engine/Source/Runtime/CoreUObject/Private/Misc/PackageName.cpp b/Engine/Source/Runtime/CoreUObject/Private/Misc/PackageName.cpp
--- a/Engine/Source/Runtime/CoreUObject/Private/Misc/PackageName.cpp
+++ b/Engine/Source/Runtime/CoreUObject/Private/Misc/PackageName.cpp
@@ -658,9 +658,16 @@ void FPackageName::InternalFilenameToLongPackageName(FStringView InFilename, FSt
 				FReadScopeLock ScopeLock(ContentMountPointCriticalSection);
 				for (const auto& Pair : Paths.ContentPathToRoot)
 				{
-					if (Result.StartsWith(Pair.ContentPath))
+					FStringView RelPath;
+					// Compare against AbsolutePath, but also try ContentPath in case ContentPath is absolute
+					if (FPathViews::TryMakeChildPathRelativeTo(Result, Pair.AbsolutePath, RelPath))
 					{
-						OutPackageName << Pair.RootPath << Result.RightChop(Pair.ContentPath.Len());
+						OutPackageName << Pair.RootPath << RelPath;
+						return;
+					}
+					if (FPathViews::TryMakeChildPathRelativeTo(Result, Pair.ContentPath, RelPath))
+					{
+						OutPackageName << Pair.RootPath << RelPath;
 						return;
 					}
 				}
@@ -740,6 +747,27 @@ FString FPackageName::FilenameToLongPackageName(const FString& InFilename)
 	FString Result;
 	if (!TryConvertFilenameToLongPackageName(InFilename, Result, &FailureReason))
 	{
+		TStringBuilder<128> ContentRoots;
+		{
+			FReadScopeLock ScopeLock(ContentMountPointCriticalSection);
+			for (const FPathPair& Pair : FLongPackagePathsSingleton::Get().ContentPathToRoot)
+			{
+				ContentRoots << TEXT("\n\t\t") << Pair.ContentPath;
+				ContentRoots << TEXT("\n\t\t\t") << Pair.AbsolutePath;
+			}
+		}
+
+		UE_LOG(LogPackageName, Display, TEXT("FilenameToLongPackageName failed, we will issue a fatal log. Diagnostics:")
+			TEXT("\n\tInFilename=%s")
+			TEXT("\n\tConvertToRelativePath=%s")
+			TEXT("\n\tConvertRelativePathToFull=%s")
+			TEXT("\n\tRootDir=%s")
+			TEXT("\n\tBaseDir=%s")
+			TEXT("\n\tContentRoots=%s"),
+			*InFilename, *IFileManager::Get().ConvertToRelativePath(*InFilename),
+			*FPaths::ConvertRelativePathToFull(InFilename), FPlatformMisc::RootDir(), FPlatformProcess::BaseDir(),
+			ContentRoots.ToString()
+		);
 		UE_LOG(LogPackageName, Fatal, TEXT("%s"), *FailureReason);
 	}
 	return Result;
@@ -2606,4 +2634,4 @@ bool FPackageNameTests::RunTest(const FString& Parameters)
 }
 
 #endif //WITH_DEV_AUTOMATION_TESTS
-#undef LOCTEXT_NAMESPACE // "PackageNames"
\ No newline at end of file
+#undef LOCTEXT_NAMESPACE // "PackageNames"
