Virtualized bulkdata is now know as EditorBulkData

#rnx

- The goal is to show a clear separation between the bulkdata object and the virtualization system.
- NOTE that all of these classes were added during 5.0 development so we are just renaming and moving them, no deprecation paths.
- Renamed FVirtualizedUntypedBulkData to FEditorBulkData
-- Also removed the derived templated types, they were never actually used
- Renamed FVirtualizedBulkDataReader to FEditorBulkDataReader
- Renamed FVirtualizedBulkDataWriter to FEditorBulkDataWriter
- Moved the renamed classes to the UE::Serialization namespace from UE::Virtualization
- Renamed the files of the renamed classes where required.
- Replaced use of LogVirtualization with LogSerialization.
- Renamed defines prefixed with VBD_* to UE_ and make sure they are undefed by the end of the cpp
- Edited unit tests to show up under "System.CoreUObject.Serialization.EditorBulkData.*"

diff --git a/Engine/Plugins/Experimental/Interchange/Runtime/Source/Import/Private/Texture/InterchangeTextureFactory.cpp b/Engine/Plugins/Experimental/Interchange/Runtime/Source/Import/Private/Texture/InterchangeTextureFactory.cpp
--- a/Engine/Plugins/Experimental/Interchange/Runtime/Source/Import/Private/Texture/InterchangeTextureFactory.cpp
+++ b/Engine/Plugins/Experimental/Interchange/Runtime/Source/Import/Private/Texture/InterchangeTextureFactory.cpp
@@ -268,7 +268,7 @@ namespace UE::Interchange::Private::InterchangeTextureFactory
 	}
 
 #if WITH_EDITORONLY_DATA
-	void SetupTextureSourceData(UTexture* Texture, const FImportImage& Image, UE::Virtualization::FVirtualizedUntypedBulkData::FSharedBufferWithID&& BufferAndId)
+	void SetupTextureSourceData(UTexture* Texture, const FImportImage& Image, UE::Serialization::FEditorBulkData::FSharedBufferWithID&& BufferAndId)
 	{
 		Texture->Source.Init(
 			Image.SizeX,
@@ -290,7 +290,7 @@ namespace UE::Interchange::Private::InterchangeTextureFactory
 		}
 	}
 
-	void SetupTexture2DSourceData(UTexture2D* Texture2D, const FImportBlockedImage& BlockedImage, UE::Virtualization::FVirtualizedUntypedBulkData::FSharedBufferWithID&& BufferAndId)
+	void SetupTexture2DSourceData(UTexture2D* Texture2D, const FImportBlockedImage& BlockedImage, UE::Serialization::FEditorBulkData::FSharedBufferWithID&& BufferAndId)
 	{
 		if (BlockedImage.BlocksData.Num() > 1)
 		{
@@ -330,7 +330,7 @@ namespace UE::Interchange::Private::InterchangeTextureFactory
 		}
 	}
 
-	void SetupTextureSourceData(UTexture* Texture, const FImportSlicedImage& SlicedImage, UE::Virtualization::FVirtualizedUntypedBulkData::FSharedBufferWithID&& BufferAndId)
+	void SetupTextureSourceData(UTexture* Texture, const FImportSlicedImage& SlicedImage, UE::Serialization::FEditorBulkData::FSharedBufferWithID&& BufferAndId)
 	{
 		Texture->Source.InitLayered(
 			SlicedImage.SizeX,
@@ -352,7 +352,7 @@ namespace UE::Interchange::Private::InterchangeTextureFactory
 		}
 	}
 
-	void SetupTextureSourceData(UTextureLightProfile* TextureLightProfile, const FImportLightProfile& LightProfile, UE::Virtualization::FVirtualizedUntypedBulkData::FSharedBufferWithID&& BufferAndId)
+	void SetupTextureSourceData(UTextureLightProfile* TextureLightProfile, const FImportLightProfile& LightProfile, UE::Serialization::FEditorBulkData::FSharedBufferWithID&& BufferAndId)
 	{
 		const FImportImage& ImportImage = LightProfile;
 		SetupTextureSourceData(TextureLightProfile, ImportImage, MoveTemp(BufferAndId));
diff --git a/Engine/Plugins/Experimental/Interchange/Runtime/Source/Import/Public/Texture/InterchangeTextureFactory.h b/Engine/Plugins/Experimental/Interchange/Runtime/Source/Import/Public/Texture/InterchangeTextureFactory.h
--- a/Engine/Plugins/Experimental/Interchange/Runtime/Source/Import/Public/Texture/InterchangeTextureFactory.h
+++ b/Engine/Plugins/Experimental/Interchange/Runtime/Source/Import/Public/Texture/InterchangeTextureFactory.h
@@ -6,7 +6,7 @@
 #include "EditorFramework/AssetImportData.h"
 #include "InterchangeFactoryBase.h"
 #include "Misc/TVariant.h"
-#include "Serialization/VirtualizedBulkData.h"
+#include "Serialization/EditorBulkData.h"
 #include "Texture/InterchangeBlockedTexturePayloadData.h"
 #include "Texture/InterchangeSlicedTexturePayloadData.h"
 #include "Texture/InterchangeTextureLightProfilePayloadData.h"
@@ -37,7 +37,7 @@ namespace UE::Interchange::Private::InterchangeTextureFactory
 		FProcessedPayload& operator=(UE::Interchange::Private::InterchangeTextureFactory::FTexturePayloadVariant&& InPayloadVariant);
 
 		UE::Interchange::Private::InterchangeTextureFactory::FTexturePayloadVariant SettingsFromPayload;
-		UE::Virtualization::FVirtualizedUntypedBulkData::FSharedBufferWithID PayloadAndId;
+		UE::Serialization::FEditorBulkData::FSharedBufferWithID PayloadAndId;
 	};
 }
 
diff --git a/Engine/Source/Developer/Virtualization/Private/PackageSubmissionChecks.cpp b/Engine/Source/Developer/Virtualization/Private/PackageSubmissionChecks.cpp
--- a/Engine/Source/Developer/Virtualization/Private/PackageSubmissionChecks.cpp
+++ b/Engine/Source/Developer/Virtualization/Private/PackageSubmissionChecks.cpp
@@ -9,7 +9,7 @@
 #include "Internationalization/Internationalization.h"
 #include "Misc/PackageName.h"
 #include "Misc/Paths.h"
-#include "Serialization/VirtualizedBulkData.h"
+#include "Serialization/EditorBulkData.h"
 #include "UObject/Linker.h"
 #include "UObject/Package.h"
 #include "UObject/PackageResourceManager.h"
diff --git a/Engine/Source/Developer/Virtualization/Private/Virtualization.cpp b/Engine/Source/Developer/Virtualization/Private/Virtualization.cpp
--- a/Engine/Source/Developer/Virtualization/Private/Virtualization.cpp
+++ b/Engine/Source/Developer/Virtualization/Private/Virtualization.cpp
@@ -7,7 +7,7 @@
 #include "MessageLogModule.h"
 #include "Misc/DelayedAutoRegister.h"
 #include "PackageSubmissionChecks.h"
-#include "Serialization/VirtualizedBulkData.h"
+#include "Serialization/EditorBulkData.h"
 #include "VirtualizationSourceControlUtilities.h"
 
 #define LOCTEXT_NAMESPACE "Virtualization"
diff --git a/Engine/Source/Editor/UnrealEd/Private/EditorDomain/BulkDataRegistryEditorDomain.cpp b/Engine/Source/Editor/UnrealEd/Private/EditorDomain/BulkDataRegistryEditorDomain.cpp
--- a/Engine/Source/Editor/UnrealEd/Private/EditorDomain/BulkDataRegistryEditorDomain.cpp
+++ b/Engine/Source/Editor/UnrealEd/Private/EditorDomain/BulkDataRegistryEditorDomain.cpp
@@ -155,7 +155,7 @@ void FBulkDataRegistryEditorDomain::OnEndLoadPackage(TConstArrayView<UPackage*>
 	}
 }
 
-void FBulkDataRegistryEditorDomain::Register(UPackage* Owner, const UE::Virtualization::FVirtualizedUntypedBulkData& BulkData)
+void FBulkDataRegistryEditorDomain::Register(UPackage* Owner, const UE::Serialization::FEditorBulkData& BulkData)
 {
 	if (!BulkData.GetIdentifier().IsValid())
 	{
@@ -163,7 +163,7 @@ void FBulkDataRegistryEditorDomain::Register(UPackage* Owner, const UE::Virtuali
 	}
 
 	FName PackageName = NAME_None;
-	UE::Virtualization::FVirtualizedUntypedBulkData CopyBulk(BulkData.CopyTornOff());
+	UE::Serialization::FEditorBulkData CopyBulk(BulkData.CopyTornOff());
 	if (Owner
 		&& Owner->GetFileSize() // We only record the BulkDataList for disk packages
 		&& !Owner->GetHasBeenEndLoaded() // We only record BulkDatas that are loaded before the package finishes loading
@@ -182,7 +182,7 @@ void FBulkDataRegistryEditorDomain::Register(UPackage* Owner, const UE::Virtuali
 	ResaveSizeTracker.Register(Owner, BulkData);
 }
 
-void FBulkDataRegistryEditorDomain::OnExitMemory(const UE::Virtualization::FVirtualizedUntypedBulkData& BulkData)
+void FBulkDataRegistryEditorDomain::OnExitMemory(const UE::Serialization::FEditorBulkData& BulkData)
 {
 	const FGuid& Key = BulkData.GetIdentifier();
 	FWriteScopeLock RegistryScopeLock(RegistryLock);
@@ -215,7 +215,7 @@ TFuture<UE::BulkDataRegistry::FMetaData> FBulkDataRegistryEditorDomain::GetMeta(
 			return Promise.GetFuture();
 		}
 
-		UE::Virtualization::FVirtualizedUntypedBulkData& BulkData = Existing->BulkData;
+		UE::Serialization::FEditorBulkData& BulkData = Existing->BulkData;
 		if (!BulkData.HasPlaceholderPayloadId())
 		{
 			TPromise<UE::BulkDataRegistry::FMetaData> Promise;
@@ -249,7 +249,7 @@ TFuture<UE::BulkDataRegistry::FMetaData> FBulkDataRegistryEditorDomain::GetMeta(
 
 TFuture<UE::BulkDataRegistry::FData> FBulkDataRegistryEditorDomain::GetData(const FGuid& BulkDataId)
 {
-	TOptional<UE::Virtualization::FVirtualizedUntypedBulkData> CopyBulk;
+	TOptional<UE::Serialization::FEditorBulkData> CopyBulk;
 	{
 		bool bIsWriteLock = false;
 		FRWScopeLock RegistryScopeLock(RegistryLock, SLT_ReadOnly);
@@ -269,7 +269,7 @@ TFuture<UE::BulkDataRegistry::FData> FBulkDataRegistryEditorDomain::GetData(cons
 
 			if (!Existing->BulkData.HasPlaceholderPayloadId() && !Existing->bHasTempPayload)
 			{
-				// The contract of FVirtualizedUntypedBulkData does not guarantee that GetCompressedPayload() is a quick operation (it may load the data
+				// The contract of FEditorBulkData does not guarantee that GetCompressedPayload() is a quick operation (it may load the data
 				// synchronously), so copy the BulkData into a temporary and call it outside the lock
 				CopyBulk.Emplace(Existing->BulkData);
 				break;
@@ -314,7 +314,7 @@ TFuture<UE::BulkDataRegistry::FData> FBulkDataRegistryEditorDomain::GetData(cons
 	}
 
 	// We are calling a function that returns a TFuture on the stack-local CopyBulk, which would cause a read-after-free if the asynchronous TFuture could
-	// read from the BulkData. However, the contract of FVirtualizedUntypedBulkData guarantees that the TFuture gets a copy of all data it needs and
+	// read from the BulkData. However, the contract of FEditorBulkData guarantees that the TFuture gets a copy of all data it needs and
 	// does not read from the BulkData after returning from GetCompressedPayload, so a read-after-free is not possible.
 	return CopyBulk->GetCompressedPayload().Next([](FCompressedBuffer Payload)
 		{
@@ -343,7 +343,7 @@ void FBulkDataRegistryEditorDomain::Tick(float DeltaTime)
 	TickCook(DeltaTime, false /* bTickComplete */);
 }
 
-void FBulkDataRegistryEditorDomain::AddPendingPackageBulkData(FName PackageName, const UE::Virtualization::FVirtualizedUntypedBulkData& BulkData)
+void FBulkDataRegistryEditorDomain::AddPendingPackageBulkData(FName PackageName, const UE::Serialization::FEditorBulkData& BulkData)
 {
 	FScopeLock PendingPackageScopeLock(&PendingPackageLock);
 	check(bActive); // Registrations should not come in after we destruct, and AsyncTasks should check bActive before calling
@@ -389,12 +389,12 @@ void FBulkDataRegistryEditorDomain::PruneTempLoadedPayloads()
 	}
 }
 
-void FBulkDataRegistryEditorDomain::WritePayloadIdToCache(FName PackageName, const UE::Virtualization::FVirtualizedUntypedBulkData& BulkData) const
+void FBulkDataRegistryEditorDomain::WritePayloadIdToCache(FName PackageName, const UE::Serialization::FEditorBulkData& BulkData) const
 {
 	check(!PackageName.IsNone());
 	TArray<uint8> Bytes;
 	FMemoryWriter Writer(Bytes);
-	const_cast<UE::Virtualization::FVirtualizedUntypedBulkData&>(BulkData).SerializeForRegistry(Writer);
+	const_cast<UE::Serialization::FEditorBulkData&>(BulkData).SerializeForRegistry(Writer);
 	UE::EditorDomain::PutBulkDataPayloadId(PackageName, BulkData.GetIdentifier(), MakeSharedBufferFromArray(MoveTemp(Bytes)));
 }
 
@@ -422,7 +422,7 @@ void FBulkDataRegistryEditorDomain::ReadPayloadIdsFromCache(FName PackageName, T
 				return;
 			}
 			FMemoryReaderView Reader(MakeArrayView(Buffer));
-			UE::Virtualization::FVirtualizedUntypedBulkData CachedBulkData;
+			UE::Serialization::FEditorBulkData CachedBulkData;
 			CachedBulkData.SerializeForRegistry(Reader);
 			const FGuid& BulkDataId = NewPending->GetBulkDataId();
 			if (Reader.IsError() || CachedBulkData.GetIdentifier() != BulkDataId)
@@ -455,7 +455,7 @@ void FBulkDataRegistryEditorDomain::ReadPayloadIdsFromCache(FName PackageName, T
 				return;
 			}
 
-			UE::Virtualization::FVirtualizedUntypedBulkData& ExistingBulkData = ExistingRegisteredBulk->BulkData;
+			UE::Serialization::FEditorBulkData& ExistingBulkData = ExistingRegisteredBulk->BulkData;
 			check(ExistingBulkData.GetIdentifier() == BulkDataId);
 			if (ExistingBulkData.HasPlaceholderPayloadId())
 			{
@@ -595,13 +595,13 @@ void FPendingPackage::ReadCache()
 		{
 			return;
 		}
-		for (const UE::Virtualization::FVirtualizedUntypedBulkData& BulkData : CachedBulkDatas)
+		for (const UE::Serialization::FEditorBulkData& BulkData : CachedBulkDatas)
 		{
 			FGuid BulkDataId = BulkData.GetIdentifier();
 			FRegisteredBulk& TargetRegisteredBulk = Owner->Registry.FindOrAdd(BulkData.GetIdentifier());
 
 			bool bCachedLocationMatches = true;
-			UE::Virtualization::FVirtualizedUntypedBulkData& TargetBulkData = TargetRegisteredBulk.BulkData;
+			UE::Serialization::FEditorBulkData& TargetBulkData = TargetRegisteredBulk.BulkData;
 			if (!TargetBulkData.GetIdentifier().IsValid())
 			{
 				TargetBulkData = BulkData;
@@ -664,7 +664,7 @@ void FPendingPackage::WriteCache()
 	}
 
 	// Sort the list by guid, to avoid indeterminism in the list
-	BulkDatas.Sort([](const UE::Virtualization::FVirtualizedUntypedBulkData& A, const UE::Virtualization::FVirtualizedUntypedBulkData& B)
+	BulkDatas.Sort([](const UE::Serialization::FEditorBulkData& A, const UE::Serialization::FEditorBulkData& B)
 		{
 			return A.GetIdentifier() < B.GetIdentifier();
 		});
@@ -676,7 +676,7 @@ void FPendingPackage::WriteCache()
 }
 
 FUpdatePayloadWorker::FUpdatePayloadWorker(FBulkDataRegistryEditorDomain* InBulkDataRegistry,
-	const UE::Virtualization::FVirtualizedUntypedBulkData& InSourceBulk)
+	const UE::Serialization::FEditorBulkData& InSourceBulk)
 	: BulkData(InSourceBulk)
 	, BulkDataRegistry(InBulkDataRegistry)
 {
@@ -772,7 +772,7 @@ void FPendingPayloadId::Cancel()
 	Request.Cancel();
 }
 
-void Serialize(FArchive& Ar, TArray<UE::Virtualization::FVirtualizedUntypedBulkData>& Datas)
+void Serialize(FArchive& Ar, TArray<UE::Serialization::FEditorBulkData>& Datas)
 {
 	int32 Num = Datas.Num();
 	Ar << Num;
@@ -792,13 +792,13 @@ void Serialize(FArchive& Ar, TArray<UE::Virtualization::FVirtualizedUntypedBulkD
 		}
 		for (int32 n = 0; n < Num; ++n)
 		{
-			UE::Virtualization::FVirtualizedUntypedBulkData& BulkData = Datas.Emplace_GetRef();
+			UE::Serialization::FEditorBulkData& BulkData = Datas.Emplace_GetRef();
 			BulkData.SerializeForRegistry(Ar);
 		}
 	}
 	else
 	{
-		for (UE::Virtualization::FVirtualizedUntypedBulkData& BulkData : Datas)
+		for (UE::Serialization::FEditorBulkData& BulkData : Datas)
 		{
 			BulkData.SerializeForRegistry(Ar);
 		}
diff --git a/Engine/Source/Editor/UnrealEd/Private/EditorDomain/BulkDataRegistryEditorDomain.h b/Engine/Source/Editor/UnrealEd/Private/EditorDomain/BulkDataRegistryEditorDomain.h
--- a/Engine/Source/Editor/UnrealEd/Private/EditorDomain/BulkDataRegistryEditorDomain.h
+++ b/Engine/Source/Editor/UnrealEd/Private/EditorDomain/BulkDataRegistryEditorDomain.h
@@ -11,7 +11,7 @@
 #include "DerivedDataRequestOwner.h"
 #include "HAL/CriticalSection.h"
 #include "Serialization/BulkDataRegistry.h"
-#include "Serialization/VirtualizedBulkData.h"
+#include "Serialization/EditorBulkData.h"
 #include "Templates/RefCounting.h"
 #include "TickableEditorObject.h"
 
@@ -28,18 +28,18 @@ class FBulkDataRegistryEditorDomain;
 struct FRegisteredBulk
 {
 	FRegisteredBulk() = default;
-	FRegisteredBulk(UE::Virtualization::FVirtualizedUntypedBulkData&& InBulkData, FName InPackageName = NAME_None)
+	FRegisteredBulk(UE::Serialization::FEditorBulkData&& InBulkData, FName InPackageName = NAME_None)
 		:BulkData(MoveTemp(InBulkData)), PackageName(InPackageName)
 	{
 	}
 
-	UE::Virtualization::FVirtualizedUntypedBulkData BulkData;
+	UE::Serialization::FEditorBulkData BulkData;
 	FName PackageName;
 	bool bHasTempPayload = false;
 };
 
 /** Serialize an array of BulkDatas into or out of bytes saved/load from the registry's persistent cache. */
-void Serialize(FArchive& Ar, TArray<UE::Virtualization::FVirtualizedUntypedBulkData>& InDatas);
+void Serialize(FArchive& Ar, TArray<UE::Serialization::FEditorBulkData>& InDatas);
 
 /** A collection of bulkdatas that should be sent to the cache for the given package. */
 class FPendingPackage
@@ -50,7 +50,7 @@ public:
 	FPendingPackage(const FPendingPackage& Other) = delete;
 
 	void Cancel();
-	void AddBulkData(const UE::Virtualization::FVirtualizedUntypedBulkData& BulkData)
+	void AddBulkData(const UE::Serialization::FEditorBulkData& BulkData)
 	{
 		BulkDatas.Add(BulkData);
 	}
@@ -72,8 +72,8 @@ private:
 		Flag_Canceled = 1 << 2,
 	};
 	FName PackageName;
-	TArray<UE::Virtualization::FVirtualizedUntypedBulkData> BulkDatas;
-	TArray<UE::Virtualization::FVirtualizedUntypedBulkData> CachedBulkDatas;
+	TArray<UE::Serialization::FEditorBulkData> BulkDatas;
+	TArray<UE::Serialization::FEditorBulkData> CachedBulkDatas;
 	UE::DerivedData::FRequestOwner BulkDataListCacheRequest;
 	FBulkDataRegistryEditorDomain* Owner;
 	/**
@@ -110,7 +110,7 @@ class FUpdatePayloadWorker : public FNonAbandonableTask
 {
 public:
 	FUpdatePayloadWorker(FBulkDataRegistryEditorDomain* InBulkDataRegistry,
-		const UE::Virtualization::FVirtualizedUntypedBulkData& InSourceBulk);
+		const UE::Serialization::FEditorBulkData& InSourceBulk);
 
 	void DoWork();
 	TStatId GetStatId() const
@@ -119,7 +119,7 @@ public:
 	}
 
 private:
-	UE::Virtualization::FVirtualizedUntypedBulkData BulkData;
+	UE::Serialization::FEditorBulkData BulkData;
 	TRefCountPtr<FTaskSharedDataLock> SharedDataLock;
 	FBulkDataRegistryEditorDomain* BulkDataRegistry;
 };
@@ -165,8 +165,8 @@ public:
 	virtual ~FBulkDataRegistryEditorDomain();
 
 	// IBulkDataRegistry interface
-	virtual void Register(UPackage* Owner, const UE::Virtualization::FVirtualizedUntypedBulkData& BulkData) override;
-	virtual void OnExitMemory(const UE::Virtualization::FVirtualizedUntypedBulkData& BulkData) override;
+	virtual void Register(UPackage* Owner, const UE::Serialization::FEditorBulkData& BulkData) override;
+	virtual void OnExitMemory(const UE::Serialization::FEditorBulkData& BulkData) override;
 	virtual TFuture<UE::BulkDataRegistry::FMetaData> GetMeta(const FGuid& BulkDataId) override;
 	virtual TFuture<UE::BulkDataRegistry::FData> GetData(const FGuid& BulkDataId) override;
 	virtual uint64 GetBulkDataResaveSize(FName PackageName) override;
@@ -178,12 +178,12 @@ public:
 	virtual TStatId GetStatId() const override { return TStatId(); }
 
 private:
-	void AddPendingPackageBulkData(FName PackageName, const UE::Virtualization::FVirtualizedUntypedBulkData& BulkData);
+	void AddPendingPackageBulkData(FName PackageName, const UE::Serialization::FEditorBulkData& BulkData);
 	void PollPendingPackages(bool bWaitForCooldown);
 	void AddTempLoadedPayload(const FGuid& RegistryKey, uint64 PayloadSize);
 	void PruneTempLoadedPayloads();
 	void OnEndLoadPackage(TConstArrayView<UPackage*> LoadedPackages);
-	void WritePayloadIdToCache(FName PackageName, const UE::Virtualization::FVirtualizedUntypedBulkData& BulkData) const;
+	void WritePayloadIdToCache(FName PackageName, const UE::Serialization::FEditorBulkData& BulkData) const;
 	void ReadPayloadIdsFromCache(FName PackageName, TArray<TRefCountPtr<FPendingPayloadId>>&& OldPendings,
 		TArray<TRefCountPtr<FPendingPayloadId>>&& NewPendings);
 
diff --git a/Engine/Source/Runtime/Core/Private/Virtualization/PayloadId.cpp b/Engine/Source/Runtime/Core/Private/Virtualization/PayloadId.cpp
--- a/Engine/Source/Runtime/Core/Private/Virtualization/PayloadId.cpp
+++ b/Engine/Source/Runtime/Core/Private/Virtualization/PayloadId.cpp
@@ -30,7 +30,7 @@ FPayloadId::FPayloadId(const FSharedBuffer& Payload)
 }
 
 // TODO: Currently only used for backwards compatibility, we might want to make this private and allow 
-// FVirtualizedUntypedBulkData::CreateFromBulkData to access it by being a friend.
+// FEditorBulkData::CreateFromBulkData to access it by being a friend.
 FPayloadId::FPayloadId(const FGuid& Guid)
 {
 	if (Guid.IsValid())
diff --git a/Engine/Source/Runtime/Core/Public/Serialization/Archive.h b/Engine/Source/Runtime/Core/Public/Serialization/Archive.h
--- a/Engine/Source/Runtime/Core/Public/Serialization/Archive.h
+++ b/Engine/Source/Runtime/Core/Public/Serialization/Archive.h
@@ -28,10 +28,7 @@ class ITargetPlatform;
 class UObject;
 class FProperty;
 struct FUntypedBulkData;
-namespace UE::Virtualization
-{
-class FVirtualizedUntypedBulkData;
-}
+namespace UE::Serialization{ class FEditorBulkData; }
 struct FArchiveSerializedPropertyChain;
 template<class TEnum> class TEnumAsByte;
 typedef TFunction<bool (double RemainingTime)> FExternalReadCallback;
@@ -1618,7 +1615,7 @@ public:
 	 * @param	BulkData	Bulk data object to associate
 	 */
 	virtual void AttachBulkData(UObject* Owner, FUntypedBulkData* BulkData) { }
-	virtual void AttachBulkData(UE::Virtualization::FVirtualizedUntypedBulkData* BulkData) {}
+	virtual void AttachBulkData(UE::Serialization::FEditorBulkData* BulkData) {}
 
 	/**
 	 * Detaches the passed in bulk data object from the linker.
@@ -1627,7 +1624,7 @@ public:
 	 * @param	bEnsureBulkDataIsLoaded	Whether to ensure that the bulk data is loaded before detaching
 	 */
 	virtual void DetachBulkData(FUntypedBulkData* BulkData, bool bEnsureBulkDataIsLoaded) { }
-	virtual void DetachBulkData(UE::Virtualization::FVirtualizedUntypedBulkData* BulkData, bool bEnsureBulkDataIsLoaded) {}
+	virtual void DetachBulkData(UE::Serialization::FEditorBulkData* BulkData, bool bEnsureBulkDataIsLoaded) {}
 
 	/**
 	* Determine if the given archive is a valid "child" of this archive. In general, this means "is exactly the same" but
diff --git a/Engine/Source/Runtime/CoreUObject/Private/Serialization/BulkDataRegistry.cpp b/Engine/Source/Runtime/CoreUObject/Private/Serialization/BulkDataRegistry.cpp
--- a/Engine/Source/Runtime/CoreUObject/Private/Serialization/BulkDataRegistry.cpp
+++ b/Engine/Source/Runtime/CoreUObject/Private/Serialization/BulkDataRegistry.cpp
@@ -11,7 +11,7 @@
 #include "Misc/ConfigCacheIni.h"
 #include "Misc/CoreDelegates.h"
 #include "Misc/ScopeRWLock.h"
-#include "Serialization/VirtualizedBulkData.h"
+#include "Serialization/EditorBulkData.h"
 #include "UObject/Package.h"
 #include "UObject/PackageResourceManager.h"
 #include "UObject/UObjectGlobals.h"
@@ -29,8 +29,8 @@ public:
 	FBulkDataRegistryNull() = default;
 	virtual ~FBulkDataRegistryNull() {}
 
-	virtual void Register(UPackage* Owner, const UE::Virtualization::FVirtualizedUntypedBulkData& BulkData) override {}
-	virtual void OnExitMemory(const UE::Virtualization::FVirtualizedUntypedBulkData& BulkData) override {}
+	virtual void Register(UPackage* Owner, const UE::Serialization::FEditorBulkData& BulkData) override {}
+	virtual void OnExitMemory(const UE::Serialization::FEditorBulkData& BulkData) override {}
 	virtual TFuture<UE::BulkDataRegistry::FMetaData> GetMeta(const FGuid& BulkDataId) override
 	{
 		TPromise<UE::BulkDataRegistry::FMetaData> Promise;
@@ -55,7 +55,7 @@ public:
 	FBulkDataRegistryTrackBulkDataToResave() = default;
 	virtual ~FBulkDataRegistryTrackBulkDataToResave() {}
 
-	virtual void Register(UPackage* Owner, const UE::Virtualization::FVirtualizedUntypedBulkData& BulkData) override
+	virtual void Register(UPackage* Owner, const UE::Serialization::FEditorBulkData& BulkData) override
 	{
 		ResaveSizeTracker.Register(Owner, BulkData);
 	}
@@ -170,7 +170,7 @@ void FResaveSizeTracker::OnPostEngineInit()
 	}
 }
 
-void FResaveSizeTracker::Register(UPackage* Owner, const UE::Virtualization::FVirtualizedUntypedBulkData& BulkData)
+void FResaveSizeTracker::Register(UPackage* Owner, const UE::Serialization::FEditorBulkData& BulkData)
 {
 	if (!BulkData.GetIdentifier().IsValid() || !BulkData.IsMemoryOnlyPayload())
 	{
diff --git a/Engine/Source/Runtime/CoreUObject/Private/Serialization/EditorBulkData.cpp b/Engine/Source/Runtime/CoreUObject/Private/Serialization/EditorBulkData.cpp
new file mode 100644
--- /dev/null
+++ b/Engine/Source/Runtime/CoreUObject/Private/Serialization/EditorBulkData.cpp
@@ -0,0 +1,1876 @@
+// Copyright Epic Games, Inc. All Rights Reserved.
+
+#include "Serialization/EditorBulkData.h"
+
+#include "Compression/OodleDataCompression.h"
+#include "HAL/FileManager.h"
+#include "HAL/IConsoleManager.h"
+#include "Misc/ConfigCacheIni.h"
+#include "Misc/PackageSegment.h"
+#include "Misc/SecureHash.h"
+#include "Misc/ScopeLock.h"
+#include "Serialization/BulkData.h"
+#include "Serialization/BulkDataRegistry.h"
+#include "UObject/LinkerLoad.h"
+#include "UObject/LinkerSave.h"
+#include "UObject/Object.h"
+#include "UObject/ObjectSaveContext.h"
+#include "Virtualization/VirtualizationSystem.h"
+
+//#if WITH_EDITORONLY_DATA
+
+/** When enabled the non-virtualized bulkdata objects will attach to the FLinkerLoader for the package that they are loaded from */
+#if WITH_EDITOR
+	#define UE_ALLOW_LINKERLOADER_ATTACHMENT 1
+#else
+	#define UE_ALLOW_LINKERLOADER_ATTACHMENT 0
+#endif //WITH_EDITOR
+
+/** When enabled we will fatal log if we detect corrupted data rather than logging an error and returning a null FCompressedBuffer/FSharedBuffer. */
+#define UE_CORRUPTED_PAYLOAD_IS_FATAL 0
+
+#if UE_CORRUPTED_PAYLOAD_IS_FATAL
+	#define UE_CORRUPTED_DATA_SEVERITY Fatal
+#else
+	#define UE_CORRUPTED_DATA_SEVERITY Error
+#endif // VBD_CORRUPTED_PAYLOAD_IS_FATAL
+
+namespace UE::Serialization
+{
+/** This console variable should only exist for testing */
+static TAutoConsoleVariable<bool> CVarShouldLoadFromSidecar(
+	TEXT("Serialization.LoadFromSidecar"),
+	false,
+	TEXT("When true FEditorBulkData will load from the sidecar file"));
+
+/** 
+ * Prefer loading from the package trailer (load the trailer, parse the look up, then load the payload) over 
+ * using the in built OffsetInFile member to load from the package file directly.
+ */
+static TAutoConsoleVariable<bool> CVarShouldLoadFromTrailer(
+	TEXT("Serialization.LoadFromTrailer"),
+	false,
+	TEXT("When true FEditorBulkData will load payloads via the package trailer rather than the package itself"));
+
+static TAutoConsoleVariable<bool> CVarShouldValidatePayload(
+	TEXT("Serialization.ValidatePayloads"),
+	false,
+	TEXT("When true FEditorBulkData validate any payload loaded from the sidecar file"));
+
+static TAutoConsoleVariable<bool> CVarShouldAllowSidecarSyncing(
+	TEXT("Serialization.AllowSidecarSyncing"),
+	false,
+	TEXT("When true FEditorBulkData will attempt to sync it's .upayload file via sourcecontrol if the first attempt to load from it fails"));
+
+/** When enabled the bulkdata object will try pushing the payload when saved to disk as part of a package.
+  * This is legacy behavior and likely to be removed
+  */
+static constexpr bool bAllowVirtualizationOnSave = false;
+
+/** When enabled virtualized payloads will be pulled and stored locally when the bulkdata object is saved
+  * to disk as part of a package. This is currently required but will likely be changed to be a config 
+  * option when we can support it.
+  */
+static constexpr bool bForceRehydrationOnSave = true;
+
+/** Wrapper around the config file option [Core.System.Experimental]EnablePackageSidecarSaving */
+bool ShouldSaveToPackageSidecar()
+{
+	static const struct FSaveToPackageSidecar
+	{
+		bool bEnabled = false;
+
+		FSaveToPackageSidecar()
+		{
+			GConfig->GetBool(TEXT("Core.System.Experimental"), TEXT("EnablePackageSidecarSaving"), bEnabled, GEngineIni);
+		}
+	} ConfigSetting;
+
+	return ConfigSetting.bEnabled;
+}
+
+#if UE_ENABLE_VIRTUALIZATION_TOGGLE
+bool ShouldAllowVirtualizationOptOut()
+{
+	static struct FAllowVirtualizationOptOut
+	{
+		bool bEnabled = true;
+
+		FAllowVirtualizationOptOut()
+		{
+			GConfig->GetBool(TEXT("Core.System.Experimental"), TEXT("AllowVirtualizationOptOut"), bEnabled, GEngineIni);
+		}
+	} AllowVirtualizationOptOut;
+
+	return AllowVirtualizationOptOut.bEnabled;
+}
+#endif // UE_ENABLE_VIRTUALIZATION_TOGGLE
+
+/** Utility for logging extended error messages when we fail to open a package for reading */
+void LogPackageOpenFailureMessage(const FPackagePath& PackagePath, EPackageSegment PackageSegment)
+{
+	// TODO: Check the various error paths here again!
+	const uint32 SystemError = FPlatformMisc::GetLastError();
+	// If we have a system error we can give a more informative error message but don't output it if the error is zero as 
+	// this can lead to very confusing error messages.
+	if (SystemError != 0)
+	{
+		TCHAR SystemErrorMsg[2048] = { 0 };
+		FPlatformMisc::GetSystemErrorMessage(SystemErrorMsg, sizeof(SystemErrorMsg), SystemError);
+		UE_LOG(LogSerialization, Error, TEXT("Could not open the file '%s' for reading due to system error: '%s' (%d))"), *PackagePath.GetDebugNameWithExtension(PackageSegment), SystemErrorMsg, SystemError);
+	}
+	else
+	{
+		UE_LOG(LogSerialization, Error, TEXT("Could not open (%s) to read FEditorBulkData with an unknown error"), *PackagePath.GetDebugNameWithExtension(PackageSegment));
+	}
+}
+
+/** 
+ * Utility used to validate the contents of a recently loaded payload.
+ * If the given payload is null, then we assume that the load failed and errors would've been raised else
+ * where in code and there is no need to validate the contents.
+ * If the contents are validated we check the loaded result against the members of a bulkdata object to 
+ * see if they match.
+ */
+bool IsDataValid(const FEditorBulkData& BulkData, const FCompressedBuffer& Payload)
+{
+	if (!Payload.IsNull())
+	{
+		if (!BulkData.HasPlaceholderPayloadId() && BulkData.GetPayloadId() != FIoHash(Payload.GetRawHash()))
+		{
+			return false;
+		}
+
+		if (Payload.GetRawSize() != BulkData.GetPayloadSize())
+		{
+			return false;
+		}
+	}
+
+	return true;
+}
+
+/** Utility for finding the FLinkerLoad associated with a given UObject */
+const FLinkerLoad* GetLinkerLoadFromOwner(UObject* Owner)
+{
+	if (Owner != nullptr)
+	{
+		UPackage* Package = Owner->GetOutermost();
+		checkf(Package != nullptr, TEXT("Owner was not a valid UPackage!"));
+
+		return FLinkerLoad::FindExistingLinkerForPackage(Package);
+	}
+	else
+	{
+		return nullptr;
+	}
+}
+
+/** Utility for finding the FPackageTrailer associated with a given UObject */
+const FPackageTrailer* GetTrailerFromOwner(UObject* Owner)
+{
+	const FLinkerLoad* Linker = GetLinkerLoadFromOwner(Owner);
+	if (Linker != nullptr)
+	{
+		return Linker->GetPackageTrailer();
+	}
+	else
+	{
+		return nullptr;
+	}
+}
+
+/** Utility for finding the package path associated with a given UObject */
+FPackagePath GetPackagePathFromOwner(UObject* Owner, EPackageSegment& OutPackageSegment)
+{
+	OutPackageSegment = EPackageSegment::Header;
+
+	const FLinkerLoad* Linker = GetLinkerLoadFromOwner(Owner);
+
+	if (Linker != nullptr)
+	{
+		return Linker->GetPackagePath();
+	}
+	else
+	{
+		return FPackagePath();
+	}
+}
+
+/** Utility for updating an existing entry in an Archive before returning the archive to it's original seek position */
+template<typename DataType>
+void UpdateArchiveData(FArchive& Ar, int64 DataPosition, DataType& Data)
+{
+	const int64 OriginalPosition = Ar.Tell();
+
+	Ar.Seek(DataPosition);
+	Ar << Data;
+
+	Ar.Seek(OriginalPosition);
+}
+
+/** Utility for accessing IVirtualizationSourceControlUtilities from the modular feature system. */
+UE::Virtualization::Experimental::IVirtualizationSourceControlUtilities* GetSourceControlInterface()
+{
+	return (UE::Virtualization::Experimental::IVirtualizationSourceControlUtilities*)IModularFeatures::Get().GetModularFeatureImplementation(FName("VirtualizationSourceControlUtilities"), 0);
+}
+
+namespace Private
+{
+
+FCompressionSettings::FCompressionSettings()
+	: Compressor(ECompressedBufferCompressor::NotSet)
+	, CompressionLevel(ECompressedBufferCompressionLevel::None)
+	, bIsSet(false)
+{
+
+}
+
+FCompressionSettings::FCompressionSettings(const FCompressedBuffer& Buffer)
+{
+	// Note that if the buffer is using a non-oodle format we consider it
+	// as not set.
+	uint64 BlockSize;
+	if (!Buffer.TryGetCompressParameters(Compressor, CompressionLevel, BlockSize))
+	{
+		Reset();
+	}
+	else
+	{
+		bIsSet = true;
+	}
+}
+
+bool FCompressionSettings::operator==(const FCompressionSettings& Other) const
+{
+	return	Compressor == Other.Compressor &&
+			CompressionLevel == Other.CompressionLevel &&
+			bIsSet == Other.bIsSet;
+}
+
+bool FCompressionSettings::operator != (const FCompressionSettings& Other) const
+{
+	return !(*this == Other);
+}
+
+void FCompressionSettings::Reset()
+{
+	Compressor = ECompressedBufferCompressor::NotSet;
+	CompressionLevel = ECompressedBufferCompressionLevel::None;
+	bIsSet = false;
+}
+
+void FCompressionSettings::Set(ECompressedBufferCompressor InCompressor, ECompressedBufferCompressionLevel InCompressionLevel)
+{
+	Compressor = InCompressor;
+	CompressionLevel = InCompressionLevel;
+	bIsSet = true;
+}
+
+void FCompressionSettings::SetToDefault()
+{
+	Compressor = ECompressedBufferCompressor::Kraken;
+	CompressionLevel = ECompressedBufferCompressionLevel::Fast;
+	bIsSet = true;
+}
+
+void FCompressionSettings::SetToDisabled()
+{
+	Compressor = ECompressedBufferCompressor::NotSet;
+	CompressionLevel = ECompressedBufferCompressionLevel::None;
+	bIsSet = true;
+}
+
+bool FCompressionSettings::IsSet() const
+{
+	return bIsSet;
+}
+
+bool FCompressionSettings::IsCompressed() const
+{
+	return bIsSet == true && CompressionLevel != ECompressedBufferCompressionLevel::None;
+}
+
+ECompressedBufferCompressor FCompressionSettings::GetCompressor() const
+{
+	return Compressor;
+}
+
+ECompressedBufferCompressionLevel FCompressionSettings::GetCompressionLevel()
+{
+	return CompressionLevel;
+}
+
+} // namespace Private
+
+FEditorBulkData::FEditorBulkData(FEditorBulkData&& Other)
+{
+	*this = MoveTemp(Other);
+}
+
+FEditorBulkData& FEditorBulkData::operator=(FEditorBulkData&& Other)
+{
+	// The same as the default move constructor, except we need to handle registration and unregistration
+	Unregister();
+	Other.Unregister();
+
+	BulkDataId = MoveTemp(Other.BulkDataId);
+	PayloadContentId = MoveTemp(Other.PayloadContentId);
+	Payload = MoveTemp(Other.Payload);
+	PayloadSize = MoveTemp(Other.PayloadSize);
+	OffsetInFile = MoveTemp(Other.OffsetInFile);
+	PackagePath = MoveTemp(Other.PackagePath);
+	PackageSegment = MoveTemp(Other.PackageSegment);
+	Flags = MoveTemp(Other.Flags);
+	CompressionSettings = MoveTemp(Other.CompressionSettings);
+
+	Other.Reset();
+
+	Register(nullptr);
+
+	return *this;
+}
+
+FEditorBulkData::FEditorBulkData(const FEditorBulkData& Other)
+{
+	*this = Other;
+}
+
+FEditorBulkData& FEditorBulkData::operator=(const FEditorBulkData& Other)
+{
+	// Torn-off BulkDatas remain torn-off even when being copied into from a non-torn-off BulkData
+	// Remaining torn-off is a work-around necessary for FTextureSource::CopyTornOff to avoid registering a new
+	// guid before setting the new BulkData to torn-off. The caller can call Reset to clear the torn-off flag.
+	bool bTornOff = false;
+	if (EnumHasAnyFlags(Flags, EFlags::IsTornOff))
+	{
+		check(!EnumHasAnyFlags(Flags, EFlags::HasRegistered));
+		BulkDataId = Other.BulkDataId;
+		bTornOff = true;
+	}
+	else
+	{
+		Unregister();
+		if (EnumHasAnyFlags(Other.Flags, EFlags::IsTornOff))
+		{
+			BulkDataId = Other.BulkDataId;
+			bTornOff = true;
+		}
+		else if (!BulkDataId.IsValid() && Other.BulkDataId.IsValid())
+		{
+			BulkDataId = FGuid::NewGuid();
+		}
+	}
+
+	PayloadContentId = Other.PayloadContentId;
+	Payload = Other.Payload;
+	PayloadSize = Other.PayloadSize;
+	OffsetInFile = Other.OffsetInFile;
+	PackagePath = Other.PackagePath;
+	PackageSegment = Other.PackageSegment;
+	Flags = Other.Flags;
+	CompressionSettings = Other.CompressionSettings;
+
+	EnumRemoveFlags(Flags, EFlags::TransientFlags);
+
+	if (bTornOff)
+	{
+		EnumAddFlags(Flags, EFlags::IsTornOff);
+	}
+	else
+	{
+		Register(nullptr);
+	}
+	return *this;
+}
+
+FEditorBulkData::~FEditorBulkData()
+{
+	if (AttachedAr != nullptr)
+	{
+		AttachedAr->DetachBulkData(this, false);
+	}
+
+	Unregister();
+}
+
+FEditorBulkData::FEditorBulkData(const FEditorBulkData& Other, ETornOff)
+{
+	EnumAddFlags(Flags, EFlags::IsTornOff);
+	*this = Other; // We rely on operator= preserving the torn-off flag
+}
+
+void FEditorBulkData::TearOff()
+{
+	Unregister();
+	EnumAddFlags(Flags, EFlags::IsTornOff);
+}
+
+void FEditorBulkData::Register(UObject* Owner)
+{
+#if WITH_EDITOR
+	if (BulkDataId.IsValid() && PayloadSize > 0 && !EnumHasAnyFlags(Flags, EFlags::IsTornOff))
+	{
+		IBulkDataRegistry::Get().Register(Owner ? Owner->GetPackage() : nullptr, *this);
+		EnumAddFlags(Flags, EFlags::HasRegistered);
+	}
+#endif
+}
+
+void FEditorBulkData::Unregister()
+{
+#if WITH_EDITOR
+	if (EnumHasAnyFlags(Flags, EFlags::HasRegistered))
+	{
+		check(!EnumHasAnyFlags(Flags, EFlags::IsTornOff));
+		IBulkDataRegistry::Get().OnExitMemory(*this);
+		EnumRemoveFlags(Flags, EFlags::HasRegistered);
+	}
+#endif
+}
+
+static FGuid CreateUniqueGuid(const FGuid& NonUniqueGuid, const UObject* Owner, const TCHAR* DebugName)
+{
+	if (NonUniqueGuid.IsValid() && Owner)
+	{
+		TStringBuilder<256> PathName;
+		Owner->GetPathName(nullptr, PathName);
+		FBlake3 Builder;
+		Builder.Update(&NonUniqueGuid, sizeof(NonUniqueGuid));
+		Builder.Update(PathName.GetData(), PathName.Len() * sizeof(*PathName.GetData()));
+		FBlake3Hash Hash = Builder.Finalize();
+		// We use the first 16 bytes of the FIoHash to create the guid, there is
+		// no specific reason why these were chosen, we could take any pattern or combination
+		// of bytes.
+		uint32* HashBytes = (uint32*)Hash.GetBytes();
+		return FGuid(HashBytes[0], HashBytes[1], HashBytes[2], HashBytes[3]);
+	}
+	else
+	{
+		UE_LOG(LogSerialization, Warning,
+			TEXT("CreateFromBulkData recieved an invalid FGuid. A temporary one will be generated until the package is next re-saved! Package: '%s'"),
+			DebugName);
+		return FGuid::NewGuid();
+	}
+}
+
+void FEditorBulkData::CreateFromBulkData(FUntypedBulkData& InBulkData, const FGuid& InGuid, UObject* Owner)
+{
+	TRACE_CPUPROFILER_EVENT_SCOPE(FEditorBulkData::CreateFromBulkData);
+	
+	checkf(!BulkDataId.IsValid(), 
+		TEXT("Calling ::CreateFromBulkData on a bulkdata object that already has a valid identifier! Package: '%s'"),
+		*InBulkData.GetPackagePath().GetDebugName());
+
+	Reset();
+
+#if UE_ALLOW_LINKERLOADER_ATTACHMENT
+	AttachedAr = InBulkData.AttachedAr;
+	if (AttachedAr != nullptr)
+	{
+		AttachedAr->AttachBulkData(this);
+	}
+#endif //VBD_ALLOW_LINKERLOADER_ATTACHMENT
+
+	// We only need to set up the bulkdata/content identifiers if we have a valid payload
+	bool bWasKeyGuidDerived = false;
+	if (InBulkData.GetBulkDataSize() > 0)
+	{
+		BulkDataId = CreateUniqueGuid(InGuid, Owner, *InBulkData.GetPackagePath().GetDebugName());
+		PayloadContentId = UE::Virtualization::FPayloadId(BulkDataId);
+		bWasKeyGuidDerived = true;
+	}
+	
+	PayloadSize = InBulkData.GetBulkDataSize();
+	
+	PackagePath = InBulkData.GetPackagePath();
+	PackageSegment = InBulkData.GetPackageSegment();
+	
+	OffsetInFile = InBulkData.GetBulkDataOffsetInFile();
+
+	// Mark that we are actually referencing a payload stored in an old bulkdata
+	// format.
+	EnumAddFlags(Flags, EFlags::ReferencesLegacyFile);
+
+	if (InBulkData.IsStoredCompressedOnDisk())
+	{
+		EnumAddFlags(Flags, EFlags::LegacyFileIsCompressed);
+	}
+	else 
+	{
+		EnumAddFlags(Flags, EFlags::DisablePayloadCompression);
+	}
+	if (bWasKeyGuidDerived)
+	{
+		EnumAddFlags(Flags, EFlags::LegacyKeyWasGuidDerived);
+	}
+	Register(Owner);
+}
+
+void FEditorBulkData::CreateLegacyUniqueIdentifier(UObject* Owner)
+{
+	if (BulkDataId.IsValid())
+	{
+		Unregister();
+		BulkDataId = CreateUniqueGuid(BulkDataId, Owner, TEXT("Unknown"));
+		Register(Owner);
+	}
+}
+
+void FEditorBulkData::Serialize(FArchive& Ar, UObject* Owner, bool bAllowRegister)
+{
+	TRACE_CPUPROFILER_EVENT_SCOPE(FEditorBulkData::Serialize);
+
+	if (Ar.IsTransacting())
+	{
+		// Do not process the transaction if the owner is mid loading (see FUntypedBulkData::Serialize)
+		bool bNeedsTransaction = Ar.IsSaving() && (!Owner || !Owner->HasAnyFlags(RF_NeedLoad));
+
+		Ar << bNeedsTransaction;
+
+		if (bNeedsTransaction)
+		{
+			if (Ar.IsLoading())
+			{
+				Unregister();
+			}
+
+			Ar << Flags;
+			Ar << BulkDataId;
+			Ar << PayloadContentId;
+			Ar << PayloadSize;
+			Ar << PackagePath;
+			Ar << PackageSegment;
+			Ar << OffsetInFile;
+
+			// TODO: We could consider compressing the payload so it takes up less space in the 
+			// undo stack or even consider storing as a tmp file on disk rather than keeping it
+			// in memory or some other caching system.
+			// Serializing full 8k texture payloads to memory on each metadata change will empty
+			// the undo stack very quickly.
+			
+			// Note that we will only serialize the payload if it is in memory. Otherwise we can
+			// continue to load the payload as needed from disk or pull from the virtualization system
+			bool bPayloadInArchive = Ar.IsSaving() ? !Payload.IsNull() : false;
+			Ar << bPayloadInArchive;
+
+			if (Ar.IsSaving())
+			{
+				if (bPayloadInArchive)
+				{
+					FCompressedBuffer CompressedPayload = FCompressedBuffer::Compress(Payload, ECompressedBufferCompressor::NotSet, ECompressedBufferCompressionLevel::None);
+					SerializeData(Ar, CompressedPayload, Flags);
+				}
+			}
+			else
+			{
+				FCompressedBuffer CompressedPayload;
+				if (bPayloadInArchive)
+				{
+					SerializeData(Ar, CompressedPayload, Flags);
+				}
+				
+				Payload = CompressedPayload.Decompress();	
+
+				Register(Owner);
+			}
+		}
+	}
+	else if (Ar.IsPersistent() && !Ar.IsObjectReferenceCollector() && !Ar.ShouldSkipBulkData())
+	{
+		FLinkerSave* LinkerSave = nullptr;
+		bool bKeepFileDataByReference = false;
+		if (Ar.IsSaving())
+		{
+			LinkerSave = Cast<FLinkerSave>(Ar.GetLinker());
+			// If we're doing a save that can refer to bulk data by reference, and our legacy data format supports it,
+			// keep any legacy data we have referenced rather than stored, to save space and avoid spending time loading it.
+			bKeepFileDataByReference = LinkerSave && LinkerSave->bProceduralSave && PackageSegment == EPackageSegment::Header;
+			if (!bKeepFileDataByReference)
+			{
+				UpdateKeyIfNeeded();
+			}
+
+			if (bAllowVirtualizationOnSave)
+			{
+				const bool bCanAttemptVirtualization = LinkerSave != nullptr;
+				if (bCanAttemptVirtualization)
+				{
+					FPackagePath LinkerPackagePath;
+					FPackagePath::TryFromPackageName(LinkerSave->LinkerRoot->GetName(), LinkerPackagePath);
+
+					PushData(LinkerPackagePath); // Note this can change various members if we are going from non-virtualized to virtualized
+				}
+			}
+		}
+		else
+		{
+			Unregister();
+		}
+
+		// Store the position in the archive of the flags in case we need to update it later
+		const int64 SavedFlagsPos = Ar.Tell();
+		Ar << Flags;
+		if (Ar.IsLoading())
+		{
+			EnumRemoveFlags(Flags, EFlags::TransientFlags);
+		}
+
+		// TODO: Can probably remove these checks before UE5 release
+		check(!Ar.IsSaving() || GetPayloadSize() == 0 || BulkDataId.IsValid()); // Sanity check to stop us saving out bad data
+		check(!Ar.IsSaving() || GetPayloadSize() == 0 || PayloadContentId.IsValid()); // Sanity check to stop us saving out bad data
+		
+		Ar << BulkDataId;
+		Ar << PayloadContentId;
+		Ar << PayloadSize;
+
+		// TODO: Can probably remove these checks before UE5 release
+		check(!Ar.IsLoading() || GetPayloadSize() == 0 || BulkDataId.IsValid()); // Sanity check to stop us loading in bad data
+		check(!Ar.IsLoading() || GetPayloadSize() == 0 || PayloadContentId.IsValid()); // Sanity check to stop us loading in bad data
+
+		if (Ar.IsSaving())
+		{
+			checkf(Ar.IsCooking() == false, TEXT("FEditorBulkData::Serialize should not be called during a cook"));
+
+			const EFlags UpdatedFlags = BuildFlagsForSerialization(Ar, bKeepFileDataByReference);
+
+			// Go back in the archive and update the flags in the archive, we will only apply the updated flags to the current
+			// object later if we detect that the package saved successfully.
+			// TODO: Not a huge fan of this, might be better to find a way to build the flags during serialization and potential callbacks 
+			// later then go back and update the flags in the Ar. Applying the updated flags only if we are saving a package to disk
+			// and the save succeeds continues to make sense.
+			UpdateArchiveData(Ar, SavedFlagsPos, UpdatedFlags);
+
+			// Write out required extra data if we're saving by reference
+			bool bWriteOutPayload = true;
+			if (IsReferencingByPackagePath(UpdatedFlags))
+			{
+				check(PackageSegment == EPackageSegment::Header); // This should have been checked before setting bKeepFileDataByReference=true
+				if (!IsStoredInPackageTrailer(UpdatedFlags))
+				{
+					Ar << OffsetInFile;
+				}
+
+				bWriteOutPayload = false;
+			}
+			else
+			{
+				bWriteOutPayload = !IsDataVirtualized(UpdatedFlags);
+			}
+
+			if (bWriteOutPayload)
+			{
+				// Need to load the payload so that we can write it out
+				FCompressedBuffer PayloadToSerialize = GetDataInternal();
+				
+				if (!TryPayloadValidationForSaving(PayloadToSerialize, LinkerSave))
+				{
+					Ar.SetError();
+					return;
+				}
+
+				RecompressForSerialization(PayloadToSerialize, UpdatedFlags);
+
+				// If we are expecting a valid payload but fail to find one something critical has broken so assert now
+				// to prevent potentially bad data being saved to disk.
+				checkf(PayloadToSerialize || GetPayloadSize() == 0, TEXT("Failed to acquire the payload for saving!"));
+
+				// If we have a valid linker then we will defer serialization of the payload so that it will
+				// be placed at the end of the output file so we don't have to seek past the payload on load.
+				// If we do not have a linker OR the linker is in text format then we should just serialize
+				// the payload directly to the archive.
+				if (LinkerSave != nullptr && !LinkerSave->IsTextFormat())
+				{	
+					if (IsStoredInPackageTrailer(UpdatedFlags))
+					{
+						// New path that will save the payload to the package trailer
+						SerializeToPackageTrailer(*LinkerSave, PayloadToSerialize, UpdatedFlags, Owner);	
+					}
+					else 
+					{
+						// Legacy path, will save the payload data to the end of the package
+						SerializeToLegacyPath(*LinkerSave, PayloadToSerialize, UpdatedFlags, Owner);
+					}	
+				}
+				else
+				{
+					// Not saving to a package so serialize inline into the archive
+					check(IsStoredInPackageTrailer(UpdatedFlags) == false);
+
+					const int64 OffsetPos = Ar.Tell();
+
+					// Write out a dummy value that we will write over once the payload is serialized
+					int64 PlaceholderValue = INDEX_NONE;
+					Ar << PlaceholderValue; // OffsetInFile
+
+					int64 DataStartOffset = Ar.Tell();
+
+					SerializeData(Ar, PayloadToSerialize, UpdatedFlags);
+					
+					UpdateArchiveData(Ar, OffsetPos, DataStartOffset);
+				}
+			}
+
+			// Some paranoid checks to make sure everything is working as expected
+			// TODO: Probably remove this when FPackageTrailer::IsEnabled is removed
+			if(IsStoredInPackageTrailer(UpdatedFlags) && PayloadContentId.IsValid())
+			{
+				check(LinkerSave != nullptr);
+				check(LinkerSave->PackageTrailerBuilder.IsValid());
+				checkf(!(IsDataVirtualized(UpdatedFlags) && IsReferencingByPackagePath(UpdatedFlags)), TEXT("Payload cannot be both virtualized and a reference"));
+
+				if (IsReferencingByPackagePath(UpdatedFlags))
+				{
+					check(LinkerSave->PackageTrailerBuilder->IsReferencedPayloadEntry(PayloadContentId));
+				}
+				else if (IsDataVirtualized(UpdatedFlags))
+				{
+					check(LinkerSave->PackageTrailerBuilder->IsVirtualizedPayloadEntry(PayloadContentId));
+				}
+				else
+				{
+					check(LinkerSave->PackageTrailerBuilder->IsLocalPayloadEntry(PayloadContentId));
+				}
+			}
+			
+			if (CanUnloadData())
+			{
+				this->CompressionSettings.Reset();
+				Payload.Reset();
+			}
+		}
+		else if (Ar.IsLoading())
+		{
+			OffsetInFile = INDEX_NONE;
+			PackagePath.Empty();
+			PackageSegment = EPackageSegment::Header;
+
+			const FPackageTrailer* Trailer = GetTrailerFromOwner(Owner);
+				
+			if (IsStoredInPackageTrailer())
+			{
+				checkf(Trailer != nullptr, TEXT("Payload was stored in a package trailer, but there no trailer loaded"));
+				// Cache the offset from the trailer (if we move the loading of the payload to the trailer 
+				// at a later point then we can skip this)
+				OffsetInFile = Trailer->FindPayloadOffsetInFile(PayloadContentId);
+			}
+			else
+			{
+				// TODO: This check is for older virtualized formats that might be seen in older test projects.
+				UE_CLOG(IsDataVirtualized(), LogSerialization, Error, TEXT("Payload in '%s' is virtualized in an older format and should be re-saved!"), *Owner->GetName());
+				if (!IsDataVirtualized())
+				{
+					Ar << OffsetInFile;
+				}
+			}
+
+			// This cannot be inside the above ::IsStoredInPackageTrailer branch due to the original prototype assets using the trailer without the StoredInPackageTrailer flag
+			if (Trailer != nullptr && Trailer->FindPayloadStatus(PayloadContentId) == EPayloadStatus::StoredVirtualized)
+			{
+				// As the virtualization process happens outside of serialization we need
+				// to check with the trailer to see if the payload is virtualized or not
+				EnumAddFlags(Flags, EFlags::IsVirtualized);
+				OffsetInFile = INDEX_NONE;
+			}
+
+			checkf(!(IsDataVirtualized() && IsReferencingByPackagePath()), TEXT("Payload cannot be both virtualized and a reference"));
+			checkf(!IsDataVirtualized() || OffsetInFile == INDEX_NONE, TEXT("Virtualized payloads should have an invalid offset"));
+			
+			if (!IsDataVirtualized())
+			{
+				// If we can lazy load then find the PackagePath, otherwise we will want to serialize immediately.
+				FArchive* CacheableArchive = Ar.GetCacheableArchive();
+				if (Ar.IsAllowingLazyLoading() && CacheableArchive != nullptr)
+				{
+					PackagePath = GetPackagePathFromOwner(Owner, PackageSegment);
+				}
+				else
+				{
+					PackagePath.Empty();
+					PackageSegment = EPackageSegment::Header;
+				}
+					
+				if (!PackagePath.IsEmpty() && CacheableArchive != nullptr)
+				{
+#if UE_ALLOW_LINKERLOADER_ATTACHMENT
+					AttachedAr = CacheableArchive;
+					AttachedAr->AttachBulkData(this);
+#endif //VBD_ALLOW_LINKERLOADER_ATTACHMENT
+				}
+				else
+				{
+					checkf(Ar.Tell() == OffsetInFile, TEXT("Attempting to load an inline payload but the offset does not match"));
+
+					// If the package path is invalid or the archive is not cacheable then we
+					// cannot rely on loading the payload at a future point on demand so we need 
+					// to load the data immediately.
+					FCompressedBuffer CompressedPayload;
+					SerializeData(Ar, CompressedPayload, Flags);
+					
+					// Only decompress if there is actual data, otherwise we might as well just 
+					// store the payload as an empty FSharedBuffer.
+					if (CompressedPayload.GetRawSize() > 0)
+					{
+						Payload = CompressedPayload.Decompress();
+					}
+					else
+					{
+						Payload.Reset();
+					}
+				}
+			}
+
+			if (bAllowRegister)
+			{
+				Register(Owner);
+			}
+		}
+	}
+}
+
+void FEditorBulkData::SerializeForRegistry(FArchive& Ar)
+{
+	if (Ar.IsSaving())
+	{
+		check(CanSaveForRegistry());
+		EFlags FlagsForSerialize = Flags;
+		EnumRemoveFlags(FlagsForSerialize, EFlags::TransientFlags);
+		Ar << FlagsForSerialize;
+	}
+	else
+	{
+		Ar << Flags;
+		EnumRemoveFlags(Flags, EFlags::TransientFlags);
+		EnumAddFlags(Flags, EFlags::IsTornOff);
+	}
+
+	Ar << BulkDataId;
+	Ar << PayloadContentId;
+	Ar << PayloadSize;
+	if (Ar.IsSaving())
+	{
+		FString PackageName = PackagePath.GetPackageName();
+		check(PackageName.IsEmpty() || PackageSegment == EPackageSegment::Header);
+		Ar << PackageName;
+	}
+	else
+	{
+		FString PackageName;
+		Ar << PackageName;
+		if (PackageName.IsEmpty())
+		{
+			PackagePath.Empty();
+		}
+		else
+		{
+			ensure(FPackagePath::TryFromPackageName(PackageName, PackagePath));
+		}
+		PackageSegment = EPackageSegment::Header;
+	}
+	Ar << OffsetInFile;
+}
+
+bool FEditorBulkData::CanSaveForRegistry() const
+{
+	return BulkDataId.IsValid() && PayloadSize > 0 && !IsMemoryOnlyPayload()
+		&& EnumHasAnyFlags(Flags, EFlags::IsTornOff) && !EnumHasAnyFlags(Flags, EFlags::HasRegistered)
+		&& (PackagePath.IsEmpty() || PackageSegment == EPackageSegment::Header);
+}
+
+
+FCompressedBuffer FEditorBulkData::LoadFromDisk() const
+{
+	TRACE_CPUPROFILER_EVENT_SCOPE(FEditorBulkData::LoadFromDisk);
+
+	if (PackagePath.IsEmpty())
+	{
+		// Bulkdata objects without a valid package path should not get this far when attempting to access a payload!
+		UE_LOG(LogSerialization, Error, TEXT("Cannot load a payload as the package path is empty!"));
+		return FCompressedBuffer();
+	}
+
+	if (HasPayloadSidecarFile() && CVarShouldLoadFromSidecar.GetValueOnAnyThread())
+	{
+		// Note that this code path is purely for debugging and not expected to be enabled by default
+		if (CVarShouldValidatePayload.GetValueOnAnyThread())
+		{
+			UE_LOG(LogSerialization, Verbose, TEXT("Validating payload loaded from sidecar file: '%s'"), *PackagePath.GetLocalFullPath(EPackageSegment::PayloadSidecar));
+
+			// Load both payloads then generate a FPayloadId from them, since this identifier is a hash of the buffers content
+			// we only need to verify them against PayloadContentId to be sure that the data is correct.
+			FCompressedBuffer SidecarBuffer = LoadFromSidecarFile();
+			FCompressedBuffer AssetBuffer = LoadFromPackageFile();
+
+			UE::Virtualization::FPayloadId SidecarId(SidecarBuffer.Decompress());
+			UE::Virtualization::FPayloadId AssetId(AssetBuffer.Decompress());
+
+			UE_CLOG(SidecarId != PayloadContentId, LogSerialization, Error, TEXT("Sidecar content did not hash correctly! Found '%s' Expected '%s'"), *SidecarId.ToString(), *PayloadContentId.ToString());
+			UE_CLOG(AssetId != PayloadContentId, LogSerialization, Error, TEXT("Asset content did not hash correctly! Found '%s' Expected '%s'"), *AssetId.ToString(), *PayloadContentId.ToString())
+
+			return SidecarBuffer;
+		}
+		else
+		{
+			return LoadFromSidecarFile();
+		}
+	}
+	else
+	{
+		if (CVarShouldLoadFromTrailer.GetValueOnAnyThread())
+		{
+			return LoadFromPackageTrailer();
+		}
+		else
+		{
+			return LoadFromPackageFile();
+		}
+	}
+}
+
+FCompressedBuffer FEditorBulkData::LoadFromPackageFile() const
+{
+	TRACE_CPUPROFILER_EVENT_SCOPE(FEditorBulkData::LoadFromPackageFile);
+
+	UE_LOG(LogSerialization, Verbose, TEXT("Attempting to load payload from the package file: '%s'"), *PackagePath.GetLocalFullPath(PackageSegment));
+
+	// Open a reader to the file
+	TUniquePtr<FArchive> BulkArchive;
+	if (!IsReferencingByPackagePath() || PackageSegment != EPackageSegment::Header)
+	{
+		FOpenPackageResult Result = IPackageResourceManager::Get().OpenReadPackage(PackagePath, PackageSegment);
+		if (Result.Format == EPackageFormat::Binary)
+		{
+			BulkArchive = MoveTemp(Result.Archive);
+		}
+	}
+	else
+	{
+		// *this may have been loaded from the EditorDomain, but saved with a reference to the bulk data in the
+		// Workspace Domain file. This was only possible if PackageSegment == Header; we checked that when serializing to the EditorDomain
+		// In that case, we need to use OpenReadExternalResource to access the Workspace Domain file
+		// In the cases where *this was loaded from the WorkspaceDomain, OpenReadExternalResource and OpenReadPackage are identical.
+		BulkArchive = IPackageResourceManager::Get().OpenReadExternalResource(EPackageExternalResource::WorkspaceDomainFile,
+			PackagePath.GetPackageName());
+	}
+
+	if (!BulkArchive.IsValid())
+	{
+		LogPackageOpenFailureMessage(PackagePath, PackageSegment);
+		return FCompressedBuffer();
+	}
+
+	checkf(OffsetInFile != INDEX_NONE, TEXT("Attempting to load '%s' from disk with an invalid OffsetInFile!"), *PackagePath.GetDebugNameWithExtension(PackageSegment));
+	// Move the correct location of the data in the file
+	BulkArchive->Seek(OffsetInFile);
+
+	// Now we can actually serialize it
+	FCompressedBuffer PayloadFromDisk;
+	SerializeData(*BulkArchive, PayloadFromDisk, Flags);
+
+	return PayloadFromDisk;
+}
+
+FCompressedBuffer FEditorBulkData::LoadFromPackageTrailer() const
+{
+	TRACE_CPUPROFILER_EVENT_SCOPE(FEditorBulkData::LoadFromPackageTrailer);
+
+	UE_LOG(LogSerialization, Verbose, TEXT("Attempting to load payload from the package trailer: '%s'"), *PackagePath.GetLocalFullPath(PackageSegment));
+
+	// TODO: Could just get the trailer from the owning FLinkerLoad if still attached
+
+	// Open a reader to the file
+	TUniquePtr<FArchive> BulkArchive;
+	if (!IsReferencingByPackagePath() || PackageSegment != EPackageSegment::Header)
+	{
+		FOpenPackageResult Result = IPackageResourceManager::Get().OpenReadPackage(PackagePath, PackageSegment);
+		if (Result.Format == EPackageFormat::Binary)
+		{
+			BulkArchive = MoveTemp(Result.Archive);
+		}
+	}
+	else
+	{
+		// *this may have been loaded from the EditorDomain, but saved with a reference to the bulk data in the
+		// Workspace Domain file. This was only possible if PackageSegment == Header; we checked that when serializing to the EditorDomain
+		// In that case, we need to use OpenReadExternalResource to access the Workspace Domain file
+		// In the cases where *this was loaded from the WorkspaceDomain, OpenReadExternalResource and OpenReadPackage are identical.
+		BulkArchive = IPackageResourceManager::Get().OpenReadExternalResource(EPackageExternalResource::WorkspaceDomainFile,
+			PackagePath.GetPackageName());
+	}
+
+	if (!BulkArchive.IsValid())
+	{
+		LogPackageOpenFailureMessage(PackagePath, PackageSegment);
+		return FCompressedBuffer();
+	}
+
+	BulkArchive->Seek(BulkArchive->TotalSize());
+
+	FPackageTrailer Trailer;
+	
+	if (Trailer.TryLoadBackwards(*BulkArchive))
+	{
+		return Trailer.LoadLocalPayload(PayloadContentId, *BulkArchive);
+	}
+	else
+	{
+		return FCompressedBuffer();
+	}	
+}
+
+FCompressedBuffer FEditorBulkData::LoadFromSidecarFileInternal(ErrorVerbosity Verbosity) const
+{
+	TRACE_CPUPROFILER_EVENT_SCOPE(FEditorBulkData::LoadFromSidecarFileInternal);
+
+	FOpenPackageResult Result = IPackageResourceManager::Get().OpenReadPackage(PackagePath, EPackageSegment::PayloadSidecar);
+	if (Result.Archive.IsValid() && Result.Format == EPackageFormat::Binary)
+	{
+		uint32 Version = INDEX_NONE;
+		*Result.Archive << Version;
+
+		if (Version != FTocEntry::PayloadSidecarFileVersion)
+		{
+			UE_CLOG(Verbosity > ErrorVerbosity::None, LogSerialization, Error, TEXT("Unknown version (%u) found in '%s'"), Version, *PackagePath.GetLocalFullPath(EPackageSegment::PayloadSidecar));
+			return FCompressedBuffer();
+		}
+
+		// First we load the table of contents so we can find the payload in the file
+		TArray<FTocEntry> TableOfContents;
+		*Result.Archive << TableOfContents;
+
+		const FTocEntry* Entry = TableOfContents.FindByPredicate([&PayloadContentId = PayloadContentId](const FTocEntry& Entry)
+			{
+				return Entry.Identifier == PayloadContentId;
+			});
+
+		if (Entry != nullptr)
+		{
+			if (Entry->OffsetInFile != INDEX_NONE)
+			{
+				// Move the correct location of the data in the file
+				Result.Archive->Seek(Entry->OffsetInFile);
+
+				// Now we can actually serialize it
+				FCompressedBuffer PayloadFromDisk;
+				SerializeData(*Result.Archive, PayloadFromDisk, EFlags::None);
+
+				return PayloadFromDisk;
+			}
+			else if(Verbosity > ErrorVerbosity::None)
+			{
+				UE_LOG(LogSerialization, Error, TEXT("Payload '%s' in '%s' has an invalid OffsetInFile!"), *PayloadContentId.ToString(), *PackagePath.GetLocalFullPath(EPackageSegment::PayloadSidecar));
+			}
+		}
+		else if(Verbosity > ErrorVerbosity::None)
+		{
+			UE_LOG(LogSerialization, Error, TEXT("Unable to find payload '%s' in '%s'"), *PayloadContentId.ToString(), *PackagePath.GetLocalFullPath(EPackageSegment::PayloadSidecar));
+		}
+	}
+	else if(Verbosity > ErrorVerbosity::None)
+	{
+		LogPackageOpenFailureMessage(PackagePath, EPackageSegment::PayloadSidecar);
+	}
+
+	return FCompressedBuffer();
+}
+
+FCompressedBuffer FEditorBulkData::LoadFromSidecarFile() const
+{
+	TRACE_CPUPROFILER_EVENT_SCOPE(FEditorBulkData::LoadFromSidecarFile);
+
+	UE_LOG(LogSerialization, Verbose, TEXT("Attempting to load payload from the sidecar file: '%s'"),
+		*PackagePath.GetLocalFullPath(EPackageSegment::PayloadSidecar));
+
+	if (CVarShouldAllowSidecarSyncing.GetValueOnAnyThread())
+	{
+		FCompressedBuffer PayloadFromDisk = LoadFromSidecarFileInternal(ErrorVerbosity::None);
+		if (PayloadFromDisk.IsNull())
+		{
+			UE_LOG(LogSerialization, Verbose, TEXT("Initial load from sidecar failed, attempting to sync the file: '%s'"),
+				*PackagePath.GetLocalFullPath(EPackageSegment::PayloadSidecar));
+
+			if (UE::Virtualization::Experimental::IVirtualizationSourceControlUtilities* SourceControlInterface = GetSourceControlInterface())
+			{
+				// SyncPayloadSidecarFile should log failure cases, so there is no need for us to add log messages here
+				if (SourceControlInterface->SyncPayloadSidecarFile(PackagePath))
+				{
+					PayloadFromDisk = LoadFromSidecarFileInternal(ErrorVerbosity::All);
+				}
+			}
+			else
+			{
+				UE_LOG(LogSerialization, Error, TEXT("Failed to find IVirtualizationSourceControlUtilities, unable to try and sync: '%s'"),
+					*PackagePath.GetLocalFullPath(EPackageSegment::PayloadSidecar));
+			}
+		}
+
+		return PayloadFromDisk;
+	}
+	else
+	{
+		return LoadFromSidecarFileInternal(ErrorVerbosity::All);
+	}
+}
+
+bool FEditorBulkData::SerializeData(FArchive& Ar, FCompressedBuffer& InPayload, const EFlags PayloadFlags) const
+{
+	TRACE_CPUPROFILER_EVENT_SCOPE(FEditorBulkData::SerializeData);
+
+	if (Ar.IsSaving())
+	{
+		Ar << InPayload;
+		return true;
+	}
+	else if (Ar.IsLoading() && !IsReferencingOldBulkData(PayloadFlags))
+	{
+		Ar << InPayload;
+		return InPayload.IsNull();
+	}
+	else if (Ar.IsLoading()) 
+	{
+		// Loading from old bulkdata format
+		const int64 Size = GetPayloadSize();
+		FUniqueBuffer LoadPayload = FUniqueBuffer::Alloc(Size);
+
+		if (EnumHasAnyFlags(PayloadFlags, EFlags::LegacyFileIsCompressed))
+		{
+			Ar.SerializeCompressed(LoadPayload.GetData(), Size, NAME_Zlib, COMPRESS_NoFlags, false);
+		}
+		else
+		{
+			Ar.Serialize(LoadPayload.GetData(), Size);
+		}
+
+		InPayload = FCompressedBuffer::Compress(LoadPayload.MoveToShared(), ECompressedBufferCompressor::NotSet, ECompressedBufferCompressionLevel::None);
+
+		return true;
+	}
+	else
+	{
+		return false;
+	}
+}
+
+void FEditorBulkData::PushData(const FPackagePath& InPackagePath)
+{
+	checkf(IsDataVirtualized() == false || Payload.IsNull(), TEXT("Cannot have a valid payload in memory if the payload is virtualized!")); // Sanity check
+
+	// We only need to push if the payload if it actually has data and it is not 
+	// currently virtualized (either we have an updated payload in memory or the 
+	// payload is currently non-virtualized and stored on disk)
+	
+	UE::Virtualization::IVirtualizationSystem& VirtualizationSystem = UE::Virtualization::IVirtualizationSystem::Get();
+	if (!IsDataVirtualized() && GetPayloadSize() > 0 && VirtualizationSystem.IsEnabled())
+	{ 
+		TRACE_CPUPROFILER_EVENT_SCOPE(FEditorBulkData::PushData);
+
+		// We should only need to load from disk at this point if we are going from
+		// a non-virtualized payload to a virtualized one. If the bulkdata is merely being
+		// edited then we should have the payload in memory already and are just accessing a
+		// reference to it.
+
+		UpdateKeyIfNeeded();
+		FCompressedBuffer PayloadToPush = GetDataInternal();
+		// TODO: If the push fails we will end up potentially re-compressing this payload for
+		// serialization, we need a better way to save the results of 'RecompressForSerialization'
+		RecompressForSerialization(PayloadToPush, Flags);
+
+		// TODO: We could make this a config option?
+		if (VirtualizationSystem.PushData(PayloadContentId, PayloadToPush, UE::Virtualization::EStorageType::Local, InPackagePath.GetPackageName()))
+		{
+			EnumAddFlags(Flags, EFlags::IsVirtualized);
+			EnumRemoveFlags(Flags, EFlags::ReferencesLegacyFile | EFlags::ReferencesWorkspaceDomain | EFlags::LegacyFileIsCompressed);
+			check(!EnumHasAnyFlags(Flags, EFlags::LegacyKeyWasGuidDerived)); // Removed by UpdateKeyIfNeeded
+
+			// Clear members associated with non-virtualized data and release the in-memory
+			// buffer.
+			PackagePath.Empty();
+			PackageSegment = EPackageSegment::Header;
+			OffsetInFile = INDEX_NONE;
+
+			// Update our information in the registry
+			Register(nullptr);
+		}
+	}	
+}
+
+FCompressedBuffer FEditorBulkData::PullData() const
+{
+	TRACE_CPUPROFILER_EVENT_SCOPE(FEditorBulkData::PullData);
+
+	FCompressedBuffer PulledPayload = UE::Virtualization::IVirtualizationSystem::Get().PullData(PayloadContentId);
+
+	if (PulledPayload)
+	{
+		checkf(	PayloadSize == PulledPayload.GetRawSize(),
+				TEXT("Mismatch between serialized length (%" INT64_FMT ") and virtualized data length (%" UINT64_FMT ")"),
+				PayloadSize,
+				PulledPayload.GetRawSize());
+	}
+	
+	return PulledPayload;
+}
+
+bool FEditorBulkData::CanUnloadData() const
+{
+	// We cannot unload the data if are unable to reload it from a file
+	return IsDataVirtualized() || (PackagePath.IsEmpty() == false && AttachedAr != nullptr);
+}
+
+bool FEditorBulkData::IsMemoryOnlyPayload() const
+{
+	return !Payload.IsNull() && !IsDataVirtualized() && PackagePath.IsEmpty();
+}
+
+void FEditorBulkData::Reset()
+{
+	// Note that we do not reset the BulkDataId
+	if (AttachedAr != nullptr)
+	{
+		AttachedAr->DetachBulkData(this, false);
+	}
+
+	Unregister();
+	PayloadContentId.Reset();
+	Payload.Reset();
+	PayloadSize = 0;
+	OffsetInFile = INDEX_NONE;
+	PackagePath.Empty();
+	PackageSegment = EPackageSegment::Header;
+	Flags = EFlags::None;
+
+	CompressionSettings.Reset();
+}
+
+void FEditorBulkData::UnloadData()
+{
+	if (CanUnloadData())
+	{
+		Payload.Reset();
+	}
+}
+
+void FEditorBulkData::DetachFromDisk(FArchive* Ar, bool bEnsurePayloadIsLoaded)
+{
+	TRACE_CPUPROFILER_EVENT_SCOPE(FEditorBulkData::DetachFromDisk);
+
+	check(Ar);
+	check(Ar == AttachedAr || AttachedAr == nullptr || AttachedAr->IsProxyOf(Ar));
+
+	if (!IsDataVirtualized() && !PackagePath.IsEmpty())
+	{
+		if (Payload.IsNull() && bEnsurePayloadIsLoaded)
+		{
+			FCompressedBuffer CompressedPayload = GetDataInternal();
+			Payload = CompressedPayload.Decompress();
+		}
+
+		PackagePath.Empty();
+		PackageSegment = EPackageSegment::Header;
+		OffsetInFile = INDEX_NONE;
+
+		EnumRemoveFlags(Flags, EFlags::ReferencesLegacyFile | EFlags::ReferencesWorkspaceDomain | EFlags::LegacyFileIsCompressed);
+
+		if (PayloadSize > 0)
+		{
+			Register(nullptr);
+		}
+		else
+		{
+			Unregister();
+		}
+	}
+
+	AttachedAr = nullptr;	
+}
+
+FGuid FEditorBulkData::GetIdentifier() const
+{
+	checkf(GetPayloadSize() == 0 || BulkDataId.IsValid(), TEXT("If bulkdata has a valid payload then it should have a valid BulkDataId"));
+	return BulkDataId;
+}
+
+void FEditorBulkData::SerializeToLegacyPath(FLinkerSave& LinkerSave, FCompressedBuffer PayloadToSerialize, EFlags UpdatedFlags, UObject* Owner)
+{
+	const int64 OffsetPos = LinkerSave.Tell();
+
+	// Write out a dummy value that we will write over once the payload is serialized
+	int64 PlaceholderValue = INDEX_NONE;
+	LinkerSave << PlaceholderValue; // OffsetInFile
+
+	// The lambda is mutable so that PayloadToSerialize is not const (due to FArchive api not accepting const values)
+	auto SerializePayload = [this, OffsetPos, PayloadToSerialize, UpdatedFlags, Owner](FLinkerSave& LinkerSave, FArchive& ExportsArchive, FArchive& DataArchive, int64 DataStartOffset) mutable
+	{
+		checkf(ExportsArchive.IsCooking() == false, TEXT("FEditorBulkData::Serialize should not be called during a cook"));
+
+		SerializeData(DataArchive, PayloadToSerialize, UpdatedFlags);
+
+		UpdateArchiveData(ExportsArchive, OffsetPos, DataStartOffset);
+
+		// If we are saving the package to disk (we have access to FLinkerSave and its filepath is valid) 
+		// then we should register a callback to be received once the package has actually been saved to 
+		// disk so that we can update the object's members to be redirected to the saved file.
+		if (!LinkerSave.GetFilename().IsEmpty())
+		{
+			// At some point saving to the sidecar file will be mutually exclusive with saving to the asset file, at that point
+			// we can split these code paths entirely for clarity. (might need to update ::BuildFlagsForSerialization at that point too!)
+			if (ShouldSaveToPackageSidecar())
+			{
+				FLinkerSave::FSidecarStorageInfo& SidecarData = LinkerSave.SidecarDataToAppend.AddZeroed_GetRef();
+				SidecarData.Identifier = PayloadContentId;
+				SidecarData.Payload = PayloadToSerialize;
+			}
+
+			auto OnSavePackage = [this, DataStartOffset, UpdatedFlags, Owner](const FPackagePath& InPackagePath, FObjectPostSaveContext ObjectSaveContext)
+			{
+				if (!ObjectSaveContext.IsUpdatingLoadedPath())
+				{
+					return;
+				}
+
+				this->PackagePath = InPackagePath;
+				check(!this->PackagePath.IsEmpty()); // LinkerSave guarantees a valid PackagePath if we're updating loaded path
+				this->OffsetInFile = DataStartOffset;
+				this->Flags = UpdatedFlags;
+
+				if (CanUnloadData())
+				{
+					this->CompressionSettings.Reset();
+					this->Payload.Reset();
+				}
+
+				// Update our information in the registry
+				// TODO: Pass Owner into Register once the AssetRegistry has been fixed to use the updated PackageGuid from the save
+				Register(nullptr);
+			};
+
+			LinkerSave.PostSaveCallbacks.Add(MoveTemp(OnSavePackage));
+		}
+	};
+
+	auto AdditionalDataCallback = [SerializePayload = MoveTemp(SerializePayload)](FLinkerSave& ExportsArchive, FArchive& DataArchive, int64 DataStartOffset) mutable
+	{
+		SerializePayload(ExportsArchive, ExportsArchive, DataArchive, DataStartOffset);
+	};
+
+	LinkerSave.AdditionalDataToAppend.Add(MoveTemp(AdditionalDataCallback));	// -V595 PVS believes that LinkerSave can potentially be nullptr at 
+																				// this point however we test LinkerSave != nullptr to enter this branch.
+}
+
+void FEditorBulkData::SerializeToPackageTrailer(FLinkerSave& LinkerSave, FCompressedBuffer PayloadToSerialize, EFlags UpdatedFlags, UObject* Owner)
+{
+	auto OnPayloadWritten = [this, UpdatedFlags](FLinkerSave& LinkerSave, const FPackageTrailer& Trailer) mutable
+	{
+		checkf(LinkerSave.IsCooking() == false, TEXT("FEditorBulkData::Serialize should not be called during a cook"));
+
+		int64 PayloadOffset = Trailer.FindPayloadOffsetInFile(PayloadContentId);
+
+		// If we are saving the package to disk (we have access to FLinkerSave and its filepath is valid) 
+		// then we should register a callback to be received once the package has actually been saved to 
+		// disk so that we can update the object's members to be redirected to the saved file.
+		if (!LinkerSave.GetFilename().IsEmpty())
+		{
+			auto OnSavePackage = [this, PayloadOffset, UpdatedFlags](const FPackagePath& InPackagePath, FObjectPostSaveContext ObjectSaveContext)
+			{
+				if (!ObjectSaveContext.IsUpdatingLoadedPath())
+				{
+					return;
+				}
+
+				this->PackagePath = InPackagePath;
+				check(!this->PackagePath.IsEmpty()); // LinkerSave guarantees a valid PackagePath if we're updating loaded path
+				this->OffsetInFile = PayloadOffset;
+				this->Flags = UpdatedFlags;
+
+				if (CanUnloadData())
+				{
+					this->CompressionSettings.Reset();
+					this->Payload.Reset();
+				}
+
+				// Update our information in the registry
+				// TODO: Pass Owner into Register once the AssetRegistry has been fixed to use the updated PackageGuid from the save
+				Register(nullptr);
+			};
+
+			LinkerSave.PostSaveCallbacks.Add(MoveTemp(OnSavePackage));
+		}
+	};
+
+	LinkerSave.PackageTrailerBuilder->AddPayload(PayloadContentId, MoveTemp(PayloadToSerialize), MoveTemp(OnPayloadWritten));
+}
+
+void FEditorBulkData::UpdatePayloadImpl(FSharedBuffer&& InPayload, UE::Virtualization::FPayloadId&& InPayloadID)
+{
+	TRACE_CPUPROFILER_EVENT_SCOPE(FEditorBulkData::UpdatePayloadImpl);
+
+	if (AttachedAr != nullptr)
+	{
+		AttachedAr->DetachBulkData(this, false);
+	}
+
+	check(AttachedAr == nullptr);
+
+	// We only take the payload if it has a length to avoid potentially holding onto a
+	// 0 byte allocation in the FSharedBuffer
+	if(InPayload.GetSize() > 0)
+	{ 
+		Payload = MoveTemp(InPayload).MakeOwned();
+	}
+	else
+	{
+		Payload.Reset();
+	}
+
+	PayloadSize = (int64)Payload.GetSize();
+	PayloadContentId = MoveTemp(InPayloadID);
+
+	EnumRemoveFlags(Flags,	EFlags::IsVirtualized |
+							EFlags::ReferencesLegacyFile |
+							EFlags::ReferencesWorkspaceDomain |
+							EFlags::LegacyFileIsCompressed |
+							EFlags::LegacyKeyWasGuidDerived);
+
+	PackagePath.Empty();
+	PackageSegment = EPackageSegment::Header;
+	OffsetInFile = INDEX_NONE;
+
+	if (PayloadSize > 0)
+	{
+		if (!BulkDataId.IsValid())
+		{
+			BulkDataId = FGuid::NewGuid();
+		}
+		Register(nullptr);
+	}
+	else
+	{
+		Unregister();
+	}
+}
+
+FCompressedBuffer FEditorBulkData::GetDataInternal() const
+{
+	TRACE_CPUPROFILER_EVENT_SCOPE(FEditorBulkData::GetDataInternal);
+
+	// Early out there isn't any data to actually load
+	if (GetPayloadSize() == 0)
+	{
+		return FCompressedBuffer();
+	}
+
+	// Check if we already have the data in memory
+	if (Payload)
+	{
+		// Note that this doesn't actually compress the data!
+		return FCompressedBuffer::Compress(Payload, ECompressedBufferCompressor::NotSet, ECompressedBufferCompressionLevel::None);
+	}
+
+	if (IsDataVirtualized())
+	{
+		FCompressedBuffer CompressedPayload = PullData();
+		
+		checkf(Payload.IsNull(), TEXT("Pulling data somehow assigned it to the bulk data object!")); //Make sure that we did not assign the buffer internally
+
+		UE_CLOG(CompressedPayload.IsNull(), LogSerialization, Error, TEXT("Failed to pull payload '%s'"), *PayloadContentId.ToString());
+		UE_CLOG(!IsDataValid(*this, CompressedPayload), LogSerialization, UE_CORRUPTED_DATA_SEVERITY, TEXT("Virtualized payload '%s' is corrupt! Check the backend storage."), *PayloadContentId.ToString());
+		
+		return CompressedPayload;
+	}
+	else
+	{
+		FCompressedBuffer CompressedPayload = LoadFromDisk();
+		
+		check(Payload.IsNull()); //Make sure that we did not assign the buffer internally
+
+		UE_CLOG(CompressedPayload.IsNull(), LogSerialization, Error, TEXT("Failed to load payload '%s"), *PayloadContentId.ToString());
+		UE_CLOG(!IsDataValid(*this, CompressedPayload), LogSerialization, UE_CORRUPTED_DATA_SEVERITY, TEXT("Payload '%s' loaded from '%s' is corrupt! Check the file on disk."),
+			*PayloadContentId.ToString(), 
+			*PackagePath.GetDebugName());
+
+		return CompressedPayload;
+	}
+}
+
+TFuture<FSharedBuffer> FEditorBulkData::GetPayload() const
+{
+	TPromise<FSharedBuffer> Promise;
+	
+	if (GetPayloadSize() == 0)
+	{
+		// Early out for 0 sized payloads
+		Promise.SetValue(FSharedBuffer());
+	}
+	else if (Payload)
+	{
+		// Avoid a unnecessary compression and decompression if we already have the uncompressed payload
+		Promise.SetValue(Payload);
+	}
+	else
+	{
+		FCompressedBuffer CompressedPayload = GetDataInternal();
+
+		// TODO: Not actually async yet!
+		Promise.SetValue(CompressedPayload.Decompress());
+	}
+
+	return Promise.GetFuture();
+}
+
+TFuture<FCompressedBuffer>FEditorBulkData::GetCompressedPayload() const
+{
+	TPromise<FCompressedBuffer> Promise;
+
+	FCompressedBuffer CompressedPayload = GetDataInternal();
+
+	// TODO: Not actually async yet!
+	Promise.SetValue(MoveTemp(CompressedPayload));
+
+	return Promise.GetFuture();
+}
+
+void FEditorBulkData::UpdatePayload(FSharedBuffer InPayload)
+{
+	TRACE_CPUPROFILER_EVENT_SCOPE(FEditorBulkData::UpdatePayload);
+	UE::Virtualization::FPayloadId NewPayloadId(InPayload);
+	UpdatePayloadImpl(MoveTemp(InPayload), MoveTemp(NewPayloadId));
+}
+
+FEditorBulkData::FSharedBufferWithID::FSharedBufferWithID(FSharedBuffer InPayload)
+	: Payload(MoveTemp(InPayload))
+	, PayloadId(Payload)
+{
+}
+
+void FEditorBulkData::UpdatePayload(FSharedBufferWithID InPayload)
+{
+	UpdatePayloadImpl(MoveTemp(InPayload.Payload), MoveTemp(InPayload.PayloadId));
+}
+
+void FEditorBulkData::SetCompressionOptions(ECompressionOptions Option)
+{
+	switch (Option)
+	{
+	case ECompressionOptions::Disabled:
+		CompressionSettings.SetToDisabled();
+		break;
+	case ECompressionOptions::Default:
+		CompressionSettings.Reset();
+		break;
+	default:
+		checkNoEntry();
+	}
+
+	if (CompressionSettings.GetCompressionLevel() == ECompressedBufferCompressionLevel::None)
+	{
+		EnumAddFlags(Flags, EFlags::DisablePayloadCompression);
+	}
+	else
+	{
+		EnumRemoveFlags(Flags, EFlags::DisablePayloadCompression);
+	}
+}
+
+void FEditorBulkData::SetCompressionOptions(ECompressedBufferCompressor Compressor, ECompressedBufferCompressionLevel CompressionLevel)
+{
+	CompressionSettings.Set(Compressor, CompressionLevel);
+
+	if (CompressionSettings.GetCompressionLevel() == ECompressedBufferCompressionLevel::None)
+	{
+		EnumAddFlags(Flags, EFlags::DisablePayloadCompression);
+	}
+	else
+	{
+		EnumRemoveFlags(Flags, EFlags::DisablePayloadCompression);
+	}
+}
+
+FCustomVersionContainer FEditorBulkData::GetCustomVersions(FArchive& InlineArchive)
+{
+	return InlineArchive.GetCustomVersions();
+}
+
+void FEditorBulkData::UpdatePayloadId()
+{
+	UpdateKeyIfNeeded();
+}
+
+#if UE_ENABLE_VIRTUALIZATION_TOGGLE
+
+void FEditorBulkData::SetVirtualizationOptOut(bool bOptOut)
+{
+	if (ShouldAllowVirtualizationOptOut())
+	{
+		bSkipVirtualization = bOptOut;
+	}	
+}
+
+#endif //UE_ENABLE_VIRTUALIZATION_TOGGLE
+
+void FEditorBulkData::UpdateKeyIfNeeded()
+{
+	// If this was created from old BulkData then the key is generated from an older FGuid, we
+	// should recalculate it based off the payload to keep the key consistent in the future.
+	if (EnumHasAnyFlags(Flags, EFlags::LegacyKeyWasGuidDerived))
+	{
+		checkf(IsDataVirtualized() == false, TEXT("Cannot have a virtualized payload if loaded from legacy BulkData")); // Sanity check
+
+		// Load the payload from disk (or memory) so that we can hash it
+		FSharedBuffer InPayload = GetDataInternal().Decompress();
+		PayloadContentId = UE::Virtualization::FPayloadId(InPayload);
+
+		// Store as the in memory payload, since this method is only called during saving 
+		// we know it will get cleared anyway.
+		Payload = InPayload;
+		EnumRemoveFlags(Flags, EFlags::LegacyKeyWasGuidDerived);
+	}
+}
+
+void FEditorBulkData::RecompressForSerialization(FCompressedBuffer& InOutPayload, EFlags PayloadFlags) const
+{
+	Private::FCompressionSettings CurrentSettings(InOutPayload);
+	Private::FCompressionSettings TargetSettings;
+
+	if (EnumHasAnyFlags(PayloadFlags, EFlags::DisablePayloadCompression))
+	{
+		// If the disable payload compression flag is set, then we should not compress the payload
+		TargetSettings.SetToDisabled(); 
+	}
+	else if (CompressionSettings.IsSet())
+	{
+		// If we have pending compression settings then we can apply them to the payload
+		TargetSettings = CompressionSettings;
+	}
+	else if(!CurrentSettings.IsCompressed()) 
+	{
+		// If we have no settings to apply to the payload and the payload is currently uncompressed then we
+		// should use the default compression settings.
+		TargetSettings.SetToDefault();
+	}
+	else
+	{
+		// If we have no settings to apply to the payload but the payload is already compressed then we can
+		// just keep the existing settings, what ever they are.
+		TargetSettings = CurrentSettings;
+	}
+	
+	// Now we will re-compress the input payload if the current compression settings differ from the desired settings
+	if (TargetSettings != CurrentSettings)
+	{
+		FCompositeBuffer DecompressedBuffer = InOutPayload.DecompressToComposite();
+
+		// If the buffer actually decompressed we can have both the compressed and the uncompressed version of the
+		// payload in memory. Compressing it will create a third version so before doing that we should reset
+		// the original compressed buffer in case that we can release it to reduce higher water mark pressure.
+		InOutPayload.Reset();
+
+		InOutPayload = FCompressedBuffer::Compress(DecompressedBuffer, TargetSettings.GetCompressor(), TargetSettings.GetCompressionLevel());
+	}
+}
+
+FEditorBulkData::EFlags FEditorBulkData::BuildFlagsForSerialization(FArchive& Ar, bool bKeepFileDataByReference) const
+{
+	if (Ar.IsSaving())
+	{
+		EFlags UpdatedFlags = Flags;
+
+		const FLinkerSave* LinkerSave = Cast<FLinkerSave>(Ar.GetLinker());
+
+		// Now update any changes to the flags that we might need to make when serializing.
+		// Note that these changes are not applied to the current object UNLESS we are saving
+		// the package, in which case the newly modified flags will be applied once we confirm
+		// that the package has saved.
+
+		bool bIsReferencingByPackagePath = IsReferencingByPackagePath(UpdatedFlags);
+		bool bCanKeepFileDataByReference = bIsReferencingByPackagePath || !PackagePath.IsEmpty();
+		if (bKeepFileDataByReference && bCanKeepFileDataByReference)
+		{
+			if (!bIsReferencingByPackagePath)
+			{
+				EnumAddFlags(UpdatedFlags, EFlags::ReferencesWorkspaceDomain);
+			}
+			EnumRemoveFlags(UpdatedFlags, EFlags::HasPayloadSidecarFile | EFlags::IsVirtualized);
+		}
+		else
+		{
+			EnumRemoveFlags(UpdatedFlags, EFlags::ReferencesLegacyFile | EFlags::ReferencesWorkspaceDomain | 
+				EFlags::LegacyFileIsCompressed | EFlags::LegacyKeyWasGuidDerived);
+
+			
+			if (LinkerSave != nullptr && !LinkerSave->GetFilename().IsEmpty() && ShouldSaveToPackageSidecar())
+			{
+				EnumAddFlags(UpdatedFlags, EFlags::HasPayloadSidecarFile);
+				EnumRemoveFlags(UpdatedFlags, EFlags::IsVirtualized);
+			}
+			else
+			{
+				EnumRemoveFlags(UpdatedFlags, EFlags::HasPayloadSidecarFile);
+
+				// Remove the virtualization flag if we are rehydrating packages on save unless
+				// referencing the payload data is allowed, in which case we can continue to save
+				// as virtualized.
+				if (LinkerSave != nullptr && !bKeepFileDataByReference && bForceRehydrationOnSave)
+				{
+					EnumRemoveFlags(UpdatedFlags, EFlags::IsVirtualized);
+				}
+			}
+		}
+
+		if (ShouldUseLegacySerialization(LinkerSave) == false)
+		{
+			EnumAddFlags(UpdatedFlags, EFlags::StoredInPackageTrailer);
+		}
+		else
+		{
+			EnumRemoveFlags(UpdatedFlags, EFlags::StoredInPackageTrailer);
+		}
+
+		return UpdatedFlags;
+	}
+	else
+	{
+		return Flags;
+	}
+}
+
+bool FEditorBulkData::TryPayloadValidationForSaving(const FCompressedBuffer& PayloadForSaving, FLinkerSave* LinkerSave) const
+{
+	if (!IsDataValid(*this, PayloadForSaving) || (GetPayloadSize() > 0 && PayloadForSaving.IsNull()))
+	{
+		const FString ErrorMessage = GetCorruptedPayloadErrorMsgForSave(LinkerSave).ToString();
+
+		ensureMsgf(false, TEXT("%s"), *ErrorMessage);
+
+		if (LinkerSave != nullptr && LinkerSave->GetOutputDevice() != nullptr)
+		{
+			LinkerSave->GetOutputDevice()->Logf(ELogVerbosity::Error, TEXT("%s"), *ErrorMessage);
+		}
+		else
+		{
+			UE_LOG(LogSerialization, Error, TEXT("%s"), *ErrorMessage);
+		}
+
+		return false;
+	}
+	else
+	{
+		return true;
+	}
+}
+
+FText FEditorBulkData::GetCorruptedPayloadErrorMsgForSave(FLinkerSave* Linker) const
+{
+	const FText GuidID = FText::FromString(GetIdentifier().ToString());
+
+	if (Linker != nullptr)
+	{
+		// We know the package we are saving to.
+		const FText PackageName = FText::FromString(Linker->LinkerRoot->GetName());
+		
+		return FText::Format(	NSLOCTEXT("Core", "Serialization_InvalidPayloadToPkg", 
+								"Attempting to save bulkdata {0} with an invalid payload to package '{1}'. The package probably needs to be reverted/recreated to fix this."), 						
+								GuidID, PackageName);
+	}
+	else if(!PackagePath.IsEmpty())
+	{
+		// We don't know where we are saving to, but we do know the package where the payload came from.
+		const FText PackageName = FText::FromString(PackagePath.GetPackageName());
+
+		return FText::Format(	NSLOCTEXT("Core", "Serialization_InvalidPayloadFromPkg", 
+								"Attempting to save bulkdata {0} with an invalid payload from package '{1}'. The package probably needs to be reverted/recreated to fix this."),
+								GuidID, PackageName);
+	}
+	else
+	{
+		// We don't know where the payload came from or where it is being saved to.
+		return FText::Format(	NSLOCTEXT("Core", "Serialization_InvalidPayloadPath", "Attempting to save bulkdata {0} with an invalid payload, source unknown"),
+								GuidID);
+	}
+}
+
+bool FEditorBulkData::ShouldUseLegacySerialization(const FLinkerSave* LinkerSave) const
+{
+	if (LinkerSave == nullptr)
+	{
+		return true;
+	}
+
+#if UE_ENABLE_VIRTUALIZATION_TOGGLE
+	return (!LinkerSave->PackageTrailerBuilder.IsValid()) || bSkipVirtualization == true;
+#else
+	return !LinkerSave->PackageTrailerBuilder.IsValid();
+#endif //UE_ENABLE_VIRTUALIZATION_TOGGLE
+}
+	
+FArchive& operator<<(FArchive& Ar, FTocEntry& Entry)
+{
+	Ar << Entry.Identifier;
+	Ar << Entry.OffsetInFile;
+	Ar << Entry.UncompressedSize;
+
+	return Ar;
+}
+
+ void operator<<(FStructuredArchive::FSlot Slot, FTocEntry& Entry)
+{
+	FStructuredArchive::FRecord Record = Slot.EnterRecord();
+
+	Record << SA_VALUE(TEXT("Identifier"), Entry.Identifier);
+	Record << SA_VALUE(TEXT("OffsetInFile"), Entry.OffsetInFile);
+	Record << SA_VALUE(TEXT("UncompressedSize"), Entry.UncompressedSize);
+}
+
+void FPayloadToc::AddEntry(const FEditorBulkData& BulkData)
+{
+	if (BulkData.GetPayloadId().IsValid())
+	{
+		Contents.Emplace(BulkData);
+	}
+}
+
+bool FPayloadToc::FindEntry(const UE::Virtualization::FPayloadId& Identifier, FTocEntry& OutEntry)
+{
+	for (const FTocEntry& Entry : Contents)
+	{
+		if (Entry.Identifier == Identifier)
+		{
+			OutEntry = Entry;
+			return true;
+		}
+	}
+
+	return false;
+}
+
+const TArray<FTocEntry>& FPayloadToc::GetContents() const
+{
+	return Contents;
+}
+
+FArchive& operator<<(FArchive& Ar, FPayloadToc& TableOfContents)
+{
+	FPayloadToc::EVersion Version = FPayloadToc::EVersion::AUTOMATIC_VERSION;
+	Ar << Version;
+
+	Ar << TableOfContents.Contents;
+
+	return Ar;
+}
+
+void operator<<(FStructuredArchive::FSlot Slot, FPayloadToc& TableOfContents)
+{
+	FStructuredArchive::FRecord Record = Slot.EnterRecord();
+
+	FPayloadToc::EVersion Version = FPayloadToc::EVersion::AUTOMATIC_VERSION;
+
+	Record << SA_VALUE(TEXT("Version"), Version);
+	Record << SA_VALUE(TEXT("Entries"), TableOfContents.Contents);
+}
+
+} // namespace UE::Serialization
+
+#undef UE_CORRUPTED_DATA_SEVERITY
+#undef UE_CORRUPTED_PAYLOAD_IS_FATAL
+#undef UE_ALLOW_LINKERLOADER_ATTACHMENT
+
+//#endif //WITH_EDITORONLY_DATA
diff --git a/Engine/Source/Runtime/CoreUObject/Private/Serialization/VirtualizedBulkData.cpp b/Engine/Source/Runtime/CoreUObject/Private/Serialization/VirtualizedBulkData.cpp
deleted file mode 100644
diff --git a/Engine/Source/Runtime/CoreUObject/Private/Tests/Serialization/EditorBulkDataTests.cpp b/Engine/Source/Runtime/CoreUObject/Private/Tests/Serialization/EditorBulkDataTests.cpp
new file mode 100644
--- /dev/null
+++ b/Engine/Source/Runtime/CoreUObject/Private/Tests/Serialization/EditorBulkDataTests.cpp
@@ -0,0 +1,513 @@
+// Copyright Epic Games, Inc. All Rights Reserved.
+
+#include "Algo/AllOf.h"
+#include "Memory/SharedBuffer.h"
+#include "Serialization/EditorBulkData.h"
+#include "Serialization/EditorBulkDataReader.h"
+#include "Serialization/EditorBulkDataWriter.h"
+#include "Serialization/LargeMemoryReader.h"
+#include "Serialization/LargeMemoryWriter.h"
+#include "Serialization/MemoryReader.h"
+#include "Serialization/MemoryWriter.h"
+#include "Templates/UniquePtr.h"
+
+#include "Misc/AutomationTest.h"
+
+#if WITH_DEV_AUTOMATION_TESTS && WITH_EDITORONLY_DATA
+
+namespace UE::Serialization
+{
+	
+constexpr const uint32 TestFlags = EAutomationTestFlags::ApplicationContextMask | EAutomationTestFlags::EngineFilter;
+
+/** Creates a buffer full of random data to make it easy to have something to test against. */
+TUniquePtr<uint8[]> CreateRandomData(int64 BufferSize)
+{
+	TUniquePtr<uint8[]> Buffer = MakeUnique<uint8[]>(BufferSize);
+
+	for (int64 Index = 0; Index < BufferSize; ++Index)
+	{
+		Buffer[Index] = (uint8)(FMath::Rand() % 255);
+	}
+
+	return Buffer;
+}
+
+/**
+ * This test creates a very basic FEditorBulkData object with in memory payload and validates that we are able to retrieve the 
+ * payload via both TFuture and callback methods. It then creates copies of the object and makes sure that we can get the payload from
+ * the copies, even when the original source object has been reset.
+ */
+IMPLEMENT_SIMPLE_AUTOMATION_TEST(FEditorBulkDataTestBasic, TEXT("System.CoreUObject.Serialization.EditorBulkData.Basic"), TestFlags)
+bool FEditorBulkDataTestBasic::RunTest(const FString& Parameters)
+{
+	const int64 BufferSize = 1024;
+	TUniquePtr<uint8[]> SourceBuffer = CreateRandomData(BufferSize);
+
+	auto ValidateBulkData = [this, &SourceBuffer, BufferSize](const FEditorBulkData& BulkDataToValidate, const TCHAR* Label)
+		{
+			FSharedBuffer RetrievedBuffer = BulkDataToValidate.GetPayload().Get();
+			TestEqual(FString::Printf(TEXT("%s buffer length"),Label), (int64)RetrievedBuffer.GetSize(), BufferSize);
+			TestTrue(FString::Printf(TEXT("SourceBuffer values == %s values"),Label), FMemory::Memcmp(SourceBuffer.Get(), RetrievedBuffer.GetData(), BufferSize) == 0);
+		};
+
+	// Create a basic bulkdata (but retain ownership of the buffer!)
+	FEditorBulkData BulkData;
+	BulkData.UpdatePayload(FSharedBuffer::MakeView(SourceBuffer.Get(), BufferSize));
+
+	// Test accessing the data from the bulkdata object
+	ValidateBulkData(BulkData, TEXT("Retrieved"));
+
+	// Create a new bulkdata object via the copy constructor
+	FEditorBulkData BulkDataCopy(BulkData);
+
+	// Create a new bulkdata object and copy by assignment (note we assign some junk data that will get overwritten)
+	FEditorBulkData BulkDataAssignment;
+	BulkDataAssignment.UpdatePayload(FUniqueBuffer::Alloc(128).MoveToShared());
+	BulkDataAssignment = BulkData;
+
+	// Test both bulkdata objects
+	ValidateBulkData(BulkDataCopy, TEXT("Copy Constructor"));
+	ValidateBulkData(BulkDataAssignment, TEXT("Copy Assignment"));
+
+	// Should not affect BulkDataAssignment/BulkDataCopy 
+	BulkData.Reset();
+
+	// Test both bulkdata objects again now that we reset the data
+	ValidateBulkData(BulkDataCopy, TEXT("Copy Constructor (after data reset)"));
+	ValidateBulkData(BulkDataAssignment, TEXT("Copy Assignment (after data reset)"));
+
+	return true;
+}
+
+/**
+ * This test will validate how FEditorBulkData behaves when it has no associated payload and make sure 
+ * that our assumptions are correct.
+ */
+IMPLEMENT_SIMPLE_AUTOMATION_TEST(FEditorBulkDataTestEmpty, TEXT("System.CoreUObject.Serialization.EditorBulkData.Empty"), TestFlags)
+bool FEditorBulkDataTestEmpty::RunTest(const FString& Parameters)
+{
+	FEditorBulkData BulkData;
+
+	// Validate the general accessors
+	TestEqual(TEXT("Return value of ::GetBulkDataSize()"), BulkData.GetPayloadSize(), (int64)0);
+	TestTrue(TEXT("Payload key is invalid"), !BulkData.GetPayloadId().IsValid());
+	TestFalse(TEXT("Return value of ::DoesPayloadNeedLoading()"), BulkData.DoesPayloadNeedLoading());
+
+	// Validate the payload accessors
+	FSharedBuffer Payload = BulkData.GetPayload().Get();
+	TestTrue(TEXT("The payload from the GetPayload TFuture is null"), Payload.IsNull());
+
+	FCompressedBuffer CompressedPayload = BulkData.GetCompressedPayload().Get();
+	TestTrue(TEXT("The payload from the GetCompressedPayload TFuture is null"), Payload.IsNull());
+
+	return true;
+}
+
+/**
+ * Test the various methods for updating the payload that a FEditorBulkData owns
+ */
+IMPLEMENT_SIMPLE_AUTOMATION_TEST(FEditorBulkDataTestUpdatePayload, TEXT("System.CoreUObject.Serialization.EditorBulkData.UpdatePayload"), TestFlags)
+bool FEditorBulkDataTestUpdatePayload::RunTest(const FString& Parameters)
+{
+	// Create a memory buffer of all zeros
+	const int64 BufferSize = 1024;
+	TUniquePtr<uint8[]> OriginalData = MakeUnique<uint8[]>(BufferSize);
+	FMemory::Memzero(OriginalData.Get(), BufferSize);
+
+	// Pass the buffer to to bulkdata but retain ownership
+	FEditorBulkData BulkData;
+	BulkData.UpdatePayload(FSharedBuffer::MakeView(OriginalData.Get(), BufferSize));
+
+	// Access the payload, edit it and push it back into the bulkdata object
+	{
+		// The payload should be the same size and same contents as the original buffer but a different
+		// memory address since we retained ownership in the TUniquePtr, so the bulkdata object should 
+		// have created it's own copy.
+		FSharedBuffer Payload = BulkData.GetPayload().Get();
+		TestEqual(TEXT("Payload length"), (int64)Payload.GetSize(), BufferSize);
+		TestNotEqual(TEXT("OriginalData and the payload should have different memory addresses"), (uint8*)OriginalData.Get(), (uint8*)Payload.GetData());
+		TestTrue(TEXT("Orginal buffer == Payload data"), FMemory::Memcmp(OriginalData.Get(), Payload.GetData(), Payload.GetSize()) == 0);
+
+		// Make a copy of the payload that we can edit
+		const uint8 NewValue = 255;
+		FSharedBuffer EditedPayload;
+		{
+			FUniqueBuffer EditablePayload = FUniqueBuffer::Clone(Payload);
+			FMemory::Memset(EditablePayload.GetData(), NewValue, EditablePayload.GetSize());
+			EditedPayload = EditablePayload.MoveToShared();
+		}
+
+		// Update the bulkdata object with the new edited payload
+		BulkData.UpdatePayload(EditedPayload);
+
+		Payload = BulkData.GetPayload().Get();
+		TestEqual(TEXT("Updated payload length"), (int64)Payload.GetSize(), BufferSize);
+		TestEqual(TEXT("Payload and EditablePayload should have the same memory addresses"), (uint8*)Payload.GetData(), (uint8*)EditedPayload.GetData());
+
+		const bool bAllElementsCorrect = Algo::AllOf(TArrayView64<uint8>((uint8*)Payload.GetData(), (int64)Payload.GetSize()), [NewValue](uint8 Val)
+			{ 
+				return Val == NewValue; 
+			});
+
+		TestTrue(TEXT("All payload elements correctly updated"), bAllElementsCorrect);
+	}
+
+	{
+		// Store the original data pointer address so we can test against it later, we should not actually use
+		// this pointer though as once we pass it to the bulkdata object we cannot be sure what happens to it.
+		uint8* OriginalDataPtr = OriginalData.Get();
+
+		// Update the bulkdata object with the original data but this time we give ownership of the buffer to
+		// the bulkdata object.
+		BulkData.UpdatePayload(FSharedBuffer::TakeOwnership(OriginalData.Release(), BufferSize, FMemory::Free));
+
+		FSharedBuffer Payload = BulkData.GetPayload().Get();
+		TestEqual(TEXT("Updated payload length"), (int64)Payload.GetSize(), BufferSize);
+		TestEqual(TEXT("Payload and OriginalDataPtr should have the same memory addresses"), (uint8*)Payload.GetData(), OriginalDataPtr);
+
+		// The original data was all zeros, so we can test for that to make sure that the contents are correct.
+		const bool bAllElementsCorrect = Algo::AllOf(TArrayView64<uint8>((uint8*)Payload.GetData(), (int64)Payload.GetSize()), [](uint8 Val)
+			{
+				return Val == 0;
+			});
+
+		TestTrue(TEXT("All payload elements correctly updated"), bAllElementsCorrect);
+	}
+
+	return true;
+}
+
+/**
+ * This test will create a buffer, then serialize it to a FEditorBulkData object via FEditorBulkDataWriter.
+ * Then we will serialize the FEditorBulkData object back to a second buffer and compare the results.
+ * If the reader and writers are working then the ReplicatedBuffer should be the same as the original SourceBuffer.
+ */
+IMPLEMENT_SIMPLE_AUTOMATION_TEST(FEditorBulkDataTestReaderWriter, TEXT("System.CoreUObject.Serialization.EditorBulkData.Reader/Writer"), TestFlags)
+bool FEditorBulkDataTestReaderWriter::RunTest(const FString& Parameters)
+{
+	const int64 BufferSize = 1024;
+
+	TUniquePtr<uint8[]> SourceBuffer = CreateRandomData(BufferSize);
+	TUniquePtr<uint8[]> ReplicatedBuffer = MakeUnique<uint8[]>(BufferSize);
+
+	FEditorBulkData BulkData;
+
+	// Serialize the SourceBuffer to BulkData 
+	{
+		FEditorBulkDataWriter WriterAr(BulkData);
+		WriterAr.Serialize(SourceBuffer.Get(), BufferSize);
+	}
+
+	// Serialize BulkData back to ReplicatedBuffer
+	{
+		FEditorBulkDataReader ReaderAr(BulkData);
+		ReaderAr.Serialize(ReplicatedBuffer.Get(), BufferSize);
+	}
+
+	// Now test that the buffer was restored to the original values
+	const bool bMemCmpResult = FMemory::Memcmp(SourceBuffer.Get(), ReplicatedBuffer.Get(), BufferSize) == 0;
+	TestTrue(TEXT("SourceBuffer values == ReplicatedBuffer values"), bMemCmpResult);
+
+	// Test writing nothing to an empty bulkdata object and then reading that bulkdata object
+	// to make sure that we deal with null buffers properly.
+	{
+		FEditorBulkData EmptyBulkData;
+		FEditorBulkDataWriter WriterAr(EmptyBulkData);
+		FEditorBulkDataReader ReaderAr(EmptyBulkData);
+	}
+	return true;
+}
+
+/** 
+ * This test will serialize a number of empty and valid bulkdata objects to a memory buffer and then serialize them back 
+ * as bulkdata objects again. This should help find problems where the data saved and loaded are mismatched.
+ * Note that this does not check the package saving code paths, only direct serialization to buffers.
+ */
+IMPLEMENT_SIMPLE_AUTOMATION_TEST(FEditorBulkDataTestSerializationToMemory, TEXT("System.CoreUObject.Serialization.EditorBulkData.SerializationToMemory"), TestFlags)
+bool FEditorBulkDataTestSerializationToMemory::RunTest(const FString& Parameters)
+{
+	const int64 BufferSize = 1024;
+
+	TUniquePtr<uint8[]> SourceBuffer = CreateRandomData(BufferSize);
+	
+	FEditorBulkData EmptyBulkData;
+	FEditorBulkData ValidBulkData;
+	ValidBulkData.UpdatePayload(FSharedBuffer::Clone(SourceBuffer.Get(), BufferSize));
+	
+	TArray<uint8> MemoryBuffer;
+	const bool bIsArPersistent = true;
+
+	// Write out
+	{
+		FMemoryWriter WriterAr(MemoryBuffer, bIsArPersistent);
+
+		ValidBulkData.Serialize(WriterAr, nullptr);
+		EmptyBulkData.Serialize(WriterAr, nullptr);
+		ValidBulkData.Serialize(WriterAr, nullptr);
+		EmptyBulkData.Serialize(WriterAr, nullptr);
+	}
+
+	// Read in
+	{
+		FEditorBulkData SerializedBulkData;
+
+		FMemoryReader ReaderAr(MemoryBuffer, bIsArPersistent);
+
+		SerializedBulkData.Serialize(ReaderAr, nullptr);
+		TestEqual(TEXT("Bulkdata identifier should remain the same"), SerializedBulkData.GetIdentifier(), ValidBulkData.GetIdentifier());
+		TestTrue(TEXT("Bulkdata payload should remain the same"), FMemory::Memcmp(SourceBuffer.Get(), SerializedBulkData.GetPayload().Get().GetData(), BufferSize) == 0);
+
+		SerializedBulkData.Serialize(ReaderAr, nullptr);
+		TestEqual(TEXT("Bulkdata identifier should remain the same"), SerializedBulkData.GetIdentifier(), EmptyBulkData.GetIdentifier());
+		TestTrue(TEXT("Bulkdata should not have a payload"), SerializedBulkData.GetPayload().Get().IsNull());
+
+		SerializedBulkData.Serialize(ReaderAr, nullptr);
+		TestEqual(TEXT("Bulkdata identifier should remain the same"), SerializedBulkData.GetIdentifier(), ValidBulkData.GetIdentifier());
+		TestTrue(TEXT("Bulkdata payload should remain the same"), FMemory::Memcmp(SourceBuffer.Get(), SerializedBulkData.GetPayload().Get().GetData(), BufferSize) == 0);
+
+		SerializedBulkData.Serialize(ReaderAr, nullptr);
+		TestEqual(TEXT("Bulkdata identifier should remain the same"), SerializedBulkData.GetIdentifier(), EmptyBulkData.GetIdentifier());
+		TestTrue(TEXT("Bulkdata should not have a payload"), SerializedBulkData.GetPayload().Get().IsNull());
+	}
+	return true;
+}
+
+/**
+ * This set of tests validate that the BulkData's identifier works how we expect it too. It should remain unique in all cases except
+ * move semantics.
+ */
+IMPLEMENT_SIMPLE_AUTOMATION_TEST(FEditorBulkDataTestIdentifiers, TEXT("System.CoreUObject.Serialization.EditorBulkData.Identifiers"), TestFlags)
+bool FEditorBulkDataTestIdentifiers::RunTest(const FString& Parameters)
+{
+	// Some basic tests with an invalid id
+	{
+		FEditorBulkData BulkData;
+		TestFalse(TEXT("BulkData with no payload should return an invalid identifier"), BulkData.GetIdentifier().IsValid());
+
+		FEditorBulkData CopiedBulkData(BulkData);
+		TestFalse(TEXT("Copying a bulkdata with an invalid id should result in an invalid id"), CopiedBulkData.GetIdentifier().IsValid());
+
+		FEditorBulkData AssignedBulkData;
+		AssignedBulkData = BulkData;
+		TestFalse(TEXT("Assigning a a bulkdata with an invalid id should result in an invalid id"), AssignedBulkData.GetIdentifier().IsValid());
+		
+		// Check that we did not change the initial object at any point
+		TestFalse(TEXT("Being copied and assigned to other objects should not affect the identifier"), BulkData.GetIdentifier().IsValid());
+	}
+
+	// Some basic tests with a valid id
+	{
+		FEditorBulkData BulkData;
+		BulkData.UpdatePayload(FUniqueBuffer::Alloc(32).MoveToShared()); // Assigning this payload should cause BulkData to gain an identifier
+		TestTrue(TEXT("BulkData with a payload should returns a valid identifier"), BulkData.GetIdentifier().IsValid());
+
+		const FGuid OriginalGuid = BulkData.GetIdentifier();
+
+		FEditorBulkData CopiedBulkData(BulkData);
+		TestNotEqual(TEXT("Copying a bulkdata with a valid id should result in a unique identifier"), BulkData.GetIdentifier(), CopiedBulkData.GetIdentifier());
+
+		FEditorBulkData AssignedBulkData;
+		AssignedBulkData = BulkData;
+		TestNotEqual(TEXT("Assignment operator creates different identifiers"), BulkData.GetIdentifier(), AssignedBulkData.GetIdentifier());
+
+		// Check that we did not change the initial object at any point
+		TestEqual(TEXT("Being copied and assigned to other objects should not affect the identifier"), BulkData.GetIdentifier(), OriginalGuid);
+
+		// Now that AssignedBulkData has a valid identifier, make sure that it is not changed if we assign something else to it.
+		const FGuid OriginalAssignedGuid = AssignedBulkData.GetIdentifier();
+		AssignedBulkData = CopiedBulkData;
+		TestEqual(TEXT("Being copied and assigned to other objects should not affect the identifier"), AssignedBulkData.GetIdentifier(), OriginalAssignedGuid);
+	}
+
+	// Test move constructor
+	{	
+		FEditorBulkData BulkData;
+		BulkData.UpdatePayload(FUniqueBuffer::Alloc(32).MoveToShared());
+
+		FGuid OriginalGuid = BulkData.GetIdentifier();
+
+		FEditorBulkData MovedBulkData = MoveTemp(BulkData);
+
+		TestEqual(TEXT("Move constructor should preserve the identifier"), MovedBulkData.GetIdentifier(), OriginalGuid);
+	}
+
+	// Test move assignment
+	{
+		FEditorBulkData BulkData;
+		BulkData.UpdatePayload(FUniqueBuffer::Alloc(32).MoveToShared());
+
+		FGuid OriginalGuid = BulkData.GetIdentifier();
+
+		FEditorBulkData MovedBulkData;
+		MovedBulkData = MoveTemp(BulkData);
+
+		TestEqual(TEXT("Move assignment should preserve the identifier"), MovedBulkData.GetIdentifier(), OriginalGuid);
+	}
+
+	// Check that resizing an array will not change the internals
+	{
+		const uint32 NumToTest = 10;
+
+		TArray<FEditorBulkData> BulkDataArray;
+		TArray<FGuid> GuidArray;
+
+		for (uint32 Index = 0; Index < NumToTest; ++Index)
+		{
+			BulkDataArray.Add(FEditorBulkData());
+
+			// Leave some with invalid ids and some with valid ones
+			if (Index % 2 == 0)
+			{
+				BulkDataArray[Index].UpdatePayload(FUniqueBuffer::Alloc(32).MoveToShared());
+			}
+
+			GuidArray.Add(BulkDataArray[Index].GetIdentifier());
+		}
+
+		// Force an internal reallocation and make sure that the identifiers are unchanged.
+		// Note that it is possible that the allocation is just resized and not reallocated.
+		BulkDataArray.Reserve(BulkDataArray.Max() * 4);
+
+		for (uint32 Index = 0; Index < NumToTest; ++Index)
+		{
+			TestEqual(TEXT(""), BulkDataArray[Index].GetIdentifier(), GuidArray[Index]);
+		}
+
+		// Now insert a new item, moving all of the existing entries and make sure that
+		// the identifiers are unchanged.
+		BulkDataArray.Insert(FEditorBulkData(), 0);
+
+		for (uint32 Index = 0; Index < NumToTest; ++Index)
+		{
+			TestEqual(TEXT(""), BulkDataArray[Index+1].GetIdentifier(), GuidArray[Index]);
+		}
+	}
+
+	// Test that adding a payload to a reset bulkdata object or one that has had a zero length payload applied
+	// will correctly show the original id once it has a valid payload.
+	{
+		FEditorBulkData BulkData;
+		BulkData.UpdatePayload(FUniqueBuffer::Alloc(32).MoveToShared());
+		
+		const FGuid OriginalGuid = BulkData.GetIdentifier();
+
+		BulkData.Reset();
+		TestEqual(TEXT("BulkData with no payload should return its original identifier"), BulkData.GetIdentifier(), OriginalGuid);
+
+		BulkData.UpdatePayload(FUniqueBuffer::Alloc(32).MoveToShared());
+		TestEqual(TEXT("Removing a payload then adding a new one should return the original identifier"), BulkData.GetIdentifier(), OriginalGuid);
+
+		BulkData.UpdatePayload(FUniqueBuffer::Alloc(0).MoveToShared());
+		TestEqual(TEXT("Setting a zero length payload should keep the original identifier"), BulkData.GetIdentifier(), OriginalGuid);
+		
+		BulkData.UpdatePayload(FUniqueBuffer::Alloc(32).MoveToShared());
+		TestEqual(TEXT("Restoring a payload should return the original identifier"), BulkData.GetIdentifier(), OriginalGuid);
+	}
+
+	// Test that serialization does not change the identifier (in this case serializing to and from a memory buffer) 
+	{
+		FEditorBulkData SrcData;
+		SrcData.UpdatePayload(FUniqueBuffer::Alloc(32).MoveToShared());
+
+		TArray<uint8> MemoryBuffer;
+		const bool bIsArPersistent = true;
+
+		FEditorBulkData DstData;
+
+		// Serialize the SourceBuffer to BulkData 
+		{
+			FMemoryWriter WriterAr(MemoryBuffer, bIsArPersistent);
+			SrcData.Serialize(WriterAr, nullptr);
+		}
+
+		// Serialize BulkData back to ReplicatedBuffer
+		{
+			FMemoryReader ReaderAr(MemoryBuffer, bIsArPersistent);
+			DstData.Serialize(ReaderAr, nullptr);
+		}
+
+		TestEqual(TEXT("Serialization should preserve the identifier"), SrcData.GetIdentifier(), DstData.GetIdentifier());
+	}
+
+	// Test that serialization does not change the identifier (in this case serializing to and from a memory buffer) 
+	{
+		FEditorBulkData SrcData;
+		SrcData.UpdatePayload(FUniqueBuffer::Alloc(32).MoveToShared());
+
+		const FGuid OriginalIdentifier = SrcData.GetIdentifier();
+		SrcData.UpdatePayload(FUniqueBuffer::Alloc(0).MoveToShared());
+
+		TArray<uint8> MemoryBuffer;
+		const bool bIsArPersistent = true;
+
+		FEditorBulkData DstData;
+
+		// Serialize the SourceBuffer to BulkData 
+		{
+			FMemoryWriter WriterAr(MemoryBuffer, bIsArPersistent);
+			SrcData.Serialize(WriterAr, nullptr);
+		}
+
+		// Serialize BulkData back to ReplicatedBuffer
+		{
+			FMemoryReader ReaderAr(MemoryBuffer, bIsArPersistent);
+			DstData.Serialize(ReaderAr, nullptr);
+		}
+
+		TestEqual(TEXT("After serialization the identifier should keep its identifier "), DstData.GetIdentifier(), OriginalIdentifier);
+
+		DstData.UpdatePayload(FUniqueBuffer::Alloc(32).MoveToShared());
+		TestEqual(TEXT("After adding a new payload the object should have the original identifier"), DstData.GetIdentifier(), OriginalIdentifier);
+	}
+
+	return true;
+}
+
+IMPLEMENT_SIMPLE_AUTOMATION_TEST(FEditorBulkDataTestZeroSizedAllocs, TEXT("System.CoreUObject.Serialization.EditorBulkData.ZeroSizedAllocs"), TestFlags)
+bool FEditorBulkDataTestZeroSizedAllocs::RunTest(const FString& Parameters)
+{
+	{
+		FEditorBulkData Empty;
+		TestNull(TEXT("Empty bulkdata returning a payload"), Empty.GetPayload().Get().GetData());
+		TestNull(TEXT("Empty bulkdata returning a compressed payload"), Empty.GetCompressedPayload().Get().Decompress().GetData());
+
+		TestFalse(TEXT("Empty payload has payload data"), Empty.HasPayloadData());
+		TestFalse(TEXT("Empty payload needs loading"), Empty.DoesPayloadNeedLoading());
+	}
+
+	{
+		FEditorBulkData EmptySrc;
+		
+		FLargeMemoryWriter ArWrite(0, true);
+		EmptySrc.Serialize(ArWrite, nullptr);
+
+		FEditorBulkData EmptyDst;
+
+		FLargeMemoryReader ArRead(ArWrite.GetData(), ArWrite.TotalSize(), ELargeMemoryReaderFlags::Persistent);
+		EmptyDst.Serialize(ArRead, nullptr);
+
+		TestNull(TEXT("Serialized bulkdata returning a payload"), EmptyDst.GetPayload().Get().GetData());
+		TestNull(TEXT("Serialized bulkdata returning a compressed payload"), EmptyDst.GetCompressedPayload().Get().Decompress().GetData());
+
+		TestFalse(TEXT("Serialized payload has payload data"), EmptyDst.HasPayloadData());
+		TestFalse(TEXT("Serialized payload needs loading"), EmptyDst.DoesPayloadNeedLoading());
+	}
+
+	{
+		FEditorBulkData ZeroAlloc;
+		ZeroAlloc.UpdatePayload(FUniqueBuffer::Alloc(0).MoveToShared());
+
+		TestNull(TEXT("ZeroAlloc bulkdata returning a payload"), ZeroAlloc.GetPayload().Get().GetData());
+		TestNull(TEXT("ZeroAlloc bulkdata returning a compressed payload"), ZeroAlloc.GetCompressedPayload().Get().Decompress().GetData());
+
+		TestFalse(TEXT("ZeroAlloc payload has payload data"), ZeroAlloc.HasPayloadData());
+		TestFalse(TEXT("ZeroAlloc payload needs loading"), ZeroAlloc.DoesPayloadNeedLoading());
+	}
+
+	return true;
+}
+
+} // namespace UE::Serialization
+
+#endif //WITH_DEV_AUTOMATION_TESTS && WITH_EDITORONLY_DATA
diff --git a/Engine/Source/Runtime/CoreUObject/Private/Tests/Serialization/VirtualizedBulkDataTests.cpp b/Engine/Source/Runtime/CoreUObject/Private/Tests/Serialization/VirtualizedBulkDataTests.cpp
deleted file mode 100644
diff --git a/Engine/Source/Runtime/CoreUObject/Private/UObject/LinkerLoad.cpp b/Engine/Source/Runtime/CoreUObject/Private/UObject/LinkerLoad.cpp
--- a/Engine/Source/Runtime/CoreUObject/Private/UObject/LinkerLoad.cpp
+++ b/Engine/Source/Runtime/CoreUObject/Private/UObject/LinkerLoad.cpp
@@ -42,7 +42,7 @@
 #include "Serialization/ArchiveUObjectFromStructuredArchive.h"
 #include "Serialization/UnversionedPropertySerialization.h"
 #include "Serialization/LoadTimeTracePrivate.h"
-#include "Serialization/VirtualizedBulkData.h"
+#include "Serialization/EditorBulkData.h"
 #include "HAL/FileManager.h"
 #include "UObject/CoreRedirects.h"
 #include "Misc/StringBuilder.h"
@@ -5539,10 +5539,10 @@ void FLinkerLoad::AttachBulkData( UObject* Owner, FUntypedBulkData* BulkData )
 	BulkDataLoaders.Add( BulkData );
 }
 
-void FLinkerLoad::AttachBulkData(UE::Virtualization::FVirtualizedUntypedBulkData* BulkData)
+void FLinkerLoad::AttachBulkData(UE::Serialization::FEditorBulkData* BulkData)
 {
-	check(VirtualizedBulkDataLoaders.Find(BulkData) == INDEX_NONE);
-	VirtualizedBulkDataLoaders.Add(BulkData);	
+	check(EditorBulkDataLoaders.Find(BulkData) == INDEX_NONE);
+	EditorBulkDataLoaders.Add(BulkData);
 }
 
 /**
@@ -5561,9 +5561,9 @@ void FLinkerLoad::DetachBulkData( FUntypedBulkData* BulkData, bool bEnsureBulkDa
 	BulkData->DetachFromArchive( this, bEnsureBulkDataIsLoaded );
 }
 
-void FLinkerLoad::DetachBulkData(UE::Virtualization::FVirtualizedUntypedBulkData* BulkData, bool bEnsureBulkDataIsLoaded)
+void FLinkerLoad::DetachBulkData(UE::Serialization::FEditorBulkData* BulkData, bool bEnsureBulkDataIsLoaded)
 {
-	int32 RemovedCount = VirtualizedBulkDataLoaders.Remove(BulkData);
+	int32 RemovedCount = EditorBulkDataLoaders.Remove(BulkData);
 	if (RemovedCount != 1)
 	{
 		UE_LOG(LogLinker, Fatal, TEXT("Detachment inconsistency: %i (%s)"), RemovedCount, *GetDebugName());
@@ -5591,13 +5591,13 @@ void FLinkerLoad::DetachAllBulkData(bool bEnsureAllBulkDataIsLoaded)
 
 	// Then virtualized bulkdata
 	{
-		TArray<UE::Virtualization::FVirtualizedUntypedBulkData*> BulkDataToDetach = VirtualizedBulkDataLoaders;
-		for (UE::Virtualization::FVirtualizedUntypedBulkData* BulkData : BulkDataToDetach)
+		TArray<UE::Serialization::FEditorBulkData*> BulkDataToDetach = EditorBulkDataLoaders;
+		for (UE::Serialization::FEditorBulkData* BulkData : BulkDataToDetach)
 		{
 			check(BulkData != nullptr);
 			BulkData->DetachFromDisk(this, bEnsureAllBulkDataIsLoaded);
 		}
-		VirtualizedBulkDataLoaders.Empty();
+		EditorBulkDataLoaders.Empty();
 	}
 }
 
diff --git a/Engine/Source/Runtime/CoreUObject/Private/UObject/PackageTrailer.cpp b/Engine/Source/Runtime/CoreUObject/Private/UObject/PackageTrailer.cpp
--- a/Engine/Source/Runtime/CoreUObject/Private/UObject/PackageTrailer.cpp
+++ b/Engine/Source/Runtime/CoreUObject/Private/UObject/PackageTrailer.cpp
@@ -48,7 +48,7 @@ void LogPackageOpenFailureMessage(const FPackagePath& PackagePath)
 	}
 	else
 	{
-		UE_LOG(LogVirtualization, Error, TEXT("Could not open (%s) to read FVirtualizedUntypedBulkData with an unknown error"), *PackagePath.GetDebugName());
+		UE_LOG(LogVirtualization, Error, TEXT("Could not open (%s) to read FPackageTrailer with an unknown error"), *PackagePath.GetDebugName());
 	}
 }
 
diff --git a/Engine/Source/Runtime/CoreUObject/Private/UObject/SavePackage.cpp b/Engine/Source/Runtime/CoreUObject/Private/UObject/SavePackage.cpp
--- a/Engine/Source/Runtime/CoreUObject/Private/UObject/SavePackage.cpp
+++ b/Engine/Source/Runtime/CoreUObject/Private/UObject/SavePackage.cpp
@@ -65,7 +65,7 @@
 #include "Serialization/ArchiveUObjectFromStructuredArchive.h"
 #include "Serialization/PackageWriter.h"
 #include "Serialization/UnversionedPropertySerialization.h"
-#include "Serialization/VirtualizedBulkData.h"
+#include "Serialization/EditorBulkData.h"
 #include "UObject/AsyncWorkSequence.h"
 #include "Misc/ScopeExit.h"
 #include "Misc/PackageAccessTracking.h"
diff --git a/Engine/Source/Runtime/CoreUObject/Private/UObject/SavePackage/SavePackageUtilities.cpp b/Engine/Source/Runtime/CoreUObject/Private/UObject/SavePackage/SavePackageUtilities.cpp
--- a/Engine/Source/Runtime/CoreUObject/Private/UObject/SavePackage/SavePackageUtilities.cpp
+++ b/Engine/Source/Runtime/CoreUObject/Private/UObject/SavePackage/SavePackageUtilities.cpp
@@ -9,16 +9,16 @@
 #include "HAL/FileManager.h"
 #include "Interfaces/ITargetPlatform.h"
 #include "Misc/AssetRegistryInterface.h"
-#include "Misc/ConfigCacheIni.h"
 #include "Misc/CommandLine.h"
-#include "Misc/Paths.h"
+#include "Misc/ConfigCacheIni.h"
 #include "Misc/PathViews.h"
+#include "Misc/Paths.h"
 #include "Misc/ScopedSlowTask.h"
 #include "SaveContext.h"
 #include "Serialization/BulkData.h"
+#include "Serialization/EditorBulkData.h"
 #include "Serialization/LargeMemoryWriter.h"
 #include "Serialization/PackageWriter.h"
-#include "Serialization/VirtualizedBulkData.h"
 #include "UObject/AsyncWorkSequence.h"
 #include "UObject/Class.h"
 #include "UObject/GCScopeLock.h"
@@ -30,9 +30,9 @@
 #include "UObject/ObjectSaveContext.h"
 #include "UObject/Package.h"
 #include "UObject/SavePackage.h"
-#include "UObject/UnrealType.h"
 #include "UObject/UObjectHash.h"
 #include "UObject/UObjectThreadContext.h"
+#include "UObject/UnrealType.h"
 
 DEFINE_LOG_CATEGORY(LogSavePackage);
 UE_TRACE_CHANNEL_DEFINE(SaveTimeChannel);
@@ -1944,12 +1944,12 @@ ESavePackageResult CreatePayloadSidecarFile(FLinkerSave& Linker, const FPackageP
 
 	FLargeMemoryWriter Ar(0, true /* bIsPersistent */);
 
-	uint32 VersionNumber = UE::Virtualization::FTocEntry::PayloadSidecarFileVersion;
+	uint32 VersionNumber = UE::Serialization::FTocEntry::PayloadSidecarFileVersion;
 	Ar << VersionNumber;
 
 	int64 TocPosition = Ar.Tell();
 
-	TArray<UE::Virtualization::FTocEntry> TableOfContents;
+	TArray<UE::Serialization::FTocEntry> TableOfContents;
 	TableOfContents.SetNum(Linker.SidecarDataToAppend.Num());
 
 	// First we write an empty table of contents to the file to reserve the space
diff --git a/Engine/Source/Runtime/CoreUObject/Private/UObject/SavePackage2.cpp b/Engine/Source/Runtime/CoreUObject/Private/UObject/SavePackage2.cpp
--- a/Engine/Source/Runtime/CoreUObject/Private/UObject/SavePackage2.cpp
+++ b/Engine/Source/Runtime/CoreUObject/Private/UObject/SavePackage2.cpp
@@ -23,12 +23,12 @@
 #include "ProfilingDebugging/CookStats.h"
 #include "Serialization/ArchiveUObjectFromStructuredArchive.h"
 #include "Serialization/ArchiveStackTrace.h"
+#include "Serialization/EditorBulkData.h"
 #include "Serialization/Formatters/JsonArchiveOutputFormatter.h"
 #include "Serialization/LargeMemoryWriter.h"
 #include "Serialization/PackageWriter.h"
 #include "Serialization/PropertyLocalizationDataGathering.h"
 #include "Serialization/UnversionedPropertySerialization.h"
-#include "Serialization/VirtualizedBulkData.h"
 #include "UObject/AsyncWorkSequence.h"
 #include "UObject/DebugSerializationFlags.h"
 #include "UObject/EditorObjectVersion.h"
diff --git a/Engine/Source/Runtime/CoreUObject/Public/Serialization/BulkData.h b/Engine/Source/Runtime/CoreUObject/Public/Serialization/BulkData.h
--- a/Engine/Source/Runtime/CoreUObject/Public/Serialization/BulkData.h
+++ b/Engine/Source/Runtime/CoreUObject/Public/Serialization/BulkData.h
@@ -244,7 +244,7 @@ private:
 public:
 	friend class FLinkerLoad;
 	friend class FExportArchive;
-	friend class UE::Virtualization::FVirtualizedUntypedBulkData; // To allow access to AttachedAr
+	friend class UE::Serialization::FEditorBulkData; // To allow access to AttachedAr
 
 	using BulkDataRangeArray = TArray<FUntypedBulkData*, TInlineAllocator<8>>;
 
diff --git a/Engine/Source/Runtime/CoreUObject/Public/Serialization/BulkDataRegistry.h b/Engine/Source/Runtime/CoreUObject/Public/Serialization/BulkDataRegistry.h
--- a/Engine/Source/Runtime/CoreUObject/Public/Serialization/BulkDataRegistry.h
+++ b/Engine/Source/Runtime/CoreUObject/Public/Serialization/BulkDataRegistry.h
@@ -17,8 +17,7 @@
 class IBulkDataRegistry;
 class UPackage;
 struct FGuid;
-struct FUntypedBulkData;
-namespace UE::Virtualization { class FVirtualizedUntypedBulkData; }
+namespace UE::Serialization { class FEditorBulkData; }
 
 DECLARE_DELEGATE_RetVal(IBulkDataRegistry*, FSetBulkDataRegistry);
 
@@ -56,9 +55,9 @@ public:
 	COREUOBJECT_API static IBulkDataRegistry& Get();
 	
 	/** Register a BulkData with the registry. Its payload and metadata will be fetchable by its GetIdentifier. */
-	virtual void Register(UPackage* Owner, const UE::Virtualization::FVirtualizedUntypedBulkData& BulkData) = 0;
+	virtual void Register(UPackage* Owner, const UE::Serialization::FEditorBulkData& BulkData) = 0;
 	/** Report that a BulkData is leaving memory and its in-memory payload (if it had one) is no longer available. */
-	virtual void OnExitMemory(const UE::Virtualization::FVirtualizedUntypedBulkData& BulkData) = 0;
+	virtual void OnExitMemory(const UE::Serialization::FEditorBulkData& BulkData) = 0;
 
 	/** Return the metadata for the given registered BulkData; returns false if not registered. */
 	virtual TFuture<UE::BulkDataRegistry::FMetaData> GetMeta(const FGuid& BulkDataId) = 0;
@@ -98,7 +97,7 @@ public:
 	COREUOBJECT_API FResaveSizeTracker();
 	COREUOBJECT_API ~FResaveSizeTracker();
 
-	COREUOBJECT_API void Register(UPackage* Owner, const UE::Virtualization::FVirtualizedUntypedBulkData& BulkData);
+	COREUOBJECT_API void Register(UPackage* Owner, const UE::Serialization::FEditorBulkData& BulkData);
 	COREUOBJECT_API uint64 GetBulkDataResaveSize(FName PackageName);
 
 private:
diff --git a/Engine/Source/Runtime/CoreUObject/Public/Serialization/EditorBulkData.h b/Engine/Source/Runtime/CoreUObject/Public/Serialization/EditorBulkData.h
new file mode 100644
--- /dev/null
+++ b/Engine/Source/Runtime/CoreUObject/Public/Serialization/EditorBulkData.h
@@ -0,0 +1,570 @@
+// Copyright Epic Games, Inc. All Rights Reserved.
+
+#pragma once
+
+#include "Async/Future.h"
+#include "Compression/CompressedBuffer.h"
+#include "HAL/Platform.h"
+#include "Memory/SharedBuffer.h"
+#include "Misc/Guid.h"
+#include "Misc/PackagePath.h"
+#include "Misc/PackageSegment.h"
+#include "Virtualization/PayloadId.h"
+
+class FArchive;
+class UObject;
+
+// When enabled it will be possible for specific editor bulkdata objects to opt out of being virtualized
+// This is a development feature and not expected to be used!
+#define UE_ENABLE_VIRTUALIZATION_TOGGLE 1
+
+//TODO: At some point it might be a good idea to uncomment this to make sure that FEditorBulkData is
+//		never used at runtime (requires a little too much reworking of some assets for now though)
+//#if WITH_EDITORONLY_DATA
+
+namespace UE::Serialization
+{
+
+namespace Private
+{
+/** A wrapper around the oodle compression settings used by FEditorBulkData. */
+struct FCompressionSettings
+{
+	COREUOBJECT_API FCompressionSettings();
+	FCompressionSettings(const FCompressedBuffer& Buffer);
+
+	[[nodiscard]] bool operator ==(const FCompressionSettings& Other) const;
+	[[nodiscard]] bool operator != (const FCompressionSettings& Other) const;
+
+	void Reset();
+
+	void Set(ECompressedBufferCompressor InCompressor, ECompressedBufferCompressionLevel InCompressionLevel);
+	void SetToDefault();
+	void SetToDisabled();
+
+	[[nodiscard]] bool IsSet() const;
+	[[nodiscard]] bool IsCompressed() const;
+
+	[[nodiscard]] ECompressedBufferCompressor GetCompressor() const;
+	[[nodiscard]] ECompressedBufferCompressionLevel GetCompressionLevel();
+
+private:
+
+	ECompressedBufferCompressor Compressor;
+	ECompressedBufferCompressionLevel CompressionLevel;
+
+	bool bIsSet;
+};
+} // namespace Private
+
+/** 
+ * A set of higher level compression options that avoid the need to set the specific
+ * oodle options.
+ */
+enum class ECompressionOptions : uint8
+{
+	/** Use default compression settings. */
+	Default,
+	/** Disable compression for the bulkdata entirely. */
+	Disabled,
+};
+
+/**
+ * The goal of this class is to provide an editor time version of BulkData that will work with the asset
+ * virtualization system.
+ *
+ * Assuming that the DDC is hot, the virtualized payloads are accessed relatively infrequently, usually when the package
+ * is being edited in the editor in some manner. So the payload access is designed around this. If the data is frequently 
+ * accessed when running the editor then the user would not gain from having it virtualized as they would end up pulling
+ * it immediately anyway.
+ *
+ * The biggest difference with normal bulkdata is that the access times might be significantly longer if the
+ * payload is not readily available which is why the only way to access the payload is by a TFuture or a callback
+ * so that the caller is forced to consider how to handle the potential stall and hopefully organize their code 
+ * in such a way that the time lag is not noticeable to the user.
+ *
+ * The second biggest difference is that the caller will own the payload memory once it is returned to them, unlike
+ * the old bulkdata class which would retain ownership. Which forces the calling code to be in control of when the 
+ * memory is actually allocated and for how long. With the old bulkdata class a single access would leave that memory
+ * allocated in a non-obvious way and would lead to memory bloats within the editor.
+ *
+ * The method ::GetGuid can be used to access a unique identifier for the payload, currently it is based on the 
+ * payload itself, so that two objects with the same payload would both have the same Guid. The intent is that we 
+ * would be able to share local copies of the payload between branches to reduce the cost of having multiple branches 
+ * with similar data on the same machine.
+ * 
+ */
+
+/** The base class with no type */
+class COREUOBJECT_API FEditorBulkData
+{
+public:
+	FEditorBulkData() = default;
+	FEditorBulkData(FEditorBulkData&& Other);
+	FEditorBulkData& operator=(FEditorBulkData&& Other);
+
+	FEditorBulkData(const FEditorBulkData& Other);
+	FEditorBulkData& operator=(const FEditorBulkData& Other);
+
+	~FEditorBulkData();
+
+	/** 
+	 * Convenience method to make it easier to convert from FUntypedBulkData to FEditorBulkData and sets the Guid 
+	 *
+	 * @param BulkData	The bulkdata object to create from.
+	 * @param Guid		A guid associated with the bulkdata object which will be used to identify the payload.
+	 *					This MUST remain the same between sessions so that the payloads key remains consistent!
+	 */
+	void CreateFromBulkData(FUntypedBulkData& BulkData, const FGuid& Guid, UObject* Owner);
+	/** Fix legacy content that created the Id from non-unique Guids. */
+	void CreateLegacyUniqueIdentifier(UObject* Owner);
+
+	/** 
+	 * Used to serialize the bulkdata to/from a FArchive
+	 * 
+	 * @param Ar	The archive to serialize the bulkdata.
+	 * @param Owner	The UObject that contains the bulkdata object, if this is a nullptr then the bulkdata will
+	 *				assume that it must serialize the payload immediately to memory as it will not be able to
+	 *				identify it's package path.
+	 * @param bAllowRegistry Legacy parameter to skip registration when loading BulkData we know we will need to
+	 *				modify the identifier of. Should always be true for non-legacy serialization.
+	 */
+	void Serialize(FArchive& Ar, UObject* Owner, bool bAllowRegister=true);
+
+	/** Reset to a truly empty state */
+	void Reset();
+
+	// TODO: Probably can just remove this as there probably isn't a good use case for unloading updated payloads as there is no
+	// way for us to restore it. In that case ::Reset might as well be used.
+	/** Unloads the data (if possible) but leaves it in a state where the data can be reloaded */
+	void UnloadData();
+
+	/** 
+	 * Removes the ability for the bulkdata object to load it's payload from disk (if it was doing so) 
+	 * 
+	 * @param Ar						The archive that the bulkdata is being detached from. This should match AttachedAr.
+	 * @param bEnsurePayloadIsLoaded	If true and the bulkdata currently is storing it's payload as a package on disk then
+	 *									the payload should be loaded into memory so that it can be accessed in the future.
+	 *									If false then the payload is not important and does not need to be loaded.
+	 */
+	void DetachFromDisk(FArchive* Ar, bool bEnsurePayloadIsLoaded);
+
+	/** Returns a unique identifier for the object itself. */
+	FGuid GetIdentifier() const;
+
+	/** Returns an unique identifier for the content of the payload. */
+	const UE::Virtualization::FPayloadId& GetPayloadId() const
+	{ 
+		return PayloadContentId; 
+	}
+
+	/** Returns the size of the payload in bytes. */
+	int64 GetPayloadSize() const 
+	{ 
+		return PayloadSize; 
+	}
+
+	/** Returns true if the bulkdata object contains a valid payload greater than zero bytes in size. */
+	bool HasPayloadData() const 
+	{ 
+		return PayloadSize > 0; 
+	}
+
+	/** Returns if the payload would require loading in order to be accessed. Returns false if the payload is already in memory or of zero length */
+	bool DoesPayloadNeedLoading() const
+	{
+		return Payload.IsNull() && PayloadSize > 0; 
+	}
+
+	/** Returns an immutable FCompressedBuffer reference to the payload data. */
+	TFuture<FSharedBuffer> GetPayload() const;
+
+	/**
+	 * Returns an immutable FCompressedBuffer reference to the payload data.
+	 *
+	 * Note that depending on the internal storage formats, the payload might not actually be compressed, but that
+	 * will be handled by the FCompressedBuffer interface. Call FCompressedBuffer::Decompress() to get access to
+	 * the payload in FSharedBuffer format.
+	 */
+	TFuture<FCompressedBuffer> GetCompressedPayload() const;
+
+	/**
+	 * Replaces the existing payload (if any) with a new one. 
+	 * It is important to consider the ownership model of the payload being passed in to the method.
+	 * 
+	 * To pass in a raw pointer use 'FSharedBuffer::...(Data, Size)' to create a valid FSharedBuffer.
+	 * Use 'FSharedBuffer::MakeView' if you want to retain ownership on the data being passed in, and use
+	 * 'FSharedBuffer::TakeOwnership' if you are okay with the bulkdata object taking over ownership of it.
+	 * The bulkdata object must own its internal buffer, so if you pass in a non-owned FSharedBuffer (ie
+	 * by using 'FSharedBuffer::MakeView') then a clone of the data will be created internally and assigned
+	 * to the bulkdata object.
+	 *
+	 * @param InPayload	The payload that this bulkdata object should reference. @see FSharedBuffer
+	 */
+	void UpdatePayload(FSharedBuffer InPayload);
+
+	/**
+	 * Utility struct used to compute the Payload ID before calling UpdatePayload
+	 */
+	struct COREUOBJECT_API FSharedBufferWithID
+	{
+		FSharedBufferWithID(FSharedBuffer InPayload);
+
+		FSharedBufferWithID() = default;
+		FSharedBufferWithID(FSharedBufferWithID&&) = default;
+		FSharedBufferWithID& operator=(FSharedBufferWithID&&) = default;
+
+		FSharedBufferWithID(const FSharedBufferWithID&) = delete;
+		FSharedBufferWithID& operator=(const FSharedBufferWithID&) = delete;
+
+	private:
+		friend FEditorBulkData;
+
+		FSharedBuffer Payload;
+		UE::Virtualization::FPayloadId PayloadId;
+	};
+
+	/**
+	 * Allows the existing payload to be replaced with a new one.
+	 *
+	 * 
+	 * To pass in a raw pointer use 'FSharedBuffer::...(Data, Size)' to create a valid FSharedBuffer.
+	 * Use 'FSharedBuffer::MakeView' if you want to retain ownership on the data being passed in, and use
+	 * 'FSharedBuffer::TakeOwnership' if you are okay with the bulkdata object taking over ownership of it.
+	 * The bulkdata object must own its internal buffer, so if you pass in a non-owned FSharedBuffer (ie
+	 * by using 'FSharedBuffer::MakeView') then a clone of the data will be created internally and assigned
+	 * to the bulkdata object.
+	 * 
+	 * Use this override if you want compute PayloadId before updating the bulkdata
+	 *
+	 * @param InPayload				The payload to update the bulkdata with
+	 */
+	void UpdatePayload(FSharedBufferWithID InPayload);
+
+	/** 
+	 * Sets the compression options to be applied to the payload during serialization.
+	 * 
+	 * These settings will continue to be used until the bulkdata object is reset, a subsequent
+	 * call to ::SetCompressionOptions is made or the owning package is serialized to disk.
+	 * 
+	 * @param Option	The high level option to use. @see UE::Serialization::ECompressionOptions
+	 */ 
+	void SetCompressionOptions(ECompressionOptions Option);
+
+	/** 
+	 * Sets the compression options to be applied to the payload during serialization.
+	 * 
+	 * These settings will continue to be used until the bulkdata object is reset, a subsequent
+	 * call to ::SetCompressionOptions is made or the owning package is serialized to disk.
+	 * 
+	 * @param Compressor		The Oodle compressor to use. @see ECompressedBufferCompressor
+	 * @param CompressionLevel	The Oodle compression level to use. @see ECompressedBufferCompressionLevel
+	 */
+	void SetCompressionOptions(ECompressedBufferCompressor Compressor, ECompressedBufferCompressionLevel CompressionLevel);
+
+	/**
+	* Get the CustomVersions used in the file containing the payload. Currently this is assumed
+	* to always be the versions in the InlineArchive
+	* 
+	* @param InlineArchive The archive that was used to load this object
+	* 
+	* @return The CustomVersions that apply to the interpretation of the payload.
+	*/
+	FCustomVersionContainer GetCustomVersions(FArchive& InlineArchive);
+
+	/**
+	 * Set this BulkData into Torn-Off mode. It will no longer register with the BulkDataRegistry, even if
+	 * copied from another BulkData, and it will pass on this flag to any BulkData copied/moved from it.
+	 * Use Reset() to remove this state. Torn-off BulkDatas share the guid with the BulkData they copy from.
+	 */
+	void TearOff();
+
+	/** Make a torn-off copy of this bulk data. */
+	FEditorBulkData CopyTornOff() const
+	{ 
+		return FEditorBulkData(*this, ETornOff());
+	}
+
+	// Functions used by the BulkDataRegistry
+
+	/** Used to serialize the bulkdata to/from a limited cache system used by the BulkDataRegistry. */
+	void SerializeForRegistry(FArchive& Ar);
+	/** Return true if the bulkdata has a source location that persists between editor processes (package file or virtualization). */
+	bool CanSaveForRegistry() const;
+	/** Return whether the BulkData has legacy payload id that needs to be updated from loaded payload before it can be used in DDC. */
+	bool HasPlaceholderPayloadId() const { return EnumHasAnyFlags(Flags, EFlags::LegacyKeyWasGuidDerived); }
+	/** Return whether the BulkData is an in-memory payload without a persistent source location. */
+	bool IsMemoryOnlyPayload() const;
+	/** Load the payload and set the correct payload id, if the bulkdata has a PlaceholderPayloadId. */
+	void UpdatePayloadId();
+
+	/**
+	 * Register this BulkData with the BulkData registry, if it is valid for registration.
+	 * This function is called automatically when necessary by EditorBulkData internals, 
+	 * but external callers can call it to provide information about the owner if it was previously
+	 * registered anonymously by UpdatePayload.
+	 */
+	void Register(UObject* Owner);
+
+#if UE_ENABLE_VIRTUALIZATION_TOGGLE
+	UE_DEPRECATED(5.0, "SetVirtualizationOptOut is an internal feature for development and will be removed without warning!")
+	void SetVirtualizationOptOut(bool bOptOut);
+#endif //UE_ENABLE_VIRTUALIZATION_TOGGLE
+
+protected:
+	enum class ETornOff {};
+	FEditorBulkData(const FEditorBulkData& Other, ETornOff);
+
+private:
+	friend struct FTocEntry;
+
+	/** Flags used to store additional meta information about the bulk data */
+	enum class EFlags : uint32
+	{
+		/** No flags are set */
+		None						= 0,
+		/** Is the data actually virtualized or not? */
+		IsVirtualized				= 1 << 0,
+		/** Does the package have access to a .upayload file? */
+		HasPayloadSidecarFile		= 1 << 1,
+		/** The bulkdata object is currently referencing a payload saved under old bulkdata formats */
+		ReferencesLegacyFile		= 1 << 2,
+		/** The legacy file being referenced is stored with Zlib compression format */
+		LegacyFileIsCompressed		= 1 << 3,
+		/** The payload should not have compression applied to it. It is assumed that the payload is already 
+			in some sort of compressed format, see the compression documentation above for more details. */
+		DisablePayloadCompression	= 1 << 4,
+		/** The legacy file being referenced derived its key from guid and it should be replaced with a key-from-hash when saved */
+		LegacyKeyWasGuidDerived		= 1 << 5,
+		/** The Guid has been registered with the BulkDataRegistry */
+		HasRegistered				= 1 << 6,
+		/** The BulkData object is a copy used only to represent the id and payload; it does not communicate with the BulkDataRegistry, and will point DDC jobs toward the original BulkData */
+		IsTornOff					= 1 << 7,
+		/** The bulkdata object references a payload stored in a WorkspaceDomain file  */
+		ReferencesWorkspaceDomain	= 1 << 8,
+		/** The payload is stored in a package trailer, so the bulkdata object will have to poll the trailer to find the payload offset */
+		StoredInPackageTrailer		= 1 << 9,
+
+		TransientFlags				= HasRegistered | IsTornOff,
+	};
+
+	FRIEND_ENUM_CLASS_FLAGS(EFlags);
+
+	/** Used to control what level of error reporting we return from some methods */
+	enum ErrorVerbosity
+	{
+		/** No errors should be logged */
+		None = 0,
+		/** Everything should be logged */
+		All
+	};
+
+	/** Old legacy path that saved the payload to the end of the package */
+	void SerializeToLegacyPath(FLinkerSave& LinkerSave, FCompressedBuffer PayloadToSerialize, EFlags UpdatedFlags, UObject* Owner);
+	/** The new path that saves payloads to the FPackageTrailer which is then appended to the end of the package file */
+	void SerializeToPackageTrailer(FLinkerSave& LinkerSave, FCompressedBuffer PayloadToSerialize, EFlags UpdatedFlags, UObject* Owner);
+
+	void UpdatePayloadImpl(FSharedBuffer&& InPayload, UE::Virtualization::FPayloadId&& InPayloadID);
+
+	FCompressedBuffer GetDataInternal() const;
+
+	FCompressedBuffer LoadFromDisk() const;
+	FCompressedBuffer LoadFromPackageFile() const;
+	FCompressedBuffer LoadFromPackageTrailer() const;
+	FCompressedBuffer LoadFromSidecarFile() const;
+	FCompressedBuffer LoadFromSidecarFileInternal(ErrorVerbosity Verbosity) const;
+
+	bool SerializeData(FArchive& Ar, FCompressedBuffer& Payload, const EFlags PayloadFlags) const;
+
+	void PushData(const FPackagePath& InPackagePath);
+	FCompressedBuffer PullData() const;
+
+	bool CanUnloadData() const;
+
+	void UpdateKeyIfNeeded();
+
+	void RecompressForSerialization(FCompressedBuffer& InOutPayload, EFlags PayloadFlags) const;
+	EFlags BuildFlagsForSerialization(FArchive& Ar, bool bKeepFileDataByReference) const;
+
+	bool IsDataVirtualized() const 
+	{ 
+		return IsDataVirtualized(Flags);
+	}
+
+	static bool IsDataVirtualized(EFlags InFlags)
+	{
+		return EnumHasAnyFlags(InFlags, EFlags::IsVirtualized);
+	}
+
+	bool HasPayloadSidecarFile() const 
+	{ 
+		return EnumHasAnyFlags(Flags, EFlags::HasPayloadSidecarFile); 
+	}
+
+	bool IsReferencingOldBulkData() const
+	{ 
+		return IsReferencingOldBulkData(Flags);
+	}
+	static bool IsReferencingOldBulkData(EFlags InFlags)
+	{
+		return EnumHasAnyFlags(InFlags, EFlags::ReferencesLegacyFile);
+	}
+
+	bool IsReferencingWorkspaceDomain() const
+	{
+		return IsReferencingWorkspaceDomain(Flags);
+	}
+	static bool IsReferencingWorkspaceDomain(EFlags InFlags)
+	{
+		return EnumHasAnyFlags(InFlags, EFlags::ReferencesWorkspaceDomain);
+	}
+
+	bool IsReferencingByPackagePath() const
+	{
+		return IsReferencingByPackagePath(Flags);
+	}
+	static bool IsReferencingByPackagePath(EFlags InFlags)
+	{
+		return EnumHasAnyFlags(InFlags, EFlags::ReferencesLegacyFile | EFlags::ReferencesWorkspaceDomain);
+	}
+
+	bool IsStoredInPackageTrailer() const
+	{
+		return IsStoredInPackageTrailer(Flags);
+	}
+	static bool IsStoredInPackageTrailer(EFlags InFlags)
+	{
+		return EnumHasAnyFlags(InFlags, EFlags::StoredInPackageTrailer);
+	}
+
+	void Unregister();
+
+	/**
+	 * Checks to make sure that the payload we are saving is what we expect it to be. If not then we need to log an error to the
+	 * user, mark the archive as having an error and return.
+	 * This error is not reasonable to expect, if it occurs then it means something has gone very wrong, which is why there is
+	 * also an ensure to make sure that any instances of it occurring as properly recorded and can be investigated.
+	 */
+	bool TryPayloadValidationForSaving(const FCompressedBuffer& PayloadForSaving, FLinkerSave* LinkerSave) const;
+
+	/** 
+	 * Utility to return an apt error message if the payload is invalid when trying to save the bulkdata. It will try to provide the best info from the given options.
+	 * 
+	 * @param Linker	The FLinkerSave being used to save the bulkdata at the point of error. This can be nullptr if the bulkdata is not being saved to a package on disk.
+	 * @return			The formatted error message.
+	 */
+	FText GetCorruptedPayloadErrorMsgForSave(FLinkerSave* Linker) const;
+
+	/** Returns true if we should use legacy serialization instead of the FPackageTrailer system. This can be removed when UE_ENABLE_VIRTUALIZATION_TOGGLE is removed. */
+	bool ShouldUseLegacySerialization(const FLinkerSave* LinkerSave) const;
+
+	/** Unique identifier for the bulkdata object itself */
+	FGuid BulkDataId;
+
+	/** Unique identifier for the contents of the payload*/
+	UE::Virtualization::FPayloadId PayloadContentId;
+
+	/** Pointer to the payload if it is held in memory (it has been updated but not yet saved to disk for example) */
+	FSharedBuffer Payload;
+
+	/** Length of the payload in bytes */
+	int64 PayloadSize = 0;
+
+	//---- The remaining members are used when the payload is not virtualized.
+	
+	/** The archive representing the file on disk containing the payload (if there is one), we keep the pointer so that the bulkdata object can be detached if needed. */
+	FArchive* AttachedAr = nullptr;
+
+	/** Offset of the payload in the file that contains it (INDEX_NONE if the payload does not come from a file)*/
+	int64 OffsetInFile = INDEX_NONE;
+
+	/** PackagePath containing the payload (this will be empty if the payload does not come from PackageResourceManager)*/
+	FPackagePath PackagePath;
+
+	/** PackageSegment to load with the packagepath (unused if the payload does not come from PackageResourceManager) */
+	EPackageSegment PackageSegment= EPackageSegment::Header;
+
+	/** A 32bit bitfield of flags */
+	EFlags Flags = EFlags::None;
+
+#if UE_ENABLE_VIRTUALIZATION_TOGGLE
+	bool bSkipVirtualization = false;
+#endif //UE_ENABLE_VIRTUALIZATION_TOGGLE
+
+	/** 
+	 * Compression settings to be applied to the payload when the package is next saved. The settings will be reset if
+	 * the payload is unloaded from memory during serialization (i.e. the payload was virtualized or the package was
+	 * saved to disk.&
+	 */
+	Private::FCompressionSettings CompressionSettings;
+};
+
+ENUM_CLASS_FLAGS(FEditorBulkData::EFlags);
+
+/** 
+ * NOTE: FPayloadToc/FTocEntry are now strictly legacy and are only used by the experimental sidecar system.
+ * They will be removed when the sidecar storage is changed to use FPackageTrailer instead.
+ */
+
+/** 
+  * Represents an entry to the table of contents found at the start of a payload sidecar file.
+  * This might be moved to it's own header and the table of contents made into a proper class 
+  * if we decide that we want to make access of the payload sidecar file a generic feature.
+  */
+struct COREUOBJECT_API FTocEntry
+{
+	FTocEntry() = default;
+	FTocEntry(const FEditorBulkData& BulkData)
+		: Identifier(BulkData.PayloadContentId)
+		, OffsetInFile(BulkData.OffsetInFile)
+		, UncompressedSize(BulkData.PayloadSize)
+	{
+
+	}
+
+	friend FArchive& operator<<(FArchive& Ar, FTocEntry& Entry);
+	friend void operator<<(FStructuredArchive::FSlot Slot, FTocEntry& Entry);
+
+	static constexpr uint32 PayloadSidecarFileVersion = 1;
+
+	/** Identifier for the payload */
+	UE::Virtualization::FPayloadId Identifier;
+	/** The offset into the file where we can find the payload, note that a virtualized payload will have an offset of INDEX_NONE */
+	int64 OffsetInFile = INDEX_NONE;
+	/** The size of the payload when uncompressed. */
+	int64 UncompressedSize = INDEX_NONE;
+};
+
+/** 
+ * A table of contents showing the location of all virtualized bulkdata payloads in a file. 
+ * NOTE: This is currently only used by the sidecar functionality.
+ */
+class COREUOBJECT_API FPayloadToc
+{
+public:
+
+	void AddEntry(const FEditorBulkData& BulkData);
+	bool FindEntry(const UE::Virtualization::FPayloadId& Identifier, FTocEntry& OutEntry);
+
+	const TArray<FTocEntry>& GetContents() const;
+
+	friend FArchive& operator<<(FArchive& Ar, FPayloadToc& TableOfContents);
+	friend void operator<<(FStructuredArchive::FSlot Slot, FPayloadToc& TableOfContents);
+
+private:
+	enum class EVersion : uint32
+	{
+		INITIAL = 0,
+
+		// -----<new versions can be added before this line>-------------------------------------------------
+		// - this needs to be the last line (see note below)
+		AUTOMATIC_VERSION_PLUS_ONE,
+		AUTOMATIC_VERSION = AUTOMATIC_VERSION_PLUS_ONE - 1
+	};
+
+	TArray<FTocEntry> Contents;
+};
+
+} // namespace UE::Serialization
+
+//#endif //WITH_EDITORONLY_DATA
diff --git a/Engine/Source/Runtime/CoreUObject/Public/Serialization/EditorBulkDataReader.h b/Engine/Source/Runtime/CoreUObject/Public/Serialization/EditorBulkDataReader.h
new file mode 100644
--- /dev/null
+++ b/Engine/Source/Runtime/CoreUObject/Public/Serialization/EditorBulkDataReader.h
@@ -0,0 +1,90 @@
+// Copyright Epic Games, Inc. All Rights Reserved.
+
+#pragma once
+
+#include "Serialization/BufferReader.h"
+#include "Serialization/EditorBulkData.h"
+
+#if WITH_EDITORONLY_DATA
+
+namespace UE::Serialization
+{
+
+namespace Private
+{
+
+/** Wraps access to the FEditorBulkData data for FEditorBulkDataReader
+so that it can be done before the FBufferReaderBase constructor is called */
+class DataAccessWrapper
+{
+protected:
+	DataAccessWrapper(FEditorBulkData& InBulkData)
+		: Payload(InBulkData.GetPayload().Get())
+
+	{
+
+	}
+
+	virtual ~DataAccessWrapper() = default;
+
+	bool IsValid() const
+	{
+		return !Payload.IsNull();
+	}
+
+	void* GetData() const
+	{
+		// It's okay to remove the const qualifier here as it will only be passed
+		// on to FBufferReaderBase, which will not change the data at all, but takes a non-const
+		// pointer so that it can free the memory if requested, which we don't.
+		return const_cast<void*>(Payload.GetData());
+	}
+
+	int64 GetDataLength() const
+	{
+		return Payload.GetSize();
+	}
+
+private:
+	FSharedBuffer Payload;
+};
+
+} // namespace Private
+
+class FEditorBulkDataReader : protected Private::DataAccessWrapper, public FBufferReaderBase
+{
+public:
+	FEditorBulkDataReader(FEditorBulkData& InBulkData, bool bIsPersistent = false)
+		: DataAccessWrapper(InBulkData)
+		, FBufferReaderBase(GetData(), GetDataLength(), false, bIsPersistent)
+	{
+	}
+
+	virtual ~FEditorBulkDataReader() = default;
+
+	/** Returns if the FEditorBulkDataReader has a valid bulkdata payload or not */
+	bool IsValid() const
+	{
+		return Private::DataAccessWrapper::IsValid();
+	}
+
+	using FArchive::operator<<; // For visibility of the overloads we don't override
+
+	virtual FArchive& operator<<(class FName& Name) override
+	{
+		// FNames are serialized as strings in BulkData
+		FString StringName;
+		*this << StringName;
+		Name = FName(*StringName);
+		return *this;
+	}
+
+	virtual FString GetArchiveName() const
+	{
+		return TEXT("FEditorBulkDataReader");
+	}
+};
+
+} // namespace UE::Serialization
+
+#endif //WITH_EDITORONLY_DATA
diff --git a/Engine/Source/Runtime/CoreUObject/Public/Serialization/EditorBulkDataWriter.h b/Engine/Source/Runtime/CoreUObject/Public/Serialization/EditorBulkDataWriter.h
new file mode 100644
--- /dev/null
+++ b/Engine/Source/Runtime/CoreUObject/Public/Serialization/EditorBulkDataWriter.h
@@ -0,0 +1,129 @@
+// Copyright Epic Games, Inc. All Rights Reserved.
+
+#pragma once
+
+#include "Serialization/EditorBulkData.h"
+
+#if WITH_EDITORONLY_DATA
+
+namespace UE::Serialization
+{
+
+class FEditorBulkDataWriter : public FArchive
+{
+public:
+	FEditorBulkDataWriter(FEditorBulkData& InBulkData, bool bIsPersistent = false)
+		: BulkData(InBulkData)
+	{
+		SetIsSaving(true);
+		SetIsPersistent(bIsPersistent);
+
+		FSharedBuffer Payload = InBulkData.GetPayload().Get();
+
+		if (Payload)
+		{
+			const int64 CurrentDataLength = Payload.GetSize();
+
+			// Clone the payload so that we have a local copy that we can
+			// append additional data to.
+			Buffer = FMemory::Malloc(CurrentDataLength, DEFAULT_ALIGNMENT);
+			FMemory::Memcpy(Buffer, Payload.GetData(), CurrentDataLength);
+
+			BufferLength = CurrentDataLength;
+
+			// Start at the end of the existing data
+			CurPos = CurrentDataLength;
+			DataLength = CurrentDataLength;
+		}
+		else
+		{
+			Buffer = nullptr;
+			BufferLength = 0;
+
+			CurPos = 0;
+			DataLength = 0;
+		}
+	}
+
+	~FEditorBulkDataWriter()
+	{
+		// Remove the slack from the allocated bulk data
+		Buffer = FMemory::Realloc(Buffer, DataLength, DEFAULT_ALIGNMENT);
+		BulkData.UpdatePayload(FSharedBuffer::TakeOwnership(Buffer, DataLength, FMemory::Free));
+	}
+
+	/** Returns if the FEditorBulkDataWriter has a valid bulkdata payload or not */
+	bool IsValid() const
+	{
+		return Buffer != nullptr;
+	}
+
+	virtual void Serialize(void* Data, int64 Num)
+	{
+		// Determine if we need to reallocate the buffer to fit the next item
+		const int64 NewPos = CurPos + Num;
+		checkf(NewPos >= CurPos, TEXT("Serialization has somehow gone backwards"));
+
+		if (NewPos > BufferLength)
+		{
+			// If so, resize to the new size + 3/8 additional slack
+			const int64 NewLength = NewPos + 3 * NewPos / 8 + 16;
+			Buffer = FMemory::Realloc(Buffer, NewLength, DEFAULT_ALIGNMENT);
+			BufferLength = NewLength;
+		}
+
+		FMemory::Memcpy(static_cast<unsigned char*>(Buffer) + CurPos, Data, Num);
+
+		CurPos += Num;
+		DataLength = FMath::Max(DataLength, CurPos);
+	}
+
+	using FArchive::operator<<; // For visibility of the overloads we don't override
+
+	virtual FArchive& operator<<(class FName& Name) override
+	{
+		// FNames are serialized as strings in BulkData
+		FString StringName = Name.ToString();
+		*this << StringName;
+		return *this;
+	}
+
+	virtual int64 Tell() { return CurPos; }
+	virtual int64 TotalSize() { return DataLength; }
+
+	virtual void Seek(int64 InPos)
+	{
+		check(InPos >= 0);
+		check(InPos <= DataLength);
+		CurPos = InPos;
+	}
+
+	virtual bool AtEnd()
+	{
+		return CurPos >= DataLength;
+	}
+
+	virtual FString GetArchiveName() const
+	{
+		return TEXT("FEditorBulkDataWriter");
+	}
+
+protected:
+	/** The target bulkdata object */
+	FEditorBulkData& BulkData;
+
+	/** Pointer to the data buffer */
+	void* Buffer;
+	/** Length of the data buffer (stored as bytes) */
+	int64 BufferLength;
+
+	/** Current position in the buffer for serialization operations */
+	int64 CurPos;
+
+	/** The length of valid data in the data buffer (stored as bytes) */
+	int64 DataLength;
+};
+
+} // namespace UE::Serialization
+
+#endif //WITH_EDITORONLY_DATA
diff --git a/Engine/Source/Runtime/CoreUObject/Public/Serialization/VirtualizedBulkData.h b/Engine/Source/Runtime/CoreUObject/Public/Serialization/VirtualizedBulkData.h
deleted file mode 100644
diff --git a/Engine/Source/Runtime/CoreUObject/Public/Serialization/VirtualizedBulkDataReader.h b/Engine/Source/Runtime/CoreUObject/Public/Serialization/VirtualizedBulkDataReader.h
deleted file mode 100644
diff --git a/Engine/Source/Runtime/CoreUObject/Public/Serialization/VirtualizedBulkDataWriter.h b/Engine/Source/Runtime/CoreUObject/Public/Serialization/VirtualizedBulkDataWriter.h
deleted file mode 100644
diff --git a/Engine/Source/Runtime/CoreUObject/Public/UObject/LinkerLoad.h b/Engine/Source/Runtime/CoreUObject/Public/UObject/LinkerLoad.h
--- a/Engine/Source/Runtime/CoreUObject/Public/UObject/LinkerLoad.h
+++ b/Engine/Source/Runtime/CoreUObject/Public/UObject/LinkerLoad.h
@@ -18,15 +18,8 @@ class ULinkerPlaceholderExportObject;
 struct FScopedSlowTask;
 struct FUntypedBulkData;
 
-namespace UE
-{
-class FPackageTrailer;
-}
-
-namespace UE::Virtualization
-{
-class FVirtualizedUntypedBulkData;
-}
+namespace UE{ class FPackageTrailer; }
+namespace UE::Serialization{ class FEditorBulkData; }
 
 /*----------------------------------------------------------------------------
 	FLinkerLoad.
@@ -240,7 +233,7 @@ public:
 #if WITH_EDITOR
 	/** Bulk data that does not need to be loaded when the linker is loaded.												*/
 	TArray<FUntypedBulkData*> BulkDataLoaders;
-	TArray<UE::Virtualization::FVirtualizedUntypedBulkData*> VirtualizedBulkDataLoaders;
+	TArray<UE::Serialization::FEditorBulkData*> EditorBulkDataLoaders;
 #endif // WITH_EDITOR
 
 	/** Hash table for exports.																								*/
@@ -835,7 +828,7 @@ private:
 	 * @param	BulkData	Bulk data object to associate
 	 */
 	virtual void AttachBulkData(UObject* Owner, FUntypedBulkData* BulkData) override;
-	virtual void AttachBulkData(UE::Virtualization::FVirtualizedUntypedBulkData* BulkData) override;
+	virtual void AttachBulkData(UE::Serialization::FEditorBulkData* BulkData) override;
 	/**
 	 * Detaches the passed in bulk data object from the linker.
 	 *
@@ -843,7 +836,7 @@ private:
 	 * @param	bEnsureBulkDataIsLoaded	Whether to ensure that the bulk data is loaded before detaching
 	 */
 	virtual void DetachBulkData(FUntypedBulkData* BulkData, bool bEnsureBulkDataIsLoaded) override;
-	virtual void DetachBulkData(UE::Virtualization::FVirtualizedUntypedBulkData* BulkData, bool bEnsureBulkDataIsLoaded) override;
+	virtual void DetachBulkData(UE::Serialization::FEditorBulkData* BulkData, bool bEnsureBulkDataIsLoaded) override;
 	/**
 	 * Detaches all attached bulk  data objects.
 	 *
diff --git a/Engine/Source/Runtime/CoreUObject/Public/UObject/LinkerSave.h b/Engine/Source/Runtime/CoreUObject/Public/UObject/LinkerSave.h
--- a/Engine/Source/Runtime/CoreUObject/Public/UObject/LinkerSave.h
+++ b/Engine/Source/Runtime/CoreUObject/Public/UObject/LinkerSave.h
@@ -15,7 +15,6 @@
 
 class FObjectPostSaveContext;
 struct FUntypedBulkData;
-namespace UE::Virtualization { class FVirtualizedUntypedBulkData; }
 
 /*----------------------------------------------------------------------------
 	FLinkerSave.
@@ -112,7 +111,7 @@ public:
 		FCompressedBuffer Payload;
 	};
 
-	/** Used by FVirtualizedUntypedBulkData to add payloads to be added to the payload sidecar file (currently an experimental feature) */
+	/** Used by FEditorBulkData to add payloads to be added to the payload sidecar file (currently an experimental feature) */
 	TArray<FSidecarStorageInfo> SidecarDataToAppend;
 	
 	/** Gathers all payloads while save the package, so that they can be stored in a single data structure @see FPackageTrailer */
diff --git a/Engine/Source/Runtime/Engine/Classes/Engine/Texture.h b/Engine/Source/Runtime/Engine/Classes/Engine/Texture.h
--- a/Engine/Source/Runtime/Engine/Classes/Engine/Texture.h
+++ b/Engine/Source/Runtime/Engine/Classes/Engine/Texture.h
@@ -12,7 +12,7 @@
 #include "Interfaces/Interface_AsyncCompilation.h"
 #include "RenderCommandFence.h"
 #include "RenderResource.h"
-#include "Serialization/VirtualizedBulkData.h"
+#include "Serialization/EditorBulkData.h"
 #include "Engine/TextureDefines.h"
 #include "MaterialShared.h"
 #include "TextureResource.h"
@@ -156,7 +156,7 @@ struct FTextureSource
 		const FTextureSourceBlock* InBlocks,
 		int32 InNumLayers,
 		int32 InNumBlocks,
-		UE::Virtualization::FVirtualizedUntypedBulkData::FSharedBufferWithID NewData);
+		UE::Serialization::FEditorBulkData::FSharedBufferWithID NewData);
 
 	ENGINE_API void InitLayered(
 		int32 NewSizeX,
@@ -175,7 +175,7 @@ struct FTextureSource
 		int32 NewNumLayers,
 		int32 NewNumMips,
 		const ETextureSourceFormat* NewLayerFormat,
-		UE::Virtualization::FVirtualizedUntypedBulkData::FSharedBufferWithID NewData
+		UE::Serialization::FEditorBulkData::FSharedBufferWithID NewData
 	);
 
 	/**
@@ -211,7 +211,7 @@ struct FTextureSource
 		int32 NewNumSlices,
 		int32 NewNumMips,
 		ETextureSourceFormat NewFormat,
-		UE::Virtualization::FVirtualizedUntypedBulkData::FSharedBufferWithID NewData
+		UE::Serialization::FEditorBulkData::FSharedBufferWithID NewData
 	);
 
 	/**
@@ -463,7 +463,7 @@ private:
 	TDontCopy<FRWLock> BulkDataLock;
 #endif
 	/** The bulk source data. */
-	UE::Virtualization::FByteVirtualizedBulkData BulkData;
+	UE::Serialization::FEditorBulkData BulkData;
 	
 	/** Number of mips that are locked. */
 	uint32 NumLockedMips;
diff --git a/Engine/Source/Runtime/Engine/Private/Texture.cpp b/Engine/Source/Runtime/Engine/Private/Texture.cpp
--- a/Engine/Source/Runtime/Engine/Private/Texture.cpp
+++ b/Engine/Source/Runtime/Engine/Private/Texture.cpp
@@ -1224,19 +1224,19 @@ void FTextureSource::InitBlocked(const ETextureSourceFormat* InLayerFormats,
 	}
 
 	BulkData.UpdatePayload(Buffer.MoveToShared());
-	BulkData.SetCompressionOptions(UE::Virtualization::ECompressionOptions::Default);
+	BulkData.SetCompressionOptions(UE::Serialization::ECompressionOptions::Default);
 }
 
 void FTextureSource::InitBlocked(const ETextureSourceFormat* InLayerFormats,
 	const FTextureSourceBlock* InBlocks,
 	int32 InNumLayers,
 	int32 InNumBlocks,
-	UE::Virtualization::FVirtualizedUntypedBulkData::FSharedBufferWithID NewData)
+	UE::Serialization::FEditorBulkData::FSharedBufferWithID NewData)
 {
 	InitBlockedImpl(InLayerFormats, InBlocks, InNumLayers, InNumBlocks);
 
 	BulkData.UpdatePayload(MoveTemp(NewData));
-	BulkData.SetCompressionOptions(UE::Virtualization::ECompressionOptions::Default);
+	BulkData.SetCompressionOptions(UE::Serialization::ECompressionOptions::Default);
 }
 
 void FTextureSource::InitLayered(
@@ -1275,7 +1275,7 @@ void FTextureSource::InitLayered(
 		BulkData.UpdatePayload(FUniqueBuffer::Alloc(TotalBytes).MoveToShared());
 	}
 
-	BulkData.SetCompressionOptions(UE::Virtualization::ECompressionOptions::Default);
+	BulkData.SetCompressionOptions(UE::Serialization::ECompressionOptions::Default);
 }
 
 void FTextureSource::InitLayered(
@@ -1285,7 +1285,7 @@ void FTextureSource::InitLayered(
 	int32 NewNumLayers,
 	int32 NewNumMips,
 	const ETextureSourceFormat* NewLayerFormat,
-	UE::Virtualization::FVirtualizedUntypedBulkData::FSharedBufferWithID NewData)
+	UE::Serialization::FEditorBulkData::FSharedBufferWithID NewData)
 {
 	InitLayeredImpl(
 		NewSizeX,
@@ -1297,7 +1297,7 @@ void FTextureSource::InitLayered(
 	);
 
 	BulkData.UpdatePayload(MoveTemp(NewData));
-	BulkData.SetCompressionOptions(UE::Virtualization::ECompressionOptions::Default);
+	BulkData.SetCompressionOptions(UE::Serialization::ECompressionOptions::Default);
 }
 
 void FTextureSource::Init(
@@ -1318,7 +1318,7 @@ void FTextureSource::Init(
 	int32 NewNumSlices,
 	int32 NewNumMips,
 	ETextureSourceFormat NewFormat,
-	UE::Virtualization::FVirtualizedUntypedBulkData::FSharedBufferWithID NewData
+	UE::Serialization::FEditorBulkData::FSharedBufferWithID NewData
 )
 {
 	InitLayered(NewSizeX, NewSizeY, NewNumSlices, 1, NewNumMips, &NewFormat, MoveTemp(NewData));
@@ -1387,11 +1387,11 @@ void FTextureSource::InitWithCompressedSourceData(
 	// Disable the internal bulkdata compression if the source data is already compressed
 	if (CompressionFormat == TSCF_None)
 	{
-		BulkData.SetCompressionOptions(UE::Virtualization::ECompressionOptions::Default);
+		BulkData.SetCompressionOptions(UE::Serialization::ECompressionOptions::Default);
 	}
 	else
 	{
-		BulkData.SetCompressionOptions(UE::Virtualization::ECompressionOptions::Disabled);
+		BulkData.SetCompressionOptions(UE::Serialization::ECompressionOptions::Disabled);
 	}
 }
 
@@ -1459,7 +1459,7 @@ void FTextureSource::Compress()
 	}
 	else
 	{
-		BulkData.SetCompressionOptions(UE::Virtualization::ECompressionOptions::Disabled);
+		BulkData.SetCompressionOptions(UE::Serialization::ECompressionOptions::Disabled);
 	}
 }
 
@@ -1554,7 +1554,7 @@ void FTextureSource::UnlockMip(int32 BlockIndex, int32 LayerIndex, int32 MipInde
 			UE_CLOG(CompressionFormat == TSCF_JPEG, LogTexture, Warning, TEXT("Call to FTextureSource::UnlockMip will cause texture source to lose it's jpeg storage format"));
 
 			BulkData.UpdatePayload(LockedMipData.Release());
-			BulkData.SetCompressionOptions(UE::Virtualization::ECompressionOptions::Default);
+			BulkData.SetCompressionOptions(UE::Serialization::ECompressionOptions::Default);
 
 			bPNGCompressed = false;
 			CompressionFormat = TSCF_None;
diff --git a/Engine/Source/Runtime/MeshDescription/Private/MeshDescription.cpp b/Engine/Source/Runtime/MeshDescription/Private/MeshDescription.cpp
--- a/Engine/Source/Runtime/MeshDescription/Private/MeshDescription.cpp
+++ b/Engine/Source/Runtime/MeshDescription/Private/MeshDescription.cpp
@@ -6,12 +6,12 @@
 #include "MeshAttributes.h"
 #include "Misc/SecureHash.h"
 #include "Serialization/BulkData.h"
+#include "Serialization/EditorBulkDataReader.h"
+#include "Serialization/EditorBulkDataWriter.h"
 #include "Serialization/LargeMemoryWriter.h"
 #include "Serialization/MemoryReader.h"
 #include "Serialization/MemoryWriter.h"
 #include "Serialization/NameAsStringProxyArchive.h"
-#include "Serialization/VirtualizedBulkDataReader.h"
-#include "Serialization/VirtualizedBulkDataWriter.h"
 #include "UObject/EnterpriseObjectVersion.h"
 #include "UObject/UE5CookerObjectVersion.h"
 #include "UObject/UE5MainStreamObjectVersion.h"
@@ -2080,7 +2080,7 @@ void FMeshDescriptionBulkData::SaveMeshDescription( FMeshDescription& MeshDescri
 	if( !MeshDescription.IsEmpty() )
 	{
 		const bool bIsPersistent = true;
-		UE::Virtualization::FVirtualizedBulkDataWriter Ar(BulkData, bIsPersistent);
+		UE::Serialization::FEditorBulkDataWriter Ar(BulkData, bIsPersistent);
 		Ar << MeshDescription;
 
 		// Preserve CustomVersions at save time so we can reuse the same ones when reloading direct from memory
@@ -2116,7 +2116,7 @@ void FMeshDescriptionBulkData::LoadMeshDescription( FMeshDescription& MeshDescri
 #endif //WITH_EDITOR
 
 		const bool bIsPersistent = true;
-		UE::Virtualization::FVirtualizedBulkDataReader Ar(BulkData, bIsPersistent);
+		UE::Serialization::FEditorBulkDataReader Ar(BulkData, bIsPersistent);
 
 		// Propagate the custom version information from the package to the bulk data, so that the MeshDescription
 		// is serialized with the same versioning.
diff --git a/Engine/Source/Runtime/MeshDescription/Public/MeshDescription.h b/Engine/Source/Runtime/MeshDescription/Public/MeshDescription.h
--- a/Engine/Source/Runtime/MeshDescription/Public/MeshDescription.h
+++ b/Engine/Source/Runtime/MeshDescription/Public/MeshDescription.h
@@ -23,7 +23,7 @@
 #include "UObject/EditorObjectVersion.h"
 #include "UObject/ReleaseObjectVersion.h"
 #include "Serialization/CustomVersion.h"
-#include "Serialization/VirtualizedBulkData.h"
+#include "Serialization/EditorBulkData.h"
 #include "Containers/StaticArray.h"
 #include "MeshDescription.generated.h"
 
@@ -1342,7 +1342,7 @@ private:
 	FRWLock       BulkDataLock;
 #endif
 	/** Internally store bulk data as bytes */
-	UE::Virtualization::FByteVirtualizedBulkData BulkData;
+	UE::Serialization::FEditorBulkData BulkData;
 
 	/** GUID associated with the data stored herein. */
 	FGuid Guid;
