CIS fix Localisation: Fix to LOC macros that need to be on a single line for parsing to work

diff --git a/Engine/Source/Developer/DerivedDataCache/Private/FileSystemDerivedDataBackend.cpp b/Engine/Source/Developer/DerivedDataCache/Private/FileSystemDerivedDataBackend.cpp
--- a/Engine/Source/Developer/DerivedDataCache/Private/FileSystemDerivedDataBackend.cpp
+++ b/Engine/Source/Developer/DerivedDataCache/Private/FileSystemDerivedDataBackend.cpp
@@ -1082,9 +1082,7 @@ bool FFileSystemCacheStore::RunSpeedTest(
 	const FString AbsoluteCachePath = IFileManager::Get().ConvertToAbsolutePathForExternalAppForRead(*CachePath);
 	if (AbsoluteCachePath.Len() > MAX_CACHE_DIR_LEN)
 	{
-		const FText ErrorMessage = FText::Format(NSLOCTEXT("DerivedDataCache", "PathTooLong",
-			"Cache path {0} is longer than {1} characters... "
-			"please adjust [DerivedDataBackendGraph] paths to be shorter (this leaves more room for cache keys)."),
+		const FText ErrorMessage = FText::Format(NSLOCTEXT("DerivedDataCache", "PathTooLong", "Cache path {0} is longer than {1} characters... please adjust [DerivedDataBackendGraph] paths to be shorter (this leaves more room for cache keys)."),
 			FText::FromString(AbsoluteCachePath), FText::AsNumber(MAX_CACHE_DIR_LEN));
 		FMessageDialog::Open(EAppMsgType::Ok, ErrorMessage);
 		UE_LOG(LogDerivedDataCache, Fatal, TEXT("%s"), *ErrorMessage.ToString());
diff --git a/Engine/Source/Runtime/CoreUObject/Private/UObject/SavePackage2.cpp b/Engine/Source/Runtime/CoreUObject/Private/UObject/SavePackage2.cpp
--- a/Engine/Source/Runtime/CoreUObject/Private/UObject/SavePackage2.cpp
+++ b/Engine/Source/Runtime/CoreUObject/Private/UObject/SavePackage2.cpp
@@ -150,14 +150,12 @@ ESavePackageResult ValidatePackage(FSaveContext& SaveContext)
 				FText ErrorText;
 				if (!Asset->GetExternalPackage() && EnumHasAnyFlags(TopLevelFlags, RF_Standalone))
 				{
-					ErrorText = FText::Format(NSLOCTEXT("SavePackage2", "AssetSaveMissingStandaloneFlag",
-						"The Asset {Name} being saved does not have any of the provided object flags (0x{Flags}); saving the package would cause data loss. "
-						"Run with -dpcvars=save.FixupStandaloneFlags=1 to add the RF_Standalone flag."), Arguments);
+					ErrorText = FText::Format(NSLOCTEXT("SavePackage2", "AssetSaveMissingStandaloneFlag", "The Asset {Name} being saved does not have any of the provided object flags (0x{Flags}); saving the package would cause data loss. Run with -dpcvars=save.FixupStandaloneFlags=1 to add the RF_Standalone flag."), 
+						Arguments);
 				}
 				else
 				{
-					ErrorText = FText::Format(NSLOCTEXT("SavePackage2", "AssetSaveMissingTopLevelFlags",
-						"The Asset {Name} being saved does not have any of the provided object flags (0x{Flags}); saving the package would cause data loss."),
+					ErrorText = FText::Format(NSLOCTEXT("SavePackage2", "AssetSaveMissingTopLevelFlags", "The Asset {Name} being saved does not have any of the provided object flags (0x{Flags}); saving the package would cause data loss."),	
 						Arguments);
 				}
 				SaveContext.GetError()->Logf(ELogVerbosity::Warning, TEXT("%s"), *ErrorText.ToString());
