Forward the additional property flags onto the FObjectPropertyBase base class, broken in CL# 18558135.

diff --git a/Engine/Source/Runtime/CoreUObject/Private/UObject/PropertyBaseObject.cpp b/Engine/Source/Runtime/CoreUObject/Private/UObject/PropertyBaseObject.cpp
--- a/Engine/Source/Runtime/CoreUObject/Private/UObject/PropertyBaseObject.cpp
+++ b/Engine/Source/Runtime/CoreUObject/Private/UObject/PropertyBaseObject.cpp
@@ -26,7 +26,7 @@ FObjectPropertyBase::FObjectPropertyBase(FFieldVariant InOwner, const UECodeGen_
 	PropertyClass = Prop.ClassFunc ? Prop.ClassFunc() : nullptr;
 }
 FObjectPropertyBase::FObjectPropertyBase(FFieldVariant InOwner, const UECodeGen_Private::FObjectPropertyParamsWithoutClass& Prop, EPropertyFlags AdditionalPropertyFlags /*= CPF_None*/)
-	: FProperty(InOwner, (const UECodeGen_Private::FPropertyParamsBaseWithOffset&)Prop)
+	: FProperty(InOwner, (const UECodeGen_Private::FPropertyParamsBaseWithOffset&)Prop, AdditionalPropertyFlags)
 	, PropertyClass(nullptr)
 {
 }
