Fix an assert for projects with virtualization and editor domain both enabled

#rnx

- EditorBulkData payloads that are updated in PostLoad were being added to the package trailer builder for local access when saving the package for the editor domain.
- The package trailer system is not able to support this (although this will change soon) so when this occurs we should force those payloads to the legacy serialization path.
- Cleaned up ::ShouldUseLegacySerialization by factoring out the code that is actually affected by UE_ENABLE_VIRTUALIZATION_TOGGLE is wrapped in it, making things easier to read.

diff --git a/Engine/Source/Runtime/CoreUObject/Private/Serialization/EditorBulkData.cpp b/Engine/Source/Runtime/CoreUObject/Private/Serialization/EditorBulkData.cpp
--- a/Engine/Source/Runtime/CoreUObject/Private/Serialization/EditorBulkData.cpp
+++ b/Engine/Source/Runtime/CoreUObject/Private/Serialization/EditorBulkData.cpp
@@ -1695,8 +1695,7 @@ FEditorBulkData::EFlags FEditorBulkData::BuildFlagsForSerialization(FArchive& Ar
 		{
 			EnumRemoveFlags(UpdatedFlags, EFlags::ReferencesLegacyFile | EFlags::ReferencesWorkspaceDomain | 
 				EFlags::LegacyFileIsCompressed | EFlags::LegacyKeyWasGuidDerived);
-
-			
+		
 			if (LinkerSave != nullptr && !LinkerSave->GetFilename().IsEmpty() && ShouldSaveToPackageSidecar())
 			{
 				EnumAddFlags(UpdatedFlags, EFlags::HasPayloadSidecarFile);
@@ -1716,13 +1715,18 @@ FEditorBulkData::EFlags FEditorBulkData::BuildFlagsForSerialization(FArchive& Ar
 			}
 		}
 
-		if (ShouldUseLegacySerialization(LinkerSave) == false)
+		// Currently we do not support storing local payloads to a trailer if it is being built for reference
+		// access (i.e. for the editor domain) and if this is detected we should force the legacy serialization
+		// path for this payload.
+		const bool bForceLegacyPath = bKeepFileDataByReference && bCanKeepFileDataByReference == false;
+
+		if (ShouldUseLegacySerialization(LinkerSave) == true || bForceLegacyPath == true)
 		{
-			EnumAddFlags(UpdatedFlags, EFlags::StoredInPackageTrailer);
+			EnumRemoveFlags(UpdatedFlags, EFlags::StoredInPackageTrailer);
 		}
 		else
 		{
-			EnumRemoveFlags(UpdatedFlags, EFlags::StoredInPackageTrailer);
+			EnumAddFlags(UpdatedFlags, EFlags::StoredInPackageTrailer);
 		}
 
 		return UpdatedFlags;
@@ -1790,16 +1794,19 @@ FText FEditorBulkData::GetCorruptedPayloadErrorMsgForSave(FLinkerSave* Linker) c
 
 bool FEditorBulkData::ShouldUseLegacySerialization(const FLinkerSave* LinkerSave) const
 {
+#if UE_ENABLE_VIRTUALIZATION_TOGGLE
+	if (bSkipVirtualization == true)
+	{
+		return true;
+	}
+#endif // UE_ENABLE_VIRTUALIZATION_TOGGLE 
+
 	if (LinkerSave == nullptr)
 	{
 		return true;
 	}
 
-#if UE_ENABLE_VIRTUALIZATION_TOGGLE
-	return (!LinkerSave->PackageTrailerBuilder.IsValid()) || bSkipVirtualization == true;
-#else
 	return !LinkerSave->PackageTrailerBuilder.IsValid();
-#endif //UE_ENABLE_VIRTUALIZATION_TOGGLE
 }
 	
 FArchive& operator<<(FArchive& Ar, FTocEntry& Entry)
