Core: Nullptr check for ReferencingProperty in log

diff --git a/Engine/Source/Runtime/CoreUObject/Private/UObject/UObjectGlobals.cpp b/Engine/Source/Runtime/CoreUObject/Private/UObject/UObjectGlobals.cpp
--- a/Engine/Source/Runtime/CoreUObject/Private/UObject/UObjectGlobals.cpp
+++ b/Engine/Source/Runtime/CoreUObject/Private/UObject/UObjectGlobals.cpp
@@ -3995,7 +3995,7 @@ private:
 			UPackage* ReferencingPackage = CurrentObject->GetOutermost();
 			if (!ReferencingPackage->HasAnyPackageFlags(PKG_PlayInEditor) && !ReferencingPackage->HasAnyFlags(RF_Transient))
 			{
-				UE_LOG(LogGarbage, Warning, TEXT("GC detected illegal reference to PIE object from content [possibly via %s]:"), *ReferencingProperty->GetFullName());
+				UE_LOG(LogGarbage, Warning, TEXT("GC detected illegal reference to PIE object from content [possibly via %s]:"), ReferencingProperty != nullptr ? *ReferencingProperty->GetFullName() : *FString());
 				UE_LOG(LogGarbage, Warning, TEXT("      PIE object: %s"), *Object->GetFullName());
 				UE_LOG(LogGarbage, Warning, TEXT("  NON-PIE object: %s"), *CurrentObject->GetFullName());
 			}
