Saving a package containing virtualized payloads will no longer force the package to 'rehydrate' and download the payloads to store locally.

#rnx

- A cvar named "Serialization.RehydrateOnSave" now controls if we rehydrate a package on save, defaulting to false.
-- At the very least we will need to be able to turn this to true to run any "rehydrate the project" UAT command in the future.
- The package trailer no longer stores a valid compressed size for a payload. In the best case scenario we only know the compressed size of the payload when it was still stored on locally on disk. Each virtualization backend is free to take the payload and recompress it anyway it sees fit etc.
- When duplicate payloads entries are added to a FPackageTrailerBuilder we now keep the original payload entry and discard the duplicate rather than replacing the existing entry with the duplicate, which was the original intent.

diff --git a/Engine/Source/Runtime/CoreUObject/Private/Serialization/EditorBulkData.cpp b/Engine/Source/Runtime/CoreUObject/Private/Serialization/EditorBulkData.cpp
--- a/Engine/Source/Runtime/CoreUObject/Private/Serialization/EditorBulkData.cpp
+++ b/Engine/Source/Runtime/CoreUObject/Private/Serialization/EditorBulkData.cpp
@@ -67,11 +67,10 @@ static TAutoConsoleVariable<bool> CVarShouldAllowSidecarSyncing(
   */
 static constexpr bool bAllowVirtualizationOnSave = false;
 
-/** When enabled virtualized payloads will be pulled and stored locally when the bulkdata object is saved
-  * to disk as part of a package. This is currently required but will likely be changed to be a config 
-  * option when we can support it.
-  */
-static constexpr bool bForceRehydrationOnSave = true;
+static TAutoConsoleVariable<bool> CVarShouldRehydrateOnSave(
+	TEXT("Serialization.RehydrateOnSave"),
+	false,
+	TEXT("When true FVirtualizedUntypedBulkData virtualized payloads will by hydrated and stored locally when saved to a package"));
 
 /** Wrapper around the config file option [Core.System.Experimental]EnablePackageSidecarSaving */
 bool ShouldSaveToPackageSidecar()
@@ -711,9 +710,8 @@ void FEditorBulkData::Serialize(FArchive& Ar, UObject* Owner, bool bAllowRegiste
 				}
 			}
 
-			// Some paranoid checks to make sure everything is working as expected
-			// TODO: Probably remove this when FPackageTrailer::IsEnabled is removed
-			if(IsStoredInPackageTrailer(UpdatedFlags) && PayloadContentId.IsValid())
+			// Make sure that the trailer builder is correct (if it is being used)
+			if (IsStoredInPackageTrailer(UpdatedFlags) && PayloadContentId.IsValid())
 			{
 				check(LinkerSave != nullptr);
 				check(LinkerSave->PackageTrailerBuilder.IsValid());
@@ -725,6 +723,7 @@ void FEditorBulkData::Serialize(FArchive& Ar, UObject* Owner, bool bAllowRegiste
 				}
 				else if (IsDataVirtualized(UpdatedFlags))
 				{
+					LinkerSave->PackageTrailerBuilder->AddVirtualizedPayload(PayloadContentId, PayloadSize);
 					check(LinkerSave->PackageTrailerBuilder->IsVirtualizedPayloadEntry(PayloadContentId));
 				}
 				else
@@ -1708,7 +1707,7 @@ FEditorBulkData::EFlags FEditorBulkData::BuildFlagsForSerialization(FArchive& Ar
 				// Remove the virtualization flag if we are rehydrating packages on save unless
 				// referencing the payload data is allowed, in which case we can continue to save
 				// as virtualized.
-				if (LinkerSave != nullptr && !bKeepFileDataByReference && bForceRehydrationOnSave)
+				if (LinkerSave != nullptr && !bKeepFileDataByReference && CVarShouldRehydrateOnSave.GetValueOnAnyThread())
 				{
 					EnumRemoveFlags(UpdatedFlags, EFlags::IsVirtualized);
 				}
diff --git a/Engine/Source/Runtime/CoreUObject/Private/UObject/PackageTrailer.cpp b/Engine/Source/Runtime/CoreUObject/Private/UObject/PackageTrailer.cpp
--- a/Engine/Source/Runtime/CoreUObject/Private/UObject/PackageTrailer.cpp
+++ b/Engine/Source/Runtime/CoreUObject/Private/UObject/PackageTrailer.cpp
@@ -102,7 +102,7 @@ FPackageTrailerBuilder FPackageTrailerBuilder::CreateFromTrailer(const FPackageT
 
 			case EPayloadAccessMode::Virtualized:
 			{
-				Builder.VirtualizedEntries.Add(Entry.Identifier, VirtualizedEntry(Entry.CompressedSize, Entry.RawSize));
+				Builder.VirtualizedEntries.Add(Entry.Identifier, VirtualizedEntry(Entry.RawSize));
 			}
 			break;
 
@@ -143,7 +143,7 @@ TUniquePtr<UE::FPackageTrailerBuilder> FPackageTrailerBuilder::CreateReferenceTo
 			
 			case EPayloadAccessMode::Virtualized:
 			{
-				Builder->VirtualizedEntries.Add(Entry.Identifier, VirtualizedEntry(Entry.CompressedSize, Entry.RawSize));
+				Builder->VirtualizedEntries.Add(Entry.Identifier, VirtualizedEntry(Entry.RawSize));
 			}
 			break;
 
@@ -169,7 +169,15 @@ void FPackageTrailerBuilder::AddPayload(const Virtualization::FPayloadId& Identi
 
 	if (Identifier.IsValid())
 	{
-		LocalEntries.Add(Identifier, LocalEntry(MoveTemp(Payload)));
+		LocalEntries.FindOrAdd(Identifier, LocalEntry(MoveTemp(Payload)));
+	}
+}
+
+void FPackageTrailerBuilder::AddVirtualizedPayload(const Virtualization::FPayloadId& Identifier, int64 RawSize)
+{
+	if (Identifier.IsValid())
+	{
+		VirtualizedEntries.FindOrAdd(Identifier, VirtualizedEntry(RawSize));
 	}
 }
 
@@ -226,7 +234,7 @@ bool FPackageTrailerBuilder::BuildAndAppendTrailer(FLinkerSave* Linker, FArchive
 		Private::FLookupTableEntry& Entry = Trailer.Header.PayloadLookupTable.AddDefaulted_GetRef();
 		Entry.Identifier = It.Key;
 		Entry.OffsetInFile = INDEX_NONE;
-		Entry.CompressedSize = It.Value.CompressedSize;
+		Entry.CompressedSize = INDEX_NONE;
 		Entry.RawSize = It.Value.RawSize;
 		Entry.AccessMode = EPayloadAccessMode::Virtualized;
 	}
@@ -443,6 +451,8 @@ bool FPackageTrailer::UpdatePayloadAsVirtualized(const Virtualization::FPayloadI
 	{
 		Entry->AccessMode = EPayloadAccessMode::Virtualized;
 		Entry->OffsetInFile = INDEX_NONE;
+		Entry->CompressedSize = INDEX_NONE; // Once the payload is virtualized we cannot be sure on the compression
+											// being used and so cannot know the compressed size
 		return true;
 	}
 	else
diff --git a/Engine/Source/Runtime/CoreUObject/Public/UObject/PackageTrailer.h b/Engine/Source/Runtime/CoreUObject/Public/UObject/PackageTrailer.h
--- a/Engine/Source/Runtime/CoreUObject/Public/UObject/PackageTrailer.h
+++ b/Engine/Source/Runtime/CoreUObject/Public/UObject/PackageTrailer.h
@@ -175,7 +175,7 @@ public:
 
 	// Methods that can be called while building the trailer
 
-	/*
+	/**
 	 * Adds a payload to the builder to be written to the trailer. Duplicate payloads will be discarded and only a 
 	 * single instance stored in the trailer.
 	 * 
@@ -184,8 +184,19 @@ public:
 	 * @param Callback		This callback will be invoked once the FPackageTrailer has been built and appended to disk.
 	 */
 	void AddPayload(const Virtualization::FPayloadId& Identifier, FCompressedBuffer Payload, AdditionalDataCallback&& Callback);
+
+	/**
+	 * Adds an already virtualized payload to the builder to be written to the trailer. When the trailer is written
+	 * the payload will have EPayloadAccessMode::Virtualized set as it's access mode. It is assumed that the payload
+	 * is already stored in the virtualization backends and it is up to the calling code to confirm this.
+	 * Duplicate payloads will be discarded and only a single instance stored in the trailer.
+	 * 
+	 * @param Identifier	The identifier of the payload
+	 * @param RawSize		The size of the payload (in bytes) when uncompressed
+	 */
+	void AddVirtualizedPayload(const Virtualization::FPayloadId& Identifier, int64 RawSize);
 	
-	/*
+	/**
 	 * @param ExportsArchive	The linker associated with the package being written to disk.
 	 * @param DataArchive		The archive where the package data has been written to. This is where the FPackageTrailer will be written to
 	 * @return True if the builder was created and appended successfully and false if any error was encountered
@@ -248,15 +259,13 @@ private:
 	struct VirtualizedEntry
 	{
 		VirtualizedEntry() = default;
-		VirtualizedEntry(int64 InCompressedSize, int64 InRawSize)
-			: CompressedSize(InCompressedSize)
-			, RawSize(InRawSize)
+		VirtualizedEntry(int64 InRawSize)
+			: RawSize(InRawSize)
 		{
 
 		}
 		~VirtualizedEntry() = default;
 
-		int64 CompressedSize = INDEX_NONE;
 		int64 RawSize = INDEX_NONE;
 	};
 
