AsyncLoading2: Disable check that all RF_Public objects have the LoadImport flag set, RF_Public can be added later to objects without export hashes that can never be imported
#rnx

diff --git a/Engine/Source/Runtime/CoreUObject/Private/Serialization/AsyncLoading2.cpp b/Engine/Source/Runtime/CoreUObject/Private/Serialization/AsyncLoading2.cpp
--- a/Engine/Source/Runtime/CoreUObject/Private/Serialization/AsyncLoading2.cpp
+++ b/Engine/Source/Runtime/CoreUObject/Private/Serialization/AsyncLoading2.cpp
@@ -988,7 +988,8 @@ private:
 		{
 			if (Object->HasAllFlags(RF_Public | RF_WasLoaded))
 			{
-				checkf(Object->HasAnyInternalFlags(EInternalObjectFlags::LoaderImport), TEXT("%s"), *Object->GetFullName());
+				// RF_Public can be added to objects after they've been created and in that case they might not have the LoaderImport flag set
+				//checkf(Object->HasAnyInternalFlags(EInternalObjectFlags::LoaderImport), TEXT("%s"), *Object->GetFullName());
 				Object->AtomicallyClearInternalFlags(EInternalObjectFlags::LoaderImport);
 			}
 		}, /* bIncludeNestedObjects*/ true);
