Serialize PKG_Cooked and PKG_ContainsNoAsset in the package flags in the FPackageFileSummary

#rnx

diff --git a/Engine/Source/Runtime/CoreUObject/Private/UObject/PackageFileSummary.cpp b/Engine/Source/Runtime/CoreUObject/Private/UObject/PackageFileSummary.cpp
--- a/Engine/Source/Runtime/CoreUObject/Private/UObject/PackageFileSummary.cpp
+++ b/Engine/Source/Runtime/CoreUObject/Private/UObject/PackageFileSummary.cpp
@@ -193,6 +193,12 @@ void operator<<(FStructuredArchive::FSlot Slot, FPackageFileSummary& Sum)
 		}
 		Record << SA_VALUE(TEXT("TotalHeaderSize"), Sum.TotalHeaderSize);
 		Record << SA_VALUE(TEXT("FolderName"), Sum.FolderName);
+
+		if (BaseArchive.IsCooking())
+		{
+			Sum.PackageFlags |= PKG_Cooked;
+		}
+
 		Record << SA_VALUE(TEXT("PackageFlags"), Sum.PackageFlags);
 
 		if (BaseArchive.IsLoading())
diff --git a/Engine/Source/Runtime/CoreUObject/Private/UObject/SavePackage2.cpp b/Engine/Source/Runtime/CoreUObject/Private/UObject/SavePackage2.cpp
--- a/Engine/Source/Runtime/CoreUObject/Private/UObject/SavePackage2.cpp
+++ b/Engine/Source/Runtime/CoreUObject/Private/UObject/SavePackage2.cpp
@@ -1822,6 +1822,7 @@ ESavePackageResult UpdatePackageHeader(FStructuredArchive::FRecord& StructuredAr
 		}
 	}
 	// Write Real Export Map
+	bool bContainsAsset = false;
 	if (!SaveContext.IsTextFormat())
 	{
 		check(Linker->Tell() == SaveContext.OffsetAfterImportMap);
@@ -1831,6 +1832,7 @@ ESavePackageResult UpdatePackageHeader(FStructuredArchive::FRecord& StructuredAr
 		for (FObjectExport& Export : Linker->ExportMap)
 		{
 			ExportTableStream.EnterElement() << Export;
+			bContainsAsset |= Export.bIsAsset;
 		}
 		check(Linker->Tell() == SaveContext.OffsetAfterExportMap);
 	}
@@ -1853,7 +1855,12 @@ ESavePackageResult UpdatePackageHeader(FStructuredArchive::FRecord& StructuredAr
 		Linker->LinkerRoot->ThisRequiresLocalizationGather(Linker->RequiresLocalizationGather());
 
 		// Update package flags from package, in case serialization has modified package flags.
-		Linker->Summary.SetPackageFlags(Linker->LinkerRoot->GetPackageFlags());
+		uint32 PackageFlags = Linker->LinkerRoot->GetPackageFlags();
+		if (!bContainsAsset)
+		{
+			PackageFlags |= PKG_ContainsNoAsset;
+		}
+		Linker->Summary.SetPackageFlags(PackageFlags);
 
 		// @todo: custom versions: when can this be checked?
 		{
@@ -2053,31 +2060,31 @@ void PostSavePackage(FSaveContext& SaveContext)
 	// Package has been saved, so unmark the NewlyCreated flag.
 	Package->ClearPackageFlags(PKG_NewlyCreated);
 
-	// Copy and modify the output SerializedPackageFlags from the PackageFlags written into the file
-	SaveContext.SerializedPackageFlags = SaveContext.GetLinker()->Summary.GetPackageFlags();
-	// Currently the PKG_ContainsNoAsset flag is not serialized as part of the summary
-	bool bContainsAsset = false;
+	// Copy and modify the output SerializedPackageFlags from the PackageFlags written into the default realm summary
+	uint32 SerializedPackageFlags = SaveContext.GetLinker()->Summary.GetPackageFlags();
+	// Consider all output packages when reflecting PKG_ContainsNoAsset to the single entry in SerializedPackageFlags and the asset registry
+	bool bContainsNoAsset = true;
 	for (FLinkerSave* Linker : SaveContext.GetLinkers())
 	{
 		// GetLinkers shouldn't return null linker
 		check(Linker);
-		for (FObjectExport& Export : Linker->ExportMap)
-		{
-			if (Export.bIsAsset)
-			{
-				bContainsAsset = true;
-				break;
-			}
-		}
+
+		const bool bLinkerContainsNoAsset = Linker->Summary.GetPackageFlags() & PKG_ContainsNoAsset;
+		bContainsNoAsset &= bLinkerContainsNoAsset;
 
 		// Call the linker post save callbacks
 		Linker->OnPostSave(SaveContext.GetTargetPackagePath(), FObjectPostSaveContext(SaveContext.GetObjectSaveContext()));
 	}
 
-	if (!bContainsAsset)
+	if (bContainsNoAsset)
+	{
+		SerializedPackageFlags |= PKG_ContainsNoAsset;
+	}
+	else
 	{
-		SaveContext.SerializedPackageFlags |= PKG_ContainsNoAsset;
+		SerializedPackageFlags &= ~PKG_ContainsNoAsset;
 	}
+	SaveContext.SerializedPackageFlags = SerializedPackageFlags;
 
 	// Notify the soft reference collector about our harvested soft references during save. 
 	// This is currently needed only for cooking which does not require editor-only references 
