Fix for FArrayProperty::ConvertFromType() to read the inner tag back when it was written.

diff --git a/Engine/Source/Runtime/CoreUObject/Private/UObject/PropertyArray.cpp b/Engine/Source/Runtime/CoreUObject/Private/UObject/PropertyArray.cpp
--- a/Engine/Source/Runtime/CoreUObject/Private/UObject/PropertyArray.cpp
+++ b/Engine/Source/Runtime/CoreUObject/Private/UObject/PropertyArray.cpp
@@ -782,13 +782,18 @@ EConvertFromTypeResult FArrayProperty::ConvertFromType(const FPropertyTag& Tag,
 		FScriptArrayHelper ScriptArrayHelper(this, ArrayPropertyData);
 		ScriptArrayHelper.EmptyAndAddValues(ElementCount);
 
+		FPropertyTag InnerPropertyTag;
+		InnerPropertyTag.Type = Tag.InnerType;
+		InnerPropertyTag.ArrayIndex = 0;
+
+		if (Slot.GetArchiveState().UEVer() >= VER_UE4_INNER_ARRAY_TAG_INFO && Tag.InnerType == NAME_StructProperty)
+		{
+			Slot.GetUnderlyingArchive() << InnerPropertyTag;
+		}
+
 		// Convert properties from old type to new type automatically if types are compatible (array case)
 		if (ElementCount > 0)
 		{
-			FPropertyTag InnerPropertyTag;
-			InnerPropertyTag.Type = Tag.InnerType;
-			InnerPropertyTag.ArrayIndex = 0;
-
 			FStructuredArchive::FStream ValueStream = Slot.EnterStream();
 
 			if (Inner->ConvertFromType(InnerPropertyTag, ValueStream.EnterElement(), ScriptArrayHelper.GetRawPtr(0), DefaultsStruct) == EConvertFromTypeResult::Converted)
