Reduced runtime UPackage size from 176b to 144b or 120b when stripping deprecated properties
This should saves around ~1.5 Mb to ~3 Mb at runtime with ~50000 packages

Added a new core define UE_STRIP_DEPRECATED_PROPERTIES that could be used to wrap deprecated properties and strip them to regain memory when a projects becomes compliant. this can be set in the project target file
Deprecated most public properties from UPackage and created accessors for them

diff --git a/Engine/Plugins/Developer/TraceSourceFiltering/Source/SourceFilteringTrace/Private/TraceWorldFiltering.cpp b/Engine/Plugins/Developer/TraceSourceFiltering/Source/SourceFilteringTrace/Private/TraceWorldFiltering.cpp
--- a/Engine/Plugins/Developer/TraceSourceFiltering/Source/SourceFilteringTrace/Private/TraceWorldFiltering.cpp
+++ b/Engine/Plugins/Developer/TraceSourceFiltering/Source/SourceFilteringTrace/Private/TraceWorldFiltering.cpp
@@ -151,7 +151,7 @@ void FTraceWorldFiltering::GetWorldDisplayString(const UWorld* InWorld, FString&
 		switch (InWorld->GetNetMode())
 		{
 		case NM_Client:
-			PostFix = FString::Printf(TEXT("Client %i"), InWorld->GetOutermost()->PIEInstanceID - 1);
+			PostFix = FString::Printf(TEXT("Client %i"), InWorld->GetOutermost()->GetPIEInstanceID() - 1);
 			break;
 		case NM_DedicatedServer:
 		case NM_ListenServer:
diff --git a/Engine/Plugins/MovieScene/ActorSequence/Source/ActorSequence/Private/ActorSequenceObjectReference.cpp b/Engine/Plugins/MovieScene/ActorSequence/Source/ActorSequence/Private/ActorSequenceObjectReference.cpp
--- a/Engine/Plugins/MovieScene/ActorSequence/Source/ActorSequence/Private/ActorSequenceObjectReference.cpp
+++ b/Engine/Plugins/MovieScene/ActorSequence/Source/ActorSequence/Private/ActorSequenceObjectReference.cpp
@@ -76,7 +76,7 @@ UObject* FActorSequenceObjectReference::Resolve(AActor* SourceActor) const
 		if (ActorId.IsValid())
 		{
 			// Fixup for PIE
-			int32 PIEInstanceID = SourceActor->GetOutermost()->PIEInstanceID;
+			int32 PIEInstanceID = SourceActor->GetOutermost()->GetPIEInstanceID();
 			FUniqueObjectGuid FixedUpId(ActorId);
 			if (PIEInstanceID != -1)
 			{
diff --git a/Engine/Plugins/MovieScene/LevelSequenceEditor/Source/LevelSequenceEditor/Private/Misc/LevelSequencePlaybackContext.cpp b/Engine/Plugins/MovieScene/LevelSequenceEditor/Source/LevelSequenceEditor/Private/Misc/LevelSequencePlaybackContext.cpp
--- a/Engine/Plugins/MovieScene/LevelSequenceEditor/Source/LevelSequenceEditor/Private/Misc/LevelSequencePlaybackContext.cpp
+++ b/Engine/Plugins/MovieScene/LevelSequenceEditor/Source/LevelSequenceEditor/Private/Misc/LevelSequencePlaybackContext.cpp
@@ -386,7 +386,7 @@ FText SLevelSequenceContextPicker::GetWorldDescription(const UWorld* World)
 		switch(World->GetNetMode())
 		{
 		case NM_Client:
-			PostFix = FText::Format(LOCTEXT("ClientPostfixFormat", " (Client {0})"), FText::AsNumber(World->GetOutermost()->PIEInstanceID - 1));
+			PostFix = FText::Format(LOCTEXT("ClientPostfixFormat", " (Client {0})"), FText::AsNumber(World->GetOutermost()->GetPIEInstanceID() - 1));
 			break;
 		case NM_DedicatedServer:
 		case NM_ListenServer:
diff --git a/Engine/Plugins/VirtualProduction/LevelSnapshots/Source/LevelSnapshotsEditor/Private/Views/Input/SLevelSnapshotsEditorContextPicker.cpp b/Engine/Plugins/VirtualProduction/LevelSnapshots/Source/LevelSnapshotsEditor/Private/Views/Input/SLevelSnapshotsEditorContextPicker.cpp
--- a/Engine/Plugins/VirtualProduction/LevelSnapshots/Source/LevelSnapshotsEditor/Private/Views/Input/SLevelSnapshotsEditorContextPicker.cpp
+++ b/Engine/Plugins/VirtualProduction/LevelSnapshots/Source/LevelSnapshotsEditor/Private/Views/Input/SLevelSnapshotsEditorContextPicker.cpp
@@ -141,7 +141,7 @@ FText SLevelSnapshotsEditorContextPicker::GetWorldDescription(const UWorld* Worl
 		switch(World->GetNetMode())
 		{
 		case NM_Client:
-			PostFix = FText::Format(LOCTEXT("ClientPostfixFormat", " (Client {0})"), FText::AsNumber(World->GetOutermost()->PIEInstanceID - 1));
+			PostFix = FText::Format(LOCTEXT("ClientPostfixFormat", " (Client {0})"), FText::AsNumber(World->GetOutermost()->GetPIEInstanceID() - 1));
 			break;
 		case NM_DedicatedServer:
 		case NM_ListenServer:
diff --git a/Engine/Source/Editor/UnrealEd/Private/Commandlets/DiffFilesCommandlet.cpp b/Engine/Source/Editor/UnrealEd/Private/Commandlets/DiffFilesCommandlet.cpp
--- a/Engine/Source/Editor/UnrealEd/Private/Commandlets/DiffFilesCommandlet.cpp
+++ b/Engine/Source/Editor/UnrealEd/Private/Commandlets/DiffFilesCommandlet.cpp
@@ -142,7 +142,7 @@ void UDiffFilesCommandlet::LoadAndDiff()
 
 	Package = LoadPackage(Package, *FString::Printf(TEXT("%s;%s"), *PackageInfos[0].FullPath, *PackageInfos[1].FullPath), LOAD_ForDiff | LOAD_ForFileDiff);
 
-	if (Package->LinkerLoad->IsTextFormat())
+	if (Package->GetLinker()->IsTextFormat())
 	{
 		UE_LOG(LogDiffFilesCommandlet, Warning, TEXT("FileDiffs are not currently supported for text based assets"));
 	}
@@ -151,6 +151,6 @@ void UDiffFilesCommandlet::LoadAndDiff()
 		FAssetRegistryModule& AssetRegistryModule = FModuleManager::LoadModuleChecked<FAssetRegistryModule>(TEXT("AssetRegistry"));
 		IAssetRegistry& AssetRegistry = AssetRegistryModule.Get();
 		IAssetRegistry::FLoadPackageRegistryData AssetData;
-		AssetRegistry.LoadPackageRegistryData(*Package->LinkerLoad->GetLoader_Unsafe(), AssetData);
+		AssetRegistry.LoadPackageRegistryData(*Package->GetLinker()->GetLoader_Unsafe(), AssetData);
 	}
 }
diff --git a/Engine/Source/Editor/UnrealEd/Private/Commandlets/PackageUtilities.cpp b/Engine/Source/Editor/UnrealEd/Private/Commandlets/PackageUtilities.cpp
--- a/Engine/Source/Editor/UnrealEd/Private/Commandlets/PackageUtilities.cpp
+++ b/Engine/Source/Editor/UnrealEd/Private/Commandlets/PackageUtilities.cpp
@@ -1629,7 +1629,7 @@ int32 UPkgInfoCommandlet::Main( const FString& Params )
 				if (LoadedPackage)
 				{
 					check(LoadedPackage == Package);
-					Linker = Package->LinkerLoad;
+					Linker = Package->GetLinker();
 					check(Linker);
 				}
 				else
diff --git a/Engine/Source/Editor/UnrealEd/Private/Commandlets/WorldPartitionConvertCommandlet.cpp b/Engine/Source/Editor/UnrealEd/Private/Commandlets/WorldPartitionConvertCommandlet.cpp
--- a/Engine/Source/Editor/UnrealEd/Private/Commandlets/WorldPartitionConvertCommandlet.cpp
+++ b/Engine/Source/Editor/UnrealEd/Private/Commandlets/WorldPartitionConvertCommandlet.cpp
@@ -1223,9 +1223,9 @@ int32 UWorldPartitionConvertCommandlet::Main(const FString& Params)
 			ALevelInstance* LevelInstanceActor = MainWorld->SpawnActor<ALevelInstance>(SpawnParams);
 			
 			FTransform LevelTransform;
-			if (SubLevelPackage->WorldTileInfo)
+			if (SubLevelPackage->GetWorldTileInfo())
 			{
-				LevelTransform = FTransform(FVector(SubLevelPackage->WorldTileInfo->Position));
+				LevelTransform = FTransform(FVector(SubLevelPackage->GetWorldTileInfo()->Position));
 			}
 			else
 			{
diff --git a/Engine/Source/Editor/UnrealEd/Private/CookOnTheFlyServer.cpp b/Engine/Source/Editor/UnrealEd/Private/CookOnTheFlyServer.cpp
--- a/Engine/Source/Editor/UnrealEd/Private/CookOnTheFlyServer.cpp
+++ b/Engine/Source/Editor/UnrealEd/Private/CookOnTheFlyServer.cpp
@@ -2877,11 +2877,11 @@ void UCookOnTheFlyServer::ReleaseCookedPlatformData(UE::Cook::FPackageData& Pack
 		if (CurrentCookMode == ECookMode::CookByTheBook)
 		{
 			UPackage* Package = PackageData.GetPackage();
-			if (Package && Package->LinkerLoad)
+			if (Package && Package->GetLinker())
 			{
 				// Loaders and their handles can have large buffers held in process memory and in the system file cache from the
 				// data that was loaded.  Keeping this for the lifetime of the cook is costly, so we try and unload it here.
-				Package->LinkerLoad->FlushCache();
+				Package->GetLinker()->FlushCache();
 			}
 		}
 	}
@@ -4106,7 +4106,7 @@ void UCookOnTheFlyServer::SaveCookedPackage(UE::Cook::FSaveCookedPackageContext&
 			SaveArgs.TopLevelFlags = Context.FlagsToCook;
 			SaveArgs.bForceByteSwapping = Context.bEndianSwap;
 			SaveArgs.bWarnOfLongFilename = false;
-			SaveArgs.SaveFlags = Context.SaveFlags;
+			SaveArgs.SaveFlags = Context.SaveFlags | SAVE_Optional;
 			SaveArgs.TargetPlatform = TargetPlatform;
 			SaveArgs.bSlowTask = false;
 			SaveArgs.SavePackageContext = Context.SavePackageContext;
@@ -8850,7 +8850,7 @@ uint32 UCookOnTheFlyServer::FullLoadAndSave(uint32& CookedPackageCount)
 						SaveArgs.TopLevelFlags = FlagsToCook;
 						SaveArgs.bForceByteSwapping = bSwap;
 						SaveArgs.bWarnOfLongFilename = false;
-						SaveArgs.SaveFlags = SaveFlags;
+						SaveArgs.SaveFlags = SaveFlags | SAVE_Optional;
 						SaveArgs.TargetPlatform = Target;
 						SaveArgs.bSlowTask = false;
 						SaveArgs.SavePackageContext = &SavePackageContext;
diff --git a/Engine/Source/Editor/UnrealEd/Private/ObjectTools.cpp b/Engine/Source/Editor/UnrealEd/Private/ObjectTools.cpp
--- a/Engine/Source/Editor/UnrealEd/Private/ObjectTools.cpp
+++ b/Engine/Source/Editor/UnrealEd/Private/ObjectTools.cpp
@@ -816,9 +816,9 @@ namespace ObjectTools
 			}
 
 			// When duplicating a World Composition map, make sure to properly initialize WorldTileInfo
-			if (Object->GetOutermost()->WorldTileInfo.IsValid())
+			if (Object->GetOutermost()->GetWorldTileInfo())
 			{
-				DupObject->GetOutermost()->WorldTileInfo = MakeUnique<FWorldTileInfo>(*Object->GetOutermost()->WorldTileInfo);
+				DupObject->GetOutermost()->SetWorldTileInfo(MakeUnique<FWorldTileInfo>(*Object->GetOutermost()->GetWorldTileInfo()));
 			}
 
 			// Notify the asset registry
@@ -3698,9 +3698,9 @@ namespace ObjectTools
 						}
 
 						// When renaming a World Composition map, make sure to properly initialize WorldTileInfo
-						if (Object->GetOutermost()->WorldTileInfo.IsValid())
+						if (Object->GetOutermost()->GetWorldTileInfo())
 						{
-							NewPackage->WorldTileInfo = MakeUnique<FWorldTileInfo>(*Object->GetOutermost()->WorldTileInfo);
+							NewPackage->SetWorldTileInfo(MakeUnique<FWorldTileInfo>(*Object->GetOutermost()->GetWorldTileInfo()));
 						}
 
 						Redirector = Cast<UObjectRedirector>(StaticFindObject(UObjectRedirector::StaticClass(), NewPackage, *NewObjectName));
@@ -3826,14 +3826,14 @@ namespace ObjectTools
 			FString OldObjectPathName = Object->GetPathName();
 			GEditor->RenameObject( Object, NewPackage, *ObjName, bLeaveRedirector ? REN_None : REN_DontCreateRedirectors );
 
-			if (OldPackage && OldPackage->MetaData)
+			if (OldPackage && OldPackage->HasMetaData())
 			{
 				// Migrate the localization ID to the new package
 				TextNamespaceUtil::ForcePackageNamespace(NewPackage, TextNamespaceUtil::GetPackageNamespace(OldPackage));
 				TextNamespaceUtil::ClearPackageNamespace(OldPackage);
 
 				// Remove any metadata from old package pointing to moved objects
-				OldPackage->MetaData->RemoveMetaDataOutsidePackage();
+				OldPackage->GetMetaData()->RemoveMetaDataOutsidePackage();
 			}
 
 			// Migrate any thumbnail from the old package to the new one
@@ -4696,30 +4696,30 @@ namespace ThumbnailTools
 	 */
 	FObjectThumbnail* CacheThumbnail( const FString& ObjectFullName, FObjectThumbnail* Thumbnail, UPackage* DestPackage )
 	{
-		FObjectThumbnail* Result = NULL;
+		FObjectThumbnail* Result = nullptr;
 
-		if ( ObjectFullName.Len() > 0 && DestPackage != NULL )
+		if ( ObjectFullName.Len() > 0 && DestPackage != nullptr)
 		{
 			// Create a new thumbnail map if we don't have one already
-			if( !DestPackage->ThumbnailMap )
+			if( !DestPackage->HasThumbnailMap() )
 			{
-				DestPackage->ThumbnailMap = MakeUnique<FThumbnailMap>();
+				DestPackage->SetThumbnailMap(MakeUnique<FThumbnailMap>());
 			}
 
 			// @todo thumbnails: Backwards compat
 			FName ObjectFullNameFName( *ObjectFullName );
-			FObjectThumbnail* CachedThumbnail = DestPackage->ThumbnailMap->Find( ObjectFullNameFName );
-			if ( Thumbnail != NULL )
+			const FObjectThumbnail* CachedThumbnail = DestPackage->GetThumbnailMap().Find( ObjectFullNameFName );
+			if ( Thumbnail != nullptr )
 			{
 				// Cache the thumbnail (possibly replacing an existing thumb!)
-				Result = &DestPackage->ThumbnailMap->Add( ObjectFullNameFName, *Thumbnail );
+				Result = &DestPackage->AccessThumbnailMap().Add( ObjectFullNameFName, *Thumbnail );
 			}
 			//only let thumbnails loaded from disk to be removed.
 			//When capturing thumbnails from the content browser, it will only exist in memory until it is saved out to a package.
 			//Don't let the recycling purge them
-			else if ((CachedThumbnail != NULL) && (CachedThumbnail->IsLoadedFromDisk()))
+			else if ((CachedThumbnail != nullptr) && (CachedThumbnail->IsLoadedFromDisk()))
 			{
-				DestPackage->ThumbnailMap->Remove( ObjectFullNameFName );
+				DestPackage->AccessThumbnailMap().Remove( ObjectFullNameFName );
 			}
 
 		}
diff --git a/Engine/Source/Editor/UnrealEd/Private/PackageTools.cpp b/Engine/Source/Editor/UnrealEd/Private/PackageTools.cpp
--- a/Engine/Source/Editor/UnrealEd/Private/PackageTools.cpp
+++ b/Engine/Source/Editor/UnrealEd/Private/PackageTools.cpp
@@ -491,7 +491,7 @@ UPackageTools::UPackageTools(const FObjectInitializer& ObjectInitializer)
 			TArray<TWeakObjectPtr<UMetaData>> PackageMetaDataWithClearedStandaloneFlag;
 			for (UPackage* PackageToUnload : PackagesToUnload)
 			{
-				UMetaData* PackageMetaData = PackageToUnload ? PackageToUnload->MetaData : nullptr;
+				UMetaData* PackageMetaData = PackageToUnload ? PackageToUnload->GetMetaData() : nullptr;
 				if ( PackageMetaData && PackageMetaData->HasAnyFlags(RF_Standalone) )
 				{
 					PackageMetaData->ClearFlags(RF_Standalone);
diff --git a/Engine/Source/Editor/UnrealEd/Private/PlayLevel.cpp b/Engine/Source/Editor/UnrealEd/Private/PlayLevel.cpp
--- a/Engine/Source/Editor/UnrealEd/Private/PlayLevel.cpp
+++ b/Engine/Source/Editor/UnrealEd/Private/PlayLevel.cpp
@@ -2272,7 +2272,7 @@ UWorld* UEditorEngine::CreatePIEWorldByDuplication(FWorldContext &WorldContext,
 
 	UPackage* PlayWorldPackage = CreatePackage(*PlayWorldMapName);
 	PlayWorldPackage->SetPackageFlags(PKG_PlayInEditor);
-	PlayWorldPackage->PIEInstanceID = WorldContext.PIEInstance;
+	PlayWorldPackage->SetPIEInstanceID(WorldContext.PIEInstance);
 	PlayWorldPackage->SetLoadedPath(InPackage->GetLoadedPath());
 	PRAGMA_DISABLE_DEPRECATION_WARNINGS
 	PlayWorldPackage->SetGuid( InPackage->GetGuid() );
@@ -2365,7 +2365,7 @@ UWorld* UEditorEngine::CreatePIEWorldFromEntry(FWorldContext &WorldContext, UWor
 	check(LoadedWorld);
 	if (LoadedWorld->GetOutermost() != GetTransientPackage())
 	{
-		LoadedWorld->GetOutermost()->PIEInstanceID = WorldContext.PIEInstance;
+		LoadedWorld->GetOutermost()->SetPIEInstanceID(WorldContext.PIEInstance);
 	}
 	// Force default GameMode class so project specific code doesn't fire off. 
 	// We want this world to truly remain empty while we wait for connect!
diff --git a/Engine/Source/Editor/UnrealEd/Private/UnrealEdSrv.cpp b/Engine/Source/Editor/UnrealEd/Private/UnrealEdSrv.cpp
--- a/Engine/Source/Editor/UnrealEd/Private/UnrealEdSrv.cpp
+++ b/Engine/Source/Editor/UnrealEd/Private/UnrealEdSrv.cpp
@@ -299,9 +299,9 @@ UPackage* UUnrealEdEngine::GeneratePackageThumbnailsIfRequired( const TCHAR* Str
 			Packages.Add( Pkg );
 
 			// Allocate a new thumbnail map if we need one
-			if( !Pkg->ThumbnailMap )
+			if( !Pkg->HasThumbnailMap() )
 			{
-				Pkg->ThumbnailMap = MakeUnique<FThumbnailMap>();
+				Pkg->SetThumbnailMap(MakeUnique<FThumbnailMap>());
 			}
 
 			// OK, now query all of the browsable objects in the package we're about to save
diff --git a/Engine/Source/Runtime/Core/Public/Misc/CoreMiscDefines.h b/Engine/Source/Runtime/Core/Public/Misc/CoreMiscDefines.h
--- a/Engine/Source/Runtime/Core/Public/Misc/CoreMiscDefines.h
+++ b/Engine/Source/Runtime/Core/Public/Misc/CoreMiscDefines.h
@@ -235,6 +235,16 @@ enum EInPlace {InPlace};
 	#define UE_DEPRECATED_FORGAME(...)
 #endif
 
+/*
+ * Macro that can be defined in the target file to strip deprecated properties in objects across the engine that check against this define.
+ * Can be used by project that have migrated away from using deprecated functions and data members to potentially gain back some memory and perf.
+ * @note This is a define that engine developer may use when deprecating properties to allow additional memory savings when a project is compliant with deprecation notice.
+ * It doesn't indicate that all deprecated properties will be stripped.
+ */
+#ifndef UE_STRIP_DEPRECATED_PROPERTIES
+	#define UE_STRIP_DEPRECATED_PROPERTIES 0
+#endif
+
 template <bool bIsDeprecated>
 struct TStaticDeprecateExpression
 {
diff --git a/Engine/Source/Runtime/CoreUObject/Private/Blueprint/BlueprintSupport.cpp b/Engine/Source/Runtime/CoreUObject/Private/Blueprint/BlueprintSupport.cpp
--- a/Engine/Source/Runtime/CoreUObject/Private/Blueprint/BlueprintSupport.cpp
+++ b/Engine/Source/Runtime/CoreUObject/Private/Blueprint/BlueprintSupport.cpp
@@ -241,7 +241,7 @@ void FBlueprintSupport::RepairDeferredDependenciesInObject(UObject* Object)
 			{
 				if (const UPackage* PlaceholderPackage = PropertyValueAsObj->GetPackage())
 				{
-					if (FLinkerLoad* PlaceholderLinker = PlaceholderPackage->LinkerLoad)
+					if (FLinkerLoad* PlaceholderLinker = PlaceholderPackage->GetLinker())
 					{
 						int32 const ImportIndex = Placeholder->PackageIndex.ToImport();
 						FObjectImport& Import = PlaceholderLinker->ImportMap[ImportIndex];
diff --git a/Engine/Source/Runtime/CoreUObject/Private/Serialization/AsyncLoading.cpp b/Engine/Source/Runtime/CoreUObject/Private/Serialization/AsyncLoading.cpp
--- a/Engine/Source/Runtime/CoreUObject/Private/Serialization/AsyncLoading.cpp
+++ b/Engine/Source/Runtime/CoreUObject/Private/Serialization/AsyncLoading.cpp
@@ -1760,7 +1760,7 @@ EAsyncPackageState::Type FAsyncPackage::LoadImports_Event()
 		if (Import->XObject)
 		{
 			ExistingPackage = CastChecked<UPackage>(Import->XObject->GetPackage());
-			PendingPackage = ExistingPackage->LinkerLoad ? static_cast<FAsyncPackage*>(ExistingPackage->LinkerLoad->AsyncRoot) : nullptr;
+			PendingPackage = ExistingPackage->GetLinker() ? static_cast<FAsyncPackage*>(ExistingPackage->GetLinker()->AsyncRoot) : nullptr;
 		}
 		const bool bCompiledInNotDynamic = IsNativeCodePackage(ExistingPackage);
 		// Our import package name is the import name
@@ -1983,7 +1983,7 @@ EAsyncPackageState::Type FAsyncPackage::SetupImports_Event()
 				{
 					ImportedPackages.Add(ImportPackage);
 
-					FLinkerLoad* ImportLinker = ImportPackage->LinkerLoad;
+					FLinkerLoad* ImportLinker = ImportPackage->GetLinker();
 					if (ImportLinker && ImportLinker->AsyncRoot)
 					{
 						check(ImportLinker->AsyncRoot != this);
@@ -2034,7 +2034,7 @@ EAsyncPackageState::Type FAsyncPackage::SetupImports_Event()
 			bool bWaitingForImport = false;
 			if (ImportPackage)
 			{
-				FLinkerLoad* ImportLinker = ImportPackage->LinkerLoad;
+				FLinkerLoad* ImportLinker = ImportPackage->GetLinker();
 
 #if USE_EVENT_DRIVEN_ASYNC_LOAD_AT_BOOT_TIME
 				if (GIsInitialLoad && !ImportLinker && ImportPackage->HasAnyPackageFlags(PKG_CompiledIn))
@@ -2572,20 +2572,20 @@ void FAsyncPackage::DumpDependencies(const TCHAR* Label, UObject* Obj)
 	}
 	UE_LOG(LogStreaming, Error, TEXT("    Obj is %s"), *Obj->GetFullName());
 	UPackage* Package = Obj->GetOutermost();
-	if (!Package->LinkerLoad)
+	if (!Package->GetLinker())
 	{
 		UE_LOG(LogStreaming, Error, TEXT("    %s has no linker"), *Package->GetFullName());
 	}
 	else
 	{
-		for (int32 LocalExportIndex = 0; LocalExportIndex < Package->LinkerLoad->ExportMap.Num(); LocalExportIndex++)
+		for (int32 LocalExportIndex = 0; LocalExportIndex < Package->GetLinker()->ExportMap.Num(); LocalExportIndex++)
 		{
-			FObjectExport& Export = Package->LinkerLoad->ExportMap[LocalExportIndex];
+			FObjectExport& Export = Package->GetLinker()->ExportMap[LocalExportIndex];
 			if (Export.Object == Obj || Export.Object == nullptr)
 			{
 				if (Export.ObjectName == Obj->GetFName())
 				{
-					DumpDependencies(TEXT(""), Package->LinkerLoad, FPackageIndex::FromExport(LocalExportIndex));
+					DumpDependencies(TEXT(""), Package->GetLinker(), FPackageIndex::FromExport(LocalExportIndex));
 				}
 			}
 		}
@@ -5450,7 +5450,7 @@ EAsyncPackageState::Type FAsyncPackage::TickAsyncPackage(bool InbUseTimeLimit, b
 				// We can only continue to PostLoad if all imported packages finished serializing their exports
 				for (UPackage* ImportedPackage : ImportedPackages)
 				{
-					if (ImportedPackage && ImportedPackage->LinkerLoad && ImportedPackage->LinkerLoad->AsyncRoot && !static_cast<FAsyncPackage*>(ImportedPackage->LinkerLoad->AsyncRoot)->bAllExportsSerialized)
+					if (ImportedPackage && ImportedPackage->GetLinker() && ImportedPackage->GetLinker()->AsyncRoot && !static_cast<FAsyncPackage*>(ImportedPackage->GetLinker()->AsyncRoot)->bAllExportsSerialized)
 					{
 						LoadingState = EAsyncPackageState::PendingImports;
 						break;
@@ -5566,7 +5566,7 @@ EAsyncPackageState::Type FAsyncPackage::CreateLinker()
 
 			// Set package specific data 
 			Package->SetPackageFlags(Desc.PackageFlags);
-			Package->PIEInstanceID = Desc.PIEInstanceID;
+			Package->SetPIEInstanceID(Desc.PIEInstanceID);
 
 			// Always store package filename we loading from
 			Package->SetLoadedPath(Desc.PackagePath);
@@ -6053,7 +6053,7 @@ EAsyncPackageState::Type FAsyncPackage::CreateImports()
 		// Keep track of all imported packages that are also being loaded so that we can wait until they also finished serializing their exports
 		if (UPackage* ImportedPackage = Cast<UPackage>(Object))
 		{
-			if (ImportedPackage->LinkerLoad && ImportedPackage->LinkerLoad->AsyncRoot)
+			if (ImportedPackage->GetLinker() && ImportedPackage->GetLinker()->AsyncRoot)
 			{
 				ImportedPackages.Add(ImportedPackage);
 			}
diff --git a/Engine/Source/Runtime/CoreUObject/Private/Serialization/AsyncLoading2.cpp b/Engine/Source/Runtime/CoreUObject/Private/Serialization/AsyncLoading2.cpp
--- a/Engine/Source/Runtime/CoreUObject/Private/Serialization/AsyncLoading2.cpp
+++ b/Engine/Source/Runtime/CoreUObject/Private/Serialization/AsyncLoading2.cpp
@@ -3591,12 +3591,12 @@ EAsyncPackageState::Type FAsyncPackage2::Event_ProcessExportBundle(FAsyncLoading
 		}
 		FExportArchive Ar(Package->AllExportDataPtr, Package->CurrentExportDataPtr, AllExportDataSize);
 		{
-			Ar.SetUEVer(Package->LinkerRoot->LinkerPackageVersion);
-			Ar.SetLicenseeUEVer(Package->LinkerRoot->LinkerLicenseeVersion);
+			Ar.SetUEVer(Package->LinkerRoot->GetLinkerPackageVersion());
+			Ar.SetLicenseeUEVer(Package->LinkerRoot->GetLinkerLicenseeVersion());
 			// Ar.SetEngineVer(Summary.SavedByEngineVersion); // very old versioning scheme
-			if (!Package->LinkerRoot->LinkerCustomVersion.GetAllVersions().IsEmpty())
+			if (!Package->LinkerRoot->GetLinkerCustomVersions().GetAllVersions().IsEmpty())
 			{
-				Ar.SetCustomVersions(Package->LinkerRoot->LinkerCustomVersion);
+				Ar.SetCustomVersions(Package->LinkerRoot->GetLinkerCustomVersions());
 			}
 			Ar.SetUseUnversionedPropertySerialization((Package->LinkerRoot->GetPackageFlags() & PKG_UnversionedProperties) != 0);
 			Ar.SetIsLoading(true);
@@ -5750,14 +5750,14 @@ void FAsyncPackage2::CreateUPackage(const FZenPackageSummary* PackageSummary, co
 		LinkerRoot->SetPackageFlagsTo(PackageSummary->PackageFlags | PKG_Cooked);
 		if (VersioningInfo)
 		{
-			LinkerRoot->LinkerPackageVersion = VersioningInfo->PackageVersion;
-			LinkerRoot->LinkerLicenseeVersion = VersioningInfo->LicenseeVersion;
-			LinkerRoot->LinkerCustomVersion = VersioningInfo->CustomVersions;
+			LinkerRoot->SetLinkerPackageVersion(VersioningInfo->PackageVersion);
+			LinkerRoot->SetLinkerLicenseeVersion(VersioningInfo->LicenseeVersion);
+			LinkerRoot->SetLinkerCustomVersions(VersioningInfo->CustomVersions);
 		}
 		else
 		{
-			LinkerRoot->LinkerPackageVersion = GPackageFileUEVersion;
-			LinkerRoot->LinkerLicenseeVersion = GPackageFileLicenseeUEVersion;
+			LinkerRoot->SetLinkerPackageVersion(GPackageFileUEVersion);
+			LinkerRoot->SetLinkerLicenseeVersion(GPackageFileLicenseeUEVersion);
 		}
 #if WITH_IOSTORE_IN_EDITOR
 		LinkerRoot->bIsCookedForEditor = !!(PackageSummary->PackageFlags & PKG_FilterEditorOnly);
@@ -5772,8 +5772,8 @@ void FAsyncPackage2::CreateUPackage(const FZenPackageSummary* PackageSummary, co
 		check(LinkerRoot->CanBeImported() == Desc.bCanBeImported);
 		check(LinkerRoot->GetPackageId() == Desc.UPackageId);
 		check(LinkerRoot->GetPackageFlags() == (PackageSummary->PackageFlags | PKG_Cooked));
-		check(LinkerRoot->LinkerPackageVersion == GPackageFileUEVersion);
-		check(LinkerRoot->LinkerLicenseeVersion == GPackageFileLicenseeUEVersion);
+		check(LinkerRoot->GetLinkerPackageVersion() == GPackageFileUEVersion);
+		check(LinkerRoot->GetLinkerLicenseeVersion() == GPackageFileLicenseeUEVersion);
 		check(LinkerRoot->HasAnyFlags(RF_WasLoaded));
 	}
 
diff --git a/Engine/Source/Runtime/CoreUObject/Private/Serialization/BulkData.cpp b/Engine/Source/Runtime/CoreUObject/Private/Serialization/BulkData.cpp
--- a/Engine/Source/Runtime/CoreUObject/Private/Serialization/BulkData.cpp
+++ b/Engine/Source/Runtime/CoreUObject/Private/Serialization/BulkData.cpp
@@ -419,7 +419,7 @@ bool FUntypedBulkData::CanLoadFromDisk() const
 	}
 	else if (UPackage* PackagePtr = Package.Get())
 	{
-		bCanLoadFromDisk = (PackagePtr->LinkerLoad != nullptr);
+		bCanLoadFromDisk = (PackagePtr->GetLinker() != nullptr);
 	}
 	return bCanLoadFromDisk;
 #endif // WITH_EDITOR
@@ -2222,9 +2222,9 @@ bool FUntypedBulkData::TryLoadDataIntoMemory(void* Dest)
 	{
 		if (UPackage* PackagePtr = Package.Get())
 		{
-			if (PackagePtr->LinkerLoad && PackagePtr->LinkerLoad->GetOwnerThreadId() == FPlatformTLS::GetCurrentThreadId())
+			if (PackagePtr->GetLinker() && PackagePtr->GetLinker()->GetOwnerThreadId() == FPlatformTLS::GetCurrentThreadId())
 			{
-				FLinkerLoad* LinkerLoad = PackagePtr->LinkerLoad;
+				FLinkerLoad* LinkerLoad = PackagePtr->GetLinker();
 				if (LinkerLoad && LinkerLoad->HasLoader())
 				{
 					FArchive* Ar = LinkerLoad;
diff --git a/Engine/Source/Runtime/CoreUObject/Private/UObject/DeferredRegistry.cpp b/Engine/Source/Runtime/CoreUObject/Private/UObject/DeferredRegistry.cpp
--- a/Engine/Source/Runtime/CoreUObject/Private/UObject/DeferredRegistry.cpp
+++ b/Engine/Source/Runtime/CoreUObject/Private/UObject/DeferredRegistry.cpp
@@ -67,7 +67,7 @@ void ReloadProcessObject(UClass* Class, const TCHAR* RenamePrefix)
 void ReloadProcessObject(UPackage* Package, const TCHAR* RenamePrefix)
 {
 #if WITH_LIVE_CODING
-	for (UFunction* Function : Package->Delegates)
+	for (UFunction* Function : Package->GetReloadDelegates())
 	{
 		if (Function->SingletonPtr != nullptr)
 		{
diff --git a/Engine/Source/Runtime/CoreUObject/Private/UObject/LazyObjectPtr.cpp b/Engine/Source/Runtime/CoreUObject/Private/UObject/LazyObjectPtr.cpp
--- a/Engine/Source/Runtime/CoreUObject/Private/UObject/LazyObjectPtr.cpp
+++ b/Engine/Source/Runtime/CoreUObject/Private/UObject/LazyObjectPtr.cpp
@@ -147,7 +147,7 @@ void FLazyObjectPtr::PossiblySerializeObjectGuid(UObject *Object, FStructuredArc
 #if WITH_EDITOR
 						else if (Object->GetOutermostObject()->GetPackage()->HasAnyPackageFlags(PKG_PlayInEditor))
 						{
-							Guid = RemapGuid(Guid, Object->GetOutermostObject()->GetPackage()->PIEInstanceID);
+							Guid = RemapGuid(Guid, Object->GetOutermostObject()->GetPackage()->GetPIEInstanceID());
 							GuidAnnotation.AddAnnotation(Object, Guid);
 						}
 #endif
diff --git a/Engine/Source/Runtime/CoreUObject/Private/UObject/LinkerLoad.cpp b/Engine/Source/Runtime/CoreUObject/Private/UObject/LinkerLoad.cpp
--- a/Engine/Source/Runtime/CoreUObject/Private/UObject/LinkerLoad.cpp
+++ b/Engine/Source/Runtime/CoreUObject/Private/UObject/LinkerLoad.cpp
@@ -504,7 +504,7 @@ FLinkerLoad* FLinkerLoad::FindExistingLinkerForPackage(const UPackage* Package)
 	FLinkerLoad* Linker = nullptr;
 	if (Package)
 	{
-		Linker = Package->LinkerLoad;
+		Linker = Package->GetLinker();
 	}
 	return Linker;
 }
@@ -644,7 +644,7 @@ FLinkerLoad* FLinkerLoad::CreateLinkerAsync(FUObjectSerializeContext* LoadContex
 	}
 
 	// Create a new linker if there isn't an existing one.
-	if( Linker == NULL )
+	if( Linker == nullptr )
 	{
 		if (GEventDrivenLoaderEnabled && FApp::IsGame() && !GIsEditor)
 		{
@@ -652,14 +652,14 @@ FLinkerLoad* FLinkerLoad::CreateLinkerAsync(FUObjectSerializeContext* LoadContex
 		}
 		Linker = new FLinkerLoad(Parent, PackagePath, LoadFlags, InstancingContext ? *InstancingContext : FLinkerInstancingContext());
 		Linker->SetSerializeContext(LoadContext);
-		Parent->LinkerLoad = Linker;
+		Parent->SetLinker(Linker);
 		if (GEventDrivenLoaderEnabled && Linker)
 		{
 			Linker->CreateLoader(Forward<TFunction<void()>>(InSummaryReadyCallback));
 		}
 	}
 	
-	check(Parent->LinkerLoad == Linker);
+	check(Parent->GetLinker() == Linker);
 
 	return Linker;
 }
@@ -828,7 +828,7 @@ FLinkerLoad::ELinkerStatus FLinkerLoad::Tick( float InTimeLimit, bool bInUseTime
 
 	if (Status == LINKER_Failed)
 	{
-		LinkerRoot->LinkerLoad = nullptr;
+		LinkerRoot->SetLinker(nullptr);
 #if WITH_EDITOR
 
 		if (LoadProgressScope)
@@ -1474,7 +1474,7 @@ FLinkerLoad::ELinkerStatus FLinkerLoad::UpdateFromPackageFileSummary()
 		LinkerRootPackage->SetChunkIDs(Summary.ChunkIDs);
 
 		// Propagate package file size
-		LinkerRootPackage->FileSize = Loader ? Loader->TotalSize() : 0;
+		LinkerRootPackage->SetFileSize(Loader ? Loader->TotalSize() : 0);
 
 		// Propagate package Guids
 		PRAGMA_DISABLE_DEPRECATION_WARNINGS
@@ -1486,14 +1486,14 @@ FLinkerLoad::ELinkerStatus FLinkerLoad::UpdateFromPackageFileSummary()
 #endif
 
 		// Remember the linker versions
-		LinkerRootPackage->LinkerPackageVersion = Summary.GetFileVersionUE();
-		LinkerRootPackage->LinkerLicenseeVersion = Summary.GetFileVersionLicenseeUE();
+		LinkerRootPackage->SetLinkerPackageVersion(Summary.GetFileVersionUE());
+		LinkerRootPackage->SetLinkerLicenseeVersion(Summary.GetFileVersionLicenseeUE());
 
 		// Only set the custom version if it is not already latest.
 		// If it is latest, we will compare against latest in GetLinkerCustomVersion
 		if (!bCustomVersionIsLatest)
 		{
-			LinkerRootPackage->LinkerCustomVersion = SummaryVersions;
+			LinkerRootPackage->SetLinkerCustomVersions(SummaryVersions);
 		}
 
 #if WITH_EDITORONLY_DATA
@@ -2359,12 +2359,13 @@ FLinkerLoad::ELinkerStatus FLinkerLoad::SerializeThumbnails( bool bForceEnableIn
 			FStructuredArchive::FArray IndexArray = IndexSlot->EnterArray(ThumbnailCount);
 
 			// Allocate a new thumbnail map if we need one
-			if (!LinkerRoot->ThumbnailMap)
+			if (!LinkerRoot->HasThumbnailMap())
 			{
-				LinkerRoot->ThumbnailMap = MakeUnique<FThumbnailMap>();
+				LinkerRoot->SetThumbnailMap(MakeUnique<FThumbnailMap>());
 			}
 
 			// Load thumbnail names and file offsets
+			FThumbnailMap& ThumbnailMap = LinkerRoot->AccessThumbnailMap();
 			TArray< FObjectFullNameAndThumbnail > ThumbnailInfoArray;
 			for (int32 CurObjectIndex = 0; CurObjectIndex < ThumbnailCount; ++CurObjectIndex)
 			{
@@ -2391,7 +2392,7 @@ FLinkerLoad::ELinkerStatus FLinkerLoad::SerializeThumbnails( bool bForceEnableIn
 				// Only bother loading thumbnails that don't already exist in memory yet.  This is because when we
 				// go to load thumbnails that aren't in memory yet when saving packages we don't want to clobber
 				// thumbnails that were freshly-generated during that editor session
-				if (!LinkerRoot->ThumbnailMap->Contains(ThumbnailInfo.ObjectFullName))
+				if (!ThumbnailMap.Contains(ThumbnailInfo.ObjectFullName))
 				{
 					// Add to list of thumbnails to load
 					ThumbnailInfoArray.Add(ThumbnailInfo);
@@ -2417,7 +2418,7 @@ FLinkerLoad::ELinkerStatus FLinkerLoad::SerializeThumbnails( bool bForceEnableIn
 				LoadedThumbnail.Serialize(DataStream.EnterElement());
 
 				// Store the data!
-				LinkerRoot->ThumbnailMap->Add(CurThumbnailInfo.ObjectFullName, LoadedThumbnail);
+				ThumbnailMap.Add(CurThumbnailInfo.ObjectFullName, LoadedThumbnail);
 			}
 		}
 	}
@@ -3792,9 +3793,11 @@ int32 FLinkerLoad::LoadMetaDataFromExportMap(bool bForcePreload)
 		// If we didn't find a MetaData, keep the existing MetaData we may have constructed after previously noticing LoadMetaDataFromExportMap didn't find one
 		if (MetaData)
 		{
+			PRAGMA_DISABLE_DEPRECATION_WARNINGS
 			UE_CLOG(LinkerRoot->MetaData && LinkerRoot->MetaData != MetaData, LogLinker, Warning,
 				TEXT("LoadMetaDataFromExportMap was called after the MetaData was already loaded, and it found a different MetaData. Discarding the previously loaded MetaData."));
 			LinkerRoot->MetaData = MetaData;
+			PRAGMA_ENABLE_DEPRECATION_WARNINGS
 		}
 	}
 
@@ -5512,7 +5515,7 @@ void FLinkerLoad::Detach()
 	// Make sure we're never associated with LinkerRoot again.
 	if (LinkerRoot)
 	{
-		LinkerRoot->LinkerLoad = nullptr;
+		LinkerRoot->SetLinker(nullptr);
 		// When detaching the linker from its package, also empty its stored list of custom versions. 
 		// This is so that object post loaded in the editor in a package that has no associated linker consider that all package custom versions as latest
 		// (i.e. when duplicating an object in the package)
diff --git a/Engine/Source/Runtime/CoreUObject/Private/UObject/ObjectBaseUtility.cpp b/Engine/Source/Runtime/CoreUObject/Private/UObject/ObjectBaseUtility.cpp
--- a/Engine/Source/Runtime/CoreUObject/Private/UObject/ObjectBaseUtility.cpp
+++ b/Engine/Source/Runtime/CoreUObject/Private/UObject/ObjectBaseUtility.cpp
@@ -42,7 +42,7 @@ FPackageFileVersion UObjectBaseUtility::GetLinkerUEVersion() const
 	else if (GetOutermost())
 	{
 		// Get the linker version associated with the package this object lives in
-		return GetOutermost()->LinkerPackageVersion;
+		return GetOutermost()->GetLinkerPackageVersion();
 	}
 	else
 	{
@@ -56,7 +56,7 @@ int32 UObjectBaseUtility::GetLinkerCustomVersion(FGuid CustomVersionKey) const
 	FLinkerLoad* Loader = GetLinker();
 
 	// No linker.
-	if( Loader == NULL )
+	if( Loader == nullptr )
 	{
 		// the _Linker reference is never set for the top-most UPackage of a package (the linker root), so if this object
 		// is the linker root, find our loader in the global list.
@@ -66,16 +66,16 @@ int32 UObjectBaseUtility::GetLinkerCustomVersion(FGuid CustomVersionKey) const
 		}
 	}
 
-	if ( Loader != NULL )
+	if ( Loader != nullptr)
 	{
 		// We have a linker so we can return its version.
 		const FCustomVersion* CustomVersion = Loader->Summary.GetCustomVersionContainer().GetVersion(CustomVersionKey);
 		return CustomVersion ? CustomVersion->Version : -1;
 	}
-	else if (GetOutermost() && GetOutermost()->LinkerCustomVersion.GetAllVersions().Num())
+	else if (GetOutermost() && GetOutermost()->GetLinkerCustomVersions().GetAllVersions().Num())
 	{
 		// Get the linker version associated with the package this object lives in
-		const FCustomVersion* CustomVersion = GetOutermost()->LinkerCustomVersion.GetVersion(CustomVersionKey);
+		const FCustomVersion* CustomVersion = GetOutermost()->GetLinkerCustomVersions().GetVersion(CustomVersionKey);
 		return CustomVersion ? CustomVersion->Version : -1;
 	}
 	// We don't have a linker associated as we e.g. might have been saved or had loaders reset, ...
@@ -97,7 +97,7 @@ int32 UObjectBaseUtility::GetLinkerLicenseeUEVersion() const
 	FLinkerLoad* Loader = GetLinker();
 
 	// No linker.
-	if( Loader == NULL )
+	if( Loader == nullptr)
 	{
 		// the _Linker reference is never set for the top-most UPackage of a package (the linker root), so if this object
 		// is the linker root, find our loader in the global list.
@@ -107,7 +107,7 @@ int32 UObjectBaseUtility::GetLinkerLicenseeUEVersion() const
 		}
 	}
 
-	if ( Loader != NULL )
+	if ( Loader != nullptr)
 	{
 		// We have a linker so we can return its version.
 		return Loader->LicenseeUEVer();
@@ -115,7 +115,7 @@ int32 UObjectBaseUtility::GetLinkerLicenseeUEVersion() const
 	else if (GetOutermost())
 	{
 		// Get the linker version associated with the package this object lives in
-		return GetOutermost()->LinkerLicenseeVersion;
+		return GetOutermost()->GetLinkerLicenseeVersion();
 	}
 	else
 	{
diff --git a/Engine/Source/Runtime/CoreUObject/Private/UObject/ObjectHandle.cpp b/Engine/Source/Runtime/CoreUObject/Private/UObject/ObjectHandle.cpp
--- a/Engine/Source/Runtime/CoreUObject/Private/UObject/ObjectHandle.cpp
+++ b/Engine/Source/Runtime/CoreUObject/Private/UObject/ObjectHandle.cpp
@@ -352,9 +352,9 @@ UClass* ResolveObjectRefClass(const FObjectRef& ObjectRef, uint32 LoadFlags /*=
 				{
 					ClassObject = (UClass*)Redirector->DestinationObject;
 				}
-				if (ClassObject->HasAnyFlags(RF_NeedLoad) && ClassPackage->LinkerLoad)
+				if (ClassObject->HasAnyFlags(RF_NeedLoad) && ClassPackage->GetLinker())
 				{
-					ClassPackage->LinkerLoad->Preload(ClassObject);
+					ClassPackage->GetLinker()->Preload(ClassObject);
 				}
 				ClassObject->GetDefaultObject(); // build the CDO if it isn't already built
 			}
@@ -378,7 +378,7 @@ public:
 
 	void DoTask(ENamedThreads::Type CurrentThread, const FGraphEventRef& MyCompletionGraphEvent)
 	{
-		Package->LinkerLoad->LoadAllObjects(true);
+		Package->GetLinker()->LoadAllObjects(true);
 	}
 
 private:
@@ -418,11 +418,11 @@ UObject* ResolveObjectRef(const FObjectRef& ObjectRef, uint32 LoadFlags /*= LOAD
 			CurrentObject = Redirector->DestinationObject;
 		}
 
-		if (!CurrentObject && !TargetPackage->IsFullyLoaded() && TargetPackage->LinkerLoad && TargetPackage->LinkerLoad->IsLoading())
+		if (!CurrentObject && !TargetPackage->IsFullyLoaded() && TargetPackage->GetLinker() && TargetPackage->GetLinker()->IsLoading())
 		{
 			if (IsInAsyncLoadingThread() || IsInGameThread())
 			{
-				TargetPackage->LinkerLoad->LoadAllObjects(true);
+				TargetPackage->GetLinker()->LoadAllObjects(true);
 			}
 			else
 			{
@@ -445,11 +445,11 @@ UObject* ResolveObjectRef(const FObjectRef& ObjectRef, uint32 LoadFlags /*= LOAD
 		}
 	}
 
-	if (CurrentObject->HasAnyFlags(RF_NeedLoad) && TargetPackage->LinkerLoad)
+	if (CurrentObject->HasAnyFlags(RF_NeedLoad) && TargetPackage->GetLinker())
 	{
 		if (IsInAsyncLoadingThread() || IsInGameThread())
 		{
-			TargetPackage->LinkerLoad->LoadAllObjects(true);
+			TargetPackage->GetLinker()->LoadAllObjects(true);
 		}
 		else
 		{
diff --git a/Engine/Source/Runtime/CoreUObject/Private/UObject/Package.cpp b/Engine/Source/Runtime/CoreUObject/Private/UObject/Package.cpp
--- a/Engine/Source/Runtime/CoreUObject/Private/UObject/Package.cpp
+++ b/Engine/Source/Runtime/CoreUObject/Private/UObject/Package.cpp
@@ -45,17 +45,17 @@ void UPackage::PostInitProperties()
 		bDirty = false;
 	}
 
+	SetLinkerPackageVersion(GPackageFileUEVersion);
+	SetLinkerLicenseeVersion(GPackageFileLicenseeUEVersion);
+
 #if WITH_EDITORONLY_DATA
-	MetaData = nullptr;
+	SetMetaData(nullptr);
 	// Always generate a new unique PersistentGuid, required for new disk packages.
 	// For existing disk packages it will be replaced with the existing PersistentGuid when loading the package summary.
 	// For existing script packages it will be replaced in ConstructUPackage with the CRC of the generated code files.
 	PersistentGuid = FGuid::NewGuid();
-#endif
-	LinkerPackageVersion = GPackageFileUEVersion;
-	LinkerLicenseeVersion = GPackageFileLicenseeUEVersion;
-	PIEInstanceID = INDEX_NONE;
-#if WITH_EDITORONLY_DATA
+
+	SetPIEInstanceID(INDEX_NONE);
 	bIsCookedForEditor = false;
 	// Mark this package as editor-only by default. As soon as something in it is accessed through a non editor-only
 	// property the flag will be removed.
@@ -71,7 +71,7 @@ void UPackage::SetDirtyFlag( bool bIsDirty )
 {
 	if ( GetOutermost() != GetTransientPackage() )
 	{
-		if ( GUndo != NULL
+		if ( GUndo != nullptr
 		// PIE and script/class packages should never end up in the transaction buffer as we cannot undo during gameplay.
 		&& !GetOutermost()->HasAnyPackageFlags(PKG_PlayInEditor|PKG_ContainsScript|PKG_CompiledIn) )
 		{
@@ -114,9 +114,8 @@ void UPackage::Serialize( FArchive& Ar )
 	}
 	if (Ar.IsCountingMemory())
 	{		
-		if (LinkerLoad)
+		if (FLinker* Loader = GetLinker())
 		{
-			FLinker* Loader = LinkerLoad;
 			Loader->Serialize(Ar);
 		}
 	}
@@ -175,28 +174,32 @@ UMetaData* UPackage::GetMetaData()
 	checkf(!FPlatformProperties::RequiresCookedData(), TEXT("MetaData is only allowed in the Editor."));
 
 #if WITH_EDITORONLY_DATA
+	PRAGMA_DISABLE_DEPRECATION_WARNINGS
+	UMetaData* LocalMetaData = MetaData;
+	PRAGMA_ENABLE_DEPRECATION_WARNINGS
+
 	// If there is no MetaData, try to find it.
-	if (MetaData == NULL)
+	if (LocalMetaData == nullptr)
 	{
-		MetaData = FindObjectFast<UMetaData>(this, FName(NAME_PackageMetaData));
+		LocalMetaData = FindObjectFast<UMetaData>(this, FName(NAME_PackageMetaData));
 
-		// If MetaData is NULL then it wasn't loaded by linker, so we have to create it.
-		if(MetaData == NULL)
+		// If MetaData is null then it wasn't loaded by linker, so we have to create it.
+		if(LocalMetaData == nullptr)
 		{
-			MetaData = NewObject<UMetaData>(this, NAME_PackageMetaData, RF_Standalone | RF_LoadCompleted);
+			LocalMetaData = NewObject<UMetaData>(this, NAME_PackageMetaData, RF_Standalone | RF_LoadCompleted);
 		}
+		SetMetaData(LocalMetaData);
 	}
+	check(LocalMetaData);
 
-	check(MetaData);
-
-	if (MetaData->HasAnyFlags(RF_NeedLoad))
+	if (LocalMetaData->HasAnyFlags(RF_NeedLoad))
 	{
-		FLinkerLoad* MetaDataLinker = MetaData->GetLinker();
+		FLinkerLoad* MetaDataLinker = LocalMetaData->GetLinker();
 		check(MetaDataLinker);
-		MetaDataLinker->Preload(MetaData);
+		MetaDataLinker->Preload(LocalMetaData);
 	}
 
-	return MetaData;
+	return LocalMetaData;
 #else
 	return nullptr;
 #endif
@@ -226,9 +229,11 @@ const FPackagePath& UPackage::GetLoadedPath() const
 void UPackage::SetLoadedPath(const FPackagePath& InPackagePath)
 {
 	LoadedPath = InPackagePath;
+#if !UE_STRIP_DEPRECATED_PROPERTIES
 	PRAGMA_DISABLE_DEPRECATION_WARNINGS
 	FileName = InPackagePath.GetPackageFName();
 	PRAGMA_ENABLE_DEPRECATION_WARNINGS
+#endif
 }
 
 /** Tags generated objects with flags */
@@ -237,10 +242,12 @@ void UPackage::TagSubobjects(EObjectFlags NewFlags)
 	Super::TagSubobjects(NewFlags);
 
 #if WITH_EDITORONLY_DATA
+	PRAGMA_DISABLE_DEPRECATION_WARNINGS
 	if (MetaData)
 	{
 		MetaData->SetFlags(NewFlags);
 	}
+	PRAGMA_ENABLE_DEPRECATION_WARNINGS
 #endif
 }
 
@@ -258,7 +265,7 @@ bool UPackage::IsFullyLoaded() const
 
 	// We set bHasBeenFullyLoaded to true when it is read for some special cases
 
-	if (FileSize != 0)
+	if (GetFileSize() != 0)
 	{
 		// If it has a filesize, it is a normal on-disk package, therefore is not a special case, and we respect the current 'false' value of bHasBeenFullyLoaded
 		return false;
@@ -298,13 +305,12 @@ bool UPackage::IsFullyLoaded() const
 void UPackage::BeginDestroy()
 {
 	// Detach linker if still attached
-	if (LinkerLoad)
+	if (FLinkerLoad* Linker = GetLinker())
 	{
 		// Detach() below will most likely null the LinkerLoad so keep a temp copy so that we can still call RemoveLinker on it
-		FLinkerLoad* LocalLinkerToRemove = LinkerLoad;
-		LocalLinkerToRemove->Detach();
-		FLinkerManager::Get().RemoveLinker(LocalLinkerToRemove);
-		LinkerLoad = nullptr;
+		Linker->Detach();
+		FLinkerManager::Get().RemoveLinker(Linker);
+		SetLinker(nullptr);
 	}
 
 	Super::BeginDestroy();
@@ -340,11 +346,13 @@ void UPackage::SetLoadedByEditorPropertiesOnly(bool bIsEditorOnly, bool bRecursi
 #endif
 
 #if WITH_EDITORONLY_DATA
+PRAGMA_DISABLE_DEPRECATION_WARNINGS
 IMPLEMENT_CORE_INTRINSIC_CLASS(UPackage, UObject,
 	{
 		Class->EmitObjectReference(STRUCT_OFFSET(UPackage, MetaData), TEXT("MetaData"));
 	}
 );
+PRAGMA_ENABLE_DEPRECATION_WARNINGS
 #else
 IMPLEMENT_CORE_INTRINSIC_CLASS(UPackage, UObject,
 	{
diff --git a/Engine/Source/Runtime/CoreUObject/Private/UObject/SavePackage.cpp b/Engine/Source/Runtime/CoreUObject/Private/UObject/SavePackage.cpp
--- a/Engine/Source/Runtime/CoreUObject/Private/UObject/SavePackage.cpp
+++ b/Engine/Source/Runtime/CoreUObject/Private/UObject/SavePackage.cpp
@@ -2050,10 +2050,10 @@ FSavePackageResultStruct UPackage::Save(UPackage* InOuter, UObject* InAsset, con
 			}
 #endif // USE_STABLE_LOCALIZATION_KEYS
 
-			if (InOuter->WorldTileInfo.IsValid())
+			if (FWorldTileInfo* WorldTileInfo = InOuter->GetWorldTileInfo())
 			{
 				// collect custom version from wc tile info
-				ExportTaggerArchive << *(InOuter->WorldTileInfo);
+				ExportTaggerArchive << *WorldTileInfo;
 			}
 
 			{
@@ -2228,9 +2228,9 @@ FSavePackageResultStruct UPackage::Save(UPackage* InOuter, UObject* InAsset, con
 				Linker->SetUseUnversionedPropertySerialization(bSaveUnversionedProperties);
 
 				// Make sure the package has the same version as the linker
-				InOuter->LinkerPackageVersion = Linker->UEVer();
-				InOuter->LinkerLicenseeVersion = Linker->LicenseeUEVer();
-				InOuter->LinkerCustomVersion = Linker->GetCustomVersions();
+				InOuter->SetLinkerPackageVersion(Linker->UEVer());
+				InOuter->SetLinkerLicenseeVersion(Linker->LicenseeUEVer());
+				InOuter->SetLinkerCustomVersions(Linker->GetCustomVersions());
 
 				if (EndSavingIfCancelled())
 				{ 
@@ -2719,10 +2719,12 @@ FSavePackageResultStruct UPackage::Save(UPackage* InOuter, UObject* InAsset, con
 				// Write fixed-length file summary to overwrite later.
 				if (SaveFlags & SAVE_KeepGUID)
 				{
+#if !UE_STRIP_DEPRECATED_PROPERTIES
 					// First generation file, keep existing GUID
 					PRAGMA_DISABLE_DEPRECATION_WARNINGS
 					Linker->Summary.Guid = InOuter->Guid;
 					PRAGMA_ENABLE_DEPRECATION_WARNINGS
+#endif
 #if WITH_EDITORONLY_DATA
 					Linker->Summary.PersistentGuid = InOuter->PersistentGuid;
 #endif
@@ -2739,10 +2741,12 @@ FSavePackageResultStruct UPackage::Save(UPackage* InOuter, UObject* InAsset, con
 #endif
 					Linker->Summary.Generations = TArray<FGenerationInfo>();
 
+#if !UE_STRIP_DEPRECATED_PROPERTIES
 					// make sure the UPackage's copy of the GUID is up to date
 					PRAGMA_DISABLE_DEPRECATION_WARNINGS
 					InOuter->Guid = Linker->Summary.Guid;
 					PRAGMA_ENABLE_DEPRECATION_WARNINGS
+#endif
 				}
 				new(Linker->Summary.Generations)FGenerationInfo(0, 0);
 
@@ -4222,7 +4226,7 @@ FSavePackageResultStruct UPackage::Save(UPackage* InOuter, UObject* InAsset, con
 						}
 						
 						// Update package FileSize value
-						InOuter->FileSize = PackageSize;
+						InOuter->SetFileSize(PackageSize);
 
 						// Warn about long package names, which may be bad for consoles with limited filename lengths.
 						if( bWarnOfLongFilename == true )
diff --git a/Engine/Source/Runtime/CoreUObject/Private/UObject/SavePackage/SaveContext.h b/Engine/Source/Runtime/CoreUObject/Private/UObject/SavePackage/SaveContext.h
--- a/Engine/Source/Runtime/CoreUObject/Private/UObject/SavePackage/SaveContext.h
+++ b/Engine/Source/Runtime/CoreUObject/Private/UObject/SavePackage/SaveContext.h
@@ -822,6 +822,20 @@ public:
 		return GetHarvestedRealm().GetLinker();
 	}
 
+	void UpdatePackageLinkerVersions()
+	{
+		FLinkerSave* Linker = GetLinker();
+		check(Linker);
+		Package->SetLinkerPackageVersion(Linker->UEVer());
+		Package->SetLinkerLicenseeVersion(Linker->LicenseeUEVer());
+		Package->SetLinkerCustomVersions(Linker->GetCustomVersions());
+	}
+
+	void UpdatePackageFileSize(int64 InFileSize)
+	{
+		Package->SetFileSize(InFileSize);
+	}
+
 	void SetLinker(TPimplPtr<FLinkerSave> InLinker)
 	{
 		GetHarvestedRealm().SetLinker(MoveTemp(InLinker));
diff --git a/Engine/Source/Runtime/CoreUObject/Private/UObject/SavePackage/SavePackageUtilities.cpp b/Engine/Source/Runtime/CoreUObject/Private/UObject/SavePackage/SavePackageUtilities.cpp
--- a/Engine/Source/Runtime/CoreUObject/Private/UObject/SavePackage/SavePackageUtilities.cpp
+++ b/Engine/Source/Runtime/CoreUObject/Private/UObject/SavePackage/SavePackageUtilities.cpp
@@ -1865,7 +1865,7 @@ void SaveThumbnails(UPackage* InOuter, FLinkerSave* Linker, FStructuredArchive::
 	// if content browser isn't enabled, clear the thumbnail map so we're not using additional memory for nothing
 	if ( !GIsEditor || IsRunningCommandlet() )
 	{
-		InOuter->ThumbnailMap.Reset();
+		InOuter->SetThumbnailMap(nullptr);
 	}
 #endif
 }
@@ -2386,10 +2386,10 @@ void SaveWorldLevelInfo(UPackage* InOuter, FLinkerSave* Linker, FStructuredArchi
 {
 	Linker->Summary.WorldTileInfoDataOffset = 0;
 	
-	if(InOuter->WorldTileInfo.IsValid())
+	if(FWorldTileInfo* WorldTileInfo = InOuter->GetWorldTileInfo())
 	{
 		Linker->Summary.WorldTileInfoDataOffset = Linker->Tell();
-		Record << SA_VALUE(TEXT("WorldLevelInfo"), *(InOuter->WorldTileInfo));
+		Record << SA_VALUE(TEXT("WorldLevelInfo"), *WorldTileInfo);
 	}
 }
 
diff --git a/Engine/Source/Runtime/CoreUObject/Private/UObject/SavePackage2.cpp b/Engine/Source/Runtime/CoreUObject/Private/UObject/SavePackage2.cpp
--- a/Engine/Source/Runtime/CoreUObject/Private/UObject/SavePackage2.cpp
+++ b/Engine/Source/Runtime/CoreUObject/Private/UObject/SavePackage2.cpp
@@ -396,9 +396,9 @@ ESavePackageResult HarvestPackage(FSaveContext& SaveContext)
 	}
 
 	// if we have a WorldTileInfo, we need to harvest its dependencies as well, i.e. Custom Version
-	if (SaveContext.GetPackage()->WorldTileInfo.IsValid())
+	if (FWorldTileInfo* WorldTileInfo = SaveContext.GetPackage()->GetWorldTileInfo())
 	{
-		Harvester << *(SaveContext.GetPackage()->WorldTileInfo);
+		Harvester << *WorldTileInfo;
 	}
 
 	// The Editor version is used as part of the check to see if a package is too old to use the gather cache, so we always have to add it if we have gathered loc for this asset
@@ -917,9 +917,7 @@ ESavePackageResult BuildLinker(FSaveContext& SaveContext)
 		}
 #endif
 		// Make sure the package has the same version as the linker
-		SaveContext.GetPackage()->LinkerPackageVersion = SaveContext.GetLinker()->UEVer();
-		SaveContext.GetPackage()->LinkerLicenseeVersion = SaveContext.GetLinker()->LicenseeUEVer();
-		SaveContext.GetPackage()->LinkerCustomVersion = SaveContext.GetLinker()->GetCustomVersions();
+		SaveContext.UpdatePackageLinkerVersions();
 	}
 
 	PRAGMA_DISABLE_DEPRECATION_WARNINGS
@@ -2358,7 +2356,7 @@ ESavePackageResult InnerSave(FSaveContext& SaveContext)
 		// Update package FileSize value
 		if (SaveContext.IsUpdatingLoadedPath())
 		{
-			SaveContext.GetPackage()->FileSize = SaveContext.PackageHeaderAndExportSize;
+			SaveContext.UpdatePackageFileSize(SaveContext.PackageHeaderAndExportSize);
 		}
 	}
 	return SaveContext.Result;
diff --git a/Engine/Source/Runtime/CoreUObject/Private/UObject/UObjectGlobals.cpp b/Engine/Source/Runtime/CoreUObject/Private/UObject/UObjectGlobals.cpp
--- a/Engine/Source/Runtime/CoreUObject/Private/UObject/UObjectGlobals.cpp
+++ b/Engine/Source/Runtime/CoreUObject/Private/UObject/UObjectGlobals.cpp
@@ -1089,7 +1089,7 @@ public:
 	FUnsafeLinkerLoad(UPackage *Package, const FPackagePath& PackagePath, const FPackagePath& DiffPackagePath, uint32 LoadFlags)
 		: FLinkerLoad(Package, PackagePath, LoadFlags)
 	{
-		Package->LinkerLoad = this;
+		Package->SetLinker(this);
 		while ( Tick(0.0, false, false, nullptr) == FLinkerLoad::LINKER_TimedOut ) 
 		{ 
 		}
@@ -1396,12 +1396,6 @@ UPackage* LoadPackageInternal(UPackage* InOuter, const FPackagePath& PackagePath
 			}
 		}
 
-		// Only set time it took to load package if the above EndLoad is the "outermost" EndLoad.
-		if( Result && !LoadContext->HasLoadedObjects() && !(LoadFlags & LOAD_Verify) )
-		{
-			Result->SetLoadTime( FPlatformTime::Seconds() - StartTime );
-		}
-
 		Linker->Flush();
 
 		if (!FPlatformProperties::RequiresCookedData())
@@ -1422,13 +1416,13 @@ UPackage* LoadPackageInternal(UPackage* InOuter, const FPackagePath& PackagePath
 				if (GGameThreadLoadCounter == 0)
 				{
 					// Sanity check to make sure that Linker is the linker that loaded our Result package or the linker has already been detached
-					check(!Result || Result->LinkerLoad == Linker || Result->LinkerLoad == nullptr);
+					check(!Result || Result->GetLinker() == Linker || Result->GetLinker() == nullptr);
 					if (Result && Linker->HasLoader())
 					{
 						ResetLoaders(Result);
 					}
 					// Reset loaders could have already deleted Linker so guard against deleting stale pointers
-					if (Result && Result->LinkerLoad)
+					if (Result && Result->GetLinker())
 					{
 						Linker->DestroyLoader();
 					}
@@ -4832,7 +4826,7 @@ namespace UECodeGen_Private
 #endif
 		}
 #if WITH_RELOAD
-		NewPackage->Delegates = MoveTemp(Delegates);
+		NewPackage->SetReloadDelegates(MoveTemp(Delegates));
 #endif
 	}
 
diff --git a/Engine/Source/Runtime/CoreUObject/Public/UObject/Package.h b/Engine/Source/Runtime/CoreUObject/Public/UObject/Package.h
--- a/Engine/Source/Runtime/CoreUObject/Public/UObject/Package.h
+++ b/Engine/Source/Runtime/CoreUObject/Public/UObject/Package.h
@@ -23,12 +23,15 @@
 // This is a dummy type which is not implemented anywhere. It's only 
 // used to flag a deprecated Conform argument to package save functions.
 class FLinkerNull;
+class FLinkerLoad;
 class FLinkerSave;
 class ITargetPlatform;
 struct FPackageSaveInfo;
 class FSavePackageContext;
 struct FSavePackageArgs;
 
+class UMetaData;
+
 /**
 * Represents the result of saving a package
 */
@@ -202,41 +205,29 @@
 	uint8 bCanBeImported:1;
 
 private:
-	/** Time in seconds it took to fully load this package. 0 if package is either in process of being loaded or has never been fully loaded.					*/
-	float LoadTime;		// TODO: strip from runtime?
+	// @note this should probably be entirely deprecated and removed but certain stat dump function are still using it, just compile it out in shipping for now
+#if !UE_BUILD_SHIPPING
+	/** Time in seconds it took to fully load this package. 0 if package is either in process of being loaded or has never been fully loaded. */
+	float LoadTime;
+#endif
 
 #if WITH_EDITORONLY_DATA
 	/** Indicates which folder to display this package under in the Generic Browser's list of packages. If not specified, package is added to the root level.	*/
-	FName	FolderName;
-#endif
+	FName FolderName;
 
-	/** GUID of package if it was loaded from disk. Changes at every save. */
-	UE_DEPRECATED(4.27, "UPackage::Guid has not been used by the engine for a long time and it will be removed.")
-	FGuid Guid;
-
-#if WITH_EDITORONLY_DATA
 	/** Persistent GUID of package if it was loaded from disk. Persistent across saves. */
 	FGuid PersistentGuid;
-#endif
-	/** Chunk IDs for the streaming install chunks this package will be placed in.  Empty for no chunk */
+
+	/** Chunk IDs for the streaming install chunks this package will be placed in.  Empty for no chunk. Used during cooking. */
 	TArray<int32> ChunkIDs;
-
-public:
-
-	virtual bool IsNameStableForNetworking() const override { return true; }		// For now, assume all packages have stable net names
-	virtual bool NeedsLoadForClient() const override { return true; }				// To avoid calling the expensive generic version, which only makes sure that the UPackage static class isn't excluded
-	virtual bool NeedsLoadForServer() const override { return true; }
-	virtual bool IsPostLoadThreadSafe() const override;
-	virtual bool IsDestructionThreadSafe() const override { return true; }
-
-#if WITH_EDITORONLY_DATA
-																					/** Sets the bLoadedByEditorPropertiesOnly flag */
-	void SetLoadedByEditorPropertiesOnly(bool bIsEditorOnly, bool bRecursive = false);
-	/** returns true when the package is only referenced by editor-only flag */
-	bool IsLoadedByEditorPropertiesOnly() const { return bLoadedByEditorPropertiesOnly; }
 #endif
 
-private:
+#if !UE_STRIP_DEPRECATED_PROPERTIES
+	/** GUID of package if it was loaded from disk. Changes at every save. */
+	UE_DEPRECATED(4.27, "UPackage::Guid has not been used by the engine for a long time and it will be removed.")
+	FGuid Guid;
+#endif
+
 	/** Package Flags */
 	uint32	PackageFlagsPrivate;
 	
@@ -247,44 +238,71 @@
 	FPackagePath LoadedPath;
 public:
 
-	/** Editor only: PIE instance ID this package belongs to, INDEX_NONE otherwise */
-	int32 PIEInstanceID;		// TODO: strip from runtime?
-
-	/** The name of the file that this package was loaded from */
-	UE_DEPRECATED(5.0, "Use GetLoadedPath instead")
-	FName	FileName;
-
-	/** Linker load associated with this package */
-	class FLinkerLoad* LinkerLoad;
-
 	/** Linker package version this package has been serialized with. This is mostly used by PostLoad **/
+	UE_DEPRECATED(5.0, "Use Get/SetLinkerPackageVersion instead")
 	FPackageFileVersion LinkerPackageVersion;
 
 	/** Linker licensee version this package has been serialized with. This is mostly used by PostLoad **/
+	UE_DEPRECATED(5.0, "Use Get/SetLinkerLicenseeVersion instead")
 	int32 LinkerLicenseeVersion;
 
 	/** Linker custom version container this package has been serialized with. This is mostly used by PostLoad **/
+	UE_DEPRECATED(5.0, "Use Get/SetLinkerCustomVersions instead")
 	FCustomVersionContainer LinkerCustomVersion;
 
+	/** Linker load associated with this package */
+	UE_DEPRECATED(5.0, "Use Get/SetLinker instead")
+	FLinkerLoad* LinkerLoad;
+
 	/** size of the file for this package; if the package was not loaded from a file or was a forced export in another package, this will be zero */
-	uint64 FileSize;			// TODO: strip from runtime?
-
-#if WITH_RELOAD
-	/** Link list of delegates registered to the package.  The next pointer chain can't be used for this. */
-	TArray<UFunction*> Delegates;
-#endif
+	UE_DEPRECATED(5.0, "Use Get/SetFileSize instead")
+	uint64 FileSize;
 
 #if WITH_EDITORONLY_DATA
 	/** Editor only: Thumbnails stored in this package */
+	UE_DEPRECATED(5.0, "Use Get/SetThumbnailMap instead")
 	TUniquePtr< FThumbnailMap > ThumbnailMap;
 
 	// MetaData for the editor, or NULL in the game
-	class UMetaData*	MetaData;
+	UE_DEPRECATED(5.0, "Use Get/HasMetaData instead")
+	class UMetaData* MetaData;
+
+	/** Editor only: PIE instance ID this package belongs to, INDEX_NONE otherwise */
+	UE_DEPRECATED(5.0, "Use Get/SetPIEInstanceID instead")
+	int32 PIEInstanceID;
 #endif
 
+#if !UE_STRIP_DEPRECATED_PROPERTIES
+	/** The name of the file that this package was loaded from */
+	UE_DEPRECATED(5.0, "Use GetLoadedPath instead")
+	FName FileName;
+#endif
+
 	// World browser information
+	UE_DEPRECATED(5.0, "Use Get/SetWorldTileInfo instead")
 	TUniquePtr< FWorldTileInfo > WorldTileInfo;
 
+#if WITH_RELOAD
+	/** Link list of delegates registered to the package.  The next pointer chain can't be used for this. */
+	UE_DEPRECATED(5.0, "Use Get/SetReloadDelegates instead")
+	TArray<UFunction*> Delegates;
+#endif
+
+public:
+
+	virtual bool IsNameStableForNetworking() const override { return true; }		// For now, assume all packages have stable net names
+	virtual bool NeedsLoadForClient() const override { return true; }				// To avoid calling the expensive generic version, which only makes sure that the UPackage static class isn't excluded
+	virtual bool NeedsLoadForServer() const override { return true; }
+	virtual bool IsPostLoadThreadSafe() const override;
+	virtual bool IsDestructionThreadSafe() const override { return true; }
+
+#if WITH_EDITORONLY_DATA
+	/** Sets the bLoadedByEditorPropertiesOnly flag */
+	void SetLoadedByEditorPropertiesOnly(bool bIsEditorOnly, bool bRecursive = false);
+	/** returns true when the package is only referenced by editor-only flag */
+	bool IsLoadedByEditorPropertiesOnly() const { return bLoadedByEditorPropertiesOnly; }
+#endif
+
 	/**
 	* Called after the C++ constructor and after the properties have been initialized, but before the config has been loaded, etc.
 	* mainly this is to emulate some behavior of when the constructor was called after the properties were initialized.
@@ -301,12 +319,98 @@
 
 	// UPackage interface.
 
+private:
+	friend class FLinkerLoad;
+	friend class FUnsafeLinkerLoad;
+	friend class FSaveContext;
+	friend struct FAsyncPackage2;
+
+	void SetLinker(FLinkerLoad* InLinker)
+	{
+		LinkerLoad = InLinker;
+	}
+
+	void SetLinkerPackageVersion(FPackageFileVersion InVersion)
+	{
+		LinkerPackageVersion = MoveTemp(InVersion);
+	}
+
+	void SetLinkerLicenseeVersion(int32 InVersion)
+	{
+		LinkerLicenseeVersion = InVersion;
+	}
+
+	void SetLinkerCustomVersions(FCustomVersionContainer InVersions)
+	{
+		LinkerCustomVersion = MoveTemp(InVersions);
+	}
+
+	void SetFileSize(int64 InFileSize)
+	{
+		FileSize = InFileSize;
+	}
+
+	void SetMetaData(UMetaData* InMetaData)
+	{
+#if WITH_EDITORONLY_DATA
+		MetaData = InMetaData;
+#endif
+	}
+
+public:
+
 	/**
-	* Sets the time it took to load this package.
-	*/
-	void SetLoadTime( float InLoadTime )
+	 * Returns the PIE instance id used by the package if any, or INDEX_NONE otherwise 
+	 */
+	int32 GetPIEInstanceID() const
 	{
+#if WITH_EDITORONLY_DATA
+		return PIEInstanceID;
+#else
+		return INDEX_NONE;
+#endif
+	}
+
+	/**
+	 * Set the PIE instance id for this package
+	 * @param InPIEInstanceID The PIE instance id to use or INDEX_NONE to remove any id set
+	 */
+	void SetPIEInstanceID(int32 InPIEInstanceID)
+	{
+#if WITH_EDITORONLY_DATA
+		PIEInstanceID = InPIEInstanceID;
+#endif
+	}
+
+	FLinkerLoad* GetLinker() const
+	{
+		return LinkerLoad;
+	}
+
+	const FPackageFileVersion& GetLinkerPackageVersion() const
+	{
+		return LinkerPackageVersion;
+	}
+
+	int32 GetLinkerLicenseeVersion() const
+	{
+		return LinkerLicenseeVersion;
+	}
+
+	const FCustomVersionContainer& GetLinkerCustomVersions() const
+	{
+		return LinkerCustomVersion;
+	}
+
+
+	/**
+	 * Sets the time it took to load this package.
+	 */
+	void SetLoadTime( float InLoadTime ) 
+	{
+#if !UE_BUILD_SHIPPING
 		LoadTime = InLoadTime;
+#endif
 	}
 
 	/**
@@ -314,9 +418,13 @@
 	*
 	* @return Time it took to load.
 	*/
-	float GetLoadTime()
+	float GetLoadTime() const
 	{
+#if !UE_BUILD_SHIPPING
 		return LoadTime;
+#else
+		return 0.0f;
+#endif
 	}
 
 #if WITH_EDITORONLY_DATA
@@ -546,56 +654,103 @@
 		check( HasThumbnailMap() );
 		return *ThumbnailMap;
 	}
+
+	/** Set the internal thumbnail map for this package. */
+	void SetThumbnailMap(TUniquePtr<FThumbnailMap> InThumbnailMap)
+	{
+		ThumbnailMap = MoveTemp(InThumbnailMap);
+	}
+
+	/** returns our persistent Guid */
+	FORCEINLINE FGuid GetPersistentGuid() const
+	{
+		return PersistentGuid;
+	}
+	/** sets a specific persistent Guid */
+	FORCEINLINE void SetPersistentGuid(FGuid NewPersistentGuid)
+	{
+		PersistentGuid = NewPersistentGuid;
+	}
 #endif
 
+#if WITH_RELOAD
+	const TArray<UFunction*>& GetReloadDelegates() const
+	{
+		return Delegates;
+	}
+
+	void SetReloadDelegates(TArray<UFunction*> InDelegates) 
+	{
+		Delegates = MoveTemp(InDelegates);
+	}
+#endif
+
+
+	/** Get the world tile info if any*/
+	FWorldTileInfo* GetWorldTileInfo() const
+	{
+		return WorldTileInfo.Get();
+	}
+
+	/** Set the world tile info */
+	void SetWorldTileInfo(TUniquePtr<FWorldTileInfo> InWorldTileInfo)
+	{
+		WorldTileInfo = MoveTemp(InWorldTileInfo);
+	}
+
 	/** returns our Guid */
 	UE_DEPRECATED(4.27, "UPackage::Guid has not been used by the engine for a long time and GetGuid will be removed.")
 	FORCEINLINE FGuid GetGuid() const
 	{
+#if !UE_STRIP_DEPRECATED_PROPERTIES
 		return Guid;
+#else
+		return FGuid();
+#endif
 	}
 	/** makes our a new fresh Guid */
 	UE_DEPRECATED(4.27, "UPackage::Guid has not been used by the engine for a long time and MakeNewGuid will be removed.")
 	FORCEINLINE FGuid MakeNewGuid()
 	{
+#if !UE_STRIP_DEPRECATED_PROPERTIES
 		Guid = FGuid::NewGuid();
 		return Guid;
+#else
+		return FGuid();
+#endif
 	}
 	/** sets a specific Guid */
 	UE_DEPRECATED(4.27, "UPackage::Guid has not been used by the engine for a long time and SetGuid will be removed.")
 	FORCEINLINE void SetGuid(FGuid NewGuid)
 	{
+#if !UE_STRIP_DEPRECATED_PROPERTIES
 		Guid = NewGuid;
+#endif
 	}
 
-#if WITH_EDITORONLY_DATA
-	/** returns our persistent Guid */
-	FORCEINLINE FGuid GetPersistentGuid() const
-	{
-		return PersistentGuid;
-	}
-	/** sets a specific persistent Guid */
-	FORCEINLINE void SetPersistentGuid(FGuid NewPersistentGuid)
-	{
-		PersistentGuid = NewPersistentGuid;
-	}
-#endif
-
 	/** returns our FileSize */
-	FORCEINLINE int64 GetFileSize()
+	FORCEINLINE int64 GetFileSize() const
 	{
 		return FileSize;
 	}
 
 	/** returns our ChunkIDs */
-	FORCEINLINE const TArray<int32>& GetChunkIDs() const
+	const TArray<int32>& GetChunkIDs() const
 	{
+#if WITH_EDITORONLY_DATA
 		return ChunkIDs;
+#else
+		static TArray<int32> Dummy;
+		return Dummy;
+#endif
 	}
+
 	/** sets our ChunkIDs */
 	FORCEINLINE void SetChunkIDs(const TArray<int32>& InChunkIDs)
 	{
+#if WITH_EDITORONLY_DATA
 		ChunkIDs = InChunkIDs;
+#endif
 	}
 
 	/** returns the unique package id */
@@ -631,6 +786,19 @@
 	////////////////////////////////////////////////////////
 	// MetaData 
 
+	/** 
+	 * Return if the package currently has an assigned metadata object
+	 * @returns true if there's an assigned metadata
+	 */
+	bool HasMetaData() const
+	{
+#if WITH_EDITORONLY_DATA
+		return MetaData != nullptr;
+#else
+		return false;
+#endif
+	}
+
 	/**
 	* Gets (after possibly creating) a metadata object for this package
 	*
@@ -701,4 +869,3 @@
 private:
 	static FSavePackageResultStruct Save2(UPackage* InPackage, UObject* InAsset, const TCHAR* InFilename, const FSavePackageArgs& SaveArgs);
 };
-PRAGMA_ENABLE_DEPRECATION_WARNINGS
diff --git a/Engine/Source/Runtime/Engine/Private/Blueprint.cpp b/Engine/Source/Runtime/Engine/Private/Blueprint.cpp
--- a/Engine/Source/Runtime/Engine/Private/Blueprint.cpp
+++ b/Engine/Source/Runtime/Engine/Private/Blueprint.cpp
@@ -686,7 +686,7 @@ void UBlueprint::PostLoad()
 				// PostLoad), and checking RF_WasLoaded makes sure we only 
 				// forcefully set bIsRegeneratingOnLoad for blueprints that need 
 				// it (ones still actively loading)
-				Blueprint.bIsRegeneratingOnLoad = bPreviousValue || ((Package->LinkerLoad != nullptr) && !Package->HasAnyFlags(RF_WasLoaded));
+				Blueprint.bIsRegeneratingOnLoad = bPreviousValue || ((Package->GetLinker() != nullptr) && !Package->HasAnyFlags(RF_WasLoaded));
 			}
 		}
 		~FScopedRegeneratingOnLoad()
diff --git a/Engine/Source/Runtime/Engine/Private/LevelStreaming.cpp b/Engine/Source/Runtime/Engine/Private/LevelStreaming.cpp
--- a/Engine/Source/Runtime/Engine/Private/LevelStreaming.cpp
+++ b/Engine/Source/Runtime/Engine/Private/LevelStreaming.cpp
@@ -358,9 +358,9 @@ void ULevelStreaming::Serialize( FArchive& Ar )
 	
 	if (Ar.IsLoading())
 	{
-		if (GetOutermost()->HasAnyPackageFlags(PKG_PlayInEditor) && GetOutermost()->PIEInstanceID != INDEX_NONE)
+		if (GetOutermost()->HasAnyPackageFlags(PKG_PlayInEditor) && GetOutermost()->GetPIEInstanceID() != INDEX_NONE)
 		{
-			RenameForPIE(GetOutermost()->PIEInstanceID);
+			RenameForPIE(GetOutermost()->GetPIEInstanceID());
 		}
 	}
 }
@@ -1107,7 +1107,7 @@ bool ULevelStreaming::RequestLevel(UWorld* PersistentWorld, bool bAllowLevelLoad
 		{
 			PackageFlags |= PKG_PlayInEditor;
 		}
-		PIEInstanceID = PersistentWorld->GetOutermost()->PIEInstanceID;
+		PIEInstanceID = PersistentWorld->GetOutermost()->GetPIEInstanceID();
 
 		const FString NonPrefixedLevelName = UWorld::StripPIEPrefixFromPackageName(DesiredPackageName.ToString(), PersistentWorld->StreamingLevelsPrefix);
 		UPackage* EditorLevelPackage = FindObjectFast<UPackage>(nullptr, FName(*NonPrefixedLevelName));
@@ -1280,7 +1280,7 @@ void ULevelStreaming::AsyncLevelLoadComplete(const FName& InPackageName, UPackag
 		{
 			if (ULevel* Level = World->PersistentLevel)
 			{
-				PrepareLoadedLevel(Level, LevelPackage, GetOutermost()->PIEInstanceID);
+				PrepareLoadedLevel(Level, LevelPackage, GetOutermost()->GetPIEInstanceID());
 			}
 			else
 			{
diff --git a/Engine/Source/Runtime/Engine/Private/ObjectTrace.cpp b/Engine/Source/Runtime/Engine/Private/ObjectTrace.cpp
--- a/Engine/Source/Runtime/Engine/Private/ObjectTrace.cpp
+++ b/Engine/Source/Runtime/Engine/Private/ObjectTrace.cpp
@@ -521,7 +521,7 @@ void FObjectTrace::OutputWorld(const UWorld* InWorld)
 
 	UE_TRACE_LOG(Object, World, ObjectChannel)
 		<< World.Id(GetObjectId(InWorld))
-		<< World.PIEInstanceId(InWorld->GetOutermost()->PIEInstanceID)
+		<< World.PIEInstanceId(InWorld->GetOutermost()->GetPIEInstanceID())
 		<< World.Type((uint8)InWorld->WorldType)
 		<< World.NetMode((uint8)InWorld->GetNetMode())
 		<< World.IsSimulating(bIsSimulating);
diff --git a/Engine/Source/Runtime/Engine/Private/Tests/AutomationCommon.cpp b/Engine/Source/Runtime/Engine/Private/Tests/AutomationCommon.cpp
--- a/Engine/Source/Runtime/Engine/Private/Tests/AutomationCommon.cpp
+++ b/Engine/Source/Runtime/Engine/Private/Tests/AutomationCommon.cpp
@@ -373,9 +373,9 @@ bool AutomationOpenMap(const FString& MapName, bool bForceReload)
 		FString ShortMapName = FPackageName::GetShortName(MapName);
 		FString ShortWorldMapName = FPackageName::GetShortName(TestWorld->GetMapName());
 
-		if (TestWorld->GetOutermost()->PIEInstanceID != INDEX_NONE)
+		if (TestWorld->GetOutermost()->GetPIEInstanceID() != INDEX_NONE)
 		{
-			FString PIEPrefix = FString::Printf(PLAYWORLD_PACKAGE_PREFIX TEXT("_%d_"), TestWorld->GetOutermost()->PIEInstanceID);
+			FString PIEPrefix = FString::Printf(PLAYWORLD_PACKAGE_PREFIX TEXT("_%d_"), TestWorld->GetOutermost()->GetPIEInstanceID());
 			ShortWorldMapName.ReplaceInline(*PIEPrefix, TEXT(""));
 		}
 		if (ShortMapName != ShortWorldMapName || bForceReload)
diff --git a/Engine/Source/Runtime/Engine/Private/UnrealEngine.cpp b/Engine/Source/Runtime/Engine/Private/UnrealEngine.cpp
--- a/Engine/Source/Runtime/Engine/Private/UnrealEngine.cpp
+++ b/Engine/Source/Runtime/Engine/Private/UnrealEngine.cpp
@@ -1102,7 +1102,7 @@ namespace
 
 		check(NewWorld);
 
-		OutPackage->PIEInstanceID = WorldContext.PIEInstance;
+		OutPackage->SetPIEInstanceID(WorldContext.PIEInstance);
 		OutPackage->SetPackageFlags(PKG_PlayInEditor);
 
 		// Rename streaming levels to PIE
diff --git a/Engine/Source/Runtime/Engine/Private/World.cpp b/Engine/Source/Runtime/Engine/Private/World.cpp
--- a/Engine/Source/Runtime/Engine/Private/World.cpp
+++ b/Engine/Source/Runtime/Engine/Private/World.cpp
@@ -3310,7 +3310,7 @@ void UWorld::RenameToPIEWorld(int32 PIEInstanceID)
 #if WITH_EDITOR
 	UPackage* WorldPackage = GetOutermost();
 
-	WorldPackage->PIEInstanceID = PIEInstanceID;
+	WorldPackage->SetPIEInstanceID(PIEInstanceID);
 	WorldPackage->SetPackageFlags(PKG_PlayInEditor);
 
 	const FString PIEPackageName = *UWorld::ConvertToPIEPackageName(WorldPackage->GetName(), PIEInstanceID);
@@ -3485,7 +3485,7 @@ UWorld* UWorld::DuplicateWorldForPIE(const FString& PackageName, UWorld* OwningW
 	}
 	else if (OwningWorld)
 	{
-		PIEInstanceID = OwningWorld->GetOutermost()->PIEInstanceID;
+		PIEInstanceID = OwningWorld->GetOutermost()->GetPIEInstanceID();
 	}
 	else
 	{
@@ -3501,7 +3501,7 @@ UWorld* UWorld::DuplicateWorldForPIE(const FString& PackageName, UWorld* OwningW
 	UWorld::WorldTypePreLoadMap.FindOrAdd(PrefixedLevelFName) = EWorldType::PIE;
 	UPackage* PIELevelPackage = CreatePackage(*PrefixedLevelName);
 	PIELevelPackage->SetPackageFlags(PKG_PlayInEditor);
-	PIELevelPackage->PIEInstanceID = PIEInstanceID;
+	PIELevelPackage->SetPIEInstanceID(PIEInstanceID);
 	PIELevelPackage->SetLoadedPath(FPackagePath::FromPackageNameChecked(PackageFName));
 	PRAGMA_DISABLE_DEPRECATION_WARNINGS
 	PIELevelPackage->SetGuid( EditorLevelPackage->GetGuid() );
@@ -6475,9 +6475,9 @@ void FSeamlessTravelHandler::SeamlessTravelLoadCallback(const FName& PackageName
 		{
 			if (World->WorldType == EWorldType::PIE)
 			{
-				if (LevelPackage->PIEInstanceID != -1)
+				if (LevelPackage->GetPIEInstanceID() != -1)
 				{
-					World->StreamingLevelsPrefix = UWorld::BuildPIEPackagePrefix(LevelPackage->PIEInstanceID);
+					World->StreamingLevelsPrefix = UWorld::BuildPIEPackagePrefix(LevelPackage->GetPIEInstanceID());
 				}
 				else
 				{
@@ -8113,7 +8113,7 @@ static ULevel* DuplicateLevelWithPrefix(ULevel* InLevel, int32 InstanceID )
 	// Create a package for duplicated level
 	UPackage* NewPackage = CreatePackage( *PrefixedPackageName );
 	NewPackage->SetPackageFlags( PKG_PlayInEditor );
-	NewPackage->PIEInstanceID = InstanceID;
+	NewPackage->SetPIEInstanceID(InstanceID);
 	NewPackage->SetLoadedPath(OriginalPackage->GetLoadedPath());
 	PRAGMA_DISABLE_DEPRECATION_WARNINGS
 	NewPackage->SetGuid( OriginalPackage->GetGuid() );
diff --git a/Engine/Source/Runtime/Engine/Private/WorldComposition.cpp b/Engine/Source/Runtime/Engine/Private/WorldComposition.cpp
--- a/Engine/Source/Runtime/Engine/Private/WorldComposition.cpp
+++ b/Engine/Source/Runtime/Engine/Private/WorldComposition.cpp
@@ -74,7 +74,7 @@ void UWorldComposition::PostDuplicate(bool bDuplicateForPIE)
 
 	if (bDuplicateForPIE)
 	{
-		FixupForPIE(GetOutermost()->PIEInstanceID);	
+		FixupForPIE(GetOutermost()->GetPIEInstanceID());	
 	}
 }
 
@@ -330,7 +330,7 @@ void UWorldComposition::Rescan()
 void UWorldComposition::ReinitializeForPIE()
 {
 	Rescan();
-	FixupForPIE(GetOutermost()->PIEInstanceID);
+	FixupForPIE(GetOutermost()->GetPIEInstanceID());
 	GetWorld()->SetStreamingLevels(TilesStreaming);
 }
 
@@ -482,14 +482,14 @@ void UWorldComposition::OnTileInfoUpdated(const FName& InPackageName, const FWor
 	UPackage* LevelPackage = Cast<UPackage>(StaticFindObjectFast(UPackage::StaticClass(), NULL, Tile->PackageName));
 	if (LevelPackage)
 	{
-		if (LevelPackage->WorldTileInfo == nullptr)
+		if (LevelPackage->GetWorldTileInfo())
 		{
-			LevelPackage->WorldTileInfo = MakeUnique<FWorldTileInfo>(Tile->Info);
+			LevelPackage->SetWorldTileInfo(MakeUnique<FWorldTileInfo>(Tile->Info));
 			PackageDirty = true;
 		}
 		else
 		{
-			*(LevelPackage->WorldTileInfo) = Tile->Info;
+			*(LevelPackage->GetWorldTileInfo()) = Tile->Info;
 		}
 
 		if (PackageDirty)
@@ -960,7 +960,7 @@ void UWorldComposition::OnLevelPostLoad(ULevel* InLevel)
 		const bool bIsDefault = (Info == FWorldTileInfo());
 		if (!bIsDefault)
 		{
-			LevelPackage->WorldTileInfo = MakeUnique<FWorldTileInfo>(Info);
+			LevelPackage->SetWorldTileInfo(MakeUnique<FWorldTileInfo>(Info));
 		}
 	}
 }
@@ -987,9 +987,9 @@ FIntVector UWorldComposition::GetLevelOffset(ULevel* InLevel) const
 	UPackage* LevelPackage = InLevel->GetOutermost();
 	
 	FIntVector LevelPosition = FIntVector::ZeroValue;
-	if (LevelPackage->WorldTileInfo)
+	if (LevelPackage->GetWorldTileInfo())
 	{
-		LevelPosition = LevelPackage->WorldTileInfo->AbsolutePosition;
+		LevelPosition = LevelPackage->GetWorldTileInfo()->AbsolutePosition;
 	}
 	
 	return LevelPosition - OwningWorld->OriginLocation;
@@ -1001,9 +1001,9 @@ FBox UWorldComposition::GetLevelBounds(ULevel* InLevel) const
 	UPackage* LevelPackage = InLevel->GetOutermost();
 	
 	FBox LevelBBox(ForceInit);
-	if (LevelPackage->WorldTileInfo)
+	if (LevelPackage->GetWorldTileInfo())
 	{
-		LevelBBox = LevelPackage->WorldTileInfo->Bounds.ShiftBy(FVector(GetLevelOffset(InLevel)));
+		LevelBBox = LevelPackage->GetWorldTileInfo()->Bounds.ShiftBy(FVector(GetLevelOffset(InLevel)));
 	}
 	
 	return LevelBBox;
diff --git a/Engine/Source/Runtime/Engine/Private/WorldPartition/WorldPartitionLevelHelper.cpp b/Engine/Source/Runtime/Engine/Private/WorldPartition/WorldPartitionLevelHelper.cpp
--- a/Engine/Source/Runtime/Engine/Private/WorldPartition/WorldPartitionLevelHelper.cpp
+++ b/Engine/Source/Runtime/Engine/Private/WorldPartition/WorldPartitionLevelHelper.cpp
@@ -182,7 +182,7 @@ ULevel* FWorldPartitionLevelHelper::CreateEmptyLevelForRuntimeCell(const UWorld*
 	{
 		check(!InPackage);
 		CellPackage->SetPackageFlags(PKG_PlayInEditor);
-		CellPackage->PIEInstanceID = InWorld->GetPackage()->PIEInstanceID;
+		CellPackage->SetPIEInstanceID(InWorld->GetPackage()->GetPIEInstanceID());
 	}
 
 	// Create World & Persistent Level
@@ -305,7 +305,7 @@ bool FWorldPartitionLevelHelper::LoadActors(ULevel* InDestLevel, TArrayView<FWor
 			FPackagePath PackagePath = FPackagePath::FromPackageNameChecked(*ActorPackageName);
 			check(DestPackage);
 			const EPackageFlags PackageFlags = DestPackage->HasAnyPackageFlags(PKG_PlayInEditor) ? PKG_PlayInEditor : PKG_None;
-			::LoadPackageAsync(PackagePath, FName(*ActorPackageInstanceNames[ChildIndex]), CompletionCallback, PackageFlags, DestPackage->PIEInstanceID, 0, InOutInstancingContext);
+			::LoadPackageAsync(PackagePath, FName(*ActorPackageInstanceNames[ChildIndex]), CompletionCallback, PackageFlags, DestPackage->GetPIEInstanceID(), 0, InOutInstancingContext);
 		}
 		else
 		{
diff --git a/Engine/Source/Runtime/Engine/Private/WorldPartition/WorldPartitionLevelStreamingDynamic.cpp b/Engine/Source/Runtime/Engine/Private/WorldPartition/WorldPartitionLevelStreamingDynamic.cpp
--- a/Engine/Source/Runtime/Engine/Private/WorldPartition/WorldPartitionLevelStreamingDynamic.cpp
+++ b/Engine/Source/Runtime/Engine/Private/WorldPartition/WorldPartitionLevelStreamingDynamic.cpp
@@ -349,7 +349,7 @@ void UWorldPartitionLevelStreamingDynamic::FinalizeRuntimeLevel()
 
 		if (OuterWorld->IsPlayInEditor())
 		{
-			int32 PIEInstanceID = GetPackage()->PIEInstanceID;
+			int32 PIEInstanceID = GetPackage()->GetPIEInstanceID();
 			check(PIEInstanceID != INDEX_NONE);
 
 			RuntimeLevel->FixupForPIE(PIEInstanceID, [&](int32 InPIEInstanceID, FSoftObjectPath& ObjectPath)
diff --git a/Engine/Source/Runtime/Engine/Private/WorldPartition/WorldPartitionRuntimeLevelStreamingCell.cpp b/Engine/Source/Runtime/Engine/Private/WorldPartition/WorldPartitionRuntimeLevelStreamingCell.cpp
--- a/Engine/Source/Runtime/Engine/Private/WorldPartition/WorldPartitionRuntimeLevelStreamingCell.cpp
+++ b/Engine/Source/Runtime/Engine/Private/WorldPartition/WorldPartitionRuntimeLevelStreamingCell.cpp
@@ -129,10 +129,10 @@ UWorldPartitionLevelStreamingDynamic* UWorldPartitionRuntimeLevelStreamingCell::
 		NewLevelStreaming->bClientOnlyVisible = GetClientOnlyVisible();
 		NewLevelStreaming->Initialize(*this);
 
-		if (OwningWorld->IsPlayInEditor() && OwningWorld->GetPackage()->HasAnyPackageFlags(PKG_PlayInEditor) && OwningWorld->GetPackage()->PIEInstanceID != INDEX_NONE)
+		if (OwningWorld->IsPlayInEditor() && OwningWorld->GetPackage()->HasAnyPackageFlags(PKG_PlayInEditor) && OwningWorld->GetPackage()->GetPIEInstanceID() != INDEX_NONE)
 		{
 			// When renaming for PIE, make sure to keep World's name so that linker can properly remap with Package's instancing context
-			NewLevelStreaming->RenameForPIE(OwningWorld->GetPackage()->PIEInstanceID, /*bKeepWorldAssetName*/true);
+			NewLevelStreaming->RenameForPIE(OwningWorld->GetPackage()->GetPIEInstanceID(), /*bKeepWorldAssetName*/true);
 		}
 
 		return NewLevelStreaming;
diff --git a/Engine/Source/Runtime/HeadMountedDisplay/Private/IXRTrackingSystem.cpp b/Engine/Source/Runtime/HeadMountedDisplay/Private/IXRTrackingSystem.cpp
--- a/Engine/Source/Runtime/HeadMountedDisplay/Private/IXRTrackingSystem.cpp
+++ b/Engine/Source/Runtime/HeadMountedDisplay/Private/IXRTrackingSystem.cpp
@@ -28,7 +28,7 @@ bool IXRTrackingSystem::IsHeadTrackingAllowedForWorld(UWorld& World) const
 {
 #if WITH_EDITOR
 	// For VR PIE only the first instance uses the headset
-	return IsHeadTrackingAllowed() && ((World.WorldType != EWorldType::PIE) || (World.GetOutermost()->PIEInstanceID == 0));
+	return IsHeadTrackingAllowed() && ((World.WorldType != EWorldType::PIE) || (World.GetOutermost()->GetPIEInstanceID() == 0));
 #else
 	return IsHeadTrackingAllowed();
 #endif
diff --git a/Engine/Source/Runtime/LevelSequence/Private/LevelSequenceBindingReference.cpp b/Engine/Source/Runtime/LevelSequence/Private/LevelSequenceBindingReference.cpp
--- a/Engine/Source/Runtime/LevelSequence/Private/LevelSequenceBindingReference.cpp
+++ b/Engine/Source/Runtime/LevelSequence/Private/LevelSequenceBindingReference.cpp
@@ -29,9 +29,9 @@ FLevelSequenceBindingReference::FLevelSequenceBindingReference(UObject* InObject
 		if (ensure(ObjectPackage))
 		{
 			// If this is being set from PIE we need to remove the pie prefix and point to the editor object
-			if (ObjectPackage->PIEInstanceID != INDEX_NONE)
+			if (ObjectPackage->GetPIEInstanceID() != INDEX_NONE)
 			{
-				FString PIEPrefix = FString::Printf(PLAYWORLD_PACKAGE_PREFIX TEXT("_%d_"), ObjectPackage->PIEInstanceID);
+				FString PIEPrefix = FString::Printf(PLAYWORLD_PACKAGE_PREFIX TEXT("_%d_"), ObjectPackage->GetPIEInstanceID());
 				FullPath.ReplaceInline(*PIEPrefix, TEXT(""));
 			}
 		}
@@ -65,7 +65,7 @@ UObject* FLevelSequenceBindingReference::Resolve(UObject* InContext, FName Strea
 		TempPath.PreSavePath();
 
 	#if WITH_EDITORONLY_DATA
-		int32 ContextPlayInEditorID = InContext ? InContext->GetOutermost()->PIEInstanceID : INDEX_NONE;
+		int32 ContextPlayInEditorID = InContext ? InContext->GetOutermost()->GetPIEInstanceID() : INDEX_NONE;
 
 		if (ContextPlayInEditorID != INDEX_NONE)
 		{
@@ -148,7 +148,7 @@ UObject* FLevelSequenceLegacyObjectReference::Resolve(UObject* InContext) const
 {
 	if (ObjectId.IsValid() && InContext != nullptr)
 	{
-		int32 PIEInstanceID = InContext->GetOutermost()->PIEInstanceID;
+		int32 PIEInstanceID = InContext->GetOutermost()->GetPIEInstanceID();
 		FUniqueObjectGuid FixedUpId = PIEInstanceID == -1 ? ObjectId : ObjectId.FixupForPIE(PIEInstanceID);
 
 		if (PIEInstanceID != -1 && FixedUpId == ObjectId)
diff --git a/Engine/Source/Runtime/MovieScene/Private/Evaluation/IMovieSceneCustomClockSource.cpp b/Engine/Source/Runtime/MovieScene/Private/Evaluation/IMovieSceneCustomClockSource.cpp
--- a/Engine/Source/Runtime/MovieScene/Private/Evaluation/IMovieSceneCustomClockSource.cpp
+++ b/Engine/Source/Runtime/MovieScene/Private/Evaluation/IMovieSceneCustomClockSource.cpp
@@ -12,9 +12,9 @@ FMovieSceneTimeController_Custom::FMovieSceneTimeController_Custom(const FSoftOb
 	UObject* PlaybackContext = WeakPlaybackContext.Get();
 
 	UPackage* Package = PlaybackContext ? PlaybackContext->GetOutermost() : nullptr;
-	if (Package && Package->PIEInstanceID != INDEX_NONE)
+	if (Package && Package->GetPIEInstanceID() != INDEX_NONE)
 	{
-		if (!ObjectPath.FixupForPIE(Package->PIEInstanceID))
+		if (!ObjectPath.FixupForPIE(Package->GetPIEInstanceID()))
 		{
 			// log error?
 		}
