Fix asserts/warnings around child anim BPs

Ensures that cleaned-up sparse class data structs are moved out of the package they are in, thus not requiring a GC step to ensure that batch compilation works correctly.
Without this, it is possible for the NewObject call in FAnimBlueprintCompilerContext::RecreateSparseClassData to attempt to re-use (as it has a fixed name) the currently-existing (but not referenced) object rather than creating a new one.
This also reverts some other changes that were made to try to address related issues.

diff --git a/Engine/Source/Editor/AnimGraph/Private/AnimBlueprintCompiler.cpp b/Engine/Source/Editor/AnimGraph/Private/AnimBlueprintCompiler.cpp
--- a/Engine/Source/Editor/AnimGraph/Private/AnimBlueprintCompiler.cpp
+++ b/Engine/Source/Editor/AnimGraph/Private/AnimBlueprintCompiler.cpp
@@ -1267,11 +1267,7 @@ void FAnimBlueprintCompilerContext::RecreateSparseClassData()
 		check(ParentAnimClass->GetSparseClassDataStruct());
 		
 		// Derive sparse class data from parent class
-		NewAnimBlueprintConstants = FindObject<UScriptStruct>(NewAnimBlueprintClass, *UAnimBlueprintGeneratedClass::GetConstantsStructName().ToString());
-		if (!NewAnimBlueprintConstants)
-		{
-			NewAnimBlueprintConstants = NewObject<UScriptStruct>(NewAnimBlueprintClass, UAnimBlueprintGeneratedClass::GetConstantsStructName(), RF_Public);
-		}
+		NewAnimBlueprintConstants = NewObject<UScriptStruct>(NewAnimBlueprintClass, UAnimBlueprintGeneratedClass::GetConstantsStructName(), RF_Public);
 		NewAnimBlueprintConstants->SetSuperStruct(ParentAnimClass->GetSparseClassDataStruct());
 
 		// Just link & assign sparse class data struct here, no additional members are added
@@ -1386,7 +1382,7 @@ void FAnimBlueprintCompilerContext::CleanAndSanitizeClass(UBlueprintGeneratedCla
 
 	// Cleanup sparse class data & stash a reference to patch the linker later in RecreateSparseClassData
 	OldSparseClassDataStruct = AnimBlueprintClassToClean->GetSparseClassDataStruct();
-	AnimBlueprintClassToClean->ClearSparseClassDataStruct();
+	AnimBlueprintClassToClean->ClearSparseClassDataStruct(Blueprint->bIsRegeneratingOnLoad);
 	
 	FAnimBlueprintGeneratedClassCompiledData CompiledData(AnimBlueprintClassToClean);
 	FAnimBlueprintCompilationBracketContext CompilerContext(this);
diff --git a/Engine/Source/Runtime/CoreUObject/Private/UObject/Class.cpp b/Engine/Source/Runtime/CoreUObject/Private/UObject/Class.cpp
--- a/Engine/Source/Runtime/CoreUObject/Private/UObject/Class.cpp
+++ b/Engine/Source/Runtime/CoreUObject/Private/UObject/Class.cpp
@@ -5251,10 +5251,6 @@ void UClass::SetSparseClassDataStruct(UScriptStruct* InSparseClassDataStruct)
 					SubClass->CleanupSparseClassData();
 					SubClassSparseClassDataStruct->SetSuperStruct(InSparseClassDataStruct);
 				}
-				else if (SubClassSparseClassDataStruct == SparseClassDataStruct)
-				{
-					SubClass->SetSparseClassDataStruct(InSparseClassDataStruct);
-				}
 			}
 		}
 		// the old type and new type may not match when we do a reload so get rid of the old data
@@ -5264,7 +5260,7 @@ void UClass::SetSparseClassDataStruct(UScriptStruct* InSparseClassDataStruct)
 	}
 }
 
-void UClass::ClearSparseClassDataStruct()
+void UClass::ClearSparseClassDataStruct(bool bInRecomplingOnLoad)
 { 
 	if (SparseClassDataStruct != nullptr)
 	{
@@ -5281,17 +5277,22 @@ void UClass::ClearSparseClassDataStruct()
 			}
 		}
 
+		auto ClearSparseClassDataStructInner = [bInRecomplingOnLoad](UClass* InClassToClear)
+		{
+			UScriptStruct* CurrentSparseClassDataStruct = InClassToClear->GetSparseClassDataStruct();
+			InClassToClear->CleanupSparseClassData();
+			InClassToClear->SparseClassDataStruct = nullptr;
+			CurrentSparseClassDataStruct->SetSuperStruct(nullptr);
+			const ERenameFlags RenameFlags = REN_DontCreateRedirectors | ((bInRecomplingOnLoad) ? REN_ForceNoResetLoaders : 0) | REN_NonTransactional | REN_DoNotDirty;
+			CurrentSparseClassDataStruct->Rename(nullptr, GetTransientPackage(), RenameFlags);
+		};
+
 		for(UClass* SubClassToClear : SubClassesToClear)
 		{
-			UScriptStruct* SubClassSparseClassDataStruct = SubClassToClear->GetSparseClassDataStruct();
-			SubClassToClear->CleanupSparseClassData();
-			SubClassToClear->SparseClassDataStruct = nullptr;
-			SubClassSparseClassDataStruct->SetSuperStruct(nullptr);	
+			ClearSparseClassDataStructInner(SubClassToClear);
 		}
 
-		CleanupSparseClassData();
-
-		SparseClassDataStruct = nullptr;
+		ClearSparseClassDataStructInner(this);
 	}
 }
 
diff --git a/Engine/Source/Runtime/CoreUObject/Public/UObject/Class.h b/Engine/Source/Runtime/CoreUObject/Public/UObject/Class.h
--- a/Engine/Source/Runtime/CoreUObject/Public/UObject/Class.h
+++ b/Engine/Source/Runtime/CoreUObject/Public/UObject/Class.h
@@ -2836,7 +2836,7 @@ public:
 	void SetSparseClassDataStruct(UScriptStruct* InSparseClassDataStruct);
 
 	/** Clears the sparse class data struct for this and all child classes that directly reference it as a super-struct */
-	void ClearSparseClassDataStruct();
+	void ClearSparseClassDataStruct(bool bInRecomplingOnLoad);
 
 	/** Assemble reference token streams for all classes if they haven't had it assembled already */
 	static void AssembleReferenceTokenStreams();
