Fix CIS warning about unbalanced pragma disable deprecation

diff --git a/Engine/Source/Runtime/CoreUObject/Public/UObject/Package.h b/Engine/Source/Runtime/CoreUObject/Public/UObject/Package.h
--- a/Engine/Source/Runtime/CoreUObject/Public/UObject/Package.h
+++ b/Engine/Source/Runtime/CoreUObject/Public/UObject/Package.h
@@ -869,3 +869,4 @@ public:
 private:
 	static FSavePackageResultStruct Save2(UPackage* InPackage, UObject* InAsset, const TCHAR* InFilename, const FSavePackageArgs& SaveArgs);
 };
+PRAGMA_ENABLE_DEPRECATION_WARNINGS
\ No newline at end of file
