Improve error messages given when a package is not able to be loaded

- Changed a number of UE_LOG messages when parsing a package's file summary to use FMessageLog instead and using the category "LoadErrors" rather than LogLinker.
-- Currently when a user encounters one of these issues they will be prompted to open the "LoadErrors" message log category via a UI popup, but since we were logging via UE_LOG and using the LogLinker category the message log would not actually contain any of the warnings/errors relating to the failed load. those messages were being routed to the output log instead.
-- By switching to FMessageLog the user will now see the warnings/errors that prevent the package from loading and can hopefully work towards a fix of their data.
- We now report an empty/corrupted package summary as an error rather than a warning. All other messages have been left with their original severity.
- Some error messages have been cleaned up and all are now using FText so could be localised in the future.

diff --git a/Engine/Source/Runtime/CoreUObject/Private/UObject/LinkerLoad.cpp b/Engine/Source/Runtime/CoreUObject/Private/UObject/LinkerLoad.cpp
--- a/Engine/Source/Runtime/CoreUObject/Private/UObject/LinkerLoad.cpp
+++ b/Engine/Source/Runtime/CoreUObject/Private/UObject/LinkerLoad.cpp
@@ -22,6 +22,7 @@
 #include "Misc/SecureHash.h"
 #include "Misc/StringBuilder.h"
 #include "ProfilingDebugging/DebuggingDefines.h"
+#include "Logging/MessageLog.h"
 #include "Logging/TokenizedMessage.h"
 #include "UObject/LinkerPlaceholderBase.h"
 #include "UObject/LinkerPlaceholderClass.h"
@@ -1211,15 +1212,20 @@ FLinkerLoad::ELinkerStatus FLinkerLoad::SerializePackageFileSummaryInternal()
 	// Check tag.
 	if (Summary.Tag != PACKAGE_FILE_TAG)
 	{
-		UE_LOG(LogLinker, Warning, TEXT("The file '%s' contains unrecognizable data, check that it is of the expected type."), *GetDebugName());
+		FMessageLog("LoadErrors").Error(FText::Format(NSLOCTEXT("Core", "LinkerLoad_PkgSumCorrupted", "The summary for the package '{0}' is invalid. Check that the file is of the expected type and not corrupted."),
+			FText::FromString(GetDebugName())));
+
 		return LINKER_Failed;
 	}
 
 	// Validate the summary.
 	if (Summary.IsFileVersionTooOld())
 	{
-		UE_LOG(LogLinker, Warning, TEXT("The file %s was saved by a previous version which is not backwards compatible with this one. Min Required Version: %i  Package Version: %i"),
-			*GetDebugName(), (int32)VER_UE4_OLDEST_LOADABLE_PACKAGE, Summary.GetFileVersionUE().FileVersionUE4);
+		FMessageLog("LoadErrors").Warning(FText::Format(NSLOCTEXT("Core", "LinkerLoad_PkgVersionTooOld", "The package '{0}' was saved with an older version which is not backwards compatible with the current process. Min Required Version: {1}  Package Version: {2}"), 
+			FText::FromString(GetDebugName()), 
+			(int32)VER_UE4_OLDEST_LOADABLE_PACKAGE, 
+			Summary.GetFileVersionUE().FileVersionUE4));
+
 		return LINKER_Failed;
 	}
 
@@ -1245,12 +1251,13 @@ FLinkerLoad::ELinkerStatus FLinkerLoad::SerializePackageFileSummaryInternal()
 	// Don't load packages that are only compatible with an engine version newer than the current one.
 	if (bLoaderNeedsEngineVersionChecks && GEnforcePackageCompatibleVersionCheck && !FEngineVersion::Current().IsCompatibleWith(Summary.CompatibleWithEngineVersion))
 	{
-		UE_LOG(LogLinker, Warning, TEXT("Asset '%s' has been saved with a newer engine and can't be loaded. CurrentEngineVersion: %s (Licensee=%d). AssetEngineVersion: %s (Licensee=%d)"),
-			*GetDebugName(),
-			*FEngineVersion::Current().ToString(),
+		FMessageLog("LoadErrors").Warning(FText::Format(NSLOCTEXT("Core", "LinkerLoad_EngineVersionIncompatible", "Package '{0}' has been saved with a newer engine version and can't be loaded. Current EngineVersion: {1} (Licensee={2}). Package EngineVersion: {3} (Licensee={4})"),
+			FText::FromString(GetDebugName()), 
+			FText::FromString(FEngineVersion::Current().ToString()),
 			FEngineVersion::Current().IsLicenseeVersion(),
-			*Summary.CompatibleWithEngineVersion.ToString(),
-			Summary.CompatibleWithEngineVersion.IsLicenseeVersion());
+			FText::FromString(Summary.CompatibleWithEngineVersion.ToString()),
+			Summary.CompatibleWithEngineVersion.IsLicenseeVersion()));
+
 		return LINKER_Failed;
 	}
 
@@ -1281,15 +1288,22 @@ FLinkerLoad::ELinkerStatus FLinkerLoad::SerializePackageFileSummaryInternal()
 	// Don't load packages that were saved with package version newer than the current one.
 	if (bLoaderNeedsEngineVersionChecks && ((Summary.IsFileVersionTooNew()) || (Summary.GetFileVersionLicenseeUE() > GPackageFileLicenseeUEVersion)))
 	{
-		UE_LOG(LogLinker, Warning, TEXT("Unable to load package (%s) PackageVersion %i, MaxExpected %i : LicenseePackageVersion %i, MaxExpected %i."),
-			*GetDebugName(), Summary.GetFileVersionUE().ToValue(), GPackageFileUEVersion.ToValue(), Summary.GetFileVersionLicenseeUE(), GPackageFileLicenseeUEVersion);
+		FMessageLog("LoadErrors").Warning(FText::Format(NSLOCTEXT("Core", "LinkerLoad_PkgVersionTooNew", "Package '{0}' contains a newer version than the current process supports. PackageVersion {1}, MaxExpected {2} : LicenseePackageVersion {3}, MaxExpected {4}."), 
+			FText::FromString(GetDebugName()), 
+			Summary.GetFileVersionUE().ToValue(), 
+			GPackageFileUEVersion.ToValue(), 
+			Summary.GetFileVersionLicenseeUE(), 
+			GPackageFileLicenseeUEVersion));
+
 		return LINKER_Failed;
 	}
 
 	// don't load packages that contain editor only data in builds that don't support that and vise versa
 	if (!FPlatformProperties::HasEditorOnlyData() && !(Summary.GetPackageFlags() & PKG_FilterEditorOnly))
 	{
-		UE_LOG(LogLinker, Warning, TEXT("Unable to load package (%s). Package contains EditorOnly data which is not supported by the current build."), *GetDebugName());
+		FMessageLog("LoadErrors").Warning(FText::Format(NSLOCTEXT("Core", "LinkerLoad_InvalidEditorOnlyData", "Unable to load package '{0}'. Package contains EditorOnly data which is not supported by the current build."), 
+			FText::FromString(GetDebugName())));
+
 		return LINKER_Failed;
 	}
 
@@ -1299,9 +1313,9 @@ FLinkerLoad::ELinkerStatus FLinkerLoad::SerializePackageFileSummaryInternal()
 		// This warning can be disabled in ini or project settings
 		if (!GAllowCookedDataInEditorBuilds)
 		{
-			UE_LOG(LogLinker, Warning,
-				TEXT("Unable to load package (%s). Package contains cooked data which is not supported by the current build. Enable 'Allow Cooked Content In The Editor' in Project Settings under 'Engine - Cooker' section to load it."),
-				*GetDebugName());
+			FMessageLog("LoadErrors").Warning(FText::Format(NSLOCTEXT("Core", "LinkerLoad_InvalidCookedData", "Unable to load package '{0}'. Package contains cooked data which is not supported by the current build. Enable 'Allow Cooked Content In The Editor' in Project Settings under 'Engine - Cooker' section to load it."),
+				FText::FromString(GetDebugName())));
+
 			return LINKER_Failed;
 		}
 	}
@@ -1326,7 +1340,9 @@ FLinkerLoad::ELinkerStatus FLinkerLoad::SerializePackageFileSummaryInternal()
 
 		if (Tag != PACKAGE_FILE_TAG)
 		{
-			UE_LOG(LogLinker, Warning, TEXT("Unable to load package (%s). Post Tag is not valid. File might be corrupted."), *GetDebugName());
+			FMessageLog("LoadErrors").Error(FText::Format(NSLOCTEXT("Core", "LinkerLoad_PkgTagCorrupted", "Unable to load package '{0}'. The end of package tag is not valid. Check that the file is of the expected type and not corrupted."), 
+				FText::FromString(GetDebugName())));
+
 			return LINKER_Failed;
 		}
 
