Fix UPackage::IsEmptyPackage not working properly with pending kill disabled.

diff --git a/Engine/Source/Runtime/CoreUObject/Private/UObject/SavePackage/SavePackageUtilities.cpp b/Engine/Source/Runtime/CoreUObject/Private/UObject/SavePackage/SavePackageUtilities.cpp
--- a/Engine/Source/Runtime/CoreUObject/Private/UObject/SavePackage/SavePackageUtilities.cpp
+++ b/Engine/Source/Runtime/CoreUObject/Private/UObject/SavePackage/SavePackageUtilities.cpp
@@ -2427,8 +2427,8 @@ bool UPackage::IsEmptyPackage(UPackage* Package, const UObject* LastReferencer)
 				return false;
 			}
 			return true;
-		// Don't consider transient, class default or pending kill objects
-		}, false, RF_Transient | RF_ClassDefaultObject, EInternalObjectFlags::PendingKill);
+		// Don't consider transient, class default or pending kill / garbage objects
+		}, false, RF_Transient | RF_Garbage | RF_ClassDefaultObject, EInternalObjectFlags::PendingKill | EInternalObjectFlags::Garbage);
 		return bIsEmpty;
 	}
 
