Remove usage of deprecated RF_Garbage flag.

diff --git a/Engine/Source/Runtime/CoreUObject/Private/UObject/SavePackage/SavePackageUtilities.cpp b/Engine/Source/Runtime/CoreUObject/Private/UObject/SavePackage/SavePackageUtilities.cpp
--- a/Engine/Source/Runtime/CoreUObject/Private/UObject/SavePackage/SavePackageUtilities.cpp
+++ b/Engine/Source/Runtime/CoreUObject/Private/UObject/SavePackage/SavePackageUtilities.cpp
@@ -2428,7 +2428,7 @@ bool UPackage::IsEmptyPackage(UPackage* Package, const UObject* LastReferencer)
 			}
 			return true;
 		// Don't consider transient, class default or pending kill / garbage objects
-		}, false, RF_Transient | RF_Garbage | RF_ClassDefaultObject, EInternalObjectFlags::PendingKill | EInternalObjectFlags::Garbage);
+		}, false, RF_Transient | RF_ClassDefaultObject, EInternalObjectFlags::PendingKill | EInternalObjectFlags::Garbage);
 		return bIsEmpty;
 	}
 
