Fixed up uses of FIoDispatcher::IsInitialized() so that they work correctly also when there's a mix of cooked packages in iostore and uncooked packages on disk
FIoDispatcher::IsInitialized() will now return true if it exists (which is almost always)
Renamed FArchive::IsUsingEventDrivenLoader() to IsLoadingFromCookedPackage and based it on the PKG_Cooked flag
#rnx

diff --git a/Engine/Source/Editor/UnrealEd/Private/CookOnTheFlyServer.cpp b/Engine/Source/Editor/UnrealEd/Private/CookOnTheFlyServer.cpp
--- a/Engine/Source/Editor/UnrealEd/Private/CookOnTheFlyServer.cpp
+++ b/Engine/Source/Editor/UnrealEd/Private/CookOnTheFlyServer.cpp
@@ -8070,19 +8070,6 @@ void UCookOnTheFlyServer::StartCookByTheBook( const FCookByTheBookStartupOptions
 
 TArray<FName> UCookOnTheFlyServer::GetNeverCookPackageFileNames(TArrayView<const FString> ExtraNeverCookDirectories)
 {
-	// if we are cooking with IoStore enabled, this will throw warnings and fail to find any packaged, but since we are running 
-	// with IoStore the state of files will hopefully be good
-	if (FIoDispatcher::IsInitialized())
-	{
-		static bool bShownMessage = false;
-		if (!bShownMessage)
-		{
-			UE_LOG(LogCook, Log, TEXT("Unable to search for packages to not cook when running with IO Dispatcher active"));
-			bShownMessage = true;
-		}
-		return TArray<FName>();
-	}
-
 	TArray<FString> NeverCookDirectories(ExtraNeverCookDirectories);
 
 	auto AddDirectoryPathArray = [&NeverCookDirectories](const TArray<FDirectoryPath>& DirectoriesToNeverCook, const TCHAR* SettingName)
diff --git a/Engine/Source/Runtime/Core/Private/IO/IoDispatcher.cpp b/Engine/Source/Runtime/Core/Private/IO/IoDispatcher.cpp
--- a/Engine/Source/Runtime/Core/Private/IO/IoDispatcher.cpp
+++ b/Engine/Source/Runtime/Core/Private/IO/IoDispatcher.cpp
@@ -20,6 +20,7 @@
 #include "Containers/Ticker.h"
 #include "IO/IoDispatcherBackend.h"
 #include "Hash/Blake3.h"
+#include "IO/PackageId.h"
 
 DEFINE_LOG_CATEGORY(LogIoDispatcher);
 
@@ -330,20 +331,22 @@ public:
 		RequestAllocator->ReleaseRef();
 	}
 
-	FIoStatus Initialize()
+	void Initialize()
 	{
-		return FIoStatus::Ok;
-	}
-
-	bool InitializePostSettings()
-	{
-		check(!Thread);
-		for (const TSharedRef<IIoDispatcherBackend>& Backend : Backends)
+		if (bIsInitialized)
 		{
-			Backend->Initialize(BackendContext);
+			return;
+		}
+		bIsInitialized = true;
+		if (!Backends.IsEmpty())
+		{
+			for (const TSharedRef<IIoDispatcherBackend>& Backend : Backends)
+			{
+				Backend->Initialize(BackendContext);
+			}
+			// If there are no mounted backends the resolve thread is not needed
+			StartThread();
 		}
-		Thread = FRunnableThread::Create(this, TEXT("IoDispatcher"), 0, TPri_AboveNormal, FPlatformAffinity::GetIoDispatcherThreadMask());
-		return true;
 	}
 
 	FIoBatchImpl* AllocBatch()
@@ -412,10 +415,15 @@ public:
 
 	void Mount(TSharedRef<IIoDispatcherBackend> Backend)
 	{
+		check(IsInGameThread());
 		Backends.Add(Backend);
-		if (Thread)
+		if (bIsInitialized)
 		{
 			Backend->Initialize(BackendContext);
+			if (!Thread)
+			{
+				StartThread();
+			}
 		}
 	}
 
@@ -468,6 +476,25 @@ public:
 			return;
 		}
 		check(Batch.TailRequest);
+
+		if (Backends.IsEmpty())
+		{
+			FIoRequestImpl* Request = Batch.HeadRequest;
+			while (Request)
+			{
+				FIoRequestImpl* NextRequest = Request->NextRequest;
+				CompleteRequest(Request, EIoErrorCode::NotFound);
+				Request->ReleaseRef();
+				Request = NextRequest;
+			}
+			Batch.HeadRequest = Batch.TailRequest = nullptr;
+			if (BatchImpl)
+			{
+				CompleteBatch(BatchImpl);
+			}
+			return;
+		}
+
 		uint32 RequestCount = 0;
 		FIoRequestImpl* Request = Batch.HeadRequest;
 		while (Request)
@@ -550,6 +577,12 @@ private:
 	friend class FIoBatch;
 	friend class FIoRequest;
 
+	void StartThread()
+	{
+		check(!Thread);
+		Thread = FRunnableThread::Create(this, TEXT("IoDispatcher"), 0, TPri_AboveNormal, FPlatformAffinity::GetIoDispatcherThreadMask());
+	}
+
 	void ProcessCompletedRequests()
 	{
 		//TRACE_CPUPROFILER_EVENT_SCOPE(ProcessCompletedRequests);
@@ -792,6 +825,7 @@ private:
 	uint64 PendingIoRequestsCount = 0;
 	int64 TotalLoaded = 0;
 	FIoRequestStats RequestStats;
+	bool bIsInitialized = false;
 };
 
 FIoDispatcher::FIoDispatcher()
@@ -864,18 +898,7 @@ HasUseIoStoreParam()
 bool
 FIoDispatcher::IsInitialized()
 {
-	if (GIoDispatcher)
-	{
-		if (HasScriptObjectsChunk(*GIoDispatcher))
-		{
-			return true;
-		}
-		if (HasUseIoStoreParam() && GIoDispatcher->Impl->HasMountedBackend())
-		{
-			return true;
-		}
-	}
-	return false;
+	return GIoDispatcher.IsValid();
 }
 
 FIoStatus
@@ -884,7 +907,7 @@ FIoDispatcher::Initialize()
 	LLM_SCOPE(ELLMTag::FileSystem);
 	check(!GIoDispatcher);
 	GIoDispatcher = MakeUnique<FIoDispatcher>();
-	return GIoDispatcher->Impl->Initialize();
+	return FIoStatus::Ok;
 }
 
 void
@@ -892,8 +915,7 @@ FIoDispatcher::InitializePostSettings()
 {
 	LLM_SCOPE(ELLMTag::FileSystem);
 	check(GIoDispatcher);
-	bool bSuccess = GIoDispatcher->Impl->InitializePostSettings();
-	UE_CLOG(!bSuccess, LogIoDispatcher, Fatal, TEXT("Failed to initialize IoDispatcher"));
+	GIoDispatcher->Impl->Initialize();
 }
 
 void
@@ -905,6 +927,7 @@ FIoDispatcher::Shutdown()
 FIoDispatcher&
 FIoDispatcher::Get()
 {
+	check(GIoDispatcher);
 	return *GIoDispatcher;
 }
 
@@ -1176,6 +1199,11 @@ FIoRequest::Release()
 	}
 }
 
+FIoChunkId CreatePackageDataChunkId(const FPackageId& PackageId)
+{
+	return CreateIoChunkId(PackageId.Value(), 0, EIoChunkType::ExportBundleData);
+}
+
 FIoChunkId CreateExternalFileChunkId(const FStringView Filename)
 {
 	check(Filename.Len() > 0);
diff --git a/Engine/Source/Runtime/Core/Private/Serialization/Archive.cpp b/Engine/Source/Runtime/Core/Private/Serialization/Archive.cpp
--- a/Engine/Source/Runtime/Core/Private/Serialization/Archive.cpp
+++ b/Engine/Source/Runtime/Core/Private/Serialization/Archive.cpp
@@ -158,6 +158,7 @@ void FArchiveState::Reset()
 	ArEngineNetVer						= FNetworkVersion::GetEngineNetworkProtocolVersion();
 	ArGameNetVer						= FNetworkVersion::GetGameNetworkProtocolVersion();
 	ArIsLoading							= false;
+	ArIsLoadingFromCookedPackage		= false;
 	ArIsSaving							= false;
 	ArIsTransacting						= false;
 	ArIsTextFormat						= false;
@@ -217,6 +218,7 @@ void FArchiveState::CopyTrivialFArchiveStatusMembers(const FArchiveState& Archiv
 	ArEngineNetVer                       = ArchiveToCopy.ArEngineNetVer;
 	ArGameNetVer                         = ArchiveToCopy.ArGameNetVer;
 	ArIsLoading                          = ArchiveToCopy.ArIsLoading;
+	ArIsLoadingFromCookedPackage         = ArchiveToCopy.ArIsLoadingFromCookedPackage;
 	ArIsSaving                           = ArchiveToCopy.ArIsSaving;
 	ArIsTransacting                      = ArchiveToCopy.ArIsTransacting;
 	ArIsTextFormat                       = ArchiveToCopy.ArIsTextFormat;
@@ -400,11 +402,6 @@ void FArchive::PopSerializedProperty(class FProperty* InProperty, const bool bIs
 	}
 }
 
-bool FArchive::IsUsingEventDrivenLoader() const
-{
-	return GEventDrivenLoaderEnabled;
-}
-
 #if WITH_EDITORONLY_DATA
 bool FArchiveState::IsEditorOnlyPropertyOnTheStack() const
 {
@@ -1372,6 +1369,11 @@ void FArchiveState::SetIsLoading(bool bInIsLoading)
 	ArIsLoading = bInIsLoading;
 }
 
+void FArchiveState::SetIsLoadingFromCookedPackage(bool bInIsLoadingFromCookedPackage)
+{
+	ArIsLoadingFromCookedPackage = bInIsLoadingFromCookedPackage;
+}
+
 void FArchiveState::SetIsSaving(bool bInIsSaving)
 {
 	ArIsSaving = bInIsSaving;
diff --git a/Engine/Source/Runtime/Core/Public/IO/IoDispatcher.h b/Engine/Source/Runtime/Core/Public/IO/IoDispatcher.h
--- a/Engine/Source/Runtime/Core/Public/IO/IoDispatcher.h
+++ b/Engine/Source/Runtime/Core/Public/IO/IoDispatcher.h
@@ -35,6 +35,7 @@ class IMappedFileHandle;
 class IMappedFileRegion;
 class FIoDirectoryIndexReaderImpl;
 struct IIoDispatcherBackend;
+class FPackageId;
 
 CORE_API DECLARE_LOG_CATEGORY_EXTERN(LogIoDispatcher, Log, All);
 
@@ -772,6 +773,7 @@ static FIoChunkId CreateIoChunkId(uint64 ChunkId, uint16 ChunkIndex, EIoChunkTyp
 	return IoChunkId;
 }
 
+CORE_API FIoChunkId CreatePackageDataChunkId(const FPackageId& PackageId);
 CORE_API FIoChunkId CreateExternalFileChunkId(const FStringView Filename);
 
 //////////////////////////////////////////////////////////////////////////
diff --git a/Engine/Source/Runtime/Core/Public/Serialization/Archive.h b/Engine/Source/Runtime/Core/Public/Serialization/Archive.h
--- a/Engine/Source/Runtime/Core/Public/Serialization/Archive.h
+++ b/Engine/Source/Runtime/Core/Public/Serialization/Archive.h
@@ -245,6 +245,12 @@ public:
 		return ArIsLoading;
 	}
 
+	/** Returns true if this archive is loading from a cooked package. */
+ 	FORCEINLINE bool IsLoadingFromCookedPackage() const
+	{
+		return ArIsLoadingFromCookedPackage;
+	}
+
 	/** Returns true if this archive is for saving data, this can also be a pre-save preparation archive. */
 	FORCEINLINE bool IsSaving() const
 	{
@@ -722,6 +728,9 @@ protected:
 	/** Whether this archive is for loading data. */
 	uint8 ArIsLoading : 1;
 
+	/** Whether this archive is loading from a cooked package. */
+	uint8 ArIsLoadingFromCookedPackage : 1;
+
 	/** Whether this archive is for saving data. */
 	uint8 ArIsSaving : 1;
 
@@ -827,6 +836,12 @@ public:
 	 */
 	virtual void SetIsLoading(bool bInIsLoading);
 
+	/** 
+	 * Sets whether the archive is loading from a cooked package.
+	 * @param bInIsLoadingFromCookedPackage  true if this archive is loading from a cooked package, false otherwise
+	 */
+	virtual void SetIsLoadingFromCookedPackage(bool bInIsLoadingFromCookedPackage);
+
 	/**
 	 * Sets whether this archive is for saving data.
 	 *
@@ -1820,6 +1835,7 @@ public:
 	}
 
 	using FArchiveState::IsLoading;
+	using FArchiveState::IsLoadingFromCookedPackage;
 	using FArchiveState::IsSaving;
 	using FArchiveState::IsTransacting;
 	using FArchiveState::IsTextFormat;
@@ -1902,10 +1918,6 @@ public:
 	 */
 	virtual bool AttachExternalReadDependency(FExternalReadCallback& ReadCallback) { return false; }
 
-	/**
-	 * Returns whether the Event Driven Loader is enabled or not.
-	 */
-	virtual bool IsUsingEventDrivenLoader() const;
 
 #if USE_STABLE_LOCALIZATION_KEYS
 	using FArchiveState::SetLocalizationNamespace;
@@ -1995,6 +2007,7 @@ private:
 
 private:
 	using FArchiveState::ArIsLoading;
+	using FArchiveState::ArIsLoadingFromCookedPackage;
 	using FArchiveState::ArIsSaving;
 	using FArchiveState::ArIsTransacting;
 	using FArchiveState::ArIsTextFormat;
@@ -2031,6 +2044,7 @@ public:
 
 public:
 	using FArchiveState::SetIsLoading;
+	using FArchiveState::SetIsLoadingFromCookedPackage;
 	using FArchiveState::SetIsSaving;
 	using FArchiveState::SetIsTransacting;
 	using FArchiveState::SetIsTextFormat;
diff --git a/Engine/Source/Runtime/CoreUObject/Private/Misc/PackageName.cpp b/Engine/Source/Runtime/CoreUObject/Private/Misc/PackageName.cpp
--- a/Engine/Source/Runtime/CoreUObject/Private/Misc/PackageName.cpp
+++ b/Engine/Source/Runtime/CoreUObject/Private/Misc/PackageName.cpp
@@ -2110,8 +2110,6 @@ bool FPackageName::IsMapPackageExtension(const TCHAR* Ext)
 
 bool FPackageName::FindPackagesInDirectory( TArray<FString>& OutPackages, const FString& RootDir )
 {
-	UE_CLOG(FIoDispatcher::IsInitialized(), LogPackageName, Error, TEXT("Can't search for packages using the filesystem when I/O dispatcher is enabled"));
-
 	// Keep track if any package has been found. Can't rely only on OutPackages.Num() > 0 as it may not be empty.
 	const int32 PreviousPackagesCount = OutPackages.Num();
 	IteratePackagesInDirectory(RootDir, [&OutPackages](const TCHAR* PackageFilename) -> bool
diff --git a/Engine/Source/Runtime/CoreUObject/Private/Serialization/AsyncLoading.cpp b/Engine/Source/Runtime/CoreUObject/Private/Serialization/AsyncLoading.cpp
--- a/Engine/Source/Runtime/CoreUObject/Private/Serialization/AsyncLoading.cpp
+++ b/Engine/Source/Runtime/CoreUObject/Private/Serialization/AsyncLoading.cpp
@@ -6755,6 +6755,11 @@ void FAsyncPackage::UpdateLoadPercentage()
 	LoadPercentage = FMath::Max(NewLoadPercentage, LoadPercentage);
 }
 
+bool FAsyncLoadingThread::ShouldAlwaysLoadPackageAsync(const FPackagePath& InPackagePath)
+{
+	return FPlatformProperties::RequiresCookedData() && GEventDrivenLoaderEnabled && EVENT_DRIVEN_ASYNC_LOAD_ACTIVE_AT_RUNTIME;
+}
+
 int32 FAsyncLoadingThread::LoadPackage(const FPackagePath& InPackagePath, FName InCustomName, FLoadPackageAsyncDelegate InCompletionDelegate, EPackageFlags InPackageFlags, int32 InPIEInstanceID, int32 InPackagePriority, const FLinkerInstancingContext* InstancingContext)
 {
 	static bool bOnce = false;
diff --git a/Engine/Source/Runtime/CoreUObject/Private/Serialization/AsyncLoading2.cpp b/Engine/Source/Runtime/CoreUObject/Private/Serialization/AsyncLoading2.cpp
--- a/Engine/Source/Runtime/CoreUObject/Private/Serialization/AsyncLoading2.cpp
+++ b/Engine/Source/Runtime/CoreUObject/Private/Serialization/AsyncLoading2.cpp
@@ -1132,11 +1132,6 @@ public:
 	void UsingCustomVersion(const FGuid& Key) override {};
 	using FArchive::operator<<; // For visibility of the overloads we don't override
 
-	virtual bool IsUsingEventDrivenLoader() const override
-	{
-		return true;
-	}
-
 	/** FExportArchive will be created on the stack so we do not want BulkData objects caching references to it. */
 	virtual FArchive* GetCacheableArchive()
 	{
@@ -2199,6 +2194,11 @@ public:
 
 	virtual void ShutdownLoading() override;
 
+	virtual bool ShouldAlwaysLoadPackageAsync(const FPackagePath& PackagePath) override
+	{
+		return true;
+	}
+
 	virtual int32 LoadPackage(
 		const FPackagePath& InPackagePath,
 		FName InCustomName,
@@ -3383,7 +3383,7 @@ void FAsyncPackage2::StartLoading(FIoBatch& IoBatch)
 	TRACE_COUNTER_SET(AsyncLoadingPendingIoRequests, LocalPendingIoRequestsCounter);
 
 	FIoReadOptions ReadOptions;
-	IoRequest = IoBatch.ReadWithCallback(CreateIoChunkId(Desc.PackageIdToLoad.Value(), 0, EIoChunkType::ExportBundleData),
+	IoRequest = IoBatch.ReadWithCallback(CreatePackageDataChunkId(Desc.PackageIdToLoad),
 		ReadOptions,
 		Desc.Priority,
 		[this](TIoStatusOr<FIoBuffer> Result)
@@ -3599,6 +3599,7 @@ EAsyncPackageState::Type FAsyncPackage2::Event_ProcessExportBundle(FAsyncLoading
 				Ar.SetCustomVersions(Package->LinkerRoot->GetLinkerCustomVersions());
 			}
 			Ar.SetUseUnversionedPropertySerialization((Package->LinkerRoot->GetPackageFlags() & PKG_UnversionedProperties) != 0);
+			Ar.SetIsLoadingFromCookedPackage((Package->LinkerRoot->GetPackageFlags() & PKG_Cooked) != 0);
 			Ar.SetIsLoading(true);
 			Ar.SetIsPersistent(true);
 			if (Package->LinkerRoot->GetPackageFlags() & PKG_FilterEditorOnly)
@@ -6091,6 +6092,11 @@ IAsyncPackageLoader* MakeAsyncPackageLoader2(FIoDispatcher& InIoDispatcher, IAsy
 
 #endif //WITH_ASYNCLOADING2
 
+bool IsPackageLoadingFromIoDispatcher(const UPackage* Package, const FArchive& Ar)
+{
+	return Package && Package->GetPackageId().IsValid() && Ar.IsLoadingFromCookedPackage();
+}
+
 #if UE_BUILD_DEVELOPMENT || UE_BUILD_DEBUG
 PRAGMA_ENABLE_OPTIMIZATION
 #endif
diff --git a/Engine/Source/Runtime/CoreUObject/Private/Serialization/AsyncLoadingPrivate.h b/Engine/Source/Runtime/CoreUObject/Private/Serialization/AsyncLoadingPrivate.h
--- a/Engine/Source/Runtime/CoreUObject/Private/Serialization/AsyncLoadingPrivate.h
+++ b/Engine/Source/Runtime/CoreUObject/Private/Serialization/AsyncLoadingPrivate.h
@@ -149,3 +149,5 @@ private:
 
 	FLinkerLoad* OwnerLinker;
 };
+
+bool IsPackageLoadingFromIoDispatcher(const UPackage* Package, const FArchive& Ar);
\ No newline at end of file
diff --git a/Engine/Source/Runtime/CoreUObject/Private/Serialization/AsyncLoadingThread.h b/Engine/Source/Runtime/CoreUObject/Private/Serialization/AsyncLoadingThread.h
--- a/Engine/Source/Runtime/CoreUObject/Private/Serialization/AsyncLoadingThread.h
+++ b/Engine/Source/Runtime/CoreUObject/Private/Serialization/AsyncLoadingThread.h
@@ -265,6 +265,8 @@ public:
 	/** Start the async loading thread */
 	void StartThread() override;
 
+	bool ShouldAlwaysLoadPackageAsync(const FPackagePath& InPackagePath) override;
+
 	int32 LoadPackage(
 			const FPackagePath& InPackagePath,
 			FName InCustomName,
diff --git a/Engine/Source/Runtime/CoreUObject/Private/Serialization/AsyncPackageLoader.cpp b/Engine/Source/Runtime/CoreUObject/Private/Serialization/AsyncPackageLoader.cpp
--- a/Engine/Source/Runtime/CoreUObject/Private/Serialization/AsyncPackageLoader.cpp
+++ b/Engine/Source/Runtime/CoreUObject/Private/Serialization/AsyncPackageLoader.cpp
@@ -558,15 +558,21 @@ void InitAsyncThread()
 #if WITH_ASYNCLOADING2
 	if (FIoDispatcher::IsInitialized())
 	{
-		GetGEDLBootNotificationManager().Disable();
-#if WITH_IOSTORE_IN_EDITOR
-		GPackageLoader = MakeEditorPackageLoader(FIoDispatcher::Get(), GetGEDLBootNotificationManager());
+		FIoDispatcher& IoDispatcher = FIoDispatcher::Get();
+		bool bHasScriptObjectsChunk = IoDispatcher.DoesChunkExist(CreateIoChunkId(0, 0, EIoChunkType::ScriptObjects));
+		bool bHasUseIoStoreParamInEditor = WITH_EDITOR && FParse::Param(FCommandLine::Get(), TEXT("UseIoStore"));
+		if (bHasScriptObjectsChunk || bHasUseIoStoreParamInEditor)
+		{
+			GetGEDLBootNotificationManager().Disable();
+#if WITH_EDITOR
+			GPackageLoader = MakeEditorPackageLoader(IoDispatcher, GetGEDLBootNotificationManager());
 #else
-		GPackageLoader.Reset(MakeAsyncPackageLoader2(FIoDispatcher::Get()));
+			GPackageLoader.Reset(MakeAsyncPackageLoader2(IoDispatcher));
 #endif
+		}
 	}
-	else
 #endif
+	if (!GPackageLoader.IsValid())
 	{
 		GPackageLoader = MakeUnique<FAsyncLoadingThread>(/** ThreadIndex = */ 0, GetGEDLBootNotificationManager());
 	}
@@ -925,6 +931,15 @@ int32 IAsyncPackageLoader::LoadPackage(
 	return LoadPackage(PackagePath, FName(InPackageName), InCompletionDelegate, InPackageFlags, InPIEInstanceID, InPackagePriority, InstancingContext);
 }
 
+bool ShouldAlwaysLoadPackageAsync(const FPackagePath& InPackagePath)
+{
+	if (!GPackageLoader)
+	{
+		return false;
+	}
+	return GPackageLoader->ShouldAlwaysLoadPackageAsync(InPackagePath);
+}
+
 int32 LoadPackageAsync(const FPackagePath& InPackagePath, FName InPackageNameToCreate /* = NAME_None*/, FLoadPackageAsyncDelegate InCompletionDelegate /*= FLoadPackageAsyncDelegate()*/, EPackageFlags InPackageFlags /*= PKG_None*/, int32 InPIEInstanceID /*= INDEX_NONE*/, int32 InPackagePriority /*= 0*/, const FLinkerInstancingContext* InstancingContext /*=nullptr*/)
 {
 	LLM_SCOPE(ELLMTag::AsyncLoading);
diff --git a/Engine/Source/Runtime/CoreUObject/Private/Serialization/BulkData.cpp b/Engine/Source/Runtime/CoreUObject/Private/Serialization/BulkData.cpp
--- a/Engine/Source/Runtime/CoreUObject/Private/Serialization/BulkData.cpp
+++ b/Engine/Source/Runtime/CoreUObject/Private/Serialization/BulkData.cpp
@@ -21,9 +21,8 @@
 #include "HAL/PlatformFileManager.h"
 #include "UObject/UObjectThreadContext.h"
 #include "ProfilingDebugging/LoadTimeTracker.h"
-#if WITH_IOSTORE_IN_EDITOR
 #include "IO/IoDispatcher.h"
-#endif
+#include "AsyncLoadingPrivate.h"
 
 /*-----------------------------------------------------------------------------
 	Constructors and operators
@@ -960,9 +959,9 @@ static FAutoConsoleVariableRef CVarMinimumBulkDataSizeForAsyncLoading(
 
 bool FUntypedBulkData::ShouldStreamBulkData(FArchive& Ar)
 {
-	if (Ar.IsUsingEventDrivenLoader())
+	if (Ar.IsLoadingFromCookedPackage())
 	{
-#if WITH_IOSTORE_IN_EDITOR
+#if WITH_EDITOR
 		// Streaming not yet implemented
 		return false;
 #else
@@ -1129,10 +1128,8 @@ void FUntypedBulkData::Serialize( FArchive& Ar, UObject* Owner, int32 Idx, bool
 			// GetLinker
 			bool bUseIOStore = false;
 #if WITH_EDITOR
-#if WITH_IOSTORE_IN_EDITOR
-			if (Ar.IsUsingEventDrivenLoader() && bPayloadInSeparateFile)
+			if (bPayloadInSeparateFile && Owner && IsPackageLoadingFromIoDispatcher(Owner->GetPackage(), Ar))
 			{
-				check(Owner && Owner->GetPackage()->GetPackageId().IsValid());
 				checkf(!(BulkDataFlags& BULKDATA_WorkspaceDomainPayload),
 					TEXT("%s IsUsingEventDrivenLoader but has a bulkdata with BULKDATA_WorkspaceDomainPayload. ")
 					TEXT("BULKDATA_WorkspaceDomainPayload is not supported with iostore."), *Ar.GetArchiveName());
@@ -1140,7 +1137,7 @@ void FUntypedBulkData::Serialize( FArchive& Ar, UObject* Owner, int32 Idx, bool
 				PackageId = Owner->GetPackage()->GetPackageId();
 				bUseIOStore = true;
 			}
-#endif // WITH_IOSTORE_IN_EDITOR
+
 			if (Owner != nullptr)
 			{
 				Linker = Owner->GetLinker();
diff --git a/Engine/Source/Runtime/CoreUObject/Private/Serialization/BulkData2.cpp b/Engine/Source/Runtime/CoreUObject/Private/Serialization/BulkData2.cpp
--- a/Engine/Source/Runtime/CoreUObject/Private/Serialization/BulkData2.cpp
+++ b/Engine/Source/Runtime/CoreUObject/Private/Serialization/BulkData2.cpp
@@ -19,6 +19,7 @@
 #include "IO/IoDispatcher.h"
 #include "Async/Async.h"
 #include "Async/MappedFileHandle.h"
+#include "AsyncLoadingPrivate.h"
 
 DEFINE_LOG_CATEGORY_STATIC(LogBulkDataRuntime, Log, All);
 
@@ -743,7 +744,7 @@ FBulkDataBase::~FBulkDataBase()
 	}
 }
 
-void FBulkDataBase::ConditionalSetInlineAlwaysAllowDiscard()
+void FBulkDataBase::ConditionalSetInlineAlwaysAllowDiscard(bool bPackageUsesIoStore)
 {
 	// If PackagePath is null and we do not have BULKDATA_UsesIoDispatcher, we will not be able to reload the bulkdata.
 	// We will not have a PackagePath if the engine is using IoStore.
@@ -770,7 +771,7 @@ void FBulkDataBase::ConditionalSetInlineAlwaysAllowDiscard()
 	// all inline data should be allowed to be discarded.
 
 #if !UE_KEEP_INLINE_RELOADING_CONSISTENT
-	if (FIoDispatcher::IsInitialized())
+	if (bPackageUsesIoStore)
 #endif // !UE_KEEP_INLINE_RELOADING_CONSISTENT
 	{
 		SetBulkDataFlags(BULKDATA_AlwaysAllowDiscard);
@@ -827,10 +828,10 @@ void FBulkDataBase::Serialize(FArchive& Ar, UObject* Owner, int32 /*Index*/, boo
 		const UPackage* Package = Owner->GetOutermost();
 		checkf(Package != nullptr, TEXT("FBulkDataBase::Serialize requires an Owner that returns a valid UPackage"));
 
-		const bool bEngineUsesIoStore = FIoDispatcher::IsInitialized();
+		const bool bPackageUsesIoStore = IsPackageLoadingFromIoDispatcher(Package, Ar);
 		const FPackagePath* PackagePath = nullptr;
 		const FLinkerLoad* Linker = nullptr;
-		if (bEngineUsesIoStore)
+		if (bPackageUsesIoStore)
 		{
 			checkf(IsInlined() || !NeedsOffsetFixup(), TEXT("IODispatcher does not support offset fixups; SaveBulkData during cooking should have added the flag BULKDATA_NoOffsetFixUp."));
 			checkf(IsInlined() || IsInSeparateFile() || !GEventDrivenLoaderEnabled,
@@ -894,7 +895,7 @@ void FBulkDataBase::Serialize(FArchive& Ar, UObject* Owner, int32 /*Index*/, boo
 			SerializeBulkData(Ar, DataBuffer, BulkDataSize);
 
 			// Apply the legacy BULKDATA_AlwaysAllowDiscard fix for inline data in some conditions
-			ConditionalSetInlineAlwaysAllowDiscard();
+			ConditionalSetInlineAlwaysAllowDiscard(bPackageUsesIoStore);
 		}
 		else
 		{
@@ -903,16 +904,16 @@ void FBulkDataBase::Serialize(FArchive& Ar, UObject* Owner, int32 /*Index*/, boo
 				ProcessDuplicateData(DuplicateFlags, DuplicateSizeOnDisk, DuplicateOffset, Package, PackagePath, Linker);
 			}
 
-			if (bAttemptFileMapping && !IsInSeparateFile() && (bEngineUsesIoStore || !Ar.IsAllowingLazyLoading()))
+			if (bAttemptFileMapping && !IsInSeparateFile() && (bPackageUsesIoStore || !Ar.IsAllowingLazyLoading()))
 			{
 				UE_CLOG(bAttemptFileMapping, LogSerialization, Error,
 					TEXT("Attempt to file map BulkData in end-of-package-file section, this is not supported when %s. Package '%s'"),
-					bEngineUsesIoStore ? TEXT("using IoDispatcher") : TEXT("archive does not support lazyload"), *Package->GetFName().ToString());
+					bPackageUsesIoStore ? TEXT("using IoDispatcher") : TEXT("archive does not support lazyload"), *Package->GetFName().ToString());
 				bShouldForceLoad = true; // Signal we want to force the BulkData to load
 			}
 			else if (bAttemptFileMapping)
 			{
-				if (bEngineUsesIoStore)
+				if (bPackageUsesIoStore)
 				{
 					check(IsInSeparateFile());
 					TIoStatusOr<FIoMappedRegion> Status = IoDispatcher->OpenMapped(CreateChunkId(), FIoReadOptions());
@@ -939,7 +940,7 @@ void FBulkDataBase::Serialize(FArchive& Ar, UObject* Owner, int32 /*Index*/, boo
 					}
 				}
 			}
-			else if (!IsInSeparateFile() && (bEngineUsesIoStore || !Ar.IsAllowingLazyLoading()))
+			else if (!IsInSeparateFile() && (bPackageUsesIoStore || !Ar.IsAllowingLazyLoading()))
 			{
 				// If the data is in the end-of-package-file section, and we can't load the package file again later either
 				// because we're using the IoDispatcher or because the archive does not support lazy loading then we have
diff --git a/Engine/Source/Runtime/CoreUObject/Private/Serialization/EditorPackageLoader.cpp b/Engine/Source/Runtime/CoreUObject/Private/Serialization/EditorPackageLoader.cpp
--- a/Engine/Source/Runtime/CoreUObject/Private/Serialization/EditorPackageLoader.cpp
+++ b/Engine/Source/Runtime/CoreUObject/Private/Serialization/EditorPackageLoader.cpp
@@ -53,6 +53,13 @@ public:
 		UncookedPackageLoader->StartThread();
 	}
 
+	virtual bool ShouldAlwaysLoadPackageAsync(const FPackagePath& InPackagePath) override
+	{
+		// Use the old loader if an uncooked package exists on disk
+		const bool bDoesUncookedPackageExist = FPackageName::DoesPackageExistEx(InPackagePath, FPackageName::EPackageLocationFilter::Uncooked) != FPackageName::EPackageLocationFilter::None;
+		return !bDoesUncookedPackageExist;
+	}
+
 	virtual int32 LoadPackage(
 		const FPackagePath& PackagePath,
 		FName CustomPackageName,
diff --git a/Engine/Source/Runtime/CoreUObject/Private/Tests/ObjectHandleTest.cpp b/Engine/Source/Runtime/CoreUObject/Private/Tests/ObjectHandleTest.cpp
--- a/Engine/Source/Runtime/CoreUObject/Private/Tests/ObjectHandleTest.cpp
+++ b/Engine/Source/Runtime/CoreUObject/Private/Tests/ObjectHandleTest.cpp
@@ -181,19 +181,6 @@ bool FObjectHandleTestResolveEngineContentTarget::RunTest(const FString& Paramet
 IMPLEMENT_CUSTOM_SIMPLE_AUTOMATION_TEST(FObjectHandleTestResolveNonExistentTarget, FObjectHandleTestBase, TEST_NAME_ROOT TEXT(".ResolveNonExistentTarget"), ObjectHandleTestFlags)
 bool FObjectHandleTestResolveNonExistentTarget::RunTest(const FString& Parameters)
 {
-	if (FPlatformProperties::RequiresCookedData())
-	{
-		if (FIoDispatcher::IsInitialized())
-		{
-			AddExpectedError(TEXT("SkipPackage"));
-		}
-		else
-		{
-			AddExpectedError(TEXT("Couldn't find file for package"));
-			AddExpectedError(TEXT("Found 0 dependent packages..."));
-		}
-	}
-
 	// Confirm we don't successfully resolve an incorrect reference to engine content
 	TestResolveFailure("/Engine/EngineResources/NonExistentPackageName_0", "DefaultTexture");
 	TestResolveFailure("/Engine/EngineResources/DefaultTexture", "NonExistentObject_0");
diff --git a/Engine/Source/Runtime/CoreUObject/Private/UObject/Class.cpp b/Engine/Source/Runtime/CoreUObject/Private/UObject/Class.cpp
--- a/Engine/Source/Runtime/CoreUObject/Private/UObject/Class.cpp
+++ b/Engine/Source/Runtime/CoreUObject/Private/UObject/Class.cpp
@@ -1284,7 +1284,7 @@ void UStruct::SerializeVersionedTaggedProperties(FStructuredArchive::FSlot Slot,
 
 	// Determine if this struct supports optional property guid's (UBlueprintGeneratedClasses Only)
 	const bool bArePropertyGuidsAvailable = (UnderlyingArchive.UEVer() >= VER_UE4_PROPERTY_GUID_IN_PROPERTY_TAG) && !FPlatformProperties::RequiresCookedData() && ArePropertyGuidsAvailable();
-	const bool bUseRedirects = (!FPlatformProperties::RequiresCookedData() || UnderlyingArchive.IsSaveGame()) && !UnderlyingArchive.IsUsingEventDrivenLoader();
+	const bool bUseRedirects = (!FPlatformProperties::RequiresCookedData() || UnderlyingArchive.IsSaveGame()) && !UnderlyingArchive.IsLoadingFromCookedPackage();
 
 	if (UnderlyingArchive.IsLoading())
 	{
@@ -1433,7 +1433,7 @@ void UStruct::SerializeVersionedTaggedProperties(FStructuredArchive::FSlot Slot,
 	#endif // WITH_EDITOR
 					// editoronly properties should be skipped if we are NOT the editor, or we are 
 					// the editor but are cooking for console (editoronly implies notforconsole)
-					if ((Property->PropertyFlags & CPF_EditorOnly) && ((!FPlatformProperties::HasEditorOnlyData() && !GForceLoadEditorOnly) || UnderlyingArchive.IsUsingEventDrivenLoader()))
+					if ((Property->PropertyFlags & CPF_EditorOnly) && ((!FPlatformProperties::HasEditorOnlyData() && !GForceLoadEditorOnly) || UnderlyingArchive.IsLoadingFromCookedPackage()))
 					{
 					}
 					// check for valid array index
@@ -4763,7 +4763,7 @@ void UClass::Serialize( FArchive& Ar )
 		if (ClassDefaultObject == NULL)
 		{
 			check(GConfig);
-			if (GEventDrivenLoaderEnabled || Ar.IsUsingEventDrivenLoader())
+			if (GEventDrivenLoaderEnabled || Ar.IsLoadingFromCookedPackage())
 			{
 				ClassDefaultObject = GetDefaultObject();
 				// we do this later anyway, once we find it and set it in the export table. 
diff --git a/Engine/Source/Runtime/CoreUObject/Private/UObject/LinkerLoad.cpp b/Engine/Source/Runtime/CoreUObject/Private/UObject/LinkerLoad.cpp
--- a/Engine/Source/Runtime/CoreUObject/Private/UObject/LinkerLoad.cpp
+++ b/Engine/Source/Runtime/CoreUObject/Private/UObject/LinkerLoad.cpp
@@ -1261,6 +1261,10 @@ FLinkerLoad::ELinkerStatus FLinkerLoad::SerializePackageFileSummaryInternal()
 		return LINKER_Failed;
 	}
 
+	bool bIsCooked = (Summary.GetPackageFlags() & PKG_Cooked) != 0;
+	SetIsLoadingFromCookedPackage(bIsCooked);
+	Loader->SetIsLoadingFromCookedPackage(bIsCooked);
+
 	// Set desired property tag format
 	bool bUseUnversionedProperties = (Summary.GetPackageFlags() & PKG_UnversionedProperties) != 0;
 	SetUseUnversionedPropertySerialization(bUseUnversionedProperties);
diff --git a/Engine/Source/Runtime/CoreUObject/Private/UObject/UObjectGlobals.cpp b/Engine/Source/Runtime/CoreUObject/Private/UObject/UObjectGlobals.cpp
--- a/Engine/Source/Runtime/CoreUObject/Private/UObject/UObjectGlobals.cpp
+++ b/Engine/Source/Runtime/CoreUObject/Private/UObject/UObjectGlobals.cpp
@@ -1131,6 +1131,8 @@ UE_TRACE_EVENT_BEGIN(CUSTOM_LOADTIMER_LOG, LoadPackageInternal, NoSync)
 	UE_TRACE_EVENT_FIELD(UE::Trace::WideString, PackageName)
 UE_TRACE_EVENT_END()
 
+bool ShouldAlwaysLoadPackageAsync(const FPackagePath& InPackagePath);
+
 UPackage* LoadPackageInternal(UPackage* InOuter, const FPackagePath& PackagePath, uint32 LoadFlags, FLinkerLoad* ImportLinker, FArchive* InReaderOverride,
 	const FLinkerInstancingContext* InstancingContext, const FPackagePath* DiffPackagePath)
 {
@@ -1153,12 +1155,7 @@ UPackage* LoadPackageInternal(UPackage* InOuter, const FPackagePath& PackagePath
 	}
 	checkf(IsInGameThread(), TEXT("Unable to load %s. Objects and Packages can only be loaded from the game thread."), *PackagePath.GetDebugName());
 
-	if ((FPlatformProperties::RequiresCookedData() && GEventDrivenLoaderEnabled
-		&& EVENT_DRIVEN_ASYNC_LOAD_ACTIVE_AT_RUNTIME)
-#if WITH_IOSTORE_IN_EDITOR
-		|| FIoDispatcher::IsInitialized()
-#endif
-		)
+	if (ShouldAlwaysLoadPackageAsync(PackagePath))
 	{
 		checkf(!InOuter || !InOuter->GetOuter(), TEXT("Loading into subpackages is not implemented.")); // Subpackages are no longer supported in UE
 		FName PackageName(InOuter ? InOuter->GetFName() : PackagePath.GetPackageFName());
@@ -1170,28 +1167,21 @@ UPackage* LoadPackageInternal(UPackage* InOuter, const FPackagePath& PackagePath
 
 		UE_TRACK_REFERENCING_PACKAGE_SCOPED(PackageName, PackageAccessTrackingOps::NAME_Load);
 
-#if WITH_IOSTORE_IN_EDITOR
-		// Use the old loader if an uncooked package exists on disk
-		const bool bDoesUncookedPackageExist = FPackageName::DoesPackageExistEx(PackagePath, FPackageName::EPackageLocationFilter::Uncooked) != FPackageName::EPackageLocationFilter::None;
-		if (!bDoesUncookedPackageExist)
-#endif
+		if (FCoreDelegates::OnSyncLoadPackage.IsBound())
 		{
-			if (FCoreDelegates::OnSyncLoadPackage.IsBound())
-			{
-				FCoreDelegates::OnSyncLoadPackage.Broadcast(PackageName.ToString());
-			}
-
-			GSyncLoadUsingAsyncLoaderCount++;
-			int32 RequestID = LoadPackageAsync(PackagePath, PackageName);
+			FCoreDelegates::OnSyncLoadPackage.Broadcast(PackageName.ToString());
+		}
 
-			if (RequestID != INDEX_NONE)
-			{
-				FlushAsyncLoading(RequestID);
-			}
-			GSyncLoadUsingAsyncLoaderCount--;
+		GSyncLoadUsingAsyncLoaderCount++;
+		int32 RequestID = LoadPackageAsync(PackagePath, PackageName);
 
-			return (InOuter ? InOuter : FindObjectFast<UPackage>(nullptr, PackageName));
+		if (RequestID != INDEX_NONE)
+		{
+			FlushAsyncLoading(RequestID);
 		}
+		GSyncLoadUsingAsyncLoaderCount--;
+
+		return (InOuter ? InOuter : FindObjectFast<UPackage>(nullptr, PackageName));
 	}
 
 	UPackage* Result = nullptr;
diff --git a/Engine/Source/Runtime/CoreUObject/Public/Serialization/AsyncPackageLoader.h b/Engine/Source/Runtime/CoreUObject/Public/Serialization/AsyncPackageLoader.h
--- a/Engine/Source/Runtime/CoreUObject/Public/Serialization/AsyncPackageLoader.h
+++ b/Engine/Source/Runtime/CoreUObject/Public/Serialization/AsyncPackageLoader.h
@@ -97,6 +97,8 @@ public:
 
 	virtual void StartThread() = 0;
 
+	virtual bool ShouldAlwaysLoadPackageAsync(const FPackagePath& PackagePath) = 0;
+
 	UE_DEPRECATED(5.0, "Call the version that takes a FPackagePath without a FGuid instead")
 	COREUOBJECT_API int32 LoadPackage(
 		const FString& InPackageName,
diff --git a/Engine/Source/Runtime/CoreUObject/Public/Serialization/BulkData2.h b/Engine/Source/Runtime/CoreUObject/Public/Serialization/BulkData2.h
--- a/Engine/Source/Runtime/CoreUObject/Public/Serialization/BulkData2.h
+++ b/Engine/Source/Runtime/CoreUObject/Public/Serialization/BulkData2.h
@@ -295,7 +295,7 @@ private:
 
 	void ProcessDuplicateData(EBulkDataFlags NewFlags, int64 NewSizeOnDisk, int64 NewOffset, const UPackage* Package,
 		const FPackagePath* PackagePath, const FLinkerLoad* Linker);
-	void ConditionalSetInlineAlwaysAllowDiscard();
+	void ConditionalSetInlineAlwaysAllowDiscard(bool bPackageUsesIoStore);
 	void SerializeDuplicateData(FArchive& Ar, EBulkDataFlags& OutBulkDataFlags, int64& OutBulkDataSizeOnDisk, int64& OutBulkDataOffsetInFile);
 	void SerializeBulkData(FArchive& Ar, void* DstBuffer, int64 DataLength);
 
diff --git a/Engine/Source/Runtime/Engine/Private/EngineUtils.cpp b/Engine/Source/Runtime/Engine/Private/EngineUtils.cpp
--- a/Engine/Source/Runtime/Engine/Private/EngineUtils.cpp
+++ b/Engine/Source/Runtime/Engine/Private/EngineUtils.cpp
@@ -294,15 +294,14 @@ bool EngineUtils::FindOrLoadAssetsByPath(const FString& Path, TArray<UObject*>&
 		return false;
 	}
 
-	using FPackageNames = TArray<FName, TInlineAllocator<16>>;
+	using FPackageNames = TSet<FName>;
 
 	auto GetPackageNamesFromPath = [](const FString& InPath, FPackageNames& OutPackageNames)
 	{
 		// There is no filesystem support for packages when using the I/O dispatcher
-		if (FIoDispatcher::IsInitialized())
+		if (FAssetRegistryModule* AssetRegistryModule = FModuleManager::LoadModulePtr<FAssetRegistryModule>(TEXT("AssetRegistry")))
 		{
-			FAssetRegistryModule& AssetRegistryModule = FModuleManager::LoadModuleChecked<FAssetRegistryModule>(TEXT("AssetRegistry"));
-			IAssetRegistry& AssetRegistry = AssetRegistryModule.Get();
+			IAssetRegistry& AssetRegistry = AssetRegistryModule->Get();
 
 			TArray<FAssetData> Assets;
 			AssetRegistry.GetAssetsByPath(FName(*InPath), Assets, true);
@@ -316,28 +315,27 @@ bool EngineUtils::FindOrLoadAssetsByPath(const FString& Path, TArray<UObject*>&
 				}
 			}
 		}
-		else
-		{
-			// Convert the package path to a filename with no extension (directory)
-			const FString FilePath = FPackageName::LongPackageNameToFilename(InPath);
 
-			// Gather the package files in that directory and subdirectories
-			TArray<FString> Filenames;
-			FPackageName::FindPackagesInDirectory(Filenames, FilePath);
+		// Convert the package path to a filename with no extension (directory)
+		const FString FilePath = FPackageName::LongPackageNameToFilename(InPath);
+
+		// Gather the package files in that directory and subdirectories
+		TArray<FString> Filenames;
+		FPackageName::FindPackagesInDirectory(Filenames, FilePath);
 
-			// Cull out map files
-			for (const FString& Filename : Filenames)
+		// Cull out map files
+		for (const FString& Filename : Filenames)
+		{
+			FStringView Extension = FPathViews::GetExtension(Filename, true);
+			if (Extension != FPackageName::GetMapPackageExtension())
 			{
-				FStringView Extension = FPathViews::GetExtension(Filename, true);
-				if (Extension != FPackageName::GetMapPackageExtension())
-				{
-					OutPackageNames.Emplace(*FPackageName::FilenameToLongPackageName(Filename));
-				}
+				OutPackageNames.Emplace(*FPackageName::FilenameToLongPackageName(Filename));
 			}
 		}
 	};
 
 	FPackageNames PackageNames;
+	PackageNames.Reserve(16);
 	GetPackageNamesFromPath(Path, PackageNames);
 	TCHAR PackageName[FName::StringBufferSize];
 
diff --git a/Engine/Source/Runtime/Engine/Private/UnrealEngine.cpp b/Engine/Source/Runtime/Engine/Private/UnrealEngine.cpp
--- a/Engine/Source/Runtime/Engine/Private/UnrealEngine.cpp
+++ b/Engine/Source/Runtime/Engine/Private/UnrealEngine.cpp
@@ -10270,7 +10270,7 @@ FGuid UEngine::GetPackageGuid(FName PackageName, bool bForPIE)
 {
 	FGuid Result(0,0,0,0);
 	// There is no package guid support when using the I/O dispatcher
-	if (FIoDispatcher::IsInitialized())
+	if (FIoDispatcher::IsInitialized() && FIoDispatcher::Get().DoesChunkExist(CreatePackageDataChunkId(FPackageId::FromName(PackageName))))
 	{
 		return Result;
 	}
diff --git a/Engine/Source/Runtime/Engine/Private/UserInterface/Console.cpp b/Engine/Source/Runtime/Engine/Private/UserInterface/Console.cpp
--- a/Engine/Source/Runtime/Engine/Private/UserInterface/Console.cpp
+++ b/Engine/Source/Runtime/Engine/Private/UserInterface/Console.cpp
@@ -285,41 +285,34 @@ void UConsole::BuildRuntimeAutoCompleteList(bool bForce)
 	{
 		auto FindPackagesInDirectory = [](TArray<FString>& OutPackages, const FString& InPath)
 		{
-			// Can't search packages using the filesystem when I/O dispatcher is enabled
-			if (FIoDispatcher::IsInitialized())
+			FString PackagePath;
+			if (FPackageName::TryConvertFilenameToLongPackageName(InPath, PackagePath))
 			{
-				FString PackagePath;
-				if (FPackageName::TryConvertFilenameToLongPackageName(InPath, PackagePath))
+				if (FAssetRegistryModule* AssetRegistryModule = FModuleManager::LoadModulePtr<FAssetRegistryModule>(TEXT("AssetRegistry")))
 				{
-					if (FAssetRegistryModule* AssetRegistryModule = FModuleManager::LoadModulePtr<FAssetRegistryModule>(TEXT("AssetRegistry")))
-					{
-						TArray<FAssetData> Assets;
-						AssetRegistryModule->Get().GetAssetsByPath(FName(*PackagePath), Assets, true);
+					TArray<FAssetData> Assets;
+					AssetRegistryModule->Get().GetAssetsByPath(FName(*PackagePath), Assets, true);
 
-						for (const FAssetData& Asset : Assets)
+					for (const FAssetData& Asset : Assets)
+					{
+						if (!!(Asset.PackageFlags & PKG_ContainsMap) && Asset.IsUAsset())
 						{
-							if (!!(Asset.PackageFlags & PKG_ContainsMap) && Asset.IsUAsset())
-							{
-								OutPackages.Emplace(Asset.AssetName.ToString());
-							}
+							OutPackages.AddUnique(Asset.AssetName.ToString());
 						}
 					}
 				}
 			}
-			else
+			TArray<FString> Filenames;
+			FPackageName::FindPackagesInDirectory(Filenames, InPath);
+
+			for (const FString& Filename : Filenames)
 			{
-				TArray<FString> Filenames;
-				FPackageName::FindPackagesInDirectory(Filenames, InPath);
+				const int32 NameStartIdx = Filename.Find(TEXT("/"), ESearchCase::IgnoreCase, ESearchDir::FromEnd);
+				const int32 ExtIdx = Filename.Find(*FPackageName::GetMapPackageExtension(), ESearchCase::IgnoreCase, ESearchDir::FromEnd);
 
-				for (const FString& Filename : Filenames)
+				if (NameStartIdx != INDEX_NONE && ExtIdx != INDEX_NONE)
 				{
-					const int32 NameStartIdx = Filename.Find(TEXT("/"), ESearchCase::IgnoreCase, ESearchDir::FromEnd);
-					const int32 ExtIdx = Filename.Find(*FPackageName::GetMapPackageExtension(), ESearchCase::IgnoreCase, ESearchDir::FromEnd);
-
-					if (NameStartIdx != INDEX_NONE && ExtIdx != INDEX_NONE)
-					{
-						OutPackages.Emplace(Filename.Mid(NameStartIdx + 1, ExtIdx - NameStartIdx - 1));
-					}
+					OutPackages.AddUnique(Filename.Mid(NameStartIdx + 1, ExtIdx - NameStartIdx - 1));
 				}
 			}
 		};
