Renamed EPackageLocationFilter enums from Cooked/Uncooked to IoDispatcher/FileSystem
Cleaned up duplicated logic between DoesPackageExist and DoesPackageExistEx
#rnx

diff --git a/Engine/Source/Editor/Kismet/Private/FindInBlueprintManager.cpp b/Engine/Source/Editor/Kismet/Private/FindInBlueprintManager.cpp
--- a/Engine/Source/Editor/Kismet/Private/FindInBlueprintManager.cpp
+++ b/Engine/Source/Editor/Kismet/Private/FindInBlueprintManager.cpp
@@ -2417,7 +2417,7 @@ void FFindInBlueprintSearchManager::OnBlueprintUnloaded(UBlueprint* InBlueprint)
 							const FPackageName::EPackageLocationFilter PackageLocation = FPackageName::DoesPackageExistEx(PackagePath, FPackageName::EPackageLocationFilter::Any, /*bMatchCaseOnDisk*/ false, &OutPackagePath);
 							if (PackageLocation != FPackageName::EPackageLocationFilter::None)
 							{
-								if (PackageLocation == FPackageName::EPackageLocationFilter::Uncooked && OutPackagePath.HasLocalPath())
+								if (PackageLocation == FPackageName::EPackageLocationFilter::FileSystem && OutPackagePath.HasLocalPath())
 								{
 									TArray<FString> FilesToScan = { OutPackagePath.GetLocalFullPath() };
 									AssetRegistryModule->Get().ScanModifiedAssetFiles(FilesToScan);
diff --git a/Engine/Source/Runtime/CoreUObject/Private/Misc/PackageName.cpp b/Engine/Source/Runtime/CoreUObject/Private/Misc/PackageName.cpp
--- a/Engine/Source/Runtime/CoreUObject/Private/Misc/PackageName.cpp
+++ b/Engine/Source/Runtime/CoreUObject/Private/Misc/PackageName.cpp
@@ -1590,70 +1590,18 @@ bool FPackageName::DoesPackageExist(const FPackagePath& PackagePath, FPackagePat
 
 bool FPackageName::DoesPackageExist(const FPackagePath& PackagePath, bool bMatchCaseOnDisk, FPackagePath* OutPackagePath)
 {
-	TRACE_CPUPROFILER_EVENT_SCOPE(FPackageName::DoesPackageExist);
-
-	// DoesPackageExist returns false for local filenames that are in unmounted directories, even if those files exist on the local disk
-	if (!PackagePath.IsMountedPath())
-	{
-		return false;
-	}
-	TStringBuilder<256> PackageName;
-	PackagePath.AppendPackageName(PackageName);
-
-	// Once we have the real Package Name, we can exit early if it's a script package - they exist only in memory.
-	if (IsScriptPackage(PackageName))
-	{
-		return false;
-	}
-
-	if (IsMemoryPackage(PackageName))
-	{
-		return false;
-	}
-
-	FText Reason;
-	if ( !FPackageName::IsValidTextForLongPackageName( PackageName, &Reason ) )
-	{
-		UE_LOG(LogPackageName, Error, TEXT( "DoesPackageExist: DoesPackageExist FAILED: '%s' is not a long packagename name. Reason: %s"), PackageName.ToString(), *Reason.ToString() );
-		return false;
-	}
-
-	// Used when I/O dispatcher is enabled
-	if (DoesPackageExistOverrideDelegate.IsBound())
-	{
-		if (DoesPackageExistOverrideDelegate.Execute(FName(*PackageName)))
-		{
-			if (OutPackagePath)
-			{
-				*OutPackagePath = PackagePath;
-				if (OutPackagePath->GetHeaderExtension() == EPackageExtension::Unspecified)
-				{
-					OutPackagePath->SetHeaderExtension(EPackageExtension::EmptyString);
-				}
-			}
-			return true;
-		}
-
-		// Try to find uncooked packages on disk when I/O store is enabled in editor builds
-#if !WITH_IOSTORE_IN_EDITOR
-		return false; 
-#endif
-	}
-
-	if (bMatchCaseOnDisk)
-	{
-		return IPackageResourceManager::Get().TryMatchCaseOnDisk(PackagePath, OutPackagePath);
-	}
-	else
-	{
-		return IPackageResourceManager::Get().DoesPackageExist(PackagePath, OutPackagePath);
-	}
+	return DoesPackageExistEx(PackagePath, EPackageLocationFilter::Any, bMatchCaseOnDisk, OutPackagePath) != EPackageLocationFilter::None;
 }
 
 FPackageName::EPackageLocationFilter FPackageName::DoesPackageExistEx(const FPackagePath& PackagePath, EPackageLocationFilter Filter, bool bMatchCaseOnDisk, FPackagePath* OutPackagePath)
 {
 	TRACE_CPUPROFILER_EVENT_SCOPE(FPackageName::DoesPackageExistEx);
 
+	// DoesPackageExist returns false for local filenames that are in unmounted directories, even if those files exist on the local disk
+	if (!PackagePath.IsMountedPath())
+	{
+		return EPackageLocationFilter::None;
+	}
 	TStringBuilder<256> PackageName;
 	PackagePath.AppendPackageName(PackageName);
 
@@ -1676,54 +1624,47 @@ FPackageName::EPackageLocationFilter FPackageName::DoesPackageExistEx(const FPac
 	}
 
 	EPackageLocationFilter Result = EPackageLocationFilter::None;
-	// Used when I/O dispatcher is enabled
-	if (((uint8)Filter & (uint8)EPackageLocationFilter::Cooked))
+	if (((uint8)Filter & (uint8)EPackageLocationFilter::IoDispatcher))
 	{
-		// For Cooked locations, return false for local filenames that are in unmounted directories
-		if (PackagePath.IsMountedPath())
+		if (DoesPackageExistOverrideDelegate.IsBound())
 		{
-			// Used when I/O dispatcher is enabled
-			if (DoesPackageExistOverrideDelegate.IsBound())
+			if (DoesPackageExistOverrideDelegate.Execute(PackagePath.GetPackageFName()))
 			{
-				if (DoesPackageExistOverrideDelegate.Execute(FName(*PackageName)))
+				if (OutPackagePath)
 				{
-					if (OutPackagePath)
+					*OutPackagePath = PackagePath;
+					if (OutPackagePath->GetHeaderExtension() == EPackageExtension::Unspecified)
 					{
-						*OutPackagePath = PackagePath;
-						if (OutPackagePath->GetHeaderExtension() == EPackageExtension::Unspecified)
-						{
-							OutPackagePath->SetHeaderExtension(EPackageExtension::EmptyString);
-						}
+						OutPackagePath->SetHeaderExtension(EPackageExtension::EmptyString);
 					}
+				}
 
-					// note that the file was found in Cooked (IOStore) location
-					Result = EPackageLocationFilter((uint8)Result | (uint8)EPackageLocationFilter::Cooked);
+				Result = EPackageLocationFilter((uint8)Result | (uint8)EPackageLocationFilter::IoDispatcher);
 
-					// if we just want to find any existence, then we are done and we can skip the on disk check lower
-					if (Filter == EPackageLocationFilter::Any)
-					{
-						return Result;
-					}
+				// if we just want to find any existence, then we are done and we can skip the on disk check lower
+				if (Filter == EPackageLocationFilter::Any)
+				{
+					return Result;
 				}
 			}
 		}
 	}
 
-	if ((uint8)Filter & (uint8)EPackageLocationFilter::Uncooked)
+	if ((uint8)Filter & (uint8)EPackageLocationFilter::FileSystem)
 	{
-		bool bFoundUncooked = false;
+		bool bFoundInFileSystem = false;
 		if (bMatchCaseOnDisk)
 		{
-			bFoundUncooked = IPackageResourceManager::Get().TryMatchCaseOnDisk(PackagePath, OutPackagePath);
+			bFoundInFileSystem = IPackageResourceManager::Get().TryMatchCaseOnDisk(PackagePath, OutPackagePath);
 		}
 		else
 		{
-			bFoundUncooked = IPackageResourceManager::Get().DoesPackageExist(PackagePath, OutPackagePath);
+			bFoundInFileSystem = IPackageResourceManager::Get().DoesPackageExist(PackagePath, OutPackagePath);
 		}
 
-		if (bFoundUncooked)
+		if (bFoundInFileSystem)
 		{
-			Result = EPackageLocationFilter((uint8)Result | (uint8)EPackageLocationFilter::Uncooked);
+			Result = EPackageLocationFilter((uint8)Result | (uint8)EPackageLocationFilter::FileSystem);
 		}
 	}
 	
diff --git a/Engine/Source/Runtime/CoreUObject/Private/Serialization/AsyncLoading.cpp b/Engine/Source/Runtime/CoreUObject/Private/Serialization/AsyncLoading.cpp
--- a/Engine/Source/Runtime/CoreUObject/Private/Serialization/AsyncLoading.cpp
+++ b/Engine/Source/Runtime/CoreUObject/Private/Serialization/AsyncLoading.cpp
@@ -5641,12 +5641,8 @@ EAsyncPackageState::Type FAsyncPackage::CreateLinker()
 				bool bDoesPackageExist = false;
 				{
 					SCOPED_LOADTIMER(CreateLinker_DoesExist);
-#if WITH_IOSTORE_IN_EDITOR 
-					// Only look for non cooked packages on disk
-					bDoesPackageExist = FPackageName::DoesPackageExistEx(PackagePath, FPackageName::EPackageLocationFilter::Uncooked, false /* bMatchCaseOnDisk */, &PackagePath) != FPackageName::EPackageLocationFilter::None;
-#else
-					bDoesPackageExist = FPackageName::DoesPackageExist(PackagePath, false /* bMatchCaseOnDisk */, &PackagePath);
-#endif
+					// Only look for packages on disk
+					bDoesPackageExist = FPackageName::DoesPackageExistEx(PackagePath, FPackageName::EPackageLocationFilter::FileSystem, false /* bMatchCaseOnDisk */, &PackagePath) != FPackageName::EPackageLocationFilter::None;
 				}
 
 				if (!bDoesPackageExist)
diff --git a/Engine/Source/Runtime/CoreUObject/Private/Serialization/EditorPackageLoader.cpp b/Engine/Source/Runtime/CoreUObject/Private/Serialization/EditorPackageLoader.cpp
--- a/Engine/Source/Runtime/CoreUObject/Private/Serialization/EditorPackageLoader.cpp
+++ b/Engine/Source/Runtime/CoreUObject/Private/Serialization/EditorPackageLoader.cpp
@@ -56,7 +56,7 @@ public:
 	virtual bool ShouldAlwaysLoadPackageAsync(const FPackagePath& InPackagePath) override
 	{
 		// Use the old loader if an uncooked package exists on disk
-		const bool bDoesUncookedPackageExist = FPackageName::DoesPackageExistEx(InPackagePath, FPackageName::EPackageLocationFilter::Uncooked) != FPackageName::EPackageLocationFilter::None;
+		const bool bDoesUncookedPackageExist = FPackageName::DoesPackageExistEx(InPackagePath, FPackageName::EPackageLocationFilter::FileSystem) != FPackageName::EPackageLocationFilter::None;
 		return !bDoesUncookedPackageExist;
 	}
 
@@ -71,7 +71,7 @@ public:
 	{
 		// Use the old loader if an uncooked package exists on disk
 		if (!bHasInitializedCookedPackageLoader ||
-			FPackageName::DoesPackageExistEx(PackagePath, FPackageName::EPackageLocationFilter::Uncooked) != FPackageName::EPackageLocationFilter::None)
+			FPackageName::DoesPackageExistEx(PackagePath, FPackageName::EPackageLocationFilter::FileSystem) != FPackageName::EPackageLocationFilter::None)
 		{
 			UE_LOG(LogEditorPackageLoader, Verbose, TEXT("Loading uncooked package '%s' from filesystem"), *PackagePath.GetDebugName());
 			return UncookedPackageLoader->LoadPackage(PackagePath, CustomPackageName, InCompletionDelegate, InPackageFlags, InPIEInstanceID, InPackagePriority, InstancingContext);
diff --git a/Engine/Source/Runtime/CoreUObject/Private/UObject/Linker.cpp b/Engine/Source/Runtime/CoreUObject/Private/UObject/Linker.cpp
--- a/Engine/Source/Runtime/CoreUObject/Private/UObject/Linker.cpp
+++ b/Engine/Source/Runtime/CoreUObject/Private/UObject/Linker.cpp
@@ -622,12 +622,9 @@ FString GetPrestreamPackageLinkerName(const TCHAR* InLongPackageName, bool bSkip
 		return FString(); // we won't load this anyway, don't prestream
 	}
 
-#if WITH_IOSTORE_IN_EDITOR 
-	// Only look for non cooked packages on disk
-	bool DoesPackageExist = FPackageName::DoesPackageExistEx(PackagePath, FPackageName::EPackageLocationFilter::Uncooked, false, &PackagePath) != FPackageName::EPackageLocationFilter::None;
-#else
-	bool DoesPackageExist = FPackageName::DoesPackageExist(PackagePath, &PackagePath);
-#endif
+
+	// Only look for packages on disk
+	bool DoesPackageExist = FPackageName::DoesPackageExistEx(PackagePath, FPackageName::EPackageLocationFilter::FileSystem, false, &PackagePath) != FPackageName::EPackageLocationFilter::None;
 
 	if (!DoesPackageExist)
 	{
diff --git a/Engine/Source/Runtime/CoreUObject/Public/Misc/PackageName.h b/Engine/Source/Runtime/CoreUObject/Public/Misc/PackageName.h
--- a/Engine/Source/Runtime/CoreUObject/Public/Misc/PackageName.h
+++ b/Engine/Source/Runtime/CoreUObject/Public/Misc/PackageName.h
@@ -361,8 +361,8 @@ public:
 	enum class EPackageLocationFilter : uint8
 	{
 		None = 0,
-		Cooked = 1,
-		Uncooked = 2,
+		IoDispatcher = 1,
+		FileSystem = 2,
 		Any = 0xFF, // special filter to find if it exists anywhere at all, and won't need to check both. in this case, as soon as one is found, DoesPackageExistEx will return true
 	};
 
@@ -374,7 +374,7 @@ public:
 	 * @param Guid If nonnull, and the package is found on disk but does not have this PackageGuid in its FPackageFileSummary::Guid, false is returned
 	 * @param bMatchCaseOnDisk If true, the OutPackagePath is modified to match the capitalization of the discovered file
 	 * @param OutPackagePath If nonnull and the package exists, set to a copy of PackagePath with the HeaderExtension set to the extension that exists on disk (and if bMatchCaseOnDisk is true, capitalization changed to match). If not found, this variable is not written
-	 * @return the set of locations where the package exists (cooked or uncooked, both or neither)
+	 * @return the set of locations where the package exists (IoDispatcher or FileSystem, both or neither)
 	 **/
 	static EPackageLocationFilter DoesPackageExistEx(const FPackagePath& PackagePath, EPackageLocationFilter Filterconst, bool bMatchCaseOnDisk = false, FPackagePath* OutPackagePath = nullptr);
 
