Add a new bitfield of flags to the payload entries in FPackageTrailer.

#rnx

- This will allow us to store per payload info in the future such as disabling virtualization for specific payloads.
- Currently there are no flags.
- Older trailer formats will update automatically.

diff --git a/Engine/Source/Runtime/CoreUObject/Private/UObject/PackageTrailer.cpp b/Engine/Source/Runtime/CoreUObject/Private/UObject/PackageTrailer.cpp
--- a/Engine/Source/Runtime/CoreUObject/Private/UObject/PackageTrailer.cpp
+++ b/Engine/Source/Runtime/CoreUObject/Private/UObject/PackageTrailer.cpp
@@ -20,6 +20,8 @@ enum class EPackageTrailerVersion : uint32
 	INITIAL = 0,
 	// Access mode is now per payload and found in FLookupTableEntry 
 	ACCESS_PER_PAYLOAD = 1,
+	// Added EPayloadAccessMode to FLookupTableEntry
+	PAYLOAD_FLAGS = 2,
 
 	// -----<new versions can be added before this line>-------------------------------------------------
 	// - this needs to be the last line (see note below)
@@ -30,7 +32,7 @@ enum class EPackageTrailerVersion : uint32
 // These asserts are here to make sure that any changes to the size of disk constants are intentional.
 // If the change was intentional then just update the assert.
 static_assert(FPackageTrailer::FHeader::StaticHeaderSizeOnDisk == 28, "FPackageTrailer::FHeader size has been changed, if this was intentional then update this assert");
-static_assert(Private::FLookupTableEntry::SizeOnDisk == 45, "FLookupTableEntry size has been changed, if this was intentional then update this assert");
+static_assert(Private::FLookupTableEntry::SizeOnDisk == 49, "FLookupTableEntry size has been changed, if this was intentional then update this assert");
 static_assert(FPackageTrailer::FFooter::SizeOnDisk == 20, "FPackageTrailer::FFooter size has been changed, if this was intentional then update this assert");
 
 /** Utility for recording failed package open reasons */
@@ -69,6 +71,11 @@ void FLookupTableEntry::Serialize(FArchive& Ar, EPackageTrailerVersion PackageTr
 	Ar << CompressedSize;
 	Ar << RawSize;
 
+	if (Ar.IsSaving() || PackageTrailerVersion >= EPackageTrailerVersion::PAYLOAD_FLAGS)
+	{
+		Ar << Flags;
+	}
+
 	if (Ar.IsSaving() || PackageTrailerVersion >= EPackageTrailerVersion::ACCESS_PER_PAYLOAD)
 	{
 		Ar << AccessMode;
diff --git a/Engine/Source/Runtime/CoreUObject/Public/UObject/PackageTrailer.h b/Engine/Source/Runtime/CoreUObject/Public/UObject/PackageTrailer.h
--- a/Engine/Source/Runtime/CoreUObject/Public/UObject/PackageTrailer.h
+++ b/Engine/Source/Runtime/CoreUObject/Public/UObject/PackageTrailer.h
@@ -99,6 +99,13 @@ enum class EPayloadAccessMode : uint8
 	Virtualized
 };
 
+/** Flags that can be set on payloads in a payload trailer */
+enum class EPayloadFlags : uint32
+{
+	/** No flags are set */
+	None
+};
+
 enum class EPackageTrailerVersion : uint32;
 
 namespace Private
@@ -107,10 +114,11 @@ namespace Private
 struct FLookupTableEntry
 {
 	/** Size of the entry when serialized to disk in bytes */
-	static constexpr uint32 SizeOnDisk = 45;	// Identifier		| 20 bytes
+	static constexpr uint32 SizeOnDisk = 49;	// Identifier		| 20 bytes
 												// OffsetInFile		| 8 bytes
 												// CompressedSize	| 8 bytes
 												// RawSize			| 8 bytes
+												// Flags			| 4 bytes
 												// AccessMode		| 1 byte
 
 	FLookupTableEntry() = default;
@@ -131,6 +139,8 @@ struct FLookupTableEntry
 	uint64 CompressedSize = INDEX_NONE;
 	/** The size of the payload when uncompressed. */
 	uint64 RawSize = INDEX_NONE;
+	/** Bitfield of flags, see @UE::EPayloadFlags */
+	EPayloadFlags Flags = EPayloadFlags::None;
 
 	EPayloadAccessMode AccessMode = EPayloadAccessMode::Local;
 };
