Ensure that the FArchiveReplaceObjectAndStructPropertyRef ctor invokes the default serialization path with the right vtable.

diff --git a/Engine/Source/Runtime/CoreUObject/Public/Serialization/ArchiveReplaceObjectAndStructPropertyRef.h b/Engine/Source/Runtime/CoreUObject/Public/Serialization/ArchiveReplaceObjectAndStructPropertyRef.h
--- a/Engine/Source/Runtime/CoreUObject/Public/Serialization/ArchiveReplaceObjectAndStructPropertyRef.h
+++ b/Engine/Source/Runtime/CoreUObject/Public/Serialization/ArchiveReplaceObjectAndStructPropertyRef.h
@@ -25,11 +25,16 @@ public:
 	 * @param Flags                 Enum specifying behavior of archive
 	 */
 	FArchiveReplaceObjectAndStructPropertyRef(UObject* InSearchObject, const TMap<T*, T*>& InReplacementMap, EArchiveReplaceObjectFlags Flags = EArchiveReplaceObjectFlags::None)
-		: FArchiveReplaceObjectRef<T>(InSearchObject, InReplacementMap, Flags)
+		: FArchiveReplaceObjectRef<T>(InSearchObject, InReplacementMap, Flags | EArchiveReplaceObjectFlags::DelayStart)
 	{
+		// Note: We intentionally add the 'DelayStart' flag above; otherwise the base class ctor will use the base archive type for serialization, and our overrides won't get called.
+		// Here we're looking for whether or not the *caller* of this ctor has asked for a delayed start; in that case, we defer to the caller as to when our serialization will start.
+		if (!(Flags & EArchiveReplaceObjectFlags::DelayStart))
+		{
+			FArchiveReplaceObjectRef<T>::SerializeSearchObject();
+		}
 	}
 
-
 	UE_DEPRECATED(5.0, "Use version that supplies flags via enum.")
 	FArchiveReplaceObjectAndStructPropertyRef
 	(
@@ -41,8 +46,14 @@ public:
 		bool bDelayStart = false,
 		bool bIgnoreClassGeneratedByRef = true
 	)
-		: FArchiveReplaceObjectRef<T>(InSearchObject, InReplacementMap, bNullPrivateRefs, bIgnoreOuterRef, bIgnoreArchetypeRef, bDelayStart, bIgnoreClassGeneratedByRef)
+		: FArchiveReplaceObjectRef<T>(InSearchObject, InReplacementMap, bNullPrivateRefs, bIgnoreOuterRef, bIgnoreArchetypeRef, true, bIgnoreClassGeneratedByRef)
 	{
+		// Note: We intentionally request a delayed start above; otherwise the base class ctor will use the base archive type for serialization, and our overrides won't get called.
+		// Here we're looking for whether or not the *caller* of this ctor has asked for a delayed start; in that case, we defer to the caller as to when our serialization will start.
+		if (!bDelayStart)
+		{
+			FArchiveReplaceObjectRef<T>::SerializeSearchObject();
+		}
 	}
 
 	/**
