Make ClassGeneratedBy WITH_EDITORONLY_DATA as it will always be null in cooked builds

diff --git a/Engine/Plugins/Animation/ControlRig/Source/ControlRig/Private/AnimNode_ControlRigBase.cpp b/Engine/Plugins/Animation/ControlRig/Source/ControlRig/Private/AnimNode_ControlRigBase.cpp
--- a/Engine/Plugins/Animation/ControlRig/Source/ControlRig/Private/AnimNode_ControlRigBase.cpp
+++ b/Engine/Plugins/Animation/ControlRig/Source/ControlRig/Private/AnimNode_ControlRigBase.cpp
@@ -39,6 +39,8 @@ void FAnimNode_ControlRigBase::OnInitializeAnimInstance(const FAnimInstanceProxy
 	UControlRig* ControlRig = GetControlRig();
 	if (Component && Component->SkeletalMesh && ControlRig)
 	{
+#if WITH_EDITORONLY_DATA
+		// ClassGeneratedBy TODO: This may be wrong in cooked builds
 		UBlueprintGeneratedClass* BlueprintClass = Cast<UBlueprintGeneratedClass>(ControlRig->GetClass());
 		if (BlueprintClass)
 		{
@@ -46,6 +48,7 @@ void FAnimNode_ControlRigBase::OnInitializeAnimInstance(const FAnimInstanceProxy
 			// node mapping container will be saved on the initialization part
 			NodeMappingContainer = Component->SkeletalMesh->GetNodeMappingContainer(Blueprint);
 		}
+#endif
 
 		// register skeletalmesh component for now
 		ControlRig->GetDataSourceRegistry()->RegisterDataSource(UControlRig::OwnerComponent, InAnimInstance->GetOwningComponent());
diff --git a/Engine/Plugins/MovieScene/ActorSequence/Source/ActorSequence/Private/ActorSequence.cpp b/Engine/Plugins/MovieScene/ActorSequence/Source/ActorSequence/Private/ActorSequence.cpp
--- a/Engine/Plugins/MovieScene/ActorSequence/Source/ActorSequence/Private/ActorSequence.cpp
+++ b/Engine/Plugins/MovieScene/ActorSequence/Source/ActorSequence/Private/ActorSequence.cpp
@@ -62,6 +62,7 @@ bool UActorSequence::IsEditable() const
 	return !Template || Template->GetTypedOuter<UActorSequenceComponent>() == GetDefault<UActorSequenceComponent>();
 }
 
+#if WITH_EDITORONLY_DATA
 UBlueprint* UActorSequence::GetParentBlueprint() const
 {
 	if (UBlueprintGeneratedClass* GeneratedClass = GetTypedOuter<UBlueprintGeneratedClass>())
@@ -70,6 +71,7 @@ UBlueprint* UActorSequence::GetParentBlueprint() const
 	}
 	return nullptr;
 }
+#endif
 
 void UActorSequence::PostInitProperties()
 {
diff --git a/Engine/Plugins/MovieScene/ActorSequence/Source/ActorSequence/Private/ActorSequenceObjectReference.cpp b/Engine/Plugins/MovieScene/ActorSequence/Source/ActorSequence/Private/ActorSequenceObjectReference.cpp
--- a/Engine/Plugins/MovieScene/ActorSequence/Source/ActorSequence/Private/ActorSequenceObjectReference.cpp
+++ b/Engine/Plugins/MovieScene/ActorSequence/Source/ActorSequence/Private/ActorSequenceObjectReference.cpp
@@ -20,6 +20,7 @@ FActorSequenceObjectReference FActorSequenceObjectReference::CreateForComponent(
 		return NewReference;
 	}
 
+#if WITH_EDITORONLY_DATA
 	UBlueprintGeneratedClass* GeneratedClass = InComponent->GetTypedOuter<UBlueprintGeneratedClass>();
 	if (GeneratedClass && GeneratedClass->SimpleConstructionScript)
 	{
@@ -36,6 +37,7 @@ FActorSequenceObjectReference FActorSequenceObjectReference::CreateForComponent(
 			}
 		}
 	}
+#endif
 
 	ensureMsgf(false, TEXT("Unable to find parent actor for component. Reference will be unresolvable."));
 	return NewReference;
diff --git a/Engine/Plugins/MovieScene/ActorSequence/Source/ActorSequence/Public/ActorSequence.h b/Engine/Plugins/MovieScene/ActorSequence/Source/ActorSequence/Public/ActorSequence.h
--- a/Engine/Plugins/MovieScene/ActorSequence/Source/ActorSequence/Public/ActorSequence.h
+++ b/Engine/Plugins/MovieScene/ActorSequence/Source/ActorSequence/Public/ActorSequence.h
@@ -35,7 +35,9 @@ public:
 	virtual ETrackSupport IsTrackSupported(TSubclassOf<class UMovieSceneTrack> InTrackClass) const override;
 #endif
 
+#if WITH_EDITORONLY_DATA
 	UBlueprint* GetParentBlueprint() const;
+#endif
 
 	bool IsEditable() const;
 	
diff --git a/Engine/Plugins/Runtime/GameplayAbilities/Source/GameplayAbilities/Private/AttributeSet.cpp b/Engine/Plugins/Runtime/GameplayAbilities/Source/GameplayAbilities/Private/AttributeSet.cpp
--- a/Engine/Plugins/Runtime/GameplayAbilities/Source/GameplayAbilities/Private/AttributeSet.cpp
+++ b/Engine/Plugins/Runtime/GameplayAbilities/Source/GameplayAbilities/Private/AttributeSet.cpp
@@ -268,7 +268,12 @@ void FGameplayAttribute::GetAllAttributeProperties(TArray<FProperty*>& OutProper
 	for (TObjectIterator<UClass> ClassIt; ClassIt; ++ClassIt)
 	{
 		UClass *Class = *ClassIt;
-		if (Class->IsChildOf(UAttributeSet::StaticClass()) && !Class->ClassGeneratedBy)
+		if (Class->IsChildOf(UAttributeSet::StaticClass()) 
+#if WITH_EDITORONLY_DATA
+			// ClassGeneratedBy TODO: This is wrong in cooked builds
+			&& !Class->ClassGeneratedBy
+#endif
+			)
 		{
 			if (UseEditorOnlyData)
 			{
@@ -311,9 +316,9 @@ void FGameplayAttribute::GetAllAttributeProperties(TArray<FProperty*>& OutProper
 			}
 		}
 
+#if WITH_EDITOR
 		if (UseEditorOnlyData)
 		{
-			#if WITH_EDITOR
 			// UAbilitySystemComponent can add 'system' attributes
 			if (Class->IsChildOf(UAbilitySystemComponent::StaticClass()) && !Class->ClassGeneratedBy)
 			{
@@ -330,8 +335,8 @@ void FGameplayAttribute::GetAllAttributeProperties(TArray<FProperty*>& OutProper
 					OutProperties.Add(Property);
 				}
 			}
-			#endif
 		}
+#endif
 	}
 }
 
diff --git a/Engine/Source/Runtime/CoreUObject/Private/AssetRegistry/AssetData.cpp b/Engine/Source/Runtime/CoreUObject/Private/AssetRegistry/AssetData.cpp
--- a/Engine/Source/Runtime/CoreUObject/Private/AssetRegistry/AssetData.cpp
+++ b/Engine/Source/Runtime/CoreUObject/Private/AssetRegistry/AssetData.cpp
@@ -89,12 +89,15 @@ FAssetData::FAssetData(const UObject* InAsset, FAssetData::ECreationFlags InCrea
 {
 	if (InAsset != nullptr)
 	{
+#if WITH_EDITORONLY_DATA
+		// ClassGeneratedBy TODO: This may be wrong in cooked builds
 		const UClass* InClass = Cast<UClass>(InAsset);
 		if (InClass && InClass->ClassGeneratedBy && !EnumHasAnyFlags(InCreationFlags, FAssetData::ECreationFlags::AllowBlueprintClass))
 		{
 			// For Blueprints, the AssetData refers to the UBlueprint and not the UBlueprintGeneratedClass
 			InAsset = InClass->ClassGeneratedBy;
 		}
+#endif
 
 		const UPackage* Package = InAsset->GetPackage();
 
diff --git a/Engine/Source/Runtime/CoreUObject/Private/Blueprint/BlueprintSupport.cpp b/Engine/Source/Runtime/CoreUObject/Private/Blueprint/BlueprintSupport.cpp
--- a/Engine/Source/Runtime/CoreUObject/Private/Blueprint/BlueprintSupport.cpp
+++ b/Engine/Source/Runtime/CoreUObject/Private/Blueprint/BlueprintSupport.cpp
@@ -268,36 +268,24 @@ void FBlueprintSupport::RepairDeferredDependenciesInObject(UObject* Object)
 
 bool FBlueprintSupport::IsInBlueprintPackage(UObject* LoadedObj)
 {
+	bool bHasBlueprintClass = false;
+
 	UPackage* Pkg = LoadedObj->GetOutermost();
 	if (Pkg && !Pkg->HasAnyPackageFlags(PKG_CompiledIn))
-	{
-		TArray<UObject*> PkgObjects;
-		GetObjectsWithOuter(Pkg, PkgObjects, /*bIncludeNestedObjects =*/false);
-		
-		UObject* PkgCDO   = nullptr;
-		UClass*  PkgClass = nullptr;
-
-		for (UObject* PkgObj : PkgObjects)
+	{	
+		ForEachObjectWithOuterBreakable(Pkg, [&bHasBlueprintClass](UObject* PkgObj)
 		{
-			if (PkgObj->HasAnyFlags(RF_ClassDefaultObject))
+			if (UClass* PkgClass = Cast<UClass>(PkgObj))
 			{
-				PkgCDO = PkgObj;
+				bHasBlueprintClass = PkgClass && PkgClass->HasAnyClassFlags(CLASS_CompiledFromBlueprint);
+				return false; // break
 			}
-			else if (UClass* AsClass = Cast<UClass>(PkgObj))
-			{
-				PkgClass = AsClass;
-			}
-		}
-		const bool bHasBlueprintClass = PkgClass && PkgClass->HasAnyClassFlags(CLASS_CompiledFromBlueprint);
 
-		return bHasBlueprintClass
-			//&& (PkgCDO && PkgCDO->GetClass() == PkgClass)
-#if WITH_EDITORONLY_DATA
-			//&& (PkgClass->ClassGeneratedBy != nullptr) && (PkgClass->ClassGeneratedBy->GetOuter() == Pkg)
-#endif
-			;
+			return true;
+		}, /*bIncludeNestedObjects =*/false);
+		
 	}
-	return false;
+	return bHasBlueprintClass;
 }
 
 static TArray<FBlueprintWarningDeclaration> BlueprintWarnings;
@@ -693,7 +681,11 @@ bool FLinkerLoad::RegenerateBlueprintClass(UClass* LoadClass, UObject* ExportObj
 {
 	auto GetClassSourceObjectLambda = [](UClass* ForClass) -> UObject*
 	{
+#if WITH_EDITORONLY_DATA
 		return ForClass->ClassGeneratedBy ? ForClass->ClassGeneratedBy : ForClass;
+#else
+		return ForClass;
+#endif
 	};
 
 	UObject* ClassSourceObject = GetClassSourceObjectLambda(LoadClass);
@@ -2387,6 +2379,9 @@ UObject* FLinkerLoad::FindImportFast(UClass* ImportClass, UObject* ImportOuter,
  */
 bool UObject::IsInBlueprint() const
 {
+	// ClassGeneratedBy TODO: This is wrong in cooked builds
+#if WITH_EDITORONLY_DATA
+
 	// Exclude blueprint classes as they may be regenerated at any time
 	// Need to exclude classes, CDOs, and their subobjects
 	const UObject* TestObject = this;
@@ -2408,6 +2403,7 @@ bool UObject::IsInBlueprint() const
  		}
  		TestObject = TestObject->GetOuter();
  	}
+#endif
 
 	return false;
 }
diff --git a/Engine/Source/Runtime/CoreUObject/Private/UObject/Class.cpp b/Engine/Source/Runtime/CoreUObject/Private/UObject/Class.cpp
--- a/Engine/Source/Runtime/CoreUObject/Private/UObject/Class.cpp
+++ b/Engine/Source/Runtime/CoreUObject/Private/UObject/Class.cpp
@@ -3688,18 +3688,17 @@ void UClass::GetDefaultObjectSubobjects(TArray<UObject*>& OutDefaultSubobjects)
 void UClass::AddReferencedObjects(UObject* InThis, FReferenceCollector& Collector)
 {
 	UClass* This = CastChecked<UClass>(InThis);
-	for( auto& Inter : This->Interfaces )
+	for (FImplementedInterface& Inter : This->Interfaces)
 	{
 		Collector.AddReferencedObject( Inter.Class, This );
 	}
 
-	for( auto& Func : This->FuncMap )
-	{
-		Collector.AddReferencedObject( Func.Value, This );
-	}
-
+	Collector.AddReferencedObjects( This->FuncMap, This );
 	Collector.AddReferencedObject( This->ClassWithin, This );
+
+#if WITH_EDITORONLY_DATA
 	Collector.AddReferencedObject( This->ClassGeneratedBy, This );
+#endif
 
 	if ( !Collector.IsIgnoringArchetypeRef() )
 	{
@@ -3743,8 +3742,10 @@ class FRestoreClassInfo: public FRestoreForUObjectOverwrite
 	UClass*			Target;
 	/** Saved ClassWithin **/
 	UClass*			Within;
+#if WITH_EDITORONLY_DATA
 	/** Saved ClassGeneratedBy */
 	UObject*		GeneratedBy;
+#endif
 	/** Saved ClassDefaultObject **/
 	UObject*		DefaultObject;
 	/** Saved ClassFlags **/
@@ -3769,7 +3770,9 @@ public:
 	FRestoreClassInfo(UClass *Save) :
 		Target(Save),
 		Within(Save->ClassWithin),
+#if WITH_EDITORONLY_DATA
 		GeneratedBy(Save->ClassGeneratedBy),
+#endif
 		DefaultObject(Save->GetDefaultsCount() ? Save->GetDefaultObject() : NULL),
 		Flags(Save->ClassFlags & CLASS_Abstract),
 		CastFlags(Save->ClassCastFlags),
@@ -3784,7 +3787,9 @@ public:
 	virtual void Restore() const
 	{
 		Target->ClassWithin = Within;
+#if WITH_EDITORONLY_DATA
 		Target->ClassGeneratedBy = GeneratedBy;
+#endif
 		Target->ClassDefaultObject = DefaultObject;
 		Target->ClassFlags |= Flags;
 		Target->ClassCastFlags |= CastFlags;
@@ -4610,6 +4615,10 @@ void UClass::Serialize( FArchive& Ar )
 
 	if (!Ar.IsIgnoringClassGeneratedByRef())
 	{
+#if !WITH_EDITORONLY_DATA
+		// Dummy variable to keep archive consistency
+		UObject* ClassGeneratedBy = nullptr;
+#endif
 		Ar << ClassGeneratedBy;
 	}
 
@@ -5056,8 +5065,8 @@ UClass::UClass(const FObjectInitializer& ObjectInitializer)
 ,	ClassFlags(CLASS_None)
 ,	ClassCastFlags(CASTCLASS_None)
 ,	ClassWithin( UObject::StaticClass() )
-,	ClassGeneratedBy(nullptr)
 #if WITH_EDITORONLY_DATA
+,	ClassGeneratedBy(nullptr)
 ,	PropertiesPendingDestruction(nullptr)
 #endif
 ,	ClassDefaultObject(nullptr)
@@ -5081,8 +5090,8 @@ UClass::UClass(const FObjectInitializer& ObjectInitializer, UClass* InBaseClass)
 ,	ClassFlags(CLASS_None)
 ,	ClassCastFlags(CASTCLASS_None)
 ,	ClassWithin(UObject::StaticClass())
-,	ClassGeneratedBy(nullptr)
 #if WITH_EDITORONLY_DATA
+,	ClassGeneratedBy(nullptr)
 ,	PropertiesPendingDestruction(nullptr)
 #endif
 ,	ClassDefaultObject(nullptr)
@@ -5140,8 +5149,8 @@ UClass::UClass
 ,	ClassFlags				( InClassFlags | CLASS_Native )
 ,	ClassCastFlags			( InClassCastFlags )
 ,	ClassWithin				( nullptr )
-,	ClassGeneratedBy		( nullptr )
 #if WITH_EDITORONLY_DATA
+,	ClassGeneratedBy		( nullptr )
 ,	PropertiesPendingDestruction( nullptr )
 #endif
 ,	ClassConfigName			()
@@ -5715,6 +5724,7 @@ bool UClass::IsClassGroupName(const TCHAR* InGroupName) const
 #endif // WITH_EDITOR || HACK_HEADER_GENERATOR
 
 
+#if WITH_EDITORONLY_DATA
 IMPLEMENT_CORE_INTRINSIC_CLASS(UClass, UStruct,
 	{
 		Class->ClassAddReferencedObjects = &UClass::AddReferencedObjects;
@@ -5725,6 +5735,17 @@ IMPLEMENT_CORE_INTRINSIC_CLASS(UClass, UStruct,
 		Class->EmitObjectArrayReference(STRUCT_OFFSET(UClass, NetFields), TEXT("NetFields"));
 	}
 );
+#else
+IMPLEMENT_CORE_INTRINSIC_CLASS(UClass, UStruct,
+	{
+		Class->ClassAddReferencedObjects = &UClass::AddReferencedObjects;
+
+		Class->EmitObjectReference(STRUCT_OFFSET(UClass, ClassDefaultObject), TEXT("ClassDefaultObject"));
+		Class->EmitObjectReference(STRUCT_OFFSET(UClass, ClassWithin), TEXT("ClassWithin"));
+		Class->EmitObjectArrayReference(STRUCT_OFFSET(UClass, NetFields), TEXT("NetFields"));
+	}
+);
+#endif
 
 void GetPrivateStaticClassBody(
 	const TCHAR* PackageName,
diff --git a/Engine/Source/Runtime/CoreUObject/Private/UObject/LinkerLoad.cpp b/Engine/Source/Runtime/CoreUObject/Private/UObject/LinkerLoad.cpp
--- a/Engine/Source/Runtime/CoreUObject/Private/UObject/LinkerLoad.cpp
+++ b/Engine/Source/Runtime/CoreUObject/Private/UObject/LinkerLoad.cpp
@@ -4658,6 +4658,7 @@ UObject* FLinkerLoad::CreateExport( int32 Index )
 			{
 				if( LoadClass->IsChildOf(UFunction::StaticClass()) )
 				{
+#if WITH_EDITORONLY_DATA
 					// In the case of a function object, the outer should be the function's class. For Blueprints, loading
 					// the outer class may also invalidate this entry in the export map. In that case, we won't actually be
 					// keeping the function object around, so there's no need to warn here about the missing parent object.
@@ -4671,11 +4672,13 @@ UObject* FLinkerLoad::CreateExport( int32 Index )
 							UE_LOG(LogLinker, Display, TEXT("CreateExport: Failed to load Parent for %s; removing parent information, but keeping function"), *GetExportFullName(Index));
 						}
 					}
+#endif
 
 					Export.SuperIndex = FPackageIndex();
 				}
 				else
 				{
+#if WITH_EDITORONLY_DATA
 					bool bFailedToLoadGeneratedStruct = false;
 					if( LoadClass->IsChildOf(UScriptStruct::StaticClass()) )
 					{
@@ -4698,7 +4701,8 @@ UObject* FLinkerLoad::CreateExport( int32 Index )
 					{
 						UE_LOG(LogLinker, Warning, TEXT("CreateExport: Failed to load Parent for %s"), *GetExportFullName(Index));
 					}
-					return NULL;
+#endif
+					return nullptr;
 				}
 			}
 			else
diff --git a/Engine/Source/Runtime/CoreUObject/Private/UObject/Obj.cpp b/Engine/Source/Runtime/CoreUObject/Private/UObject/Obj.cpp
--- a/Engine/Source/Runtime/CoreUObject/Private/UObject/Obj.cpp
+++ b/Engine/Source/Runtime/CoreUObject/Private/UObject/Obj.cpp
@@ -1643,6 +1643,8 @@ void UObject::BuildSubobjectMapping(UObject* OtherObject, TMap<UObject*, UObject
 		FString NewSubObjectName = InSubObject->GetName();
 
 		UClass* OtherSubObjectClass = InSubObject->GetClass();
+#if WITH_EDITORONLY_DATA
+		// ClassGeneratedBy TODO: This may be wrong in cooked builds
 		if (OtherSubObjectClass->ClassGeneratedBy && OtherSubObjectClass->ClassGeneratedBy->GetOutermost() == ThisPackage)
 		{
 			// This is a generated class type, so we actually need to use the new generated class type from the new package otherwise our type check will fail
@@ -1651,6 +1653,7 @@ void UObject::BuildSubobjectMapping(UObject* OtherObject, TMap<UObject*, UObject
 
 			OtherSubObjectClass = LoadObject<UClass>(OtherPackage, *NewClassName);
 		}
+#endif
 
 		//UObject* OtherSubObject = StaticLoadObject(OtherSubObjectClass, OtherObject, *NewSubObjectName, nullptr, LOAD_Quiet | LOAD_NoRedirects, nullptr, true);
 		UObject* OtherSubObject = StaticFindObjectFast(OtherSubObjectClass, OtherObject, *NewSubObjectName);
diff --git a/Engine/Source/Runtime/CoreUObject/Private/UObject/ObjectHandle.cpp b/Engine/Source/Runtime/CoreUObject/Private/UObject/ObjectHandle.cpp
--- a/Engine/Source/Runtime/CoreUObject/Private/UObject/ObjectHandle.cpp
+++ b/Engine/Source/Runtime/CoreUObject/Private/UObject/ObjectHandle.cpp
@@ -269,10 +269,15 @@ FObjectRef MakeObjectRef(const UObject* Object)
 	{
 		return {NAME_None, NAME_None, NAME_None, FObjectPathId()};
 	}
-	
+
+#if WITH_EDITORONLY_DATA	
+	// ClassGeneratedBy TODO: This may be wrong in cooked builds
 	UObject* ClassGeneratedBy = Object->GetClass()->ClassGeneratedBy;
 	UPackage* ClassGeneratedByPackage = ClassGeneratedBy ? ClassGeneratedBy->GetOutermost() : nullptr;
 	return FObjectRef {GetNameOrNone(Object->GetOutermost()), GetNameOrNone(ClassGeneratedByPackage), GetNameOrNone(ClassGeneratedBy), FObjectPathId(Object)};
+#else
+	return FObjectRef{ GetNameOrNone(Object->GetOutermost()), NAME_None, NAME_None, FObjectPathId(Object) };
+#endif
 }
 
 FObjectRef MakeObjectRef(FPackedObjectRef PackedObjectRef)
@@ -293,9 +298,14 @@ FPackedObjectRef MakePackedObjectRef(const UObject* Object)
 	}
 
 	FName PackageName = GetNameOrNone(Object->GetOutermost());
+#if WITH_EDITORONLY_DATA
+	// ClassGeneratedBy TODO: This may be wrong in cooked builds
 	UObject* ClassGeneratedBy = Object->GetClass()->ClassGeneratedBy;
 	UPackage* ClassGeneratedByPackage = ClassGeneratedBy ? ClassGeneratedBy->GetOutermost() : nullptr;
 	return ObjectHandle_Private::MakePackedObjectRef(PackageName, GetNameOrNone(ClassGeneratedByPackage), GetNameOrNone(ClassGeneratedBy), FObjectPathId(Object));
+#else
+	return ObjectHandle_Private::MakePackedObjectRef(PackageName, NAME_None, NAME_None, FObjectPathId(Object));
+#endif
 }
 
 FPackedObjectRef MakePackedObjectRef(const FObjectRef& ObjectRef)
diff --git a/Engine/Source/Runtime/CoreUObject/Private/UObject/SavePackage/SaveContext.cpp b/Engine/Source/Runtime/CoreUObject/Private/UObject/SavePackage/SaveContext.cpp
--- a/Engine/Source/Runtime/CoreUObject/Private/UObject/SavePackage/SaveContext.cpp
+++ b/Engine/Source/Runtime/CoreUObject/Private/UObject/SavePackage/SaveContext.cpp
@@ -34,7 +34,9 @@ void FSaveContext::MarkUnsaveable(UObject* InObject)
 	// if this is the class default object, make sure it's not
 	// marked transient for any reason, as we need it to be saved
 	// to disk (unless it's associated with a transient generated class)
+#if WITH_EDITORONLY_DATA
 	ensureAlways(!InObject->HasAllFlags(RF_ClassDefaultObject | RF_Transient) || (InObject->GetClass()->ClassGeneratedBy != nullptr && InObject->GetClass()->HasAnyFlags(RF_Transient)));
+#endif
 }
 
 bool FSaveContext::IsUnsaveable(UObject* InObject, bool bEmitWarning) const
diff --git a/Engine/Source/Runtime/CoreUObject/Private/UObject/SavePackage/SavePackageUtilities.cpp b/Engine/Source/Runtime/CoreUObject/Private/UObject/SavePackage/SavePackageUtilities.cpp
--- a/Engine/Source/Runtime/CoreUObject/Private/UObject/SavePackage/SavePackageUtilities.cpp
+++ b/Engine/Source/Runtime/CoreUObject/Private/UObject/SavePackage/SavePackageUtilities.cpp
@@ -352,7 +352,11 @@ void CheckObjectPriorToSave(FArchiveUObject& Ar, UObject* InObj, UPackage* InSav
 	}
 
 	if ( InObj->HasAnyFlags(RF_ClassDefaultObject)
-		&& (InObj->GetClass()->ClassGeneratedBy == nullptr || !InObj->GetClass()->HasAnyFlags(RF_Transient)) )
+#if WITH_EDITORONLY_DATA
+		// ClassGeneratedBy TODO: This may be wrong in cooked builds
+		&& (InObj->GetClass()->ClassGeneratedBy == nullptr || !InObj->GetClass()->HasAnyFlags(RF_Transient)) 
+#endif
+		)
 	{
 		// if this is the class default object, make sure it's not
 		// marked transient for any reason, as we need it to be saved
diff --git a/Engine/Source/Runtime/CoreUObject/Private/UObject/UObjectBase.cpp b/Engine/Source/Runtime/CoreUObject/Private/UObject/UObjectBase.cpp
--- a/Engine/Source/Runtime/CoreUObject/Private/UObject/UObjectBase.cpp
+++ b/Engine/Source/Runtime/CoreUObject/Private/UObject/UObjectBase.cpp
@@ -269,6 +269,7 @@ void UObjectBase::SetExternalPackage(UPackage* InPackage)
 	}
 }
 
+#if WITH_EDITOR
 void UObjectBase::SetClass(UClass* NewClass)
 {
 #if STATS || ENABLE_STATNAMEDEVENTS_UOBJECT
@@ -286,7 +287,7 @@ void UObjectBase::SetClass(UClass* NewClass)
 #endif
 	HashObject(this);
 }
-
+#endif
 
 /**
  * Checks to see if the object appears to be valid
diff --git a/Engine/Source/Runtime/CoreUObject/Private/UObject/UObjectGlobals.cpp b/Engine/Source/Runtime/CoreUObject/Private/UObject/UObjectGlobals.cpp
--- a/Engine/Source/Runtime/CoreUObject/Private/UObject/UObjectGlobals.cpp
+++ b/Engine/Source/Runtime/CoreUObject/Private/UObject/UObjectGlobals.cpp
@@ -2119,10 +2119,12 @@ UObject* StaticDuplicateObjectEx( FObjectDuplicationParameters& Parameters )
 		Params.Name = Parameters.DestName;
 		Params.SetFlags = Parameters.ApplyFlags | Parameters.SourceObject->GetMaskedFlags(Parameters.FlagMask);
 		Params.InternalSetFlags = Parameters.ApplyInternalFlags | (Parameters.SourceObject->GetInternalFlags() & Parameters.InternalFlagMask);
-		Params.Template = Parameters.SourceObject->GetArchetype()->GetClass() == Parameters.DestClass ? Parameters.SourceObject->GetArchetype() : nullptr;
 		Params.bCopyTransientsFromClassDefaults = true;
 		Params.InstanceGraph = &InstanceGraph;
 
+		UObject* Archetype = Parameters.SourceObject->GetArchetype();
+		Params.Template = Archetype->GetClass() == Parameters.DestClass ? Archetype : nullptr;
+
 		DupRootObject = StaticConstructObject_Internal(Params);
 	}
 
@@ -3747,6 +3749,8 @@ protected:
 
 		auto GetBlueprintObjectNameLambda = [](const UObject* InSerializingObject) -> FString
 		{
+#if WITH_EDITORONLY_DATA
+			// ClassGeneratedBy TODO: This may be wrong in cooked builds
 			if (InSerializingObject)
 			{
 				const UClass* BPGC = InSerializingObject->GetTypedOuter<UClass>();
@@ -3755,6 +3759,7 @@ protected:
 					return BPGC->ClassGeneratedBy->GetFullName();
 				}
 			}
+#endif
 
 			return TEXT("NULL");
 		};
diff --git a/Engine/Source/Runtime/CoreUObject/Public/UObject/UObjectBase.h b/Engine/Source/Runtime/CoreUObject/Public/UObject/UObjectBase.h
--- a/Engine/Source/Runtime/CoreUObject/Public/UObject/UObjectBase.h
+++ b/Engine/Source/Runtime/CoreUObject/Public/UObject/UObjectBase.h
@@ -229,8 +229,10 @@ private:
 	friend class FBlueprintCompileReinstancer;
 	friend class FContextObjectManager;
 
+#if WITH_EDITOR
 	/** This is used by the reinstancer to re-class and re-archetype the current instances of a class before recompiling */
 	void SetClass(UClass* NewClass);
+#endif
 
 #if HACK_HEADER_GENERATOR
 	// Required by UHT makefiles for internal data serialization.
diff --git a/Engine/Source/Runtime/Engine/Classes/Engine/Blueprint.h b/Engine/Source/Runtime/Engine/Classes/Engine/Blueprint.h
--- a/Engine/Source/Runtime/Engine/Classes/Engine/Blueprint.h
+++ b/Engine/Source/Runtime/Engine/Classes/Engine/Blueprint.h
@@ -942,9 +942,9 @@ public:
 	virtual void PreSave(const class ITargetPlatform* TargetPlatform) override;
 	PRAGMA_ENABLE_DEPRECATION_WARNINGS
 	virtual void PreSave(FObjectPreSaveContext ObjectSaveContext) override;
+	virtual void GetPreloadDependencies(TArray<UObject*>& OutDeps) override;
 #endif // WITH_EDITORONLY_DATA
 	virtual void Serialize(FArchive& Ar) override;
-	virtual void GetPreloadDependencies(TArray<UObject*>& OutDeps) override;
 	virtual FString GetDesc(void) override;
 	virtual void TagSubobjects(EObjectFlags NewFlags) override;
 	virtual bool NeedsLoadForClient() const override;
@@ -952,6 +952,7 @@ public:
 	virtual bool NeedsLoadForEditorGame() const override;
 	//~ End UObject Interface
 
+#if WITH_EDITORONLY_DATA
 	/** Get the Blueprint object that generated the supplied class */
 	static UBlueprint* GetBlueprintFromClass(const UClass* InClass);
 
@@ -963,6 +964,7 @@ public:
 	 * @return						true if there were no status errors in any of the parent blueprints, otherwise false
 	 */
 	static bool GetBlueprintHierarchyFromClass(const UClass* InClass, TArray<UBlueprint*>& OutBlueprintParents);
+#endif
 
 	/**
 	 * Gets an array of all BPGCs used to generate this class and its parents.  0th elements is the BPGC used to generate InClass
diff --git a/Engine/Source/Runtime/Engine/Private/Animation/AnimNotify.cpp b/Engine/Source/Runtime/Engine/Private/Animation/AnimNotify.cpp
--- a/Engine/Source/Runtime/Engine/Private/Animation/AnimNotify.cpp
+++ b/Engine/Source/Runtime/Engine/Private/Animation/AnimNotify.cpp
@@ -53,21 +53,22 @@ class UWorld* UAnimNotify::GetWorld() const
 
 FString UAnimNotify::GetNotifyName_Implementation() const
 {
-	UObject* ClassGeneratedBy = GetClass()->ClassGeneratedBy;
 	FString NotifyName;
 
-	if(ClassGeneratedBy)
+#if WITH_EDITORONLY_DATA
+	if (UObject* ClassGeneratedBy = GetClass()->ClassGeneratedBy)
 	{
 		// GeneratedBy will be valid for blueprint types and gives a clean name without a suffix
 		NotifyName = ClassGeneratedBy->GetName();
 	}
 	else
+#endif
 	{
 		// Native notify classes are clean without a suffix otherwise
 		NotifyName = GetClass()->GetName();
 	}
 
-	NotifyName = NotifyName.Replace(TEXT("AnimNotify_"), TEXT(""), ESearchCase::CaseSensitive);
+	NotifyName.ReplaceInline(TEXT("AnimNotify_"), TEXT(""), ESearchCase::CaseSensitive);
 	
 	return NotifyName;
 }
diff --git a/Engine/Source/Runtime/Engine/Private/Animation/AnimNotifyState.cpp b/Engine/Source/Runtime/Engine/Private/Animation/AnimNotifyState.cpp
--- a/Engine/Source/Runtime/Engine/Private/Animation/AnimNotifyState.cpp
+++ b/Engine/Source/Runtime/Engine/Private/Animation/AnimNotifyState.cpp
@@ -79,21 +79,22 @@ void UAnimNotifyState::BranchingPointNotifyEnd(FBranchingPointNotifyPayload& Bra
 
 FString UAnimNotifyState::GetNotifyName_Implementation() const
 {
-	UObject* ClassGeneratedBy = GetClass()->ClassGeneratedBy;
 	FString NotifyName;
-
-	if(ClassGeneratedBy)
+	
+#if WITH_EDITORONLY_DATA
+	if (UObject* ClassGeneratedBy = GetClass()->ClassGeneratedBy)
 	{
 		// GeneratedBy will be valid for blueprint types and gives a clean name without a suffix
 		NotifyName = ClassGeneratedBy->GetName();
 	}
 	else
+#endif
 	{
 		// Native notify classes are clean without a suffix otherwise
 		NotifyName = GetClass()->GetName();
 	}
 
-	NotifyName = NotifyName.Replace(TEXT("AnimNotifyState_"), TEXT(""), ESearchCase::CaseSensitive);
+	NotifyName.ReplaceInline(TEXT("AnimNotifyState_"), TEXT(""), ESearchCase::CaseSensitive);
 	
 	return NotifyName;
 }
diff --git a/Engine/Source/Runtime/Engine/Private/Blueprint.cpp b/Engine/Source/Runtime/Engine/Private/Blueprint.cpp
--- a/Engine/Source/Runtime/Engine/Private/Blueprint.cpp
+++ b/Engine/Source/Runtime/Engine/Private/Blueprint.cpp
@@ -411,7 +411,6 @@ void UBlueprint::PreSave(FObjectPreSaveContext ObjectSaveContext)
 		FFindInBlueprintSearchManager::Get().AddOrUpdateBlueprintSearchMetadata(this, Flags, OverrideVersion);
 	}
 }
-#endif // WITH_EDITORONLY_DATA
 
 void UBlueprint::GetPreloadDependencies(TArray<UObject*>& OutDeps)
 {
@@ -424,6 +423,7 @@ void UBlueprint::GetPreloadDependencies(TArray<UObject*>& OutDeps)
 		}
 	}
 }
+#endif // WITH_EDITORONLY_DATA
 
 void UBlueprint::Serialize(FArchive& Ar)
 {
@@ -1380,6 +1380,7 @@ void UBlueprint::BeginDestroy()
 
 #endif // WITH_EDITOR
 
+#if WITH_EDITORONLY_DATA
 UBlueprint* UBlueprint::GetBlueprintFromClass(const UClass* InClass)
 {
 	UBlueprint* BP = NULL;
@@ -1400,9 +1401,7 @@ bool UBlueprint::GetBlueprintHierarchyFromClass(const UClass* InClass, TArray<UB
 	{
 		OutBlueprintParents.Add(BP);
 
-#if WITH_EDITORONLY_DATA
 		bNoErrors &= (BP->Status != BS_Error);
-#endif // #if WITH_EDITORONLY_DATA
 
 		// If valid, use stored ParentClass rather than the actual UClass::GetSuperClass(); handles the case when the class has not been recompiled yet after a reparent operation.
 		if (BP->ParentClass)
@@ -1418,6 +1417,7 @@ bool UBlueprint::GetBlueprintHierarchyFromClass(const UClass* InClass, TArray<UB
 
 	return bNoErrors;
 }
+#endif
 
 bool UBlueprint::GetBlueprintHierarchyFromClass(const UClass* InClass, TArray<UBlueprintGeneratedClass*>& OutBlueprintParents)
 {
@@ -1429,14 +1429,14 @@ bool UBlueprint::GetBlueprintHierarchyFromClass(const UClass* InClass, TArray<UB
 	{
 		OutBlueprintParents.Add(CurrentClass);
 
+#if WITH_EDITORONLY_DATA
+		// ClassGeneratedBy TODO: This may be wrong in cooked builds
 		UBlueprint* BP = UBlueprint::GetBlueprintFromClass(CurrentClass);
 
-#if WITH_EDITORONLY_DATA
 		if (BP)
 		{
 			bNoErrors &= (BP->Status != BS_Error);
 		}
-#endif // #if WITH_EDITORONLY_DATA
 
 		// If valid, use stored ParentClass rather than the actual UClass::GetSuperClass(); handles the case when the class has not been recompiled yet after a reparent operation.
 		if (BP && BP->ParentClass)
@@ -1444,6 +1444,7 @@ bool UBlueprint::GetBlueprintHierarchyFromClass(const UClass* InClass, TArray<UB
 			CurrentClass = Cast<UBlueprintGeneratedClass>(BP->ParentClass);
 		}
 		else
+#endif // #if WITH_EDITORONLY_DATA
 		{
 			check(CurrentClass);
 			CurrentClass = Cast<UBlueprintGeneratedClass>(CurrentClass->GetSuperClass());
diff --git a/Engine/Source/Runtime/Engine/Private/BlueprintGeneratedClass.cpp b/Engine/Source/Runtime/Engine/Private/BlueprintGeneratedClass.cpp
--- a/Engine/Source/Runtime/Engine/Private/BlueprintGeneratedClass.cpp
+++ b/Engine/Source/Runtime/Engine/Private/BlueprintGeneratedClass.cpp
@@ -1345,12 +1345,14 @@ uint8* UBlueprintGeneratedClass::GetPersistentUberGraphFrame(UObject* Obj, UFunc
 void UBlueprintGeneratedClass::CreatePersistentUberGraphFrame(UObject* Obj, bool bCreateOnlyIfEmpty, bool bSkipSuperClass, UClass* OldClass) const
 {
 #if USE_UBER_GRAPH_PERSISTENT_FRAME
+#if WITH_EDITORONLY_DATA
 	/** Macros should not create uber graph frames as they have no uber graph. If UBlueprints are cooked out the macro class probably does not exist as well */
 	UBlueprint* Blueprint = Cast<UBlueprint>(ClassGeneratedBy);
 	if (Blueprint && Blueprint->BlueprintType == BPTYPE_MacroLibrary)
 	{
 		return;
 	}
+#endif
 
 	ensure(!UberGraphFramePointerProperty == !UberGraphFunction);
 	if (Obj && UsePersistentUberGraphFrame() && UberGraphFramePointerProperty && UberGraphFunction)
diff --git a/Engine/Source/Runtime/Engine/Private/MemberReference.cpp b/Engine/Source/Runtime/Engine/Private/MemberReference.cpp
--- a/Engine/Source/Runtime/Engine/Private/MemberReference.cpp
+++ b/Engine/Source/Runtime/Engine/Private/MemberReference.cpp
@@ -70,6 +70,7 @@ void FMemberReference::SetDirect(const FName InMemberName, const FGuid InMemberG
 	MemberScope.Empty();
 }
 
+#if WITH_EDITOR
 void FMemberReference::SetGivenSelfScope(const FName InMemberName, const FGuid InMemberGuid, TSubclassOf<class UObject> InMemberParentClass, TSubclassOf<class UObject> SelfScope) const
 {
 	MemberName = InMemberName;
@@ -87,6 +88,7 @@ void FMemberReference::SetGivenSelfScope(const FName InMemberName, const FGuid I
 		MemberParent = nullptr;
 	}
 }
+#endif
 
 void FMemberReference::SetLocalMember(FName InMemberName, UStruct* InScope, const FGuid InMemberGuid)
 {
diff --git a/Engine/Source/Runtime/Foliage/Private/FoliageTypeObject.cpp b/Engine/Source/Runtime/Foliage/Private/FoliageTypeObject.cpp
--- a/Engine/Source/Runtime/Foliage/Private/FoliageTypeObject.cpp
+++ b/Engine/Source/Runtime/Foliage/Private/FoliageTypeObject.cpp
@@ -65,6 +65,7 @@ void FFoliageTypeObject::SetClean()
 	}
 }
 
+#if WITH_EDITORONLY_DATA
 void FFoliageTypeObject::PostSerialize(const FArchive& Ar)
 {
 	if (Ar.UEVer() < VER_UE4_FOLIAGE_WITH_ASSET_OR_CLASS)
@@ -76,3 +77,4 @@ void FFoliageTypeObject::PostSerialize(const FArchive& Ar)
 		}
 	}
 }
+#endif
\ No newline at end of file
diff --git a/Engine/Source/Runtime/Foliage/Private/InstancedFoliage.cpp b/Engine/Source/Runtime/Foliage/Private/InstancedFoliage.cpp
--- a/Engine/Source/Runtime/Foliage/Private/InstancedFoliage.cpp
+++ b/Engine/Source/Runtime/Foliage/Private/InstancedFoliage.cpp
@@ -716,10 +716,12 @@ void UFoliageType::PostLoad()
 	}
 }
 
+#if WITH_EDITOR
 bool UFoliageType::IsNotAssetOrBlueprint() const
 {
 	return IsAsset() == false && Cast<UBlueprint>(GetClass()->ClassGeneratedBy) == nullptr;
 }
+#endif
 
 FVector UFoliageType::GetRandomScale() const
 {
@@ -2949,7 +2951,13 @@ UFoliageType* AInstancedFoliageActor::GetLocalFoliageTypeForSource(const UObject
 	{
 		UFoliageType* FoliageType = Pair.Key;
 		// Check that the type is neither an asset nor blueprint instance
-		if (FoliageType && FoliageType->GetSource() == InSource && !FoliageType->IsAsset() && !FoliageType->GetClass()->ClassGeneratedBy)
+
+		// ClassGeneratedBy TODO: This is wrong in cooked builds
+		if (FoliageType && FoliageType->GetSource() == InSource && !FoliageType->IsAsset()
+#if WITH_EDITORONLY_DATA
+		 && !FoliageType->GetClass()->ClassGeneratedBy
+#endif
+		)
 		{
 			ReturnType = FoliageType;
 			Info = &*Pair.Value;
diff --git a/Engine/Source/Runtime/Foliage/Public/FoliageTypeObject.h b/Engine/Source/Runtime/Foliage/Public/FoliageTypeObject.h
--- a/Engine/Source/Runtime/Foliage/Public/FoliageTypeObject.h
+++ b/Engine/Source/Runtime/Foliage/Public/FoliageTypeObject.h
@@ -36,7 +36,9 @@ struct FOLIAGE_API FFoliageTypeObject
 	bool IsDirty() const;
 	void SetClean();
 
+#if WITH_EDITORONLY_DATA
 	void PostSerialize(const FArchive& Ar);
+#endif
 
 private:
 	/** The foliage type that will be spawned by the procedural foliage simulation */
@@ -55,6 +57,7 @@ private:
 	TSubclassOf<UFoliageType_InstancedStaticMesh> Type_DEPRECATED;
 };
 
+#if WITH_EDITORONLY_DATA
 template<>
 struct TStructOpsTypeTraits<FFoliageTypeObject> : public TStructOpsTypeTraitsBase2<FFoliageTypeObject>
 {
@@ -63,3 +66,4 @@ struct TStructOpsTypeTraits<FFoliageTypeObject> : public TStructOpsTypeTraitsBas
 		WithPostSerialize = true,
 	};
 };
+#endif
\ No newline at end of file
