Hide UClass::Implements to prevent accidental misuse.

diff --git a/Engine/Source/Runtime/CoreUObject/Public/UObject/Class.h b/Engine/Source/Runtime/CoreUObject/Public/UObject/Class.h
--- a/Engine/Source/Runtime/CoreUObject/Public/UObject/Class.h
+++ b/Engine/Source/Runtime/CoreUObject/Public/UObject/Class.h
@@ -3378,7 +3378,17 @@ private:
 		return UObject::IsA(Parent);
 	}
 
-	/** 
+	/**
+	 * This signature intentionally hides the method declared in UObject to make it private.
+	 * Call ImplementsInterface instead; Hidden because calling Implements on a class almost always indicates an error where the caller should use ImplementsInterface
+	 */
+	template <typename T>
+	bool Implements() const
+	{
+		return UObject::Implements<T>();
+	}
+
+	/**
 	 * This signature intentionally hides the method declared in UObject to make it private.
 	 * Call FindFunctionByName instead; This method will search for a function declared in UClass instead of the class it was called on
 	 */
