Fix IsDefaultObject to not return true for CDOs
Cleaned up out of date comment about function implementation

diff --git a/Engine/Source/Runtime/CoreUObject/Private/UObject/UObjectBaseUtility.cpp b/Engine/Source/Runtime/CoreUObject/Private/UObject/UObjectBaseUtility.cpp
--- a/Engine/Source/Runtime/CoreUObject/Private/UObject/UObjectBaseUtility.cpp
+++ b/Engine/Source/Runtime/CoreUObject/Private/UObject/UObjectBaseUtility.cpp
@@ -480,7 +480,7 @@ void* UObjectBaseUtility::GetNativeInterfaceAddress(UClass* InterfaceClass)
 
 bool UObjectBaseUtility::IsDefaultSubobject() const
 {
-	const bool bIsInstanced = GetOuter() && (GetOuter()->HasAnyFlags(RF_ClassDefaultObject) || ((UObject*)this)->GetArchetype() != GetClass()->GetDefaultObject(false));
+	const bool bIsInstanced = !HasAnyFlags(RF_ClassDefaultObject) && GetOuter() && (GetOuter()->HasAnyFlags(RF_ClassDefaultObject) || ((UObject*)this)->GetArchetype() != GetClass()->GetDefaultObject(false));
 	return bIsInstanced;
 }
 
diff --git a/Engine/Source/Runtime/CoreUObject/Public/UObject/UObjectBaseUtility.h b/Engine/Source/Runtime/CoreUObject/Public/UObject/UObjectBaseUtility.h
--- a/Engine/Source/Runtime/CoreUObject/Public/UObject/UObjectBaseUtility.h
+++ b/Engine/Source/Runtime/CoreUObject/Public/UObject/UObjectBaseUtility.h
@@ -686,7 +686,6 @@ public:
 
 	/**
 	 * Returns whether this component was instanced from a component/subobject template, or if it is a component/subobject template.
-	 * This is based on a name comparison with the outer class instance lookup table
 	 *
 	 * @return	true if this component was instanced from a template.  false if this component was created manually at runtime.
 	 */
