Remove ClassGeneratedBy TODO from validated cases
Fixup one case in SavePackageUtilities
#rnx

diff --git a/Engine/Plugins/Animation/ControlRig/Source/ControlRig/Private/AnimNode_ControlRigBase.cpp b/Engine/Plugins/Animation/ControlRig/Source/ControlRig/Private/AnimNode_ControlRigBase.cpp
--- a/Engine/Plugins/Animation/ControlRig/Source/ControlRig/Private/AnimNode_ControlRigBase.cpp
+++ b/Engine/Plugins/Animation/ControlRig/Source/ControlRig/Private/AnimNode_ControlRigBase.cpp
@@ -40,7 +40,6 @@ void FAnimNode_ControlRigBase::OnInitializeAnimInstance(const FAnimInstanceProxy
 	if (Component && Component->SkeletalMesh && ControlRig)
 	{
 #if WITH_EDITORONLY_DATA
-		// ClassGeneratedBy TODO: This may be wrong in cooked builds
 		UBlueprintGeneratedClass* BlueprintClass = Cast<UBlueprintGeneratedClass>(ControlRig->GetClass());
 		if (BlueprintClass)
 		{
diff --git a/Engine/Source/Runtime/CoreUObject/Private/UObject/Obj.cpp b/Engine/Source/Runtime/CoreUObject/Private/UObject/Obj.cpp
--- a/Engine/Source/Runtime/CoreUObject/Private/UObject/Obj.cpp
+++ b/Engine/Source/Runtime/CoreUObject/Private/UObject/Obj.cpp
@@ -1644,7 +1644,6 @@ void UObject::BuildSubobjectMapping(UObject* OtherObject, TMap<UObject*, UObject
 
 		UClass* OtherSubObjectClass = InSubObject->GetClass();
 #if WITH_EDITORONLY_DATA
-		// ClassGeneratedBy TODO: This may be wrong in cooked builds
 		if (OtherSubObjectClass->ClassGeneratedBy && OtherSubObjectClass->ClassGeneratedBy->GetOutermost() == ThisPackage)
 		{
 			// This is a generated class type, so we actually need to use the new generated class type from the new package otherwise our type check will fail
@@ -1655,7 +1654,6 @@ void UObject::BuildSubobjectMapping(UObject* OtherObject, TMap<UObject*, UObject
 		}
 #endif
 
-		//UObject* OtherSubObject = StaticLoadObject(OtherSubObjectClass, OtherObject, *NewSubObjectName, nullptr, LOAD_Quiet | LOAD_NoRedirects, nullptr, true);
 		UObject* OtherSubObject = StaticFindObjectFast(OtherSubObjectClass, OtherObject, *NewSubObjectName);
 		ObjectMapping.Emplace(InSubObject, OtherSubObject);
 
diff --git a/Engine/Source/Runtime/CoreUObject/Private/UObject/ObjectHandle.cpp b/Engine/Source/Runtime/CoreUObject/Private/UObject/ObjectHandle.cpp
--- a/Engine/Source/Runtime/CoreUObject/Private/UObject/ObjectHandle.cpp
+++ b/Engine/Source/Runtime/CoreUObject/Private/UObject/ObjectHandle.cpp
@@ -271,7 +271,6 @@ FObjectRef MakeObjectRef(const UObject* Object)
 	}
 
 #if WITH_EDITORONLY_DATA	
-	// ClassGeneratedBy TODO: This may be wrong in cooked builds
 	UObject* ClassGeneratedBy = Object->GetClass()->ClassGeneratedBy;
 	UPackage* ClassGeneratedByPackage = ClassGeneratedBy ? ClassGeneratedBy->GetOutermost() : nullptr;
 	return FObjectRef {GetNameOrNone(Object->GetOutermost()), GetNameOrNone(ClassGeneratedByPackage), GetNameOrNone(ClassGeneratedBy), FObjectPathId(Object)};
@@ -299,7 +298,6 @@ FPackedObjectRef MakePackedObjectRef(const UObject* Object)
 
 	FName PackageName = GetNameOrNone(Object->GetOutermost());
 #if WITH_EDITORONLY_DATA
-	// ClassGeneratedBy TODO: This may be wrong in cooked builds
 	UObject* ClassGeneratedBy = Object->GetClass()->ClassGeneratedBy;
 	UPackage* ClassGeneratedByPackage = ClassGeneratedBy ? ClassGeneratedBy->GetOutermost() : nullptr;
 	return ObjectHandle_Private::MakePackedObjectRef(PackageName, GetNameOrNone(ClassGeneratedByPackage), GetNameOrNone(ClassGeneratedBy), FObjectPathId(Object));
diff --git a/Engine/Source/Runtime/CoreUObject/Private/UObject/SavePackage/SavePackageUtilities.cpp b/Engine/Source/Runtime/CoreUObject/Private/UObject/SavePackage/SavePackageUtilities.cpp
--- a/Engine/Source/Runtime/CoreUObject/Private/UObject/SavePackage/SavePackageUtilities.cpp
+++ b/Engine/Source/Runtime/CoreUObject/Private/UObject/SavePackage/SavePackageUtilities.cpp
@@ -351,12 +351,7 @@ void CheckObjectPriorToSave(FArchiveUObject& Ar, UObject* InObj, UPackage* InSav
 		InObj->SetFlags(RF_Transient);
 	}
 
-	if ( InObj->HasAnyFlags(RF_ClassDefaultObject)
-#if WITH_EDITORONLY_DATA
-		// ClassGeneratedBy TODO: This may be wrong in cooked builds
-		&& (InObj->GetClass()->ClassGeneratedBy == nullptr || !InObj->GetClass()->HasAnyFlags(RF_Transient)) 
-#endif
-		)
+	if ( InObj->HasAnyFlags(RF_ClassDefaultObject) && (InObj->GetClass()->IsNative() || !InObj->GetClass()->HasAnyFlags(RF_Transient)) )
 	{
 		// if this is the class default object, make sure it's not
 		// marked transient for any reason, as we need it to be saved
diff --git a/Engine/Source/Runtime/Engine/Private/Blueprint.cpp b/Engine/Source/Runtime/Engine/Private/Blueprint.cpp
--- a/Engine/Source/Runtime/Engine/Private/Blueprint.cpp
+++ b/Engine/Source/Runtime/Engine/Private/Blueprint.cpp
@@ -1430,7 +1430,6 @@ bool UBlueprint::GetBlueprintHierarchyFromClass(const UClass* InClass, TArray<UB
 		OutBlueprintParents.Add(CurrentClass);
 
 #if WITH_EDITORONLY_DATA
-		// ClassGeneratedBy TODO: This may be wrong in cooked builds
 		UBlueprint* BP = UBlueprint::GetBlueprintFromClass(CurrentClass);
 
 		if (BP)
diff --git a/Engine/Source/Runtime/Engine/Private/EdGraph/EdGraph.cpp b/Engine/Source/Runtime/Engine/Private/EdGraph/EdGraph.cpp
--- a/Engine/Source/Runtime/Engine/Private/EdGraph/EdGraph.cpp
+++ b/Engine/Source/Runtime/Engine/Private/EdGraph/EdGraph.cpp
@@ -82,9 +82,9 @@ UEdGraph::UEdGraph(const FObjectInitializer& ObjectInitializer)
 	bAllowDeletion = true;
 }
 
+#if WITH_EDITORONLY_DATA
 void UEdGraph::BuildSubobjectMapping(UObject* OtherObject, TMap<UObject*, UObject*>& ObjectMapping) const
 {
-#if WITH_EDITORONLY_DATA
 	UEdGraph* OtherGraph = CastChecked<UEdGraph>(OtherObject);
 
 	auto FindMatchingNode = [](UEdGraph* InGraphToSearch, UEdGraphNode* InNodeToFind) -> UEdGraphNode*
@@ -131,12 +131,10 @@ void UEdGraph::BuildSubobjectMapping(UObject* OtherObject, TMap<UObject*, UObjec
 			ObjectMapping.Emplace(GraphNode, OtherGraphNode);
 		}
 	}
-#endif // WITH_EDITORONLY_DATA
 
 	Super::BuildSubobjectMapping(OtherObject, ObjectMapping);
 }
 
-#if WITH_EDITORONLY_DATA
 void UEdGraph::PostInitProperties()
 {
 	Super::PostInitProperties();
