Fix IsInBlueprint to work correctly in cooked builds
Give useful name in cooked builds when handling invalid reference errors
#rnx

diff --git a/Engine/Source/Runtime/CoreUObject/Private/Blueprint/BlueprintSupport.cpp b/Engine/Source/Runtime/CoreUObject/Private/Blueprint/BlueprintSupport.cpp
--- a/Engine/Source/Runtime/CoreUObject/Private/Blueprint/BlueprintSupport.cpp
+++ b/Engine/Source/Runtime/CoreUObject/Private/Blueprint/BlueprintSupport.cpp
@@ -2379,31 +2379,24 @@ UObject* FLinkerLoad::FindImportFast(UClass* ImportClass, UObject* ImportOuter,
  */
 bool UObject::IsInBlueprint() const
 {
-	// ClassGeneratedBy TODO: This is wrong in cooked builds
-#if WITH_EDITORONLY_DATA
-
 	// Exclude blueprint classes as they may be regenerated at any time
 	// Need to exclude classes, CDOs, and their subobjects
 	const UObject* TestObject = this;
  	while (TestObject)
  	{
  		const UClass *ClassObject = dynamic_cast<const UClass*>(TestObject);
-		if (ClassObject 
-			&& ClassObject->HasAnyClassFlags(CLASS_CompiledFromBlueprint) 
-			&& ClassObject->ClassGeneratedBy)
+		if (ClassObject && !ClassObject->IsNative())
  		{
  			return true;
  		}
 		else if (TestObject->HasAnyFlags(RF_ClassDefaultObject) 
 			&& TestObject->GetClass() 
-			&& TestObject->GetClass()->HasAnyClassFlags(CLASS_CompiledFromBlueprint) 
-			&& TestObject->GetClass()->ClassGeneratedBy)
+			&& !TestObject->GetClass()->IsNative())
  		{
  			return true;
  		}
  		TestObject = TestObject->GetOuter();
  	}
-#endif
 
 	return false;
 }
diff --git a/Engine/Source/Runtime/CoreUObject/Private/UObject/UObjectGlobals.cpp b/Engine/Source/Runtime/CoreUObject/Private/UObject/UObjectGlobals.cpp
--- a/Engine/Source/Runtime/CoreUObject/Private/UObject/UObjectGlobals.cpp
+++ b/Engine/Source/Runtime/CoreUObject/Private/UObject/UObjectGlobals.cpp
@@ -3749,17 +3749,22 @@ protected:
 
 		auto GetBlueprintObjectNameLambda = [](const UObject* InSerializingObject) -> FString
 		{
-#if WITH_EDITORONLY_DATA
-			// ClassGeneratedBy TODO: This may be wrong in cooked builds
 			if (InSerializingObject)
 			{
-				const UClass* BPGC = InSerializingObject->GetTypedOuter<UClass>();
-				if (BPGC && BPGC->ClassGeneratedBy)
+				if (const UClass* BPGC = InSerializingObject->GetTypedOuter<UClass>())
 				{
-					return BPGC->ClassGeneratedBy->GetFullName();
+#if WITH_EDITORONLY_DATA
+					if (BPGC->ClassGeneratedBy)
+					{
+						return BPGC->ClassGeneratedBy->GetFullName();
+					}
+					else
+#endif
+					{
+						return BPGC->GetFullName();
+					}
 				}
 			}
-#endif
 
 			return TEXT("NULL");
 		};
