Avoid unnecessary lookup of Archetype
General cleanup of component langugage where actually it is more generally for all subobjects
Fix comments and add default parameters
#rnx

diff --git a/Engine/Source/Editor/UnrealEd/Private/EditorObject.cpp b/Engine/Source/Editor/UnrealEd/Private/EditorObject.cpp
--- a/Engine/Source/Editor/UnrealEd/Private/EditorObject.cpp
+++ b/Engine/Source/Editor/UnrealEd/Private/EditorObject.cpp
@@ -608,7 +608,7 @@ static const TCHAR* ImportProperties(
 				{
 					checkSlow(ComponentTemplate->GetArchetype() == Archetype);
 					ReplacementMap.Add(Archetype, ComponentTemplate);
-					InstanceGraph.AddNewInstance(ComponentTemplate);
+					InstanceGraph.AddNewInstance(ComponentTemplate, Archetype);
 				}
 				if (OldComponent)
 				{
diff --git a/Engine/Source/Runtime/CoreUObject/Private/UObject/CoreNative.cpp b/Engine/Source/Runtime/CoreUObject/Private/UObject/CoreNative.cpp
--- a/Engine/Source/Runtime/CoreUObject/Private/UObject/CoreNative.cpp
+++ b/Engine/Source/Runtime/CoreUObject/Private/UObject/CoreNative.cpp
@@ -63,7 +63,7 @@ IMPLEMENT_MODULE( FCoreUObjectModule, CoreUObject );
 
 // if we are not using compiled in natives, we still need this as a base class for intrinsics
 #if !USE_COMPILED_IN_NATIVES
-COREUOBJECT_API class UClass* Z_Construct_UClass_UObject();
+COREUOBJECT_API UClass* Z_Construct_UClass_UObject();
 extern FClassRegistrationInfo Z_Registration_Info_UClass_UObject;
 UClass* Z_Construct_UClass_UObject()
 {
@@ -137,17 +137,17 @@ UObject* FObjectInstancingGraph::GetInstancedSubobject( UObject* SourceSubobject
 		if ( !bShouldInstance && CurrentValue->GetOuter() == CurrentObject->GetArchetype() )
 		{
 			// this code is intended to catch cases where SourceRoot contains subobjects assigned to instanced object properties, where the subobject's class
-			// contains components, and the class of the subobject is outside of the inheritance hierarchy of the SourceRoot, for example, a weapon
+			// contains subobjects, and the class of the subobject is outside of the inheritance hierarchy of the SourceRoot, for example, a weapon
 			// class which contains UIObject subobject definitions in its defaultproperties, where the property referencing the UIObjects is marked instanced.
 			bShouldInstance = true;
 
-			// if this case is triggered, ensure that the CurrentValue of the component property is still pointing to the template component.
+			// if this case is triggered, ensure that the CurrentValue of the subobject property is still pointing to the template subobject.
 			check(SourceSubobject == CurrentValue);
 		}
 
 		if ( bShouldInstance )
 		{
-			// search for the unique component instance that corresponds to this component template
+			// search for the unique subobject instance that corresponds to this subobject template
 			InstancedSubobject = GetDestinationObject(SourceSubobject);
 			if ( InstancedSubobject == nullptr )
 			{
@@ -157,12 +157,12 @@ UObject* FObjectInstancingGraph::GetInstancedSubobject( UObject* SourceSubobject
 				}
 				else
 				{
-					// if the Outer for the component currently assigned to this property is the same as the object that we're instancing components for,
-					// the component does not need to be instanced; otherwise, there are two possiblities:
+					// if the Outer for the subobject currently assigned to this property is the same as the object that we're instancing subobjects for,
+					// the subobject does not need to be instanced; otherwise, there are two possiblities:
 					// 1. CurrentValue is a template and needs to be instanced
-					// 2. CurrentValue is an instanced component, in which case it should already be in InstanceGraph, UNLESS the component was created
+					// 2. CurrentValue is an instanced subobject, in which case it should already be in InstanceGraph, UNLESS the subobject was created
 					//		at runtime (editinline export properties, for example).  If that is the case, CurrentValue will be an instance that is not linked
-					//		to the component template referenced by CurrentObject's archetype, and in this case, we also don't want to re-instance the component template
+					//		to the subobject template referenced by CurrentObject's archetype, and in this case, we also don't want to re-instance the subobject template
 
 					const bool bIsRuntimeInstance = CurrentValue != SourceSubobject && CurrentValue->GetOuter() == CurrentObject;
 					if (bIsRuntimeInstance )
@@ -171,7 +171,7 @@ UObject* FObjectInstancingGraph::GetInstancedSubobject( UObject* SourceSubobject
 					}
 					else
 					{
-						// If the component template is relevant in this context(client vs server vs editor), instance it.
+						// If the subobject template is relevant in this context(client vs server vs editor), instance it.
 						const bool bShouldLoadForClient = SourceSubobject->NeedsLoadForClient();
 						const bool bShouldLoadForServer = SourceSubobject->NeedsLoadForServer();
 						const bool bShouldLoadForEditor = ( GIsEditor && ( bShouldLoadForClient || !CurrentObject->RootPackageHasAnyFlags(PKG_PlayInEditor) ) );
@@ -180,8 +180,8 @@ UObject* FObjectInstancingGraph::GetInstancedSubobject( UObject* SourceSubobject
 						{
 							// this is the first time the instance corresponding to SourceSubobject has been requested
 
-							// get the object instance corresponding to the source component's Outer - this is the object that
-							// will be used as the Outer for the destination component
+							// get the object instance corresponding to the source subobject's Outer - this is the object that
+							// will be used as the Outer for the destination subobject
 							UObject* SubobjectOuter = GetDestinationObject(SourceSubobject->GetOuter());
 
 							// In the event we're templated off a deep nested UObject hierarchy, with several links to objects nested in the object
@@ -191,7 +191,7 @@ UObject* FObjectInstancingGraph::GetInstancedSubobject( UObject* SourceSubobject
 							{
 								SubobjectOuter = GetInstancedSubobject(SourceSubobject->GetOuter(), SourceSubobject->GetOuter(), CurrentObject, bDoNotCreateNewInstance, bAllowSelfReference);
 
-								checkf(SubobjectOuter && SubobjectOuter != INVALID_OBJECT, TEXT("No corresponding destination object found for '%s' while attempting to instance component '%s'"), *SourceSubobject->GetOuter()->GetFullName(), *SourceSubobject->GetFullName());
+								checkf(SubobjectOuter && SubobjectOuter != INVALID_OBJECT, TEXT("No corresponding destination object found for '%s' while attempting to instance subobject '%s'"), *SourceSubobject->GetOuter()->GetFullName(), *SourceSubobject->GetFullName());
 							}
 
 							FName SubobjectName = SourceSubobject->GetFName();
@@ -217,7 +217,7 @@ UObject* FObjectInstancingGraph::GetInstancedSubobject( UObject* SourceSubobject
 
 							if (!InstancedSubobject)
 							{
-								// finally, create the component instance
+								// finally, create the subobject instance
 								FStaticConstructObjectParameters Params(SourceSubobject->GetClass());
 								Params.Outer = SubobjectOuter;
 								Params.Name = SubobjectName;
@@ -233,7 +233,7 @@ UObject* FObjectInstancingGraph::GetInstancedSubobject( UObject* SourceSubobject
 			}
 			else if ( IsLoadingObject() && InstancedSubobject->GetClass()->HasAnyClassFlags(CLASS_HasInstancedReference) )
 			{
-				/* When loading an object from disk, in some cases we have a component which has a reference to another component in DestinationObject which
+				/* When loading an object from disk, in some cases we have a subobject which has a reference to another subobject in DestinationObject which
 					wasn't serialized and hasn't yet been instanced.  For example, the PointLight class declared two component templates:
 
 						Begin DrawLightRadiusComponent0
@@ -262,7 +262,7 @@ UObject* FObjectInstancingGraph::GetInstancedSubobject( UObject* SourceSubobject
 }
 
 
-UObject* FObjectInstancingGraph::InstancePropertyValue( class UObject* ComponentTemplate, class UObject* CurrentValue, class UObject* Owner, bool bIsTransient, bool bCausesInstancing, bool bAllowSelfReference )
+UObject* FObjectInstancingGraph::InstancePropertyValue(UObject* SubObjectTemplate, UObject* CurrentValue, UObject* Owner, bool bIsTransient, bool bCausesInstancing, bool bAllowSelfReference)
 {
 	UObject* NewValue = CurrentValue;
 
@@ -280,22 +280,22 @@ UObject* FObjectInstancingGraph::InstancePropertyValue( class UObject* Component
 		return NewValue; // not instancing
 	}
 
-	// if the object we're instancing the components for (Owner) has the current component's outer in its archetype chain, and its archetype has a nullptr value
-	// for this component property it means that the archetype didn't instance its component, so we shouldn't either.
+	// if the object we're instancing the subobjects for (Owner) has the current subobject's outer in its archetype chain, and its archetype has a nullptr value
+	// for this subobject property it means that the archetype didn't instance its subobject, so we shouldn't either.
 
-	if (ComponentTemplate == nullptr && CurrentValue != nullptr && (Owner && Owner->IsBasedOnArchetype(CurrentValue->GetOuter())))
+	if (SubObjectTemplate == nullptr && CurrentValue != nullptr && (Owner && Owner->IsBasedOnArchetype(CurrentValue->GetOuter())))
 	{
 		NewValue = nullptr;
 	}
 	else
 	{
-		if ( ComponentTemplate == nullptr )
+		if (SubObjectTemplate == nullptr )
 		{
-			// should only be here if our archetype doesn't contain this component property
-			ComponentTemplate = CurrentValue;
+			// should only be here if our archetype doesn't contain this subobject property
+			SubObjectTemplate = CurrentValue;
 		}
 
-		UObject* MaybeNewValue = GetInstancedSubobject(ComponentTemplate, CurrentValue, Owner, bAllowSelfReference, bAllowSelfReference);
+		UObject* MaybeNewValue = GetInstancedSubobject(SubObjectTemplate, CurrentValue, Owner, bAllowSelfReference, bAllowSelfReference);
 		if ( MaybeNewValue != INVALID_OBJECT )
 		{
 			NewValue = MaybeNewValue;
@@ -304,7 +304,7 @@ UObject* FObjectInstancingGraph::InstancePropertyValue( class UObject* Component
 	return NewValue;
 }
 
-void FObjectInstancingGraph::AddNewObject(class UObject* ObjectInstance, UObject* InArchetype /*= nullptr*/)
+void FObjectInstancingGraph::AddNewObject(UObject* ObjectInstance, UObject* InArchetype /*= nullptr*/)
 {
 	check(!GEventDrivenLoaderEnabled || !InArchetype || !InArchetype->HasAnyFlags(RF_NeedLoad));
 
diff --git a/Engine/Source/Runtime/CoreUObject/Public/UObject/UObjectGlobals.h b/Engine/Source/Runtime/CoreUObject/Public/UObject/UObjectGlobals.h
--- a/Engine/Source/Runtime/CoreUObject/Public/UObject/UObjectGlobals.h
+++ b/Engine/Source/Runtime/CoreUObject/Public/UObject/UObjectGlobals.h
@@ -926,24 +926,23 @@ public:
 
 	/**
 	* Create a component or subobject only to be used with the editor.
-	* @param	TReturnType					class of return type, all overrides must be of this type
 	* @param	Outer						outer to construct the subobject in
-	* @param	ReturnType					type of the new component 
 	* @param	SubobjectName				name of the new component
+	* @param	ReturnType					type of the new component
 	* @param	bTransient					true if the component is being assigned to a transient property
 	*/
 	UObject* CreateEditorOnlyDefaultSubobject(UObject* Outer, FName SubobjectName, UClass* ReturnType, bool bTransient = false) const;
 
 	/**
 	 * Create a component or subobject
-	 * @param	TReturnType					class of return type, all overrides must be of this type
-	 * @param	TClassToConstructByDefault	if the derived class has not overridden, create a component of this type (default is TReturnType)
-	 * @param	Outer						outer to construct the subobject in
-	 * @param	SubobjectName				name of the new component
-	 * @param bIsRequired			true if the component is required and will always be created even if DoNotCreateDefaultSubobject was specified.
-	 * @param bIsTransient		true if the component is being assigned to a transient property
+	 * @param	Outer                       outer to construct the subobject in
+	 * @param	SubobjectName               name of the new component
+	 * @param	ReturnType                  class of return type, all overrides must be of this type
+	 * @param	ClassToConstructByDefault   if the derived class has not overridden, create a component of this type
+	 * @param	bIsRequired                 true if the component is required and will always be created even if DoNotCreateDefaultSubobject was specified.
+	 * @param	bIsTransient                true if the component is being assigned to a transient property
 	 */
-	UObject* CreateDefaultSubobject(UObject* Outer, FName SubobjectFName, UClass* ReturnType, UClass* ClassToCreateByDefault, bool bIsRequired, bool bIsTransient) const;
+	UObject* CreateDefaultSubobject(UObject* Outer, FName SubobjectFName, UClass* ReturnType, UClass* ClassToCreateByDefault, bool bIsRequired = true, bool bIsTransient = false) const;
 
 	/**
 	 * Sets the class to use for a subobject defined in a base class, the class must be a subclass of the class used by the base class.
