EditorBulkData and the virtualization system now use FIoHash directly and FPayloadId is removed

#rnx

- A lot of files have been touched but most of this is changing FPayload::IsValid to !FIoHash::IsZero, the main changes are in EditorBulkData/EditorBulkDataTests
- The only difference between FPayloadId and FIoHash is that the former considered the hash of an invalid or empty buffer to be 'invalid' too where as FIoHash would return a valid hash
-- To keep behaviour the same, we only hash payloads in EditorBulkData using a utility method ::HashBuffer which will not hash empty or invalid payloads and return a default FIoHash instead.
-- Unit tests have been extended to prove that the behaviour has not changed.

diff --git a/Engine/Source/Developer/Virtualization/Private/DumpPackagePayloadInfoCommand.cpp b/Engine/Source/Developer/Virtualization/Private/DumpPackagePayloadInfoCommand.cpp
--- a/Engine/Source/Developer/Virtualization/Private/DumpPackagePayloadInfoCommand.cpp
+++ b/Engine/Source/Developer/Virtualization/Private/DumpPackagePayloadInfoCommand.cpp
@@ -31,8 +31,8 @@ void DumpPackagePayloadInfo(const TArray<FString>& Args)
 
 		if (FPackagePath::TryFromMountedName(Arg, Path))
 		{
-			TArray<Virtualization::FPayloadId> LocalPayloadIds;
-			TArray<Virtualization::FPayloadId> VirtualizedPayloadIds;
+			TArray<FIoHash> LocalPayloadIds;
+			TArray<FIoHash> VirtualizedPayloadIds;
 
 			if (!UE::FindPayloadsInPackageFile(Path, UE::EPayloadFilter::Local, LocalPayloadIds))
 			{
diff --git a/Engine/Source/Developer/Virtualization/Private/PackageSubmissionChecks.cpp b/Engine/Source/Developer/Virtualization/Private/PackageSubmissionChecks.cpp
--- a/Engine/Source/Developer/Virtualization/Private/PackageSubmissionChecks.cpp
+++ b/Engine/Source/Developer/Virtualization/Private/PackageSubmissionChecks.cpp
@@ -158,7 +158,7 @@ void OnPrePackageSubmission(const TArray<FString>& FilesToSubmit, TArray<FText>&
 		FPackagePath Path;
 		FPackageTrailer Trailer;
 
-		TArray<FPayloadId> LocalPayloads;
+		TArray<FIoHash> LocalPayloads;
 		int32 PayloadIndex = INDEX_NONE;
 
 		bool bWasTrailerUpdated = false;
@@ -169,7 +169,7 @@ void OnPrePackageSubmission(const TArray<FString>& FilesToSubmit, TArray<FText>&
 	TArray<FPackageInfo> Packages;
 	Packages.Reserve(FilesToSubmit.Num());
 
-	TArray<FPayloadId> AllLocalPayloads;
+	TArray<FIoHash> AllLocalPayloads;
 	AllLocalPayloads.Reserve(FilesToSubmit.Num());
 
 	// From the list of files to submit we need to find all of the valid packages that contain
@@ -289,9 +289,9 @@ void OnPrePackageSubmission(const TArray<FString>& FilesToSubmit, TArray<FText>&
 
 		PackageInfo.PayloadIndex = PayloadsToSubmit.Num();
 
-		for (const FPayloadId& PayloadId : PackageInfo.LocalPayloads)
+		for (const FIoHash& PayloadId : PackageInfo.LocalPayloads)
 		{
-			checkf(PayloadId.IsValid(), TEXT("PackageTrailer for package '%s' should not contain invalid FPayloadIds"), *PackageInfo.Path.GetDebugName());
+			checkf(!PayloadId.IsZero(), TEXT("PackageTrailer for package '%s' should not contain invalid FIoHashs"), *PackageInfo.Path.GetDebugName());
 			
 			FCompressedBuffer Payload = PackageInfo.Trailer.LoadLocalPayload(PayloadId, *PackageAr);
 
diff --git a/Engine/Source/Developer/Virtualization/Private/VirtualizationDDCBackend.cpp b/Engine/Source/Developer/Virtualization/Private/VirtualizationDDCBackend.cpp
--- a/Engine/Source/Developer/Virtualization/Private/VirtualizationDDCBackend.cpp
+++ b/Engine/Source/Developer/Virtualization/Private/VirtualizationDDCBackend.cpp
@@ -3,7 +3,6 @@
 #include "VirtualizationDDCBackend.h"
 
 #include "Misc/Parse.h"
-#include "Virtualization/PayloadId.h"
 #include "DerivedDataCache.h"
 #include "DerivedDataCacheRecord.h"
 #include "DerivedDataRequestOwner.h"
@@ -12,10 +11,10 @@
 namespace UE::Virtualization
 {
 
-/** Utility function to help convert from UE::Virtualization::FPayloadId to UE::DerivedData::FValueId */
-static UE::DerivedData::FValueId ToDerivedDataValueId(const FPayloadId& Id)
+/** Utility function to help convert from UE::Virtualization::FIoHash to UE::DerivedData::FValueId */
+static UE::DerivedData::FValueId ToDerivedDataValueId(const FIoHash& Id)
 {
-	return UE::DerivedData::FValueId::FromHash(Id.GetIdentifier());
+	return UE::DerivedData::FValueId::FromHash(Id);
 }
 
 FDDCBackend::FDDCBackend(FStringView ConfigName, FStringView InDebugName)
@@ -71,7 +70,7 @@ bool FDDCBackend::Initialize(const FString& ConfigEntry)
 	return true;	
 }
 
-EPushResult FDDCBackend::PushData(const FPayloadId& Id, const FCompressedBuffer& Payload, const FString& PackageContext)
+EPushResult FDDCBackend::PushData(const FIoHash& Id, const FCompressedBuffer& Payload, const FString& PackageContext)
 {
 	TRACE_CPUPROFILER_EVENT_SCOPE(FDDCBackend::PushData);
 
@@ -85,10 +84,10 @@ EPushResult FDDCBackend::PushData(const FPayloadId& Id, const FCompressedBuffer&
 	
 	UE::DerivedData::FCacheKey Key;
 	Key.Bucket = Bucket;
-	Key.Hash = Id.GetIdentifier();
+	Key.Hash = Id;
 
 	UE::DerivedData::FValue DerivedDataValue(Payload);
-	check(DerivedDataValue.GetRawHash() == Id.GetIdentifier());
+	check(DerivedDataValue.GetRawHash() == Id);
 
 	UE::DerivedData::FCacheRecordBuilder RecordBuilder(Key);
 	RecordBuilder.AddValue(ToDerivedDataValueId(Id), DerivedDataValue);
@@ -116,7 +115,7 @@ EPushResult FDDCBackend::PushData(const FPayloadId& Id, const FCompressedBuffer&
 	}
 }
 
-FCompressedBuffer FDDCBackend::PullData(const FPayloadId& Id)
+FCompressedBuffer FDDCBackend::PullData(const FIoHash& Id)
 {
 	TRACE_CPUPROFILER_EVENT_SCOPE(FDDCBackend::PullData);
 
@@ -124,7 +123,7 @@ FCompressedBuffer FDDCBackend::PullData(const FPayloadId& Id)
 
 	UE::DerivedData::FCacheKey Key;
 	Key.Bucket = Bucket; 
-	Key.Hash = Id.GetIdentifier();
+	Key.Hash = Id;
 
 	UE::DerivedData::FRequestOwner Owner(UE::DerivedData::EPriority::Blocking);
 
@@ -148,7 +147,7 @@ FCompressedBuffer FDDCBackend::PullData(const FPayloadId& Id)
 	return ResultData;
 }
 
-bool FDDCBackend::DoesPayloadExist(const FPayloadId& Id)
+bool FDDCBackend::DoesPayloadExist(const FIoHash& Id)
 {
 	TRACE_CPUPROFILER_EVENT_SCOPE(FDDCBackend::DoesPayloadExist);
 
@@ -156,7 +155,7 @@ bool FDDCBackend::DoesPayloadExist(const FPayloadId& Id)
 
 	UE::DerivedData::FCacheKey Key;
 	Key.Bucket = Bucket;
-	Key.Hash = Id.GetIdentifier();
+	Key.Hash = Id;
 
 	UE::DerivedData::FRequestOwner Owner(UE::DerivedData::EPriority::Blocking);
 	
diff --git a/Engine/Source/Developer/Virtualization/Private/VirtualizationDDCBackend.h b/Engine/Source/Developer/Virtualization/Private/VirtualizationDDCBackend.h
--- a/Engine/Source/Developer/Virtualization/Private/VirtualizationDDCBackend.h
+++ b/Engine/Source/Developer/Virtualization/Private/VirtualizationDDCBackend.h
@@ -36,11 +36,11 @@ private:
 
 	virtual bool Initialize(const FString& ConfigEntry) override;
 
-	virtual EPushResult PushData(const FPayloadId& Id, const FCompressedBuffer& Payload, const FString& PackageContext) override;
+	virtual EPushResult PushData(const FIoHash& Id, const FCompressedBuffer& Payload, const FString& PackageContext) override;
 
-	virtual FCompressedBuffer PullData(const FPayloadId& Id) override;
+	virtual FCompressedBuffer PullData(const FIoHash& Id) override;
 	
-	virtual bool DoesPayloadExist(const FPayloadId& Id) override;
+	virtual bool DoesPayloadExist(const FIoHash& Id) override;
 
 private:
 	/** The bucket being used to group together the virtualized payloads in storage */
diff --git a/Engine/Source/Developer/Virtualization/Private/VirtualizationFileBackend.cpp b/Engine/Source/Developer/Virtualization/Private/VirtualizationFileBackend.cpp
--- a/Engine/Source/Developer/Virtualization/Private/VirtualizationFileBackend.cpp
+++ b/Engine/Source/Developer/Virtualization/Private/VirtualizationFileBackend.cpp
@@ -6,7 +6,6 @@
 #include "HAL/PlatformProcess.h"
 #include "Misc/Parse.h"
 #include "Misc/Paths.h"
-#include "Virtualization/PayloadId.h"
 #include "VirtualizationUtilities.h"
 
 namespace UE::Virtualization
@@ -54,7 +53,7 @@ bool FFileSystemBackend::Initialize(const FString& ConfigEntry)
 	return true;
 }
 
-EPushResult FFileSystemBackend::PushData(const FPayloadId& Id, const FCompressedBuffer& Payload, const FString& PackageContext)
+EPushResult FFileSystemBackend::PushData(const FIoHash& Id, const FCompressedBuffer& Payload, const FString& PackageContext)
 {
 	TRACE_CPUPROFILER_EVENT_SCOPE(FFileSystemBackend::PushData);
 
@@ -143,7 +142,7 @@ EPushResult FFileSystemBackend::PushData(const FPayloadId& Id, const FCompressed
 	return EPushResult::Success;
 }
 
-FCompressedBuffer FFileSystemBackend::PullData(const FPayloadId& Id)
+FCompressedBuffer FFileSystemBackend::PullData(const FIoHash& Id)
 {
 	TRACE_CPUPROFILER_EVENT_SCOPE(FFileSystemBackend::PullData);
 
@@ -176,7 +175,7 @@ FCompressedBuffer FFileSystemBackend::PullData(const FPayloadId& Id)
 	return FCompressedBuffer::Load(*FileAr);
 }
 
-bool FFileSystemBackend::DoesPayloadExist(const FPayloadId& Id)
+bool FFileSystemBackend::DoesPayloadExist(const FIoHash& Id)
 {
 	TRACE_CPUPROFILER_EVENT_SCOPE(FFileSystemBackend::DoesPayloadExist);
 
@@ -186,7 +185,7 @@ bool FFileSystemBackend::DoesPayloadExist(const FPayloadId& Id)
 	return IFileManager::Get().FileExists(FilePath.ToString());
 }
 
-void FFileSystemBackend::CreateFilePath(const FPayloadId& PayloadId, FStringBuilderBase& OutPath)
+void FFileSystemBackend::CreateFilePath(const FIoHash& PayloadId, FStringBuilderBase& OutPath)
 {
 	TStringBuilder<52> PayloadPath;
 	Utils::PayloadIdToPath(PayloadId, PayloadPath);
diff --git a/Engine/Source/Developer/Virtualization/Private/VirtualizationFileBackend.h b/Engine/Source/Developer/Virtualization/Private/VirtualizationFileBackend.h
--- a/Engine/Source/Developer/Virtualization/Private/VirtualizationFileBackend.h
+++ b/Engine/Source/Developer/Virtualization/Private/VirtualizationFileBackend.h
@@ -39,15 +39,15 @@ private:
 
 	virtual bool Initialize(const FString& ConfigEntry) override;
 
-	virtual EPushResult PushData(const FPayloadId& Id, const FCompressedBuffer& Payload, const FString& PackageContext) override;
+	virtual EPushResult PushData(const FIoHash& Id, const FCompressedBuffer& Payload, const FString& PackageContext) override;
 
-	virtual FCompressedBuffer PullData(const FPayloadId& Id) override;
+	virtual FCompressedBuffer PullData(const FIoHash& Id) override;
 
-	virtual bool DoesPayloadExist(const FPayloadId& Id) override;
+	virtual bool DoesPayloadExist(const FIoHash& Id) override;
 	
 private:
 
-	void CreateFilePath(const FPayloadId& PayloadId, FStringBuilderBase& OutPath);
+	void CreateFilePath(const FIoHash& PayloadId, FStringBuilderBase& OutPath);
 
 	TUniquePtr<FArchive> OpenFileForReading(const TCHAR* FilePath);
 
diff --git a/Engine/Source/Developer/Virtualization/Private/VirtualizationHttpBackend.cpp b/Engine/Source/Developer/Virtualization/Private/VirtualizationHttpBackend.cpp
--- a/Engine/Source/Developer/Virtualization/Private/VirtualizationHttpBackend.cpp
+++ b/Engine/Source/Developer/Virtualization/Private/VirtualizationHttpBackend.cpp
@@ -1,7 +1,6 @@
 // Copyright Epic Games, Inc. All Rights Reserved.
 
 #include "VirtualizationHttpBackend.h"
-#include "Virtualization/PayloadId.h"
 
 // TODO: Our libcurl implementation does not currently support MacOS
 // (not registering this will cause a fatal log error if the backend is actually used)
@@ -1374,7 +1373,7 @@ bool FHttpBackend::Initialize(const FString& ConfigEntry)
 	return true;
 }
 
-EPushResult FHttpBackend::PushData(const FPayloadId& Id, const FCompressedBuffer& CompressedPayload, const FString& PackageContext)
+EPushResult FHttpBackend::PushData(const FIoHash& Id, const FCompressedBuffer& CompressedPayload, const FString& PackageContext)
 {
 	using namespace Utility;
 
@@ -1492,7 +1491,7 @@ EPushResult FHttpBackend::PushData(const FPayloadId& Id, const FCompressedBuffer
 	return EPushResult::Failed;
 }
 
-FCompressedBuffer FHttpBackend::PullData(const FPayloadId& Id)
+FCompressedBuffer FHttpBackend::PullData(const FIoHash& Id)
 {
 	using namespace Utility;
 
@@ -1609,7 +1608,7 @@ FCompressedBuffer FHttpBackend::PullData(const FPayloadId& Id)
 	}
 }
 
-bool FHttpBackend::DoesPayloadExist(const FPayloadId& Id)
+bool FHttpBackend::DoesPayloadExist(const FIoHash& Id)
 {
 	using namespace Utility;
 
@@ -1875,7 +1874,7 @@ bool FHttpBackend::ShouldRetryOnError(int64 ResponseCode)
 	return false;
 }
 
-bool FHttpBackend::PostChunk(const TArrayView<const uint8>& ChunkData, const FPayloadId& PayloadId, FString& OutHashAsString)
+bool FHttpBackend::PostChunk(const TArrayView<const uint8>& ChunkData, const FIoHash& PayloadId, FString& OutHashAsString)
 {
 	TRACE_CPUPROFILER_EVENT_SCOPE(FHttpBackend::PostChunk);
 
@@ -1924,7 +1923,7 @@ bool FHttpBackend::PostChunk(const TArrayView<const uint8>& ChunkData, const FPa
 	return false;
 }
 
-bool FHttpBackend::PullChunk(const FString& Hash, const FPayloadId& PayloadId, uint8* DataPtr, int64 BufferSize)
+bool FHttpBackend::PullChunk(const FString& Hash, const FIoHash& PayloadId, uint8* DataPtr, int64 BufferSize)
 {
 	TRACE_CPUPROFILER_EVENT_SCOPE(FHttpBackend::PullChunk);
 
diff --git a/Engine/Source/Developer/Virtualization/Private/VirtualizationHttpBackend.h b/Engine/Source/Developer/Virtualization/Private/VirtualizationHttpBackend.h
--- a/Engine/Source/Developer/Virtualization/Private/VirtualizationHttpBackend.h
+++ b/Engine/Source/Developer/Virtualization/Private/VirtualizationHttpBackend.h
@@ -36,11 +36,11 @@ private:
 
 	virtual bool Initialize(const FString& ConfigEntry) override;
 
-	virtual EPushResult PushData(const FPayloadId& Id, const FCompressedBuffer& CompressedPayload, const FString& PackageContext) override;
+	virtual EPushResult PushData(const FIoHash& Id, const FCompressedBuffer& CompressedPayload, const FString& PackageContext) override;
 
-	virtual FCompressedBuffer PullData(const FPayloadId& Id) override;
+	virtual FCompressedBuffer PullData(const FIoHash& Id) override;
 
-	virtual bool DoesPayloadExist(const FPayloadId& Id) override;
+	virtual bool DoesPayloadExist(const FIoHash& Id) override;
 
 private:
 
@@ -56,8 +56,8 @@ private:
 
 	bool ShouldRetryOnError(int64 ResponseCode);
 
-	bool PostChunk(const TArrayView<const uint8>& ChunkData, const FPayloadId& PayloadId, FString& OutHashAsString);
-	bool PullChunk(const FString& Hash, const FPayloadId& PayloadId, uint8* DataPtr, int64 BufferSize);
+	bool PostChunk(const TArrayView<const uint8>& ChunkData, const FIoHash& PayloadId, FString& OutHashAsString);
+	bool PullChunk(const FString& Hash, const FIoHash& PayloadId, uint8* DataPtr, int64 BufferSize);
 	bool DoesChunkExist(const FString& Hash);
 
 	/** Address of the service*/
diff --git a/Engine/Source/Developer/Virtualization/Private/VirtualizationManager.cpp b/Engine/Source/Developer/Virtualization/Private/VirtualizationManager.cpp
--- a/Engine/Source/Developer/Virtualization/Private/VirtualizationManager.cpp
+++ b/Engine/Source/Developer/Virtualization/Private/VirtualizationManager.cpp
@@ -312,7 +312,7 @@ bool FVirtualizationManager::IsPushingEnabled(EStorageType StorageType) const
 	}
 }
 
-bool FVirtualizationManager::PushData(const FPayloadId& Id, const FCompressedBuffer& Payload, EStorageType StorageType, const FString& Context)
+bool FVirtualizationManager::PushData(const FIoHash& Id, const FCompressedBuffer& Payload, EStorageType StorageType, const FString& Context)
 {
 	FPushRequest Request(Id, Payload, Context);
 	return FVirtualizationManager::PushData(MakeArrayView(&Request, 1), StorageType);
@@ -339,7 +339,7 @@ bool FVirtualizationManager::PushData(TArrayView<FPushRequest> Requests, EStorag
 		OriginalToValidatedRequest[Index] = INDEX_NONE;
 
 		FPushRequest& Request = Requests[Index];
-		if (!Request.Identifier.IsValid() || Request.Payload.GetCompressedSize() == 0)
+		if (Request.Identifier.IsZero() || Request.Payload.GetCompressedSize() == 0)
 		{
 			Request.Status = FPushRequest::EStatus::Invalid;
 			continue;
@@ -430,14 +430,14 @@ bool FVirtualizationManager::PushData(TArrayView<FPushRequest> Requests, EStorag
 	return StorageType == EStorageType::Local ? ErrorCount < Backends.Num() : ErrorCount == 0;
 }
 
-FCompressedBuffer FVirtualizationManager::PullData(const FPayloadId& Id)
+FCompressedBuffer FVirtualizationManager::PullData(const FIoHash& Id)
 {
 	TRACE_CPUPROFILER_EVENT_SCOPE(FVirtualizationManager::PullData);
 
-	if (!Id.IsValid())
+	if (Id.IsZero())
 	{
 		// TODO: See below, should errors here be fatal?
-		UE_LOG(LogVirtualization, Error, TEXT("Attempting to pull a virtualized payload with an invalid FPayloadId"));
+		UE_LOG(LogVirtualization, Error, TEXT("Attempting to pull a virtualized payload with an invalid FIoHash"));
 		return FCompressedBuffer();
 	}
 
@@ -484,7 +484,7 @@ FCompressedBuffer FVirtualizationManager::PullData(const FPayloadId& Id)
 	return FCompressedBuffer();
 }
 
-bool FVirtualizationManager::DoPayloadsExist(TArrayView<const FPayloadId> Ids, EStorageType StorageType, TArray<FPayloadStatus>& OutStatuses)
+bool FVirtualizationManager::DoPayloadsExist(TArrayView<const FIoHash> Ids, EStorageType StorageType, TArray<FPayloadStatus>& OutStatuses)
 {
 	TRACE_CPUPROFILER_EVENT_SCOPE(FVirtualizationManager::DoPayloadsExist);
 
@@ -497,7 +497,7 @@ bool FVirtualizationManager::DoPayloadsExist(TArrayView<const FPayloadId> Ids, E
 
 	for (int32 Index = 0; Index < Ids.Num(); ++Index)
 	{
-		OutStatuses[Index] = Ids[Index].IsValid() ? FPayloadStatus::NotFound : FPayloadStatus::Invalid;
+		OutStatuses[Index] = Ids[Index].IsZero() ? FPayloadStatus::Invalid : FPayloadStatus::NotFound;
 	}
 
 	FBackendArray& Backends = StorageType == EStorageType::Local ? LocalCachableBackends : PersistentStorageBackends;
@@ -522,7 +522,7 @@ bool FVirtualizationManager::DoPayloadsExist(TArrayView<const FPayloadId> Ids, E
 
 			for (int32 Index = 0; Index < Ids.Num(); ++Index)
 			{
-				if (Ids[Index].IsValid() && Results[Index])
+				if (!Ids[Index].IsZero() && Results[Index])
 				{
 					HitCount[Index]++;
 				}
@@ -533,7 +533,7 @@ bool FVirtualizationManager::DoPayloadsExist(TArrayView<const FPayloadId> Ids, E
 	// Now we total up the hit count for each payload to see if it was found in none, all or some of the backends
 	for (int32 Index = 0; Index < Ids.Num(); ++Index)
 	{
-		if (Ids[Index].IsValid())
+		if (!Ids[Index].IsZero())
 		{
 			if (HitCount[Index] == 0)
 			{
@@ -877,7 +877,7 @@ void FVirtualizationManager::AddBackend(TUniquePtr<IVirtualizationBackend> Backe
 	UE_LOG(LogVirtualization, Log, TEXT("Mounted backend: %s"), *BackendRef->GetDebugName());
 }
 
-void FVirtualizationManager::CachePayload(const FPayloadId& Id, const FCompressedBuffer& Payload, const IVirtualizationBackend* BackendSource)
+void FVirtualizationManager::CachePayload(const FIoHash& Id, const FCompressedBuffer& Payload, const IVirtualizationBackend* BackendSource)
 {
 	TRACE_CPUPROFILER_EVENT_SCOPE(FVirtualizationManager::CachePayload);
 
@@ -908,7 +908,7 @@ void FVirtualizationManager::CachePayload(const FPayloadId& Id, const FCompresse
 	}
 }
 
-bool FVirtualizationManager::TryCacheDataToBackend(IVirtualizationBackend& Backend, const FPayloadId& Id, const FCompressedBuffer& Payload)
+bool FVirtualizationManager::TryCacheDataToBackend(IVirtualizationBackend& Backend, const FIoHash& Id, const FCompressedBuffer& Payload)
 {
 	COOK_STAT(FCookStats::FScopedStatsCounter Timer(Profiling::GetCacheStats(Backend)));
 	const EPushResult Result = Backend.PushData(Id, Payload, FString());
@@ -951,7 +951,7 @@ bool FVirtualizationManager::TryPushDataToBackend(IVirtualizationBackend& Backen
 	return bPushResult;
 }
 
-FCompressedBuffer FVirtualizationManager::PullDataFromBackend(IVirtualizationBackend& Backend, const FPayloadId& Id)
+FCompressedBuffer FVirtualizationManager::PullDataFromBackend(IVirtualizationBackend& Backend, const FIoHash& Id)
 {
 	COOK_STAT(FCookStats::FScopedStatsCounter Timer(Profiling::GetPullStats(Backend)));
 	FCompressedBuffer Payload = Backend.PullData(Id);
diff --git a/Engine/Source/Developer/Virtualization/Private/VirtualizationManager.h b/Engine/Source/Developer/Virtualization/Private/VirtualizationManager.h
--- a/Engine/Source/Developer/Virtualization/Private/VirtualizationManager.h
+++ b/Engine/Source/Developer/Virtualization/Private/VirtualizationManager.h
@@ -6,7 +6,6 @@
 #include "HAL/CriticalSection.h"
 #include "Logging/LogMacros.h"
 #include "Templates/UniquePtr.h"
-#include "Virtualization/PayloadId.h"
 
 #include "Virtualization/VirtualizationSystem.h"
 
@@ -94,12 +93,12 @@ private:
 	virtual bool IsEnabled() const override;
 	virtual bool IsPushingEnabled(EStorageType StorageType) const override;
 	
-	virtual bool PushData(const FPayloadId& Id, const FCompressedBuffer& Payload, EStorageType StorageType, const FString& Context) override;
+	virtual bool PushData(const FIoHash& Id, const FCompressedBuffer& Payload, EStorageType StorageType, const FString& Context) override;
 	virtual bool PushData(TArrayView<FPushRequest> Requests, EStorageType StorageType) override;
 
-	virtual FCompressedBuffer PullData(const FPayloadId& Id) override;
+	virtual FCompressedBuffer PullData(const FIoHash& Id) override;
 
-	virtual bool DoPayloadsExist(TArrayView<const FPayloadId> Ids, EStorageType StorageType, TArray<FPayloadStatus>& OutStatuses) override;
+	virtual bool DoPayloadsExist(TArrayView<const FIoHash> Ids, EStorageType StorageType, TArray<FPayloadStatus>& OutStatuses) override;
 
 	virtual FPayloadActivityInfo GetAccumualtedPayloadActivityInfo() const override;
 
@@ -124,12 +123,12 @@ private:
 
 	void AddBackend(TUniquePtr<IVirtualizationBackend> Backend, FBackendArray& PushArray);
 
-	void CachePayload(const FPayloadId& Id, const FCompressedBuffer& Payload, const IVirtualizationBackend* BackendSource);
+	void CachePayload(const FIoHash& Id, const FCompressedBuffer& Payload, const IVirtualizationBackend* BackendSource);
 
-	bool TryCacheDataToBackend(IVirtualizationBackend& Backend, const FPayloadId& Id, const FCompressedBuffer& Payload);
+	bool TryCacheDataToBackend(IVirtualizationBackend& Backend, const FIoHash& Id, const FCompressedBuffer& Payload);
 	bool TryPushDataToBackend(IVirtualizationBackend& Backend, TArrayView<FPushRequest> Requests);
 
-	FCompressedBuffer PullDataFromBackend(IVirtualizationBackend& Backend, const FPayloadId& Id);
+	FCompressedBuffer PullDataFromBackend(IVirtualizationBackend& Backend, const FIoHash& Id);
 
 	/** 
 	 * Determines if a package should be virtualized or not based on it's package path and the current 
diff --git a/Engine/Source/Developer/Virtualization/Private/VirtualizationSourceControlBackend.cpp b/Engine/Source/Developer/Virtualization/Private/VirtualizationSourceControlBackend.cpp
--- a/Engine/Source/Developer/Virtualization/Private/VirtualizationSourceControlBackend.cpp
+++ b/Engine/Source/Developer/Virtualization/Private/VirtualizationSourceControlBackend.cpp
@@ -4,6 +4,7 @@
 
 #include "Containers/Ticker.h"
 #include "HAL/FileManager.h"
+#include "IO/IoHash.h"
 #include "ISourceControlModule.h"
 #include "ISourceControlProvider.h"
 #include "Logging/MessageLog.h"
@@ -12,7 +13,6 @@
 #include "Misc/Paths.h"
 #include "Misc/ScopeExit.h"
 #include "SourceControlOperations.h"
-#include "Virtualization/PayloadId.h"
 #include "VirtualizationSourceControlUtilities.h"
 #include "VirtualizationUtilities.h"
 
@@ -165,7 +165,7 @@ bool FSourceControlBackend::Initialize(const FString& ConfigEntry)
 	return true;
 }
 
-FCompressedBuffer FSourceControlBackend::PullData(const FPayloadId& Id)
+FCompressedBuffer FSourceControlBackend::PullData(const FIoHash& Id)
 {
 	TRACE_CPUPROFILER_EVENT_SCOPE(FSourceControlBackend::PullData);
 
@@ -194,11 +194,11 @@ FCompressedBuffer FSourceControlBackend::PullData(const FPayloadId& Id)
 	return FCompressedBuffer::FromCompressed(Buffer);
 }
 
-bool FSourceControlBackend::DoesPayloadExist(const FPayloadId& Id)
+bool FSourceControlBackend::DoesPayloadExist(const FIoHash& Id)
 {
 	TArray<bool> Result;
 
-	if (FSourceControlBackend::DoPayloadsExist(MakeArrayView<const FPayloadId>(&Id, 1), Result))
+	if (FSourceControlBackend::DoPayloadsExist(MakeArrayView<const FIoHash>(&Id, 1), Result))
 	{
 		check(Result.Num() == 1);
 		return Result[0];
@@ -209,7 +209,7 @@ bool FSourceControlBackend::DoesPayloadExist(const FPayloadId& Id)
 	}
 }
 
-EPushResult FSourceControlBackend::PushData(const FPayloadId& Id, const FCompressedBuffer& Payload, const FString& Context)
+EPushResult FSourceControlBackend::PushData(const FIoHash& Id, const FCompressedBuffer& Payload, const FString& Context)
 {
 	FPushRequest Request(Id, Payload, Context);
 	return FSourceControlBackend::PushData(MakeArrayView(&Request, 1)) ? EPushResult::Success : EPushResult::Failed;
@@ -444,7 +444,7 @@ bool FSourceControlBackend::PushData(TArrayView<FPushRequest> Requests)
 	return true;
 }
 
-bool FSourceControlBackend::DoPayloadsExist(TArrayView<const FPayloadId> PayloadIds, TArray<bool>& OutResults)
+bool FSourceControlBackend::DoPayloadsExist(TArrayView<const FIoHash> PayloadIds, TArray<bool>& OutResults)
 {
 	ISourceControlProvider& SCCProvider = ISourceControlModule::Get().GetProvider();
 
@@ -453,9 +453,9 @@ bool FSourceControlBackend::DoPayloadsExist(TArrayView<const FPayloadId> Payload
 
 	TArray<FSourceControlStateRef> PathStates;
 
-	for (const FPayloadId& PayloadId : PayloadIds)
+	for (const FIoHash& PayloadId : PayloadIds)
 	{
-		if (PayloadId.IsValid())
+		if (!PayloadId.IsZero())
 		{
 			TStringBuilder<52> LocalPayloadPath;
 			Utils::PayloadIdToPath(PayloadId, LocalPayloadPath);
@@ -478,7 +478,7 @@ bool FSourceControlBackend::DoPayloadsExist(TArrayView<const FPayloadId> Payload
 	int32 StatusIndex = 0;
 	for (int32 Index = 0; Index < PayloadIds.Num(); ++Index)
 	{
-		if (PayloadIds[Index].IsValid())
+		if (!PayloadIds[Index].IsZero())
 		{
 			OutResults[Index] = PathStates[StatusIndex++]->IsSourceControlled();
 		}
@@ -487,7 +487,7 @@ bool FSourceControlBackend::DoPayloadsExist(TArrayView<const FPayloadId> Payload
 	return true;
 }
 
-void FSourceControlBackend::CreateDepotPath(const FPayloadId& PayloadId, FStringBuilderBase& OutPath)
+void FSourceControlBackend::CreateDepotPath(const FIoHash& PayloadId, FStringBuilderBase& OutPath)
 {
 	TStringBuilder<52> PayloadPath;
 	Utils::PayloadIdToPath(PayloadId, PayloadPath);
diff --git a/Engine/Source/Developer/Virtualization/Private/VirtualizationSourceControlBackend.h b/Engine/Source/Developer/Virtualization/Private/VirtualizationSourceControlBackend.h
--- a/Engine/Source/Developer/Virtualization/Private/VirtualizationSourceControlBackend.h
+++ b/Engine/Source/Developer/Virtualization/Private/VirtualizationSourceControlBackend.h
@@ -38,18 +38,18 @@ private:
 
 	virtual bool Initialize(const FString& ConfigEntry) override;
 	
-	virtual EPushResult PushData(const FPayloadId& Id, const FCompressedBuffer& Payload, const FString& Context) override;
+	virtual EPushResult PushData(const FIoHash& Id, const FCompressedBuffer& Payload, const FString& Context) override;
 	virtual bool PushData(TArrayView<FPushRequest> Requests) override;
 
-	virtual FCompressedBuffer PullData(const FPayloadId& Id) override;
+	virtual FCompressedBuffer PullData(const FIoHash& Id) override;
 
-	virtual bool DoesPayloadExist(const FPayloadId& Id) override;
+	virtual bool DoesPayloadExist(const FIoHash& Id) override;
 	
-	virtual bool DoPayloadsExist(TArrayView<const FPayloadId> PayloadIds, TArray<bool>& OutResults) override;
+	virtual bool DoPayloadsExist(TArrayView<const FIoHash> PayloadIds, TArray<bool>& OutResults) override;
 
 private:
 
-	void CreateDepotPath(const FPayloadId& PayloadId, FStringBuilderBase& OutPath);
+	void CreateDepotPath(const FIoHash& PayloadId, FStringBuilderBase& OutPath);
 
 	/** Will display a FMessage notification to the user on the next valid engine tick to try and keep them aware of connection failures */
 	void OnConnectionError();
diff --git a/Engine/Source/Developer/Virtualization/Private/VirtualizationUtilities.cpp b/Engine/Source/Developer/Virtualization/Private/VirtualizationUtilities.cpp
--- a/Engine/Source/Developer/Virtualization/Private/VirtualizationUtilities.cpp
+++ b/Engine/Source/Developer/Virtualization/Private/VirtualizationUtilities.cpp
@@ -2,13 +2,13 @@
 
 #include "VirtualizationUtilities.h"
 
+#include "IO/IoHash.h"
 #include "Misc/StringBuilder.h"
-#include "Virtualization/PayloadId.h"
 
 namespace UE::Virtualization::Utils
 {
 
-void PayloadIdToPath(const FPayloadId& Id, FStringBuilderBase& OutPath)
+void PayloadIdToPath(const FIoHash& Id, FStringBuilderBase& OutPath)
 {
 	OutPath.Reset();
 	OutPath << Id;
@@ -23,7 +23,7 @@ void PayloadIdToPath(const FPayloadId& Id, FStringBuilderBase& OutPath)
 	OutPath << TEXT(".payload");
 }
 
-FString PayloadIdToPath(const FPayloadId& Id)
+FString PayloadIdToPath(const FIoHash& Id)
 {
 	TStringBuilder<52> Path;
 	PayloadIdToPath(Id, Path);
diff --git a/Engine/Source/Developer/Virtualization/Private/VirtualizationUtilities.h b/Engine/Source/Developer/Virtualization/Private/VirtualizationUtilities.h
--- a/Engine/Source/Developer/Virtualization/Private/VirtualizationUtilities.h
+++ b/Engine/Source/Developer/Virtualization/Private/VirtualizationUtilities.h
@@ -5,23 +5,20 @@
 #include "Containers/StringFwd.h"
 #include "Containers/UnrealString.h"
 
-namespace UE::Virtualization
-{
-	class FPayloadId;
-} // namespace UE::Virtualization
+struct FIoHash;
 
 namespace UE::Virtualization::Utils
 {
 
 /** 
- * Converts a FPayloadId into a file path.
+ * Converts a FIoHash into a file path.
  *
- * This utility will take a FPayloadId and return a file path that is
+ * This utility will take an FIoHash and return a file path that is
  * 3 directories deep. The first six characters of the id will be used to
  * create the directory names, with each directory using two characters.
  * The remaining thirty four characters will be used as the file name.
  * Lastly the extension '.payload' will be applied to complete the path/
- * Example: FPayloadId 0139d6d5d477e32dfd2abd3c5bc8ea8507e8eef8 becomes
+ * Example: FIoHash 0139d6d5d477e32dfd2abd3c5bc8ea8507e8eef8 becomes
  *			01/39/d6/d5d477e32dfd2abd3c5bc8ea8507e8eef8.payload'
  * 
  * @param	Id The payload identifier used to create the file path.
@@ -29,17 +26,17 @@ namespace UE::Virtualization::Utils
  *			The string builder must have a capacity of at least 52 characters 
 			to avoid reallocation.
  */
-void PayloadIdToPath(const FPayloadId& Id, FStringBuilderBase& OutPath);
+void PayloadIdToPath(const FIoHash& Id, FStringBuilderBase& OutPath);
 
 /** 
- * Converts a FPayloadId into a file path.
+ * Converts a FIoHash into a file path.
  * 
  * See above for further details
  * 
  * @param	Id The payload identifier used to create the file path.
  * @return	The resulting file path.
  */
-FString PayloadIdToPath(const FPayloadId& Id);
+FString PayloadIdToPath(const FIoHash& Id);
 
 /**
  * Fill in the given string builder with the human readable message of the current system
diff --git a/Engine/Source/Developer/Virtualization/Public/IVirtualizationBackend.h b/Engine/Source/Developer/Virtualization/Public/IVirtualizationBackend.h
--- a/Engine/Source/Developer/Virtualization/Public/IVirtualizationBackend.h
+++ b/Engine/Source/Developer/Virtualization/Public/IVirtualizationBackend.h
@@ -7,9 +7,10 @@
 #include "Features/IModularFeature.h"
 #include "Features/IModularFeatures.h"
 #include "Templates/UniquePtr.h"
-#include "Virtualization/PayloadId.h"
 #include "Virtualization/VirtualizationSystem.h"
 
+struct FIoHash;
+
 namespace UE::Virtualization
 {
 
@@ -84,7 +85,7 @@ public:
 	 * @param Payload	A potentially compressed buffer representing the payload
 	 * @return			The result of the push operation
 	 */
-	virtual EPushResult PushData(const FPayloadId& Id, const FCompressedBuffer& Payload, const FString& PackageContext) = 0;
+	virtual EPushResult PushData(const FIoHash& Id, const FCompressedBuffer& Payload, const FString& PackageContext) = 0;
 
 	virtual bool PushData(TArrayView<FPushRequest> Requests)
 	{
@@ -121,7 +122,7 @@ public:
 	 *				operation succeeded and a null FCompressedBuffer
 	 *				if it did not.
 	 */
-	virtual FCompressedBuffer PullData(const FPayloadId& Id) = 0;
+	virtual FCompressedBuffer PullData(const FIoHash& Id) = 0;
 	
 	/**
 	 * Checks if a payload exists in the backends storage.
@@ -130,21 +131,21 @@ public:
 	 * 
 	 * @return True if the backend storage already contains the payload, otherwise false
 	 */
-	virtual bool DoesPayloadExist(const FPayloadId& Id) = 0;
+	virtual bool DoesPayloadExist(const FIoHash& Id) = 0;
 	
 	/**
 	 * Checks if a number of payload exists in the backends storage.
 	 *
-	 * @param[in]	PayloadIds	An array of FPayloadId that should be checked
+	 * @param[in]	PayloadIds	An array of FIoHash that should be checked
 	 * @param[out]	OutResults	An array to contain the result, true if the payload
 	 *							exists in the backends storage, false if not.
 	 *							This array will be resized to match the size of PayloadIds.
 	 * 
 	 * @return True if the operation completed without error, otherwise false
 	 */
-	virtual bool DoPayloadsExist(TArrayView<const FPayloadId> PayloadIds, TArray<bool>& OutResults)
+	virtual bool DoPayloadsExist(TArrayView<const FIoHash> PayloadIds, TArray<bool>& OutResults)
 	{
-		// This is the default implementation that just calls ::DoesExist on each FPayloadId in the
+		// This is the default implementation that just calls ::DoesExist on each FIoHash in the
 		// array, one at a time. 
 		// Backends may override this with their own implementations if it can be done with less
 		// overhead by performing the check on the entire batch instead.
diff --git a/Engine/Source/Editor/DerivedDataEditor/Private/SVirtualAssetsStatistics.cpp b/Engine/Source/Editor/DerivedDataEditor/Private/SVirtualAssetsStatistics.cpp
--- a/Engine/Source/Editor/DerivedDataEditor/Private/SVirtualAssetsStatistics.cpp
+++ b/Engine/Source/Editor/DerivedDataEditor/Private/SVirtualAssetsStatistics.cpp
@@ -35,7 +35,7 @@ SVirtualAssetsStatisticsDialog::~SVirtualAssetsStatisticsDialog()
 	System.GetNotificationEvent().RemoveAll(this);
 }
 
-void SVirtualAssetsStatisticsDialog::OnNotificationEvent(IVirtualizationSystem::ENotification Notification, const FPayloadId& PayloadId)
+void SVirtualAssetsStatisticsDialog::OnNotificationEvent(IVirtualizationSystem::ENotification Notification, const FIoHash& PayloadId)
 {
 	FScopeLock SocpeLock(&NotificationCS);
 	
diff --git a/Engine/Source/Editor/DerivedDataEditor/Private/SVirtualAssetsStatistics.h b/Engine/Source/Editor/DerivedDataEditor/Private/SVirtualAssetsStatistics.h
--- a/Engine/Source/Editor/DerivedDataEditor/Private/SVirtualAssetsStatistics.h
+++ b/Engine/Source/Editor/DerivedDataEditor/Private/SVirtualAssetsStatistics.h
@@ -28,7 +28,7 @@ private:
 
 	FText GetNotificationText() const;
 
-	void OnNotificationEvent(UE::Virtualization::IVirtualizationSystem::ENotification Notification, const UE::Virtualization::FPayloadId& PayloadId);
+	void OnNotificationEvent(UE::Virtualization::IVirtualizationSystem::ENotification Notification, const FIoHash& PayloadId);
 
 	SVerticalBox::FSlot* GridSlot = nullptr;
 
diff --git a/Engine/Source/Editor/UnrealEd/Private/EditorDomain/BulkDataRegistryEditorDomain.cpp b/Engine/Source/Editor/UnrealEd/Private/EditorDomain/BulkDataRegistryEditorDomain.cpp
--- a/Engine/Source/Editor/UnrealEd/Private/EditorDomain/BulkDataRegistryEditorDomain.cpp
+++ b/Engine/Source/Editor/UnrealEd/Private/EditorDomain/BulkDataRegistryEditorDomain.cpp
@@ -219,7 +219,7 @@ TFuture<UE::BulkDataRegistry::FMetaData> FBulkDataRegistryEditorDomain::GetMeta(
 		if (!BulkData.HasPlaceholderPayloadId())
 		{
 			TPromise<UE::BulkDataRegistry::FMetaData> Promise;
-			Promise.SetValue(UE::BulkDataRegistry::FMetaData{ true, BulkData.GetPayloadId().GetIdentifier(), static_cast<uint64>(BulkData.GetPayloadSize()) });
+			Promise.SetValue(UE::BulkDataRegistry::FMetaData{ true, BulkData.GetPayloadId(), static_cast<uint64>(BulkData.GetPayloadSize()) });
 			return Promise.GetFuture();
 		}
 
diff --git a/Engine/Source/Runtime/Core/Private/Tests/Virtualization/PayloadIdTests.cpp b/Engine/Source/Runtime/Core/Private/Tests/Virtualization/PayloadIdTests.cpp
deleted file mode 100644
diff --git a/Engine/Source/Runtime/Core/Private/Virtualization/PayloadId.cpp b/Engine/Source/Runtime/Core/Private/Virtualization/PayloadId.cpp
deleted file mode 100644
diff --git a/Engine/Source/Runtime/Core/Private/Virtualization/VirtualizationSystem.cpp b/Engine/Source/Runtime/Core/Private/Virtualization/VirtualizationSystem.cpp
--- a/Engine/Source/Runtime/Core/Private/Virtualization/VirtualizationSystem.cpp
+++ b/Engine/Source/Runtime/Core/Private/Virtualization/VirtualizationSystem.cpp
@@ -29,7 +29,7 @@ public:
 		return false;
 	}
 
-	virtual bool PushData(const FPayloadId& Id, const FCompressedBuffer& Payload, EStorageType StorageType, const FString& Context) override
+	virtual bool PushData(const FIoHash& Id, const FCompressedBuffer& Payload, EStorageType StorageType, const FString& Context) override
 	{
 		return false;
 	}
@@ -39,12 +39,12 @@ public:
 		return false;
 	}
 
-	virtual FCompressedBuffer PullData(const FPayloadId& Id) override
+	virtual FCompressedBuffer PullData(const FIoHash& Id) override
 	{
 		return FCompressedBuffer();
 	}
 
-	virtual bool DoPayloadsExist(TArrayView<const FPayloadId> Ids, EStorageType StorageType, TArray<FPayloadStatus>& OutStatuses) override
+	virtual bool DoPayloadsExist(TArrayView<const FIoHash> Ids, EStorageType StorageType, TArray<FPayloadStatus>& OutStatuses) override
 	{
 		OutStatuses.Reset();
 
diff --git a/Engine/Source/Runtime/Core/Public/Virtualization/PayloadId.h b/Engine/Source/Runtime/Core/Public/Virtualization/PayloadId.h
deleted file mode 100644
diff --git a/Engine/Source/Runtime/Core/Public/Virtualization/VirtualizationSystem.h b/Engine/Source/Runtime/Core/Public/Virtualization/VirtualizationSystem.h
--- a/Engine/Source/Runtime/Core/Public/Virtualization/VirtualizationSystem.h
+++ b/Engine/Source/Runtime/Core/Public/Virtualization/VirtualizationSystem.h
@@ -4,9 +4,9 @@
 
 #include "Compression/CompressedBuffer.h"
 #include "Containers/StringFwd.h"
+#include "IO/IoHash.h"
 #include "Features/IModularFeature.h"
 #include "Features/IModularFeatures.h"
-#include "Virtualization/PayloadId.h"
 
 class FPackagePath;
 
@@ -70,7 +70,7 @@ struct FPushRequest
 		Success
 	};
 
-	FPushRequest(const FPayloadId& InIdentifier, const FCompressedBuffer& InPayload, const FString& InContext)
+	FPushRequest(const FIoHash& InIdentifier, const FCompressedBuffer& InPayload, const FString& InContext)
 		: Identifier(InIdentifier)
 		, Payload(InPayload)
 		, Context(InContext)
@@ -78,7 +78,7 @@ struct FPushRequest
 
 	}
 
-	FPushRequest(const FPayloadId& InIdentifier, FCompressedBuffer&& InPayload, FString&& InContext)
+	FPushRequest(const FIoHash& InIdentifier, FCompressedBuffer&& InPayload, FString&& InContext)
 		: Identifier(InIdentifier)
 		, Payload(InPayload)
 		, Context(InContext)
@@ -87,7 +87,7 @@ struct FPushRequest
 	}
 
 	/** The identifier of the payload */
-	FPayloadId Identifier;
+	FIoHash Identifier;
 	/** The payload data */
 	FCompressedBuffer Payload;
 	/** A string containing context for the payload, typically a package name */
@@ -144,7 +144,7 @@ public:
 	 * 
 	 * @return	True if at least one backend now contains the payload, otherwise false.
 	 */
-	virtual bool PushData(const FPayloadId& Id, const FCompressedBuffer& Payload, EStorageType StorageType, const FString& Context) = 0;
+	virtual bool PushData(const FIoHash& Id, const FCompressedBuffer& Payload, EStorageType StorageType, const FString& Context) = 0;
 	
 	/**
 	 * Push one or more payloads to a backend storage system. @See FPushRequest.
@@ -170,7 +170,7 @@ public:
 	 *			up to the caller if they want to retain the payload in compressed or uncompressed format.  If no
 	 *			backend contained the payload then an empty invalid FCompressedBuffer will be returned.
 	 */
-	virtual FCompressedBuffer PullData(const FPayloadId& Id) = 0;
+	virtual FCompressedBuffer PullData(const FIoHash& Id) = 0;
 
 	
 
@@ -185,7 +185,7 @@ public:
 	 * @return	True if the operation succeeded and the contents of OutStatuses is valid. False if errors were 
 	 * 			encountered in which case the contents of OutStatuses should be ignored.
 	 */
-	virtual bool DoPayloadsExist(TArrayView<const FPayloadId> Ids, EStorageType StorageType, TArray<FPayloadStatus>& OutStatuses) = 0;
+	virtual bool DoPayloadsExist(TArrayView<const FIoHash> Ids, EStorageType StorageType, TArray<FPayloadStatus>& OutStatuses) = 0;
 
 	using GetPayloadActivityInfoFuncRef = TFunctionRef<void(const FString& DebugName, const FString& ConfigName, const FPayloadActivityInfo& PayloadInfo)>;
 
@@ -208,7 +208,7 @@ public:
 	};
 
 	/** Declare delegate for notifications*/
-	DECLARE_MULTICAST_DELEGATE_TwoParams(FOnNotification, ENotification, const FPayloadId&);
+	DECLARE_MULTICAST_DELEGATE_TwoParams(FOnNotification, ENotification, const FIoHash&);
 
 	virtual FOnNotification& GetNotificationEvent() = 0;
 };
diff --git a/Engine/Source/Runtime/CoreUObject/Private/Serialization/EditorBulkData.cpp b/Engine/Source/Runtime/CoreUObject/Private/Serialization/EditorBulkData.cpp
--- a/Engine/Source/Runtime/CoreUObject/Private/Serialization/EditorBulkData.cpp
+++ b/Engine/Source/Runtime/CoreUObject/Private/Serialization/EditorBulkData.cpp
@@ -73,7 +73,7 @@ static TAutoConsoleVariable<bool> CVarShouldRehydrateOnSave(
 	TEXT("When true FVirtualizedUntypedBulkData virtualized payloads will by hydrated and stored locally when saved to a package"));
 
 /** Wrapper around the config file option [Core.System.Experimental]EnablePackageSidecarSaving */
-bool ShouldSaveToPackageSidecar()
+static bool ShouldSaveToPackageSidecar()
 {
 	static const struct FSaveToPackageSidecar
 	{
@@ -106,7 +106,7 @@ bool ShouldAllowVirtualizationOptOut()
 #endif // UE_ENABLE_VIRTUALIZATION_TOGGLE
 
 /** Utility for logging extended error messages when we fail to open a package for reading */
-void LogPackageOpenFailureMessage(const FPackagePath& PackagePath, EPackageSegment PackageSegment)
+static void LogPackageOpenFailureMessage(const FPackagePath& PackagePath, EPackageSegment PackageSegment)
 {
 	// TODO: Check the various error paths here again!
 	const uint32 SystemError = FPlatformMisc::GetLastError();
@@ -131,7 +131,7 @@ void LogPackageOpenFailureMessage(const FPackagePath& PackagePath, EPackageSegme
  * If the contents are validated we check the loaded result against the members of a bulkdata object to 
  * see if they match.
  */
-bool IsDataValid(const FEditorBulkData& BulkData, const FCompressedBuffer& Payload)
+static bool IsDataValid(const FEditorBulkData& BulkData, const FCompressedBuffer& Payload)
 {
 	if (!Payload.IsNull())
 	{
@@ -166,7 +166,7 @@ const FLinkerLoad* GetLinkerLoadFromOwner(UObject* Owner)
 }
 
 /** Utility for finding the FPackageTrailer associated with a given UObject */
-const FPackageTrailer* GetTrailerFromOwner(UObject* Owner)
+static const FPackageTrailer* GetTrailerFromOwner(UObject* Owner)
 {
 	const FLinkerLoad* Linker = GetLinkerLoadFromOwner(Owner);
 	if (Linker != nullptr)
@@ -180,7 +180,7 @@ const FPackageTrailer* GetTrailerFromOwner(UObject* Owner)
 }
 
 /** Utility for finding the package path associated with a given UObject */
-FPackagePath GetPackagePathFromOwner(UObject* Owner, EPackageSegment& OutPackageSegment)
+static FPackagePath GetPackagePathFromOwner(UObject* Owner, EPackageSegment& OutPackageSegment)
 {
 	OutPackageSegment = EPackageSegment::Header;
 
@@ -196,6 +196,53 @@ FPackagePath GetPackagePathFromOwner(UObject* Owner, EPackageSegment& OutPackage
 	}
 }
 
+/** Utility for hashing a payload, will return a default FIoHash if the payload is invalid or of zero length */
+static FIoHash HashPayload(const FSharedBuffer& InPayload)
+{
+	if (InPayload.GetSize() > 0)
+	{
+		return FIoHash::HashBuffer(InPayload);
+	}
+	else
+	{
+		return FIoHash();
+	}
+}
+
+/** Returns the FIoHash of a FGuid */
+static FIoHash GuidToIoHash(const FGuid& Guid)
+{
+	if (Guid.IsValid())
+	{
+		// Hash each element individually rather than making assumptions about
+		// the internal layout of FGuid and treating it as a contiguous buffer.
+		// Slightly slower, but safer.
+		FBlake3 Hash;
+
+		Hash.Update(&Guid[0], sizeof(uint32));
+		Hash.Update(&Guid[1], sizeof(uint32));
+		Hash.Update(&Guid[2], sizeof(uint32));
+		Hash.Update(&Guid[3], sizeof(uint32));
+
+		return FIoHash(Hash.Finalize());
+	}
+	else
+	{
+		return FIoHash();
+	}
+}
+
+FGuid IoHashToGuid(const FIoHash& Hash)
+{
+	// We use the first 16 bytes of the FIoHash to create the guid, there is
+	// no specific reason why these were chosen, we could take any pattern or combination
+	// of bytes.
+	// Note that if the input hash is invalid (all zeros) then the FGuid returned will
+	// also be considered as invalid
+	uint32* HashBytes = (uint32*)Hash.GetBytes();
+	return FGuid(HashBytes[0], HashBytes[1], HashBytes[2], HashBytes[3]);
+}
+
 /** Utility for updating an existing entry in an Archive before returning the archive to it's original seek position */
 template<typename DataType>
 void UpdateArchiveData(FArchive& Ar, int64 DataPosition, DataType& Data)
@@ -476,7 +523,7 @@ void FEditorBulkData::CreateFromBulkData(FUntypedBulkData& InBulkData, const FGu
 	if (InBulkData.GetBulkDataSize() > 0)
 	{
 		BulkDataId = CreateUniqueGuid(InGuid, Owner, *InBulkData.GetPackagePath().GetDebugName());
-		PayloadContentId = UE::Virtualization::FPayloadId(BulkDataId);
+		PayloadContentId = GuidToIoHash(BulkDataId);
 		bWasKeyGuidDerived = true;
 	}
 	
@@ -617,7 +664,7 @@ void FEditorBulkData::Serialize(FArchive& Ar, UObject* Owner, bool bAllowRegiste
 
 		// TODO: Can probably remove these checks before UE5 release
 		check(!Ar.IsSaving() || GetPayloadSize() == 0 || BulkDataId.IsValid()); // Sanity check to stop us saving out bad data
-		check(!Ar.IsSaving() || GetPayloadSize() == 0 || PayloadContentId.IsValid()); // Sanity check to stop us saving out bad data
+		check(!Ar.IsSaving() || GetPayloadSize() == 0 || !PayloadContentId.IsZero()); // Sanity check to stop us saving out bad data
 		
 		Ar << BulkDataId;
 		Ar << PayloadContentId;
@@ -625,7 +672,7 @@ void FEditorBulkData::Serialize(FArchive& Ar, UObject* Owner, bool bAllowRegiste
 
 		// TODO: Can probably remove these checks before UE5 release
 		check(!Ar.IsLoading() || GetPayloadSize() == 0 || BulkDataId.IsValid()); // Sanity check to stop us loading in bad data
-		check(!Ar.IsLoading() || GetPayloadSize() == 0 || PayloadContentId.IsValid()); // Sanity check to stop us loading in bad data
+		check(!Ar.IsLoading() || GetPayloadSize() == 0 || !PayloadContentId.IsZero()); // Sanity check to stop us loading in bad data
 
 		if (Ar.IsSaving())
 		{
@@ -711,7 +758,7 @@ void FEditorBulkData::Serialize(FArchive& Ar, UObject* Owner, bool bAllowRegiste
 			}
 
 			// Make sure that the trailer builder is correct (if it is being used)
-			if (IsStoredInPackageTrailer(UpdatedFlags) && PayloadContentId.IsValid())
+			if (IsStoredInPackageTrailer(UpdatedFlags) && !PayloadContentId.IsZero())
 			{
 				check(LinkerSave != nullptr);
 				check(LinkerSave->PackageTrailerBuilder.IsValid());
@@ -900,8 +947,8 @@ FCompressedBuffer FEditorBulkData::LoadFromDisk() const
 			FCompressedBuffer SidecarBuffer = LoadFromSidecarFile();
 			FCompressedBuffer AssetBuffer = LoadFromPackageFile();
 
-			UE::Virtualization::FPayloadId SidecarId(SidecarBuffer.Decompress());
-			UE::Virtualization::FPayloadId AssetId(AssetBuffer.Decompress());
+			const FIoHash SidecarId = HashPayload(SidecarBuffer.Decompress());
+			const FIoHash AssetId = HashPayload(AssetBuffer.Decompress());
 
 			UE_CLOG(SidecarId != PayloadContentId, LogSerialization, Error, TEXT("Sidecar content did not hash correctly! Found '%s' Expected '%s'"), *LexToString(SidecarId), *LexToString(PayloadContentId));
 			UE_CLOG(AssetId != PayloadContentId, LogSerialization, Error, TEXT("Asset content did not hash correctly! Found '%s' Expected '%s'"), *LexToString(AssetId), *LexToString(PayloadContentId))
@@ -1174,7 +1221,7 @@ void FEditorBulkData::PushData(const FPackagePath& InPackagePath)
 		// serialization, we need a better way to save the results of 'RecompressForSerialization'
 		RecompressForSerialization(PayloadToPush, Flags);
 
-		// TODO: We could make this a config option?
+		// TODO: We could make the storage type a config option?
 		if (VirtualizationSystem.PushData(PayloadContentId, PayloadToPush, UE::Virtualization::EStorageType::Local, InPackagePath.GetPackageName()))
 		{
 			EnumAddFlags(Flags, EFlags::IsVirtualized);
@@ -1399,7 +1446,7 @@ void FEditorBulkData::SerializeToPackageTrailer(FLinkerSave& LinkerSave, FCompre
 	LinkerSave.PackageTrailerBuilder->AddPayload(PayloadContentId, MoveTemp(PayloadToSerialize), MoveTemp(OnPayloadWritten));
 }
 
-void FEditorBulkData::UpdatePayloadImpl(FSharedBuffer&& InPayload, UE::Virtualization::FPayloadId&& InPayloadID)
+void FEditorBulkData::UpdatePayloadImpl(FSharedBuffer&& InPayload, FIoHash&& InPayloadID)
 {
 	TRACE_CPUPROFILER_EVENT_SCOPE(FEditorBulkData::UpdatePayloadImpl);
 
@@ -1531,7 +1578,7 @@ TFuture<FCompressedBuffer>FEditorBulkData::GetCompressedPayload() const
 void FEditorBulkData::UpdatePayload(FSharedBuffer InPayload)
 {
 	TRACE_CPUPROFILER_EVENT_SCOPE(FEditorBulkData::UpdatePayload);
-	UE::Virtualization::FPayloadId NewPayloadId(InPayload);
+	FIoHash NewPayloadId = HashPayload(InPayload);
 	UpdatePayloadImpl(MoveTemp(InPayload), MoveTemp(NewPayloadId));
 }
 
@@ -1616,7 +1663,7 @@ void FEditorBulkData::UpdateKeyIfNeeded()
 
 		// Load the payload from disk (or memory) so that we can hash it
 		FSharedBuffer InPayload = GetDataInternal().Decompress();
-		PayloadContentId = UE::Virtualization::FPayloadId(InPayload);
+		PayloadContentId = HashPayload(InPayload);
 
 		// Store as the in memory payload, since this method is only called during saving 
 		// we know it will get cleared anyway.
@@ -1826,13 +1873,13 @@ FArchive& operator<<(FArchive& Ar, FTocEntry& Entry)
 
 void FPayloadToc::AddEntry(const FEditorBulkData& BulkData)
 {
-	if (BulkData.GetPayloadId().IsValid())
+	if (!BulkData.GetPayloadId().IsZero())
 	{
 		Contents.Emplace(BulkData);
 	}
 }
 
-bool FPayloadToc::FindEntry(const UE::Virtualization::FPayloadId& Identifier, FTocEntry& OutEntry)
+bool FPayloadToc::FindEntry(const FIoHash& Identifier, FTocEntry& OutEntry)
 {
 	for (const FTocEntry& Entry : Contents)
 	{
diff --git a/Engine/Source/Runtime/CoreUObject/Private/Tests/Serialization/EditorBulkDataTests.cpp b/Engine/Source/Runtime/CoreUObject/Private/Tests/Serialization/EditorBulkDataTests.cpp
--- a/Engine/Source/Runtime/CoreUObject/Private/Tests/Serialization/EditorBulkDataTests.cpp
+++ b/Engine/Source/Runtime/CoreUObject/Private/Tests/Serialization/EditorBulkDataTests.cpp
@@ -87,19 +87,31 @@ bool FEditorBulkDataTestBasic::RunTest(const FString& Parameters)
 IMPLEMENT_SIMPLE_AUTOMATION_TEST(FEditorBulkDataTestEmpty, TEXT("System.CoreUObject.Serialization.EditorBulkData.Empty"), TestFlags)
 bool FEditorBulkDataTestEmpty::RunTest(const FString& Parameters)
 {
-	FEditorBulkData BulkData;
+	auto Validate = [](const TCHAR* Id, FAutomationTestBase& Test, const FEditorBulkData& BulkData)
+	{
+		// Validate the general accessors
+		Test.TestEqual(FString::Printf(TEXT("(%s) Return value of ::GetBulkDataSize()"), Id), BulkData.GetPayloadSize(), (int64)0);
+		Test.TestTrue(FString::Printf(TEXT("(%s) Payload key is invalid"), Id), BulkData.GetPayloadId().IsZero());
+		Test.TestFalse(FString::Printf(TEXT("(%s) Return value of ::DoesPayloadNeedLoading()"), Id), BulkData.DoesPayloadNeedLoading());
+
+		// Validate the payload accessors
+		FSharedBuffer Payload = BulkData.GetPayload().Get();
+		Test.TestTrue(FString::Printf(TEXT("(%s) The payload from the GetPayload TFuture is null"), Id), Payload.IsNull());
 
-	// Validate the general accessors
-	TestEqual(TEXT("Return value of ::GetBulkDataSize()"), BulkData.GetPayloadSize(), (int64)0);
-	TestTrue(TEXT("Payload key is invalid"), !BulkData.GetPayloadId().IsValid());
-	TestFalse(TEXT("Return value of ::DoesPayloadNeedLoading()"), BulkData.DoesPayloadNeedLoading());
+		FCompressedBuffer CompressedPayload = BulkData.GetCompressedPayload().Get();
+		Test.TestTrue(FString::Printf(TEXT("(%s) The payload from the GetCompressedPayload TFuture is null"), Id), Payload.IsNull());
+	};
 
-	// Validate the payload accessors
-	FSharedBuffer Payload = BulkData.GetPayload().Get();
-	TestTrue(TEXT("The payload from the GetPayload TFuture is null"), Payload.IsNull());
+	FEditorBulkData DefaultBulkData;
+	Validate(TEXT("DefaultBulkData"), *this, DefaultBulkData);
+		
+	FEditorBulkData NullPayloadBulkData;
+	NullPayloadBulkData.UpdatePayload(FSharedBuffer());
+	Validate(TEXT("NullPayloadBulkData"), *this, NullPayloadBulkData);
 
-	FCompressedBuffer CompressedPayload = BulkData.GetCompressedPayload().Get();
-	TestTrue(TEXT("The payload from the GetCompressedPayload TFuture is null"), Payload.IsNull());
+	FEditorBulkData ZeroLengthPayloadBulkData;
+	ZeroLengthPayloadBulkData.UpdatePayload(FUniqueBuffer::Alloc(0).MoveToShared());
+	Validate(TEXT("ZeroLengthPayloadBulkData"), *this, ZeroLengthPayloadBulkData);
 
 	return true;
 }
@@ -508,6 +520,31 @@ bool FEditorBulkDataTestZeroSizedAllocs::RunTest(const FString& Parameters)
 	return true;
 }
 
+/** This tests the function UE::Serialization::IoHashToGuid which is closely used with the FEditorBulkData system */
+IMPLEMENT_SIMPLE_AUTOMATION_TEST(FEditorBulkDataIoHashToGuid, TEXT("System.CoreUObject.Serialization.EditorBulkData.IoHashToGuid"), TestFlags)
+bool FEditorBulkDataIoHashToGuid::RunTest(const FString& Parameters)
+{
+	// Test that an empty hash will give an invalid FGuid
+	FIoHash DefaultHash;
+	FGuid InvalidGuid = IoHashToGuid(DefaultHash);
+
+	TestTrue(TEXT("Calling::IoHashToGuid on a all zero FIoHash should result in an invalid FGuid"), !InvalidGuid.IsValid());
+	
+	// Test that finding the FGuid of a known hash results in the FGuid we expect. If not then the generation algorithm has
+	// changed. The failing test should remind whoever changed the algorithm to double check that changing the results 
+	// will not have knock on effects.
+	const uint8 KnownHashData[20] = {1, 2, 3, 4, 5, 6, 7, 8, 9, 10, 11, 12, 13, 14, 15, 16, 17, 18 , 19};
+
+	FIoHash KnownHash(KnownHashData);
+	FGuid KnownGuid = IoHashToGuid(KnownHash);
+	FGuid KnownResult(TEXT("04030201-0807-0605-0C0B-0A09100F0E0D"));
+
+	TestEqual(TEXT("That the result of hashing known data will be unchanged"), KnownGuid, KnownResult);
+
+	return true;
+}
+
+
 } // namespace UE::Serialization
 
 #endif //WITH_DEV_AUTOMATION_TESTS && WITH_EDITORONLY_DATA
diff --git a/Engine/Source/Runtime/CoreUObject/Private/UObject/PackageTrailer.cpp b/Engine/Source/Runtime/CoreUObject/Private/UObject/PackageTrailer.cpp
--- a/Engine/Source/Runtime/CoreUObject/Private/UObject/PackageTrailer.cpp
+++ b/Engine/Source/Runtime/CoreUObject/Private/UObject/PackageTrailer.cpp
@@ -57,7 +57,7 @@ void LogPackageOpenFailureMessage(const FPackagePath& PackagePath)
 namespace Private
 {
 
-FLookupTableEntry::FLookupTableEntry(const Virtualization::FPayloadId& InIdentifier, uint64 InRawSize)
+FLookupTableEntry::FLookupTableEntry(const FIoHash& InIdentifier, uint64 InRawSize)
 	: Identifier(InIdentifier)
 	, RawSize(InRawSize)
 {
@@ -90,7 +90,7 @@ FPackageTrailerBuilder FPackageTrailerBuilder::CreateFromTrailer(const FPackageT
 
 	for (const Private::FLookupTableEntry& Entry : Trailer.Header.PayloadLookupTable)
 	{
-		checkf(Entry.Identifier.IsValid(), TEXT("PackageTrailer for package should not contain invalid FPayloadIds. Package '%s'"), *PackageName.ToString());
+		checkf(!Entry.Identifier.IsZero(), TEXT("PackageTrailer for package should not contain invalid FIoHash entry. Package '%s'"), *PackageName.ToString());
 
 		switch (Entry.AccessMode)
 		{
@@ -129,7 +129,7 @@ TUniquePtr<UE::FPackageTrailerBuilder> FPackageTrailerBuilder::CreateReferenceTo
 
 	for (const Private::FLookupTableEntry& Entry : Trailer.Header.PayloadLookupTable)
 	{
-		checkf(Entry.Identifier.IsValid(), TEXT("PackageTrailer for package should not contain invalid FPayloadIds. Package '%s'"), *PackageName.ToString());
+		checkf(!Entry.Identifier.IsZero(), TEXT("PackageTrailer for package should not contain invalid FIoHash entry. Package '%s'"), *PackageName.ToString());
 
 		switch (Entry.AccessMode)
 		{
@@ -170,19 +170,19 @@ FPackageTrailerBuilder::FPackageTrailerBuilder(const FName& InPackageName)
 {
 }
 
-void FPackageTrailerBuilder::AddPayload(const Virtualization::FPayloadId& Identifier, FCompressedBuffer Payload, AdditionalDataCallback&& Callback)
+void FPackageTrailerBuilder::AddPayload(const FIoHash& Identifier, FCompressedBuffer Payload, AdditionalDataCallback&& Callback)
 {
 	Callbacks.Emplace(MoveTemp(Callback));
 
-	if (Identifier.IsValid())
+	if (!Identifier.IsZero())
 	{
 		LocalEntries.FindOrAdd(Identifier, LocalEntry(MoveTemp(Payload)));
 	}
 }
 
-void FPackageTrailerBuilder::AddVirtualizedPayload(const Virtualization::FPayloadId& Identifier, int64 RawSize)
+void FPackageTrailerBuilder::AddVirtualizedPayload(const FIoHash& Identifier, int64 RawSize)
 {
-	if (Identifier.IsValid())
+	if (!Identifier.IsZero())
 	{
 		VirtualizedEntries.FindOrAdd(Identifier, VirtualizedEntry(RawSize));
 	}
@@ -208,9 +208,9 @@ bool FPackageTrailerBuilder::BuildAndAppendTrailer(FLinkerSave* Linker, FArchive
 	Trailer.Header.PayloadsDataLength = 0;
 	Trailer.Header.PayloadLookupTable.Reserve(LocalEntries.Num() + VirtualizedEntries.Num());
 
-	for (const TPair<Virtualization::FPayloadId, LocalEntry>& It : LocalEntries)
+	for (const TPair<FIoHash, LocalEntry>& It : LocalEntries)
 	{
-		checkf(It.Key.IsValid(), TEXT("PackageTrailer should not contain invalid FPayloadIds. Package '%s'"), *PackageName.ToString());
+		checkf(!It.Key.IsZero(), TEXT("PackageTrailer should not contain invalid FIoHash values. Package '%s'"), *PackageName.ToString());
 
 		Private::FLookupTableEntry& Entry = Trailer.Header.PayloadLookupTable.AddDefaulted_GetRef();
 		Entry.Identifier = It.Key;
@@ -222,9 +222,9 @@ bool FPackageTrailerBuilder::BuildAndAppendTrailer(FLinkerSave* Linker, FArchive
 		Trailer.Header.PayloadsDataLength += It.Value.Payload.GetCompressedSize();
 	}
 
-	for (const TPair<Virtualization::FPayloadId, ReferencedEntry>& It : ReferencedEntries)
+	for (const TPair<FIoHash, ReferencedEntry>& It : ReferencedEntries)
 	{
-		checkf(It.Key.IsValid(), TEXT("PackageTrailer should not contain invalid FPayloadIds. Package '%s'"), *PackageName.ToString());
+		checkf(!It.Key.IsZero(), TEXT("PackageTrailer should not contain invalid FIoHash values. Package '%s'"), *PackageName.ToString());
 
 		Private::FLookupTableEntry& Entry = Trailer.Header.PayloadLookupTable.AddDefaulted_GetRef();
 		Entry.Identifier = It.Key;
@@ -234,9 +234,9 @@ bool FPackageTrailerBuilder::BuildAndAppendTrailer(FLinkerSave* Linker, FArchive
 		Entry.AccessMode = EPayloadAccessMode::Referenced;
 	}
 
-	for (const TPair<Virtualization::FPayloadId, VirtualizedEntry>& It : VirtualizedEntries)
+	for (const TPair<FIoHash, VirtualizedEntry>& It : VirtualizedEntries)
 	{
-		checkf(It.Key.IsValid(), TEXT("PackageTrailer should not contain invalid FPayloadIds. Package '%s'"), *PackageName.ToString());
+		checkf(!It.Key.IsZero(), TEXT("PackageTrailer should not contain invalid FIoHash values. Package '%s'"), *PackageName.ToString());
 
 		Private::FLookupTableEntry& Entry = Trailer.Header.PayloadLookupTable.AddDefaulted_GetRef();
 		Entry.Identifier = It.Key;
@@ -256,7 +256,7 @@ bool FPackageTrailerBuilder::BuildAndAppendTrailer(FLinkerSave* Linker, FArchive
 
 	const int64 PayloadPosInFile = DataArchive.Tell();
 
-	for (TPair<Virtualization::FPayloadId, LocalEntry>& It : LocalEntries)
+	for (TPair<FIoHash, LocalEntry>& It : LocalEntries)
 	{
 		DataArchive << It.Value.Payload;
 	}
@@ -283,17 +283,17 @@ bool FPackageTrailerBuilder::IsEmpty() const
 	return LocalEntries.IsEmpty() && ReferencedEntries.IsEmpty() && VirtualizedEntries.IsEmpty();
 }
 
-bool FPackageTrailerBuilder::IsLocalPayloadEntry(const Virtualization::FPayloadId& Identifier) const
+bool FPackageTrailerBuilder::IsLocalPayloadEntry(const FIoHash& Identifier) const
 {
 	return LocalEntries.Find(Identifier) != nullptr;
 }
 
-bool FPackageTrailerBuilder::IsReferencedPayloadEntry(const Virtualization::FPayloadId& Identifier) const
+bool FPackageTrailerBuilder::IsReferencedPayloadEntry(const FIoHash& Identifier) const
 {
 	return ReferencedEntries.Find(Identifier) != nullptr;
 }
 
-bool FPackageTrailerBuilder::IsVirtualizedPayloadEntry(const Virtualization::FPayloadId& Identifier) const
+bool FPackageTrailerBuilder::IsVirtualizedPayloadEntry(const FIoHash& Identifier) const
 {
 	return VirtualizedEntries.Find(Identifier) != nullptr;
 }
@@ -424,7 +424,7 @@ bool FPackageTrailer::TryLoadBackwards(FArchive& Ar)
 	return TryLoad(Ar);
 }
 
-FCompressedBuffer FPackageTrailer::LoadLocalPayload(const Virtualization::FPayloadId& Id, FArchive& Ar) const
+FCompressedBuffer FPackageTrailer::LoadLocalPayload(const FIoHash& Id, FArchive& Ar) const
 {
 	// TODO: This method should be able to load the payload in all cases, but we need a good way of passing the Archive/PackagePath 
 	// to the trailer etc. Would work if we stored the package path in the trailer.
@@ -447,7 +447,7 @@ FCompressedBuffer FPackageTrailer::LoadLocalPayload(const Virtualization::FPaylo
 	return Payload;
 }
 
-bool FPackageTrailer::UpdatePayloadAsVirtualized(const Virtualization::FPayloadId& Identifier)
+bool FPackageTrailer::UpdatePayloadAsVirtualized(const FIoHash& Identifier)
 {
 	Private::FLookupTableEntry* Entry = Header.PayloadLookupTable.FindByPredicate([&Identifier](const Private::FLookupTableEntry& Entry)->bool
 		{
@@ -468,7 +468,7 @@ bool FPackageTrailer::UpdatePayloadAsVirtualized(const Virtualization::FPayloadI
 	}
 }
 
-EPayloadStatus FPackageTrailer::FindPayloadStatus(const Virtualization::FPayloadId& Id) const
+EPayloadStatus FPackageTrailer::FindPayloadStatus(const FIoHash& Id) const
 {
 	const Private::FLookupTableEntry* Entry = Header.PayloadLookupTable.FindByPredicate([&Id](const Private::FLookupTableEntry& Entry)->bool
 		{
@@ -501,9 +501,9 @@ EPayloadStatus FPackageTrailer::FindPayloadStatus(const Virtualization::FPayload
 	}
 }
 
-int64 FPackageTrailer::FindPayloadOffsetInFile(const Virtualization::FPayloadId& Id) const
+int64 FPackageTrailer::FindPayloadOffsetInFile(const FIoHash& Id) const
 {
-	if (Id.IsValid())
+	if (!Id.IsZero())
 	{
 		const Private::FLookupTableEntry* Entry = Header.PayloadLookupTable.FindByPredicate([&Id](const Private::FLookupTableEntry& Entry)->bool
 			{
@@ -546,9 +546,9 @@ int64 FPackageTrailer::GetTrailerLength() const
 	return Header.HeaderLength + Header.PayloadsDataLength + FFooter::SizeOnDisk;
 }
 
-TArray<Virtualization::FPayloadId> FPackageTrailer::GetPayloads(EPayloadFilter Type) const
+TArray<FIoHash> FPackageTrailer::GetPayloads(EPayloadFilter Type) const
 {
-	TArray<Virtualization::FPayloadId> Identifiers;
+	TArray<FIoHash> Identifiers;
 	Identifiers.Reserve(Header.PayloadLookupTable.Num());
 
 	for (const Private::FLookupTableEntry& Entry : Header.PayloadLookupTable)
@@ -651,7 +651,7 @@ FArchive& operator<<(FArchive& Ar, FPackageTrailer::FFooter& Footer)
 	return Ar;
 }
 
-bool FindPayloadsInPackageFile(const FPackagePath& PackagePath, EPayloadFilter Filter, TArray<Virtualization::FPayloadId>& OutPayloadIds)
+bool FindPayloadsInPackageFile(const FPackagePath& PackagePath, EPayloadFilter Filter, TArray<FIoHash>& OutPayloadIds)
 {
 	if (FPackageName::IsTextPackageExtension(PackagePath.GetHeaderExtension()))
 	{
diff --git a/Engine/Source/Runtime/CoreUObject/Public/Serialization/EditorBulkData.h b/Engine/Source/Runtime/CoreUObject/Public/Serialization/EditorBulkData.h
--- a/Engine/Source/Runtime/CoreUObject/Public/Serialization/EditorBulkData.h
+++ b/Engine/Source/Runtime/CoreUObject/Public/Serialization/EditorBulkData.h
@@ -5,11 +5,11 @@
 #include "Async/Future.h"
 #include "Compression/CompressedBuffer.h"
 #include "HAL/Platform.h"
+#include "IO/IoHash.h"
 #include "Memory/SharedBuffer.h"
 #include "Misc/Guid.h"
 #include "Misc/PackagePath.h"
 #include "Misc/PackageSegment.h"
-#include "Virtualization/PayloadId.h"
 
 class FArchive;
 class UObject;
@@ -153,7 +153,7 @@ public:
 	FGuid GetIdentifier() const;
 
 	/** Returns an unique identifier for the content of the payload. */
-	const UE::Virtualization::FPayloadId& GetPayloadId() const
+	const FIoHash& GetPayloadId() const
 	{ 
 		return PayloadContentId; 
 	}
@@ -221,7 +221,7 @@ public:
 		friend FEditorBulkData;
 
 		FSharedBuffer Payload;
-		UE::Virtualization::FPayloadId PayloadId;
+		FIoHash PayloadId;
 	};
 
 	/**
@@ -364,7 +364,7 @@ private:
 	/** The new path that saves payloads to the FPackageTrailer which is then appended to the end of the package file */
 	void SerializeToPackageTrailer(FLinkerSave& LinkerSave, FCompressedBuffer PayloadToSerialize, EFlags UpdatedFlags, UObject* Owner);
 
-	void UpdatePayloadImpl(FSharedBuffer&& InPayload, UE::Virtualization::FPayloadId&& InPayloadID);
+	void UpdatePayloadImpl(FSharedBuffer&& InPayload, FIoHash&& InPayloadID);
 
 	FCompressedBuffer GetDataInternal() const;
 
@@ -462,7 +462,7 @@ private:
 	FGuid BulkDataId;
 
 	/** Unique identifier for the contents of the payload*/
-	UE::Virtualization::FPayloadId PayloadContentId;
+	FIoHash PayloadContentId;
 
 	/** Pointer to the payload if it is held in memory (it has been updated but not yet saved to disk for example) */
 	FSharedBuffer Payload;
@@ -528,7 +528,7 @@ struct COREUOBJECT_API FTocEntry
 	static constexpr uint32 PayloadSidecarFileVersion = 1;
 
 	/** Identifier for the payload */
-	UE::Virtualization::FPayloadId Identifier;
+	FIoHash Identifier;
 	/** The offset into the file where we can find the payload, note that a virtualized payload will have an offset of INDEX_NONE */
 	int64 OffsetInFile = INDEX_NONE;
 	/** The size of the payload when uncompressed. */
@@ -544,7 +544,7 @@ class COREUOBJECT_API FPayloadToc
 public:
 
 	void AddEntry(const FEditorBulkData& BulkData);
-	bool FindEntry(const UE::Virtualization::FPayloadId& Identifier, FTocEntry& OutEntry);
+	bool FindEntry(const FIoHash& Identifier, FTocEntry& OutEntry);
 
 	const TArray<FTocEntry>& GetContents() const;
 
@@ -565,6 +565,12 @@ private:
 	TArray<FTocEntry> Contents;
 };
 
+/** 
+ * Returns a FGuid representation of a FIoHash, to be used in existing code paths that require the id of 
+ * an FEditorBulkData payload to be in FGuid form 
+ */
+COREUOBJECT_API FGuid IoHashToGuid(const FIoHash& Hash);
+
 } // namespace UE::Serialization
 
 //#endif //WITH_EDITORONLY_DATA
diff --git a/Engine/Source/Runtime/CoreUObject/Public/UObject/LinkerSave.h b/Engine/Source/Runtime/CoreUObject/Public/UObject/LinkerSave.h
--- a/Engine/Source/Runtime/CoreUObject/Public/UObject/LinkerSave.h
+++ b/Engine/Source/Runtime/CoreUObject/Public/UObject/LinkerSave.h
@@ -11,7 +11,6 @@
 #include "UObject/ObjectResource.h"
 #include "UObject/PackageTrailer.h"
 #include "UObject/UObjectThreadContext.h"
-#include "Virtualization/PayloadId.h"
 
 class FObjectPostSaveContext;
 struct FUntypedBulkData;
@@ -107,7 +106,7 @@ public:
 	
 	struct FSidecarStorageInfo
 	{
-		UE::Virtualization::FPayloadId Identifier;
+		FIoHash Identifier;
 		FCompressedBuffer Payload;
 	};
 
diff --git a/Engine/Source/Runtime/CoreUObject/Public/UObject/PackageTrailer.h b/Engine/Source/Runtime/CoreUObject/Public/UObject/PackageTrailer.h
--- a/Engine/Source/Runtime/CoreUObject/Public/UObject/PackageTrailer.h
+++ b/Engine/Source/Runtime/CoreUObject/Public/UObject/PackageTrailer.h
@@ -4,8 +4,8 @@
 
 #include "Compression/CompressedBuffer.h"
 #include "Containers/Map.h"
+#include "IO/IoHash.h"
 #include "UObject/NameTypes.h"
-#include "Virtualization/PayloadId.h"
 
 class FArchive;
 class FLinkerSave;
@@ -122,7 +122,7 @@ struct FLookupTableEntry
 												// AccessMode		| 1 byte
 
 	FLookupTableEntry() = default;
-	FLookupTableEntry(const Virtualization::FPayloadId& InIdentifier, uint64 InRawSize);
+	FLookupTableEntry(const FIoHash& InIdentifier, uint64 InRawSize);
 
 	void Serialize(FArchive& Ar, EPackageTrailerVersion PackageTrailerVersion);
 	
@@ -132,7 +132,7 @@ struct FLookupTableEntry
 	}
 
 	/** Identifier for the payload */
-	Virtualization::FPayloadId Identifier;
+	FIoHash Identifier;
 	/** The offset into the file where we can find the payload, note that a virtualized payload will have an offset of INDEX_NONE */
 	int64 OffsetInFile = INDEX_NONE;
 	/** The size of the payload when compressed. This will be the same value as RawSize if the payload is not compressed */
@@ -193,7 +193,7 @@ public:
 	 * @param Payload		The payload data
 	 * @param Callback		This callback will be invoked once the FPackageTrailer has been built and appended to disk.
 	 */
-	void AddPayload(const Virtualization::FPayloadId& Identifier, FCompressedBuffer Payload, AdditionalDataCallback&& Callback);
+	void AddPayload(const FIoHash& Identifier, FCompressedBuffer Payload, AdditionalDataCallback&& Callback);
 
 	/**
 	 * Adds an already virtualized payload to the builder to be written to the trailer. When the trailer is written
@@ -204,7 +204,7 @@ public:
 	 * @param Identifier	The identifier of the payload
 	 * @param RawSize		The size of the payload (in bytes) when uncompressed
 	 */
-	void AddVirtualizedPayload(const Virtualization::FPayloadId& Identifier, int64 RawSize);
+	void AddVirtualizedPayload(const FIoHash& Identifier, int64 RawSize);
 	
 	/**
 	 * @param ExportsArchive	The linker associated with the package being written to disk.
@@ -216,9 +216,9 @@ public:
 	/** Returns if the builder has any payload entries or not */
 	[[nodiscard]] bool IsEmpty() const;
 
-	[[nodiscard]] bool IsLocalPayloadEntry(const Virtualization::FPayloadId& Identifier) const;
-	[[nodiscard]] bool IsReferencedPayloadEntry(const Virtualization::FPayloadId& Identifier) const;
-	[[nodiscard]] bool IsVirtualizedPayloadEntry(const Virtualization::FPayloadId& Identifier) const;
+	[[nodiscard]] bool IsLocalPayloadEntry(const FIoHash& Identifier) const;
+	[[nodiscard]] bool IsReferencedPayloadEntry(const FIoHash& Identifier) const;
+	[[nodiscard]] bool IsVirtualizedPayloadEntry(const FIoHash& Identifier) const;
 
 	/** Returns the total number of payload entries in the builder */
 	[[nodiscard]] int32 GetNumPayloads() const;
@@ -285,11 +285,11 @@ private:
 	FName PackageName;
 
 	/** Payloads that will be stored locally when the trailer is written to disk */
-	TMap<Virtualization::FPayloadId, LocalEntry> LocalEntries;
+	TMap<FIoHash, LocalEntry> LocalEntries;
 	/** Payloads that reference entries in another trailer */
-	TMap<Virtualization::FPayloadId, ReferencedEntry> ReferencedEntries;
+	TMap<FIoHash, ReferencedEntry> ReferencedEntries;
 	/** Payloads that are already virtualized and so will not be written to disk */
-	TMap<Virtualization::FPayloadId, VirtualizedEntry> VirtualizedEntries;
+	TMap<FIoHash, VirtualizedEntry> VirtualizedEntries;
 	/** Callbacks to invoke once the trailer has been written to the end of a package */
 	TArray<AdditionalDataCallback> Callbacks;
 };
@@ -350,7 +350,7 @@ public:
 	 * @return	The payload in the form of a FCompressedBuffer. If the payload does not exist in the trailer or is not
 	 *			stored locally in the trailer then the FCompressedBuffer will be null.
 	 */
-	[[nodiscard]] FCompressedBuffer LoadLocalPayload(const Virtualization::FPayloadId& Id, FArchive& Ar) const;
+	[[nodiscard]] FCompressedBuffer LoadLocalPayload(const FIoHash& Id, FArchive& Ar) const;
 
 	/** 
 	 * Calling this indicates that the payload has been virtualized and will no longer be stored on disk. 
@@ -358,19 +358,19 @@ public:
 	 * @param Identifier The payload that has been virtualized
 	 * @return True if the payload was in the trailer, otherwise false
 	 */
-	[[nodiscard]] bool UpdatePayloadAsVirtualized(const Virtualization::FPayloadId& Identifier);
+	[[nodiscard]] bool UpdatePayloadAsVirtualized(const FIoHash& Identifier);
 
 	/** Attempt to find the status of the given payload. @See EPayloadStatus */
-	[[nodiscard]] EPayloadStatus FindPayloadStatus(const Virtualization::FPayloadId& Id) const;
+	[[nodiscard]] EPayloadStatus FindPayloadStatus(const FIoHash& Id) const;
 
 	/** Returns the absolute offset of the payload in the package file, invalid and virtualized payloads will return INDEX_NONE */
-	[[nodiscard]] int64 FindPayloadOffsetInFile(const Virtualization::FPayloadId& Id) const;
+	[[nodiscard]] int64 FindPayloadOffsetInFile(const FIoHash& Id) const;
 
 	/** Returns the total size of the of the trailer on disk in bytes */
 	[[nodiscard]] int64 GetTrailerLength() const;
 
 	/** Returns an array of the payloads that match the given filter type. @See EPayloadType */
-	[[nodiscard]] TArray<Virtualization::FPayloadId> GetPayloads(EPayloadFilter Type) const;
+	[[nodiscard]] TArray<FIoHash> GetPayloads(EPayloadFilter Type) const;
 
 	/** Returns the number of payloads that the trailer owns that match the given filter type. @See EPayloadType */
 	[[nodiscard]] int32 GetNumPayloads(EPayloadFilter Type) const;
@@ -452,6 +452,6 @@ private:
  *
  * @return 				True if the package was parsed successfully (although it might not have contained any payloads) and false if opening or parsing the package file failed.
  */
-[[nodiscard]] COREUOBJECT_API bool FindPayloadsInPackageFile(const FPackagePath& PackagePath, EPayloadFilter Filter, TArray<Virtualization::FPayloadId>& OutPayloadIds);
+[[nodiscard]] COREUOBJECT_API bool FindPayloadsInPackageFile(const FPackagePath& PackagePath, EPayloadFilter Filter, TArray<FIoHash>& OutPayloadIds);
 
 } //namespace UE
\ No newline at end of file
diff --git a/Engine/Source/Runtime/Engine/Private/ShaderCompiler/ShaderCompiler.cpp b/Engine/Source/Runtime/Engine/Private/ShaderCompiler/ShaderCompiler.cpp
--- a/Engine/Source/Runtime/Engine/Private/ShaderCompiler/ShaderCompiler.cpp
+++ b/Engine/Source/Runtime/Engine/Private/ShaderCompiler/ShaderCompiler.cpp
@@ -70,7 +70,6 @@
 #include "DerivedDataCache.h"
 #include "DerivedDataCacheRecord.h"
 #include "DerivedDataRequestOwner.h"
-#include "Virtualization/PayloadId.h"
 
 #include "TextureCompiler.h"
 #include "Rendering/StaticLightingSystemInterface.h"
diff --git a/Engine/Source/Runtime/Engine/Private/Texture.cpp b/Engine/Source/Runtime/Engine/Private/Texture.cpp
--- a/Engine/Source/Runtime/Engine/Private/Texture.cpp
+++ b/Engine/Source/Runtime/Engine/Private/Texture.cpp
@@ -2042,7 +2042,7 @@ void FTextureSource::UseHashAsGuid()
 		checkf(LockState == ELockState::None, TEXT("UseHashAsGuid shouldn't be called in-between LockMip/UnlockMip"));
 
 		bGuidIsHash = true;
-		Id = BulkData.GetPayloadId().ToGuid();
+		Id = UE::Serialization::IoHashToGuid(BulkData.GetPayloadId());
 	}
 	else
 	{
diff --git a/Engine/Source/Runtime/MeshDescription/Private/MeshDescription.cpp b/Engine/Source/Runtime/MeshDescription/Private/MeshDescription.cpp
--- a/Engine/Source/Runtime/MeshDescription/Private/MeshDescription.cpp
+++ b/Engine/Source/Runtime/MeshDescription/Private/MeshDescription.cpp
@@ -2215,7 +2215,7 @@ FGuid FMeshDescriptionBulkData::GetHash() const
 {
 	if (BulkData.GetPayloadSize() > 0)
 	{
-		return BulkData.GetPayloadId().ToGuid();
+		return UE::Serialization::IoHashToGuid(BulkData.GetPayloadId());
 	}
 
 	return FGuid();
