Replace use of LogVirtualization with LogSerialization for FPackageTrailer.

#rnx

diff --git a/Engine/Source/Runtime/CoreUObject/Private/UObject/PackageTrailer.cpp b/Engine/Source/Runtime/CoreUObject/Private/UObject/PackageTrailer.cpp
--- a/Engine/Source/Runtime/CoreUObject/Private/UObject/PackageTrailer.cpp
+++ b/Engine/Source/Runtime/CoreUObject/Private/UObject/PackageTrailer.cpp
@@ -46,11 +46,11 @@ void LogPackageOpenFailureMessage(const FPackagePath& PackagePath)
 	{
 		TCHAR SystemErrorMsg[2048] = { 0 };
 		FPlatformMisc::GetSystemErrorMessage(SystemErrorMsg, sizeof(SystemErrorMsg), SystemError);
-		UE_LOG(LogVirtualization, Error, TEXT("Could not open the file '%s' for reading due to system error: '%s' (%d))"), *PackagePath.GetDebugName(), SystemErrorMsg, SystemError);
+		UE_LOG(LogSerialization, Error, TEXT("Could not open the file '%s' for reading due to system error: '%s' (%d))"), *PackagePath.GetDebugName(), SystemErrorMsg, SystemError);
 	}
 	else
 	{
-		UE_LOG(LogVirtualization, Error, TEXT("Could not open (%s) to read FPackageTrailer with an unknown error"), *PackagePath.GetDebugName());
+		UE_LOG(LogSerialization, Error, TEXT("Could not open (%s) to read FPackageTrailer with an unknown error"), *PackagePath.GetDebugName());
 	}
 }
 
@@ -655,7 +655,7 @@ bool FindPayloadsInPackageFile(const FPackagePath& PackagePath, EPayloadFilter F
 {
 	if (FPackageName::IsTextPackageExtension(PackagePath.GetHeaderExtension()))
 	{
-		UE_LOG(LogVirtualization, Warning, TEXT("Attempting to call 'FindPayloadsInPackageFile' on a text based asset '%s' this is not currently supported"), *PackagePath.GetDebugName());
+		UE_LOG(LogSerialization, Warning, TEXT("Attempting to call 'FindPayloadsInPackageFile' on a text based asset '%s' this is not currently supported"), *PackagePath.GetDebugName());
 		return false;
 	}
 
@@ -674,13 +674,13 @@ bool FindPayloadsInPackageFile(const FPackagePath& PackagePath, EPayloadFilter F
 		}
 		else
 		{
-			UE_LOG(LogVirtualization, Warning, TEXT("Failed to parse the FPackageTrailer for '%s'"), *PackagePath.GetDebugName());
+			UE_LOG(LogSerialization, Warning, TEXT("Failed to parse the FPackageTrailer for '%s'"), *PackagePath.GetDebugName());
 			return false;
 		}	
 	}
 	else
 	{
-		UE_LOG(LogVirtualization, Warning, TEXT("Unable to open '%s' for reading"), *PackagePath.GetDebugName());
+		UE_LOG(LogSerialization, Warning, TEXT("Unable to open '%s' for reading"), *PackagePath.GetDebugName());
 		return false;
 	}
 }
