Merging changes from another stream:

FReferenceChainSearch will no longer store raw pointers to UObjects when constructing reference chains and instead will only preserve basic information about UObjects separate from UObjects themselves

Making FUObjectArray::ObjectToObjectItem accept a pointer to a const object

Removing redundant code from FindStaleClusters and fixing indents

diff --git a/Engine/Source/Runtime/CoreUObject/Private/Blueprint/BlueprintSupport.cpp b/Engine/Source/Runtime/CoreUObject/Private/Blueprint/BlueprintSupport.cpp
--- a/Engine/Source/Runtime/CoreUObject/Private/Blueprint/BlueprintSupport.cpp
+++ b/Engine/Source/Runtime/CoreUObject/Private/Blueprint/BlueprintSupport.cpp
@@ -400,13 +400,16 @@ void FBlueprintSupport::ValidateNoExternalRefsToSkeletons()
 		bool bBadRefs = false;
 		for(const FReferenceChainSearch::FReferenceChain* Chain : RefChainSearch.GetReferenceChains())
 		{
-			if(Chain->GetRootNode()->Object->GetOutermost() != SkeletonClass->GetOutermost())
+			UObject* ChainRootObject = Chain->GetRootNode()->ObjectInfo->TryResolveObject();
+			checkf(ChainRootObject, TEXT("Unable to resolve reference chain root object %s"), *Chain->GetRootNode()->ObjectInfo->GetPathName());
+			if(ChainRootObject->GetOutermost() != SkeletonClass->GetOutermost())
 			{
 				bBadRefs = true;
 				for (int32 NodeIndex = 1; bBadRefs && NodeIndex < Chain->Num(); ++NodeIndex)
 				{
 					// if there's a skeleton class (or an object outered to a skeleton class) somewhere in the chain, then it's fine:
-					UObject* ObjectReferencingSkeletonClass = Chain->GetNode(NodeIndex)->Object;
+					UObject* ObjectReferencingSkeletonClass = Chain->GetNode(NodeIndex)->ObjectInfo->TryResolveObject();
+					checkf(ChainRootObject, TEXT("Unable to resolve reference object referencing skeleton class %s"), *Chain->GetNode(NodeIndex)->ObjectInfo->GetPathName());
 					if (UClass* AsClass = Cast<UClass>(ObjectReferencingSkeletonClass))
 					{
 						if (IsSkeleton(AsClass))
diff --git a/Engine/Source/Runtime/CoreUObject/Private/UObject/GCObjectInfo.cpp b/Engine/Source/Runtime/CoreUObject/Private/UObject/GCObjectInfo.cpp
new file mode 100644
--- /dev/null
+++ b/Engine/Source/Runtime/CoreUObject/Private/UObject/GCObjectInfo.cpp
@@ -0,0 +1,67 @@
+// Copyright Epic Games, Inc. All Rights Reserved.
+
+/*=============================================================================
+	GCObjectInfo.cpp: Info about object participating in Garbage Collection code.
+=============================================================================*/
+
+#include "UObject/GCObjectInfo.h"
+
+UObject* FGCObjectInfo::TryResolveObject()
+{
+	return StaticFindObject(UObject::StaticClass(), nullptr, *GetPathName());
+}
+
+void FGCObjectInfo::GetPathName(FStringBuilderBase& ResultString) const
+{
+	if (this != nullptr)
+	{
+		if (Outer)
+		{
+			Outer->GetPathName(ResultString);
+
+			// SUBOBJECT_DELIMITER_CHAR is used to indicate that this object's outer is not a UPackage
+			if (Outer->Class->Name != NAME_Package && Outer->Outer->Class->Name == NAME_Package)
+			{
+				ResultString << SUBOBJECT_DELIMITER_CHAR;
+			}
+			else
+			{
+				ResultString << TEXT('.');
+			}
+		}
+		Name.AppendString(ResultString);
+	}
+	else
+	{
+		ResultString << TEXT("None");
+	}
+}
+
+FString FGCObjectInfo::GetPathName() const
+{
+	TStringBuilder<256> ResultBuilder;
+	GetPathName(ResultBuilder);
+	return ResultBuilder.ToString();
+}
+
+FString FGCObjectInfo::GetFullName() const
+{
+	return FString::Printf(TEXT("%s %s"), *GetClassName(), *GetPathName());
+}
+
+FGCObjectInfo* FGCObjectInfo::FindOrAddInfoHelper(UObject* InObject, TMap<UObject*, FGCObjectInfo*>& InOutObjectToInfoMap)
+{
+	FGCObjectInfo** ExistingObjInfo = InOutObjectToInfoMap.Find(InObject);
+	if (ExistingObjInfo)
+	{
+		return *ExistingObjInfo;
+	}
+
+	FGCObjectInfo* NewInfo = new FGCObjectInfo(InObject);
+	InOutObjectToInfoMap.Add(InObject, NewInfo);
+
+	NewInfo->Class = FindOrAddInfoHelper(InObject->GetClass(), InOutObjectToInfoMap);
+	NewInfo->Outer = InObject->GetOuter() ? FindOrAddInfoHelper(InObject->GetOuter(), InOutObjectToInfoMap) : nullptr;
+
+	return NewInfo;
+};
\ No newline at end of file
diff --git a/Engine/Source/Runtime/CoreUObject/Private/UObject/PackageReload.cpp b/Engine/Source/Runtime/CoreUObject/Private/UObject/PackageReload.cpp
--- a/Engine/Source/Runtime/CoreUObject/Private/UObject/PackageReload.cpp
+++ b/Engine/Source/Runtime/CoreUObject/Private/UObject/PackageReload.cpp
@@ -263,10 +263,12 @@ void DumpExternalReferences(UObject* InObject, UPackage* InPackage)
 			for (int32 NodeIndex = 0; NodeIndex < ObjectRefChain->Num(); ++NodeIndex)
 			{
 				const FReferenceChainSearch::FGraphNode* ObjectRefChainLink = ObjectRefChain->GetNode(NodeIndex);
-				const bool bIsExternalRef = ObjectRefChainLink->Object->GetOutermost() != InPackage;
+				UObject* LinkObject = ObjectRefChainLink->ObjectInfo->TryResolveObject();
+				checkf(LinkObject, TEXT("Unable to resolve object: %s when dumping external references"), *ObjectRefChainLink->ObjectInfo->GetPathName());
+				const bool bIsExternalRef = LinkObject->GetOutermost() != InPackage;
 				if (bIsExternalRef)
 				{
-					ExternalRefDumps.Emplace(ObjectRefChainLink->Object->GetFullName());
+					ExternalRefDumps.Emplace(ObjectRefChainLink->ObjectInfo->GetFullName());
 				}
 			}
 		}
diff --git a/Engine/Source/Runtime/CoreUObject/Private/UObject/ReferenceChainSearch.cpp b/Engine/Source/Runtime/CoreUObject/Private/UObject/ReferenceChainSearch.cpp
--- a/Engine/Source/Runtime/CoreUObject/Private/UObject/ReferenceChainSearch.cpp
+++ b/Engine/Source/Runtime/CoreUObject/Private/UObject/ReferenceChainSearch.cpp
@@ -11,34 +11,35 @@
 DEFINE_LOG_CATEGORY_STATIC(LogReferenceChain, Log, All);
 
 // Returns true if the object can't be collected by GC
-static FORCEINLINE bool IsNonGCObject(UObject* Object, EReferenceChainSearchMode SearchMode)
+static FORCEINLINE bool IsNonGCObject(FGCObjectInfo* Object, EReferenceChainSearchMode SearchMode)
 {
-	FUObjectItem* ObjectItem = GUObjectArray.ObjectToObjectItem(Object);
-	return (ObjectItem->IsRootSet() ||
-		ObjectItem->HasAnyFlags(EInternalObjectFlags::GarbageCollectionKeepFlags) ||
-		(GARBAGE_COLLECTION_KEEPFLAGS != RF_NoFlags && Object->HasAnyFlags(GARBAGE_COLLECTION_KEEPFLAGS) && !(SearchMode & EReferenceChainSearchMode::FullChain))
-		
-		);
+	return Object->HasAnyInternalFlags(EInternalObjectFlags::GarbageCollectionKeepFlags | EInternalObjectFlags::RootSet) ||
+		(GARBAGE_COLLECTION_KEEPFLAGS != RF_NoFlags && Object->HasAnyFlags(GARBAGE_COLLECTION_KEEPFLAGS) && !(SearchMode & EReferenceChainSearchMode::FullChain));
 }
 
-FReferenceChainSearch::FGraphNode* FReferenceChainSearch::FindOrAddNode(TMap<UObject*, FGraphNode*>& AllNodes, UObject* InObjectToFindNodeFor)
+FReferenceChainSearch::FGraphNode* FReferenceChainSearch::FindOrAddNode(FGCObjectInfo* InObjectInfo)
 {
 	FGraphNode* ObjectNode = nullptr;
-	FGraphNode** ExistingObjectNode = AllNodes.Find(InObjectToFindNodeFor);
+	FGraphNode** ExistingObjectNode = AllNodes.Find(InObjectInfo);
 	if (ExistingObjectNode)
 	{
 		ObjectNode = *ExistingObjectNode;
-		check(ObjectNode->Object == InObjectToFindNodeFor);
+		check(ObjectNode->ObjectInfo == InObjectInfo);
 	}
 	else
 	{
 		ObjectNode = new FGraphNode();
-		ObjectNode->Object = InObjectToFindNodeFor;
-		AllNodes.Add(InObjectToFindNodeFor, ObjectNode);
+		ObjectNode->ObjectInfo = InObjectInfo;
+		AllNodes.Add(InObjectInfo, ObjectNode);
 	}
 	return ObjectNode;
 }
 
+FReferenceChainSearch::FGraphNode* FReferenceChainSearch::FindOrAddNode(UObject* InObjectToFindNodeFor)
+{
+	return FindOrAddNode(FGCObjectInfo::FindOrAddInfoHelper(InObjectToFindNodeFor, ObjectToInfoMap));
+}
+
 int32 FReferenceChainSearch::BuildReferenceChains(FGraphNode* TargetNode, TArray<FReferenceChain*>& ProducedChains, int32 ChainDepth, const int32 VisitCounter, EReferenceChainSearchMode SearchMode)
 {
 	int32 ProducedChainsCount = 0;
@@ -47,7 +48,7 @@ int32 FReferenceChainSearch::BuildReferenceChains(FGraphNode* TargetNode, TArray
 		TargetNode->Visited = VisitCounter;
 
 		// Stop at root objects
-		if (!IsNonGCObject(TargetNode->Object, SearchMode))
+		if (!IsNonGCObject(TargetNode->ObjectInfo, SearchMode))
 		{			
 			for (FGraphNode* ReferencedByNode : TargetNode->ReferencedByObjects)
 			{
@@ -198,7 +199,7 @@ void FReferenceChainSearch::BuildReferenceChainsForDirectReferences(FGraphNode*
 {
 	for (FGraphNode* ReferencedByNode : TargetNode->ReferencedByObjects)
 	{
-		if (!(SearchMode & EReferenceChainSearchMode::ExternalOnly) || !ReferencedByNode->Object->IsIn(TargetNode->Object))
+		if (!(SearchMode & EReferenceChainSearchMode::ExternalOnly) || !ReferencedByNode->ObjectInfo->IsIn(TargetNode->ObjectInfo))
 		{
 			FReferenceChain* Chain = new FReferenceChain();
 			Chain->AddNode(TargetNode);
@@ -209,7 +210,7 @@ void FReferenceChainSearch::BuildReferenceChainsForDirectReferences(FGraphNode*
 	}
 }
 
-FString FReferenceChainSearch::GetObjectFlags(UObject* InObject)
+FString FReferenceChainSearch::GetObjectFlags(FGCObjectInfo* InObject)
 {
 	FString Flags;
 	if (InObject->IsRooted())
@@ -250,17 +251,16 @@ FString FReferenceChainSearch::GetObjectFlags(UObject* InObject)
 		Flags += TEXT("(asyncloading) ");
 	}
 
-	if (GUObjectArray.IsDisregardForGC(InObject))
+	if (InObject->IsDisregardForGC())
 	{
 		Flags += TEXT("(NeverGCed) ");
 	}
 
-	FUObjectItem* ReferencedByObjectItem = GUObjectArray.ObjectToObjectItem(InObject);
-	if (ReferencedByObjectItem->HasAnyFlags(EInternalObjectFlags::ClusterRoot))
+	if (InObject->HasAnyInternalFlags(EInternalObjectFlags::ClusterRoot))
 	{
 		Flags += TEXT("(ClusterRoot) ");
 	}
-	if (ReferencedByObjectItem->GetOwnerIndex() > 0)
+	if (InObject->GetOwnerIndex() > 0)
 	{
 		Flags += TEXT("(Clustered) ");
 	}
@@ -306,7 +306,7 @@ void FReferenceChainSearch::DumpChain(FReferenceChainSearch::FReferenceChain* Ch
 		bool bPostCallbackContinue = true;
 		const int32 RootIndex = Chain->Num() - 1;
 		FNodeReferenceInfo ReferenceInfo = Chain->GetReferenceInfo(RootIndex);
-		UObject* ReferencerObject = Chain->GetNode(RootIndex)->Object;
+		FGCObjectInfo* ReferencerObject = Chain->GetNode(RootIndex)->ObjectInfo;
 		{
 			FCallbackParams Params;
 			Params.Referencer = nullptr;
@@ -326,7 +326,7 @@ void FReferenceChainSearch::DumpChain(FReferenceChainSearch::FReferenceChain* Ch
 		// Roots are at the end so iterate from the last to the first node
 		for (int32 NodeIndex = RootIndex - 1; NodeIndex >= 0 && bPostCallbackContinue; --NodeIndex)
 		{
-			UObject* Object = Chain->GetNode(NodeIndex)->Object;
+			FGCObjectInfo* Object = Chain->GetNode(NodeIndex)->ObjectInfo;
 
 			FCallbackParams Params;
 			Params.Referencer = ReferencerObject;
@@ -338,9 +338,10 @@ void FReferenceChainSearch::DumpChain(FReferenceChainSearch::FReferenceChain* Ch
 			if (ReferenceInfo.Type == EReferenceType::Property)
 			{
 				FString ReferencingPropertyName;
+				UClass* ReferencerClass = Cast<UClass>(ReferencerObject->GetClass()->TryResolveObject());
 				TArray<FProperty*> ReferencingProperties;
 
-				if (FGCStackSizeHelper::ConvertPathToProperties(ReferencerObject->GetClass(), ReferenceInfo.ReferencerName, ReferencingProperties))
+				if (ReferencerClass && FGCStackSizeHelper::ConvertPathToProperties(ReferencerClass, ReferenceInfo.ReferencerName, ReferencingProperties))
 				{
 					FProperty* InnermostProperty = ReferencingProperties.Last();
 					FProperty* OutermostProperty = ReferencingProperties[0];
@@ -362,10 +363,16 @@ void FReferenceChainSearch::DumpChain(FReferenceChainSearch::FReferenceChain* Ch
 					{
 						ClassName = TEXT("UObject");
 					}
+					else if (ReferencerClass)
+					{
+						// Use the native class name when possible
+						ClassName = ReferencerClass->GetPrefixCPP();
+						ClassName += ReferencerClass->GetName();
+					}
 					else
 					{
-						ClassName = ReferencerObject->GetClass()->GetPrefixCPP();
-						ClassName += ReferencerObject->GetClass()->GetName();
+						// Revert to the internal class name if not
+						ClassName = ReferencerObject->GetClassName();
 					}
 					ReferencingPropertyName = FString::Printf(TEXT("UObject* %s::%s"), *ClassName, *ReferenceInfo.ReferencerName.ToString());
 				}
@@ -381,8 +388,16 @@ void FReferenceChainSearch::DumpChain(FReferenceChainSearch::FReferenceChain* Ch
 				FString UObjectOrGCObjectName;
 				if (ReferenceInfo.ReferencerName.IsNone())
 				{
-					UObjectOrGCObjectName = ReferencerObject->GetClass()->GetPrefixCPP();
-					UObjectOrGCObjectName += ReferencerObject->GetClass()->GetName();
+					UClass* ReferencerClass = Cast<UClass>(ReferencerObject->GetClass()->TryResolveObject());
+					if (ReferencerClass)
+					{
+						UObjectOrGCObjectName = ReferencerClass->GetPrefixCPP();
+						UObjectOrGCObjectName += ReferencerClass->GetName();
+					}
+					else
+					{
+						UObjectOrGCObjectName += ReferencerObject->GetClassName();
+					}
 				}
 				else
 				{
@@ -443,7 +458,7 @@ bool FReferenceChainSearch::FReferenceChain::IsExternal() const
 	if (Nodes.Num() > 1)
 	{
 		// Reference is external if the root (the last node) is not in the first node (target)
-		return !Nodes.Last()->Object->IsIn(Nodes[0]->Object);
+		return !Nodes.Last()->ObjectInfo->IsIn(Nodes[0]->ObjectInfo);
 	}
 	else
 	{
@@ -458,49 +473,57 @@ class FDirectReferenceProcessor : public FSimpleReferenceProcessorBase
 {	
 	UObject* ObjectToFindReferencesTo;
 	TSet<FReferenceChainSearch::FObjectReferenceInfo>& ReferencedObjects;
+	TMap<UObject*, FGCObjectInfo*>& ObjectToInfoMap;
 
 public:
 
-	FDirectReferenceProcessor(UObject* InObjectToFindReferencesTo, TSet<FReferenceChainSearch::FObjectReferenceInfo>& InReferencedObjects)
+	FDirectReferenceProcessor(UObject* InObjectToFindReferencesTo, TSet<FReferenceChainSearch::FObjectReferenceInfo>& InReferencedObjects, TMap<UObject*, FGCObjectInfo*>& InObjectToInfoMap)
 		: ObjectToFindReferencesTo(InObjectToFindReferencesTo)
-		, ReferencedObjects(InReferencedObjects)		
+		, ReferencedObjects(InReferencedObjects)
+		, ObjectToInfoMap(InObjectToInfoMap)
 	{
 	}
+
 	FORCEINLINE void HandleTokenStreamObjectReference(FGCArrayStruct& ObjectsToSerializeStruct, UObject* ReferencingObject, UObject*& Object, const int32 TokenIndex, const EGCTokenType TokenType, bool bAllowReferenceElimination)
-	{
-		FReferenceChainSearch::FObjectReferenceInfo RefInfo(Object);
-		if (Object && !ReferencedObjects.Contains(RefInfo))
+	{		
+		if (Object)
 		{
-			if (TokenIndex >= 0)
-			{
-				FTokenInfo TokenInfo = ReferencingObject->GetClass()->ReferenceTokenStream.GetTokenInfo(TokenIndex);
-				RefInfo.ReferencerName = TokenInfo.Name;
-				RefInfo.Type = FReferenceChainSearch::EReferenceType::Property;
-			}
-			else
-			{
-				RefInfo.Type = FReferenceChainSearch::EReferenceType::AddReferencedObjects;
-				RefInfo.NumStackFrames = FPlatformStackWalk::CaptureStackBackTrace(RefInfo.StackFrames, FReferenceChainSearch::FObjectReferenceInfo::MaxStackFrames);
+			FGCObjectInfo* ObjectInfo = FGCObjectInfo::FindOrAddInfoHelper(Object, ObjectToInfoMap);
 
-				if (FGCObject::GGCObjectReferencer && (!ReferencingObject || ReferencingObject == FGCObject::GGCObjectReferencer))
+			FReferenceChainSearch::FObjectReferenceInfo RefInfo(ObjectInfo);
+			if (!ReferencedObjects.Contains(RefInfo))
+			{
+				if (TokenIndex >= 0)
+				{
+					FTokenInfo TokenInfo = ReferencingObject->GetClass()->ReferenceTokenStream.GetTokenInfo(TokenIndex);
+					RefInfo.ReferencerName = TokenInfo.Name;
+					RefInfo.Type = FReferenceChainSearch::EReferenceType::Property;
+				}
+				else
 				{
-					FString RefName;
-					if (FGCObject::GGCObjectReferencer->GetReferencerName(Object, RefName, true))
+					RefInfo.Type = FReferenceChainSearch::EReferenceType::AddReferencedObjects;
+					RefInfo.NumStackFrames = FPlatformStackWalk::CaptureStackBackTrace(RefInfo.StackFrames, FReferenceChainSearch::FObjectReferenceInfo::MaxStackFrames);
+
+					if (FGCObject::GGCObjectReferencer && (!ReferencingObject || ReferencingObject == FGCObject::GGCObjectReferencer))
 					{
-						RefInfo.ReferencerName = FName(*RefName);
+						FString RefName;
+						if (FGCObject::GGCObjectReferencer->GetReferencerName(Object, RefName, true))
+						{
+							RefInfo.ReferencerName = FName(*RefName);
+						}
+						else if (ReferencingObject)
+						{
+							RefInfo.ReferencerName = *ReferencingObject->GetFullName();
+						}
 					}
 					else if (ReferencingObject)
 					{
 						RefInfo.ReferencerName = *ReferencingObject->GetFullName();
 					}
 				}
-				else if (ReferencingObject)
-				{
-					RefInfo.ReferencerName = *ReferencingObject->GetFullName();
-				}
-			}
 
-			ReferencedObjects.Add(RefInfo);
+				ReferencedObjects.Add(RefInfo);
+			}
 		}
 	}
 };
@@ -523,10 +546,18 @@ public:
 
 FReferenceChainSearch::FReferenceChainSearch(UObject* InObjectToFindReferencesTo, EReferenceChainSearchMode Mode /*= EReferenceChainSearchMode::PrintResults*/)
 	: ObjectToFindReferencesTo(InObjectToFindReferencesTo)
+	, SearchMode(Mode)
 {
-	check(InObjectToFindReferencesTo);
+	FSlowHeartBeatScope DisableHangDetection; // This function can be very slow
+
+	check(ObjectToFindReferencesTo);
+	ObjectInfoToFindReferencesTo = FGCObjectInfo::FindOrAddInfoHelper(ObjectToFindReferencesTo, ObjectToInfoMap);
+
+	// First pass is to find all direct references for each object
+	FindDirectReferencesForObjects();
 
-	PerformSearch(Mode);
+	// Second pass creates all reference chains
+	PerformSearch();
 
 	if (!!(Mode & (EReferenceChainSearchMode::PrintResults|EReferenceChainSearchMode::PrintAllResults)))
 	{
@@ -534,21 +565,15 @@ FReferenceChainSearch::FReferenceChainSearch(UObject* InObjectToFindReferencesTo
 	}
 }
 
+
 FReferenceChainSearch::~FReferenceChainSearch()
 {
 	Cleanup();
 }
 
-void FReferenceChainSearch::PerformSearch(EReferenceChainSearchMode SearchMode)
+void FReferenceChainSearch::PerformSearch()
 {
-	FSlowHeartBeatScope DisableHangDetection; // This function can be very slow
-
-	// First pass is to find all direct references for each object
-	{
-		FindDirectReferencesForObjects();
-	}
-
-	FGraphNode* ObjectNodeToFindReferencesTo = FindOrAddNode(AllNodes, ObjectToFindReferencesTo);
+	FGraphNode* ObjectNodeToFindReferencesTo = FindOrAddNode(ObjectToFindReferencesTo);
 	check(ObjectNodeToFindReferencesTo);
 
 	// Now it's time to build the reference chain from all of the objects that reference the object to find references to
@@ -565,7 +590,7 @@ void FReferenceChainSearch::PerformSearch(EReferenceChainSearchMode SearchMode)
 void FReferenceChainSearch::FindDirectReferencesForObjects()
 {
 	TSet<FObjectReferenceInfo> ReferencedObjects;
-	FDirectReferenceProcessor Processor(ObjectToFindReferencesTo, ReferencedObjects);
+	FDirectReferenceProcessor Processor(ObjectToFindReferencesTo, ReferencedObjects, ObjectToInfoMap);
 	TFastReferenceCollector<
 		FDirectReferenceProcessor, 
 		FDirectReferenceCollector, 
@@ -579,7 +604,7 @@ void FReferenceChainSearch::FindDirectReferencesForObjects()
 	{
 		FUObjectItem* ObjItem = *It;
 		UObject* Object = static_cast<UObject*>(ObjItem->Object);
-		FGraphNode* ObjectNode = FindOrAddNode(AllNodes, Object);
+		FGraphNode* ObjectNode = FindOrAddNode(Object);
 
 		// Find direct references
 		ReferencedObjects.Reset();
@@ -587,10 +612,10 @@ void FReferenceChainSearch::FindDirectReferencesForObjects()
 		ObjectsToProcess.Add(Object);
 		ReferenceCollector.CollectReferences(ArrayStruct);
 
-		// Build the reference tree
+		// Build direct reference tree
 		for (FObjectReferenceInfo& ReferenceInfo : ReferencedObjects)
 		{
-			FGraphNode* ReferencedObjectNode = FindOrAddNode(AllNodes, ReferenceInfo.Object);
+			FGraphNode* ReferencedObjectNode = FindOrAddNode(ReferenceInfo.Object);
 			ObjectNode->ReferencedObjects.Add(FNodeReferenceInfo(ReferencedObjectNode, ReferenceInfo.Type, ReferenceInfo.ReferencerName, ReferenceInfo.StackFrames, ReferenceInfo.NumStackFrames));
 			ReferencedObjectNode->ReferencedByObjects.Add(ObjectNode);
 		}
@@ -627,10 +652,10 @@ void FReferenceChainSearch::PrintResults(TFunctionRef<bool(FCallbackParams& Para
 	}
 	else
 	{
-		check(ObjectToFindReferencesTo);
+		check(ObjectInfoToFindReferencesTo);
 		UE_LOG(LogReferenceChain, Log, TEXT("%s%s is not currently reachable."),
-			*GetObjectFlags(ObjectToFindReferencesTo),
-			*ObjectToFindReferencesTo->GetFullName()
+			*GetObjectFlags(ObjectInfoToFindReferencesTo),
+			*ObjectInfoToFindReferencesTo->GetFullName()
 		);
 	}
 }
@@ -653,8 +678,8 @@ FString FReferenceChainSearch::GetRootPath(TFunctionRef<bool(FCallbackParams& Pa
 	else
 	{
 		return FString::Printf(TEXT("%s%s is not currently reachable."),
-			*GetObjectFlags(ObjectToFindReferencesTo),
-			*ObjectToFindReferencesTo->GetFullName()
+			*GetObjectFlags(ObjectInfoToFindReferencesTo),
+			*ObjectInfoToFindReferencesTo->GetFullName()
 		);
 	}
 }
@@ -667,9 +692,15 @@ void FReferenceChainSearch::Cleanup()
 	}
 	ReferenceChains.Empty();
 
-	for (TPair<UObject*, FGraphNode*>& ObjectNodePair : AllNodes)
+	for (TPair<FGCObjectInfo*, FGraphNode*>& ObjectNodePair : AllNodes)
 	{
 		delete ObjectNodePair.Value;
 	}
 	AllNodes.Empty();
+
+	for (TPair<UObject*, FGCObjectInfo*>& ObjectToInfoPair : ObjectToInfoMap)
+	{
+		delete ObjectToInfoPair.Value;
+	}
+	ObjectToInfoMap.Empty();
 }
diff --git a/Engine/Source/Runtime/CoreUObject/Private/UObject/UObjectClusters.cpp b/Engine/Source/Runtime/CoreUObject/Private/UObject/UObjectClusters.cpp
--- a/Engine/Source/Runtime/CoreUObject/Private/UObject/UObjectClusters.cpp
+++ b/Engine/Source/Runtime/CoreUObject/Private/UObject/UObjectClusters.cpp
@@ -474,12 +474,11 @@ void FindStaleClusters(const TArray<FString>& Args)
 			if (SearchRefs.GetReferenceChains().Num() > 0)
 			{
 				for (const FReferenceChainSearch::FReferenceChain* ReferenceChain : SearchRefs.GetReferenceChains())
-					{
-					UObject* ReferencingObj = ReferenceChain->GetRootNode()->Object;
-						bReferenced = true;
-						break;
-					}
+				{
+					bReferenced = true;
+					break;
 				}
+			}
 			if (!bReferenced)
 			{
 				NumStaleClusters++;
diff --git a/Engine/Source/Runtime/CoreUObject/Public/UObject/GCObjectInfo.h b/Engine/Source/Runtime/CoreUObject/Public/UObject/GCObjectInfo.h
new file mode 100644
--- /dev/null
+++ b/Engine/Source/Runtime/CoreUObject/Public/UObject/GCObjectInfo.h
@@ -0,0 +1,133 @@
+// Copyright Epic Games, Inc. All Rights Reserved.
+
+/*=============================================================================
+	GarbageCollectionHistory.h: Unreal realtime garbage collection history
+=============================================================================*/
+
+#pragma once
+
+#include "UObject/GarbageCollection.h"
+#include "UObject/Class.h"
+
+/**
+ * Structure containing information about a UObject participating in Garbage Collection.
+ * It's purpose is to avoid holding onto direct references to UObjects which may have already been Garbage Collected.
+ * FGCObjectInfo interface mimics that of UObject.
+ **/
+class COREUOBJECT_API FGCObjectInfo
+{
+public:
+
+	FGCObjectInfo() = default;
+	explicit FGCObjectInfo(const UObject* Obj)
+		: Name(Obj->GetFName())
+		, Flags(Obj->GetFlags())
+		, InternalFlags(Obj->GetInternalFlags())
+		, ClusterRootIndex(GUObjectArray.ObjectToObjectItem(Obj)->ClusterRootIndex)
+		, bDisregardForGC(GUObjectArray.IsDisregardForGC(Obj))
+	{
+
+	}
+
+private:
+
+	/** Name of the object */
+	FName Name;
+	/** Pointer to class info */
+	FGCObjectInfo* Class = nullptr;
+	/** Pointer to Outer info */
+	FGCObjectInfo* Outer = nullptr;
+	/** Captured Object flags */
+	EObjectFlags Flags = RF_NoFlags;
+	/** Captured Internal flags */
+	EInternalObjectFlags InternalFlags = EInternalObjectFlags::None;
+	/** Object's FGCObjectItem cluster root index */
+	int32 ClusterRootIndex = -1;
+	/** True if the object was inside of the disregard for GC set */
+	bool bDisregardForGC = false;
+
+public:
+
+	/**
+	 * Tries to find an existing object matching the stored Path.
+	 * @returns Pointer to an object this info struct represents if the object is still alive, nullptr otherwise
+	 */
+	UObject* TryResolveObject();
+
+	FGCObjectInfo* GetClass() const
+	{
+		return Class;
+	}
+
+	FGCObjectInfo* GetOuter() const
+	{
+		return Outer;
+	}
+
+	inline bool IsIn(const FGCObjectInfo* MaybeOuter) const
+	{		
+		for (const FGCObjectInfo* Obj = this; Obj; Obj = Obj->Outer)
+		{
+			if (Obj == MaybeOuter)
+			{
+				return true;
+			}
+		}
+		return false;
+	}
+
+	bool HasAnyFlags(EObjectFlags InFlags) const
+	{
+		return !!(Flags & InFlags);
+	}
+
+	bool HasAnyInternalFlags(EInternalObjectFlags InFlags) const
+	{
+		return !!(InternalFlags & InFlags);
+	}
+
+	bool IsRooted() const
+	{
+		return HasAnyInternalFlags(EInternalObjectFlags::RootSet);
+	}
+
+	bool IsNative() const
+	{
+		return HasAnyInternalFlags(EInternalObjectFlags::Native);
+	}
+
+	int32 GetOwnerIndex() const
+	{
+		return ClusterRootIndex;
+	}
+
+	FString GetClassName() const
+	{
+		check(Class);
+		return Class->Name.ToString();
+	}
+
+	bool IsValid() const
+	{
+		PRAGMA_DISABLE_DEPRECATION_WARNINGS
+		return !HasAnyInternalFlags(EInternalObjectFlags::Garbage | EInternalObjectFlags::PendingKill);
+		PRAGMA_ENABLE_DEPRECATION_WARNINGS
+	}
+
+	bool IsDisregardForGC() const
+	{
+		return bDisregardForGC;
+	}
+
+	void GetPathName(FStringBuilderBase& ResultString) const;
+	FString GetPathName() const;
+	FString GetFullName() const;
+
+	friend uint32 GetTypeHash(const FGCObjectInfo& Info)
+	{
+		return HashCombine(GetTypeHash(Info.Outer), GetTypeHash(Info.Name));
+	}
+
+	/** Helper function for adding info about an UObject into UObject to FGCObjectInfo map */
+	static FGCObjectInfo* FindOrAddInfoHelper(UObject* InObject, TMap<UObject*, FGCObjectInfo*>& InOutObjectToInfoMap);
+};
diff --git a/Engine/Source/Runtime/CoreUObject/Public/UObject/ReferenceChainSearch.h b/Engine/Source/Runtime/CoreUObject/Public/UObject/ReferenceChainSearch.h
--- a/Engine/Source/Runtime/CoreUObject/Public/UObject/ReferenceChainSearch.h
+++ b/Engine/Source/Runtime/CoreUObject/Public/UObject/ReferenceChainSearch.h
@@ -4,8 +4,7 @@
 
 #include "CoreMinimal.h"
 #include "UObject/UObjectGlobals.h"
-#include "UObject/Class.h"
-#include "UObject/GarbageCollection.h"
+#include "UObject/GCObjectInfo.h"
 #include "HAL/ThreadHeartBeat.h"
 
 /** Search mode flags */
@@ -138,23 +137,21 @@ public:
 	/** Single node in the reference graph */
 	struct FGraphNode
 	{
-		FGraphNode()
-			: Object(nullptr)
-			, Visited(0)
-		{}
-
 		/** Object pointer */
-		UObject* Object;
+		UE_DEPRECATED(5.0, "Direct Object reference has been deprecated. Use ObjectInfo member variable instead.")
+		UObject* Object = nullptr;
+		/** Object pointer */
+		FGCObjectInfo* ObjectInfo = nullptr;
 		/** Objects referenced by this object with reference info */
 		TSet< TReferenceInfo<FGraphNode> > ReferencedObjects;
 		/** Objects that have references to this object */
 		TSet<FGraphNode*> ReferencedByObjects;
 		/** Non-zero if this node has been already visited during reference search */
-		int32 Visited;
+		int32 Visited = 0;
 	};
 
 	/** Convenience type definitions to avoid template hell */
-	typedef TReferenceInfo<UObject> FObjectReferenceInfo;
+	typedef TReferenceInfo<FGCObjectInfo> FObjectReferenceInfo;
 	typedef TReferenceInfo<FGraphNode> FNodeReferenceInfo;
 
 	/** Reference chain. The first object in the list is the target object and the last object is a root object */
@@ -224,9 +221,9 @@ public:
 	struct FCallbackParams
 	{
 		/** Referenced object */
-		UObject* Object = nullptr;
+		FGCObjectInfo* Object = nullptr;
 		/** Object that is referencing the current object */
-		UObject* Referencer = nullptr;
+		FGCObjectInfo* Referencer = nullptr;
 		/** Information about the type of reference (Referencer -> Object) */
 		const FNodeReferenceInfo* ReferenceInfo = nullptr;
 		/** For use when outputting custom information: current indent */
@@ -272,14 +269,21 @@ public:
 private:
 
 	/** The object we're going to look for references to */
-	UObject* ObjectToFindReferencesTo;
+	UObject* ObjectToFindReferencesTo = nullptr;
+	FGCObjectInfo* ObjectInfoToFindReferencesTo = nullptr;
+
+	/** Search mode and options */
+	EReferenceChainSearchMode SearchMode = EReferenceChainSearchMode::Default;
+
 	/** All reference chain found during the search */
 	TArray<FReferenceChain*> ReferenceChains;
 	/** All nodes created during the search */
-	TMap<UObject*, FGraphNode*> AllNodes;
+	TMap<FGCObjectInfo*, FGraphNode*> AllNodes;
+	/** Maps UObject pointers to object info structs */
+	TMap<UObject*, FGCObjectInfo*> ObjectToInfoMap;
 
 	/** Performs the search */
-	void PerformSearch(EReferenceChainSearchMode SearchMode);
+	void PerformSearch();
 
 	/** Performs the search */
 	void FindDirectReferencesForObjects();
@@ -288,7 +292,9 @@ private:
 	void Cleanup();
 
 	/** Tries to find a node for an object and if it doesn't exists creates a new one and returns it */
-	static FGraphNode* FindOrAddNode(TMap<UObject*, FGraphNode*>& AllNodes, UObject* InObjectToFindNodeFor);
+	FGraphNode* FindOrAddNode(UObject* InObjectToFindNodeFor);
+	FGraphNode* FindOrAddNode(FGCObjectInfo* InObjectInfo);
+
 	/** Builds reference chains */
 	static int32 BuildReferenceChains(FGraphNode* TargetNode, TArray<FReferenceChain*>& ProducedChains, int32 ChainDepth, const int32 VisitCounter, EReferenceChainSearchMode SearchMode);
 	/** Builds reference chains */
@@ -301,7 +307,7 @@ private:
 	static void RemoveDuplicatedChains(TArray<FReferenceChain*>& AllChains);
 
 	/** Returns a string with all flags (we care about) set on an object */
-	static FString GetObjectFlags(UObject* InObject);
+	static FString GetObjectFlags(FGCObjectInfo* InObject);
 	/** Dumps a reference chain to log */
 	static void DumpChain(FReferenceChainSearch::FReferenceChain* Chain, TFunctionRef<bool(FCallbackParams& Params)> ReferenceCallback, FOutputDevice& Out);
 };
diff --git a/Engine/Source/Runtime/CoreUObject/Public/UObject/UObjectArray.h b/Engine/Source/Runtime/CoreUObject/Public/UObject/UObjectArray.h
--- a/Engine/Source/Runtime/CoreUObject/Public/UObject/UObjectArray.h
+++ b/Engine/Source/Runtime/CoreUObject/Public/UObject/UObjectArray.h
@@ -743,7 +743,7 @@ public:
 		return ObjectItem;
 	}
 
-	FORCEINLINE FUObjectItem* ObjectToObjectItem(UObjectBase* Object)
+	FORCEINLINE FUObjectItem* ObjectToObjectItem(const UObjectBase* Object)
 	{
 		FUObjectItem* ObjectItem = IndexToObject(Object->InternalIndex);
 		return ObjectItem;
diff --git a/Engine/Source/Runtime/Engine/Private/UnrealEngine.cpp b/Engine/Source/Runtime/Engine/Private/UnrealEngine.cpp
--- a/Engine/Source/Runtime/Engine/Private/UnrealEngine.cpp
+++ b/Engine/Source/Runtime/Engine/Private/UnrealEngine.cpp
@@ -14928,13 +14928,13 @@ void UEngine::FindAndPrintStaleReferencesToObject(UObject* ObjectToFindReference
 {
 	check(ObjectToFindReferencesTo);
 
-	UObject* OutGarbageObject = nullptr;
-	UObject* OutReferencingObject = nullptr;
+	FGCObjectInfo* OutGarbageObject = nullptr;
+	FGCObjectInfo* OutReferencingObject = nullptr;
 	FReferenceChainSearch RefChainSearch(ObjectToFindReferencesTo, EReferenceChainSearchMode::Default);
 	RefChainSearch.PrintResults([&OutGarbageObject, &OutReferencingObject](FReferenceChainSearch::FCallbackParams& Params)
 	{
 		check(Params.Object);
-		if (!IsValid(Params.Object))
+		if (!Params.Object->IsValid())
 		{
 			// We may find many chains that lead to the leak but for brevity we only report the first one in the fatal error below
 			if (!OutGarbageObject)
@@ -14955,7 +14955,7 @@ void UEngine::FindAndPrintStaleReferencesToObject(UObject* ObjectToFindReference
 	FString PathToCulprit = RefChainSearch.GetRootPath([&bFirstGarbageObjectFound, ObjectToFindReferencesTo](FReferenceChainSearch::FCallbackParams& Params)
 	{
 		check(Params.Object);
-		if (!IsValid(Params.Object) && !bFirstGarbageObjectFound)
+		if (!Params.Object->IsValid() && !bFirstGarbageObjectFound)
 		{
 			bFirstGarbageObjectFound = true;
 			check(Params.Out);
