Change ArchetypeAnnotation to use FUObjectAnnotationChunked.
Saves some memory and prevents very large reallocs which are slow and cause fragmentation

diff --git a/Engine/Source/Runtime/CoreUObject/Private/UObject/UObjectArchetype.cpp b/Engine/Source/Runtime/CoreUObject/Private/UObject/UObjectArchetype.cpp
--- a/Engine/Source/Runtime/CoreUObject/Private/UObject/UObjectArchetype.cpp
+++ b/Engine/Source/Runtime/CoreUObject/Private/UObject/UObjectArchetype.cpp
@@ -48,7 +48,7 @@ struct FArchetypeInfo
 	int32 ArchetypeIndex;
 };
 
-static FUObjectAnnotationDense<FArchetypeInfo, true> ArchetypeAnnotation;
+static FUObjectAnnotationChunked<FArchetypeInfo, true> ArchetypeAnnotation;
 
 #endif // UE_CACHE_ARCHETYPE
 
