Manually roll back changelist 18825124 - changing archetype annotation to use chunked annotation type

diff --git a/Engine/Source/Runtime/CoreUObject/Private/UObject/UObjectArchetype.cpp b/Engine/Source/Runtime/CoreUObject/Private/UObject/UObjectArchetype.cpp
--- a/Engine/Source/Runtime/CoreUObject/Private/UObject/UObjectArchetype.cpp
+++ b/Engine/Source/Runtime/CoreUObject/Private/UObject/UObjectArchetype.cpp
@@ -48,7 +48,7 @@ struct FArchetypeInfo
 	int32 ArchetypeIndex;
 };
 
-static FUObjectAnnotationChunked<FArchetypeInfo, true> ArchetypeAnnotation;
+static FUObjectAnnotationDense<FArchetypeInfo, true> ArchetypeAnnotation;
 
 #endif // UE_CACHE_ARCHETYPE
 
