Move ContainsNoAsset search out of a !IsTextFormat() block so that it doesn't get lost when saving as a text asset

diff --git a/Engine/Source/Runtime/CoreUObject/Private/UObject/SavePackage2.cpp b/Engine/Source/Runtime/CoreUObject/Private/UObject/SavePackage2.cpp
--- a/Engine/Source/Runtime/CoreUObject/Private/UObject/SavePackage2.cpp
+++ b/Engine/Source/Runtime/CoreUObject/Private/UObject/SavePackage2.cpp
@@ -1820,7 +1820,6 @@ ESavePackageResult UpdatePackageHeader(FStructuredArchive::FRecord& StructuredAr
 		}
 	}
 	// Write Real Export Map
-	bool bContainsAsset = false;
 	if (!SaveContext.IsTextFormat())
 	{
 		check(Linker->Tell() == SaveContext.OffsetAfterImportMap);
@@ -1830,11 +1829,17 @@ ESavePackageResult UpdatePackageHeader(FStructuredArchive::FRecord& StructuredAr
 		for (FObjectExport& Export : Linker->ExportMap)
 		{
 			ExportTableStream.EnterElement() << Export;
-			bContainsAsset |= Export.bIsAsset;
 		}
 		check(Linker->Tell() == SaveContext.OffsetAfterExportMap);
 	}
 
+	// Figure out if at leat one export is marked as an asset
+	bool bContainsAsset = false;
+	for (FObjectExport& Export : Linker->ExportMap)
+	{
+		bContainsAsset |= Export.bIsAsset;
+	}
+
 	// Update Summary
 	// Write Real Summary
 	{
