Reduced runtime UPackage size from 176b to 144b or 120b when stripping deprecated properties
This should saves around ~1.5 Mb to ~3 Mb at runtime with ~50000 packages

Added a new core define UE_STRIP_DEPRECATED_PROPERTIES that could be used to wrap deprecated properties and strip them to regain memory when a projects becomes compliant. this can be set in the project target file
Deprecated most public properties from UPackage and created accessors for them

Merged from cl  18646219, 18738937, 18747375

diff --git a/Engine/Source/Runtime/CoreUObject/Private/UObject/SavePackage/SaveContext.h b/Engine/Source/Runtime/CoreUObject/Private/UObject/SavePackage/SaveContext.h
--- a/Engine/Source/Runtime/CoreUObject/Private/UObject/SavePackage/SaveContext.h
+++ b/Engine/Source/Runtime/CoreUObject/Private/UObject/SavePackage/SaveContext.h
@@ -841,6 +841,18 @@ public:
 		GetHarvestedRealm().SetLinker(MoveTemp(InLinker));
 	}
 
+	void UpdatePackageLinkerVersions()
+	{
+		Package->SetLinkerPackageVersion(Linker->UEVer());
+		Package->SetLinkerLicenseeVersion(Linker->LicenseeUEVer());
+		Package->SetLinkerCustomVersions(Linker->GetCustomVersions());
+	}
+
+	void UpdatePackageFileSize(int64 InFileSize)
+	{
+		Package->SetFileSize(InFileSize);
+	}
+
 	bool CloseLinkerArchives()
 	{
 		return GetHarvestedRealm().CloseLinkerArchives();
