Fix compilation error from merge issue that didn't result in a conflict

diff --git a/Engine/Source/Runtime/CoreUObject/Private/UObject/SavePackage/SaveContext.h b/Engine/Source/Runtime/CoreUObject/Private/UObject/SavePackage/SaveContext.h
--- a/Engine/Source/Runtime/CoreUObject/Private/UObject/SavePackage/SaveContext.h
+++ b/Engine/Source/Runtime/CoreUObject/Private/UObject/SavePackage/SaveContext.h
@@ -841,18 +841,6 @@ public:
 		GetHarvestedRealm().SetLinker(MoveTemp(InLinker));
 	}
 
-	void UpdatePackageLinkerVersions()
-	{
-		Package->SetLinkerPackageVersion(Linker->UEVer());
-		Package->SetLinkerLicenseeVersion(Linker->LicenseeUEVer());
-		Package->SetLinkerCustomVersions(Linker->GetCustomVersions());
-	}
-
-	void UpdatePackageFileSize(int64 InFileSize)
-	{
-		Package->SetFileSize(InFileSize);
-	}
-
 	bool CloseLinkerArchives()
 	{
 		return GetHarvestedRealm().CloseLinkerArchives();
