Fixed a crash caused by a implicit conversion of a FSharedBuffer to a FIoHash.

diff --git a/Engine/Source/Runtime/CoreUObject/Private/Serialization/EditorBulkData.cpp b/Engine/Source/Runtime/CoreUObject/Private/Serialization/EditorBulkData.cpp
--- a/Engine/Source/Runtime/CoreUObject/Private/Serialization/EditorBulkData.cpp
+++ b/Engine/Source/Runtime/CoreUObject/Private/Serialization/EditorBulkData.cpp
@@ -1584,7 +1584,7 @@ void FEditorBulkData::UpdatePayload(FSharedBuffer InPayload)
 
 FEditorBulkData::FSharedBufferWithID::FSharedBufferWithID(FSharedBuffer InPayload)
 	: Payload(MoveTemp(InPayload))
-	, PayloadId(Payload)
+	, PayloadId(HashPayload(Payload))
 {
 }
 
