Ensure BPGC ARO their sparse class data

#localization none
#tests Verified a Linux server no longer crashes with this change

diff --git a/Engine/Source/Runtime/CoreUObject/Private/UObject/Class.cpp b/Engine/Source/Runtime/CoreUObject/Private/UObject/Class.cpp
--- a/Engine/Source/Runtime/CoreUObject/Private/UObject/Class.cpp
+++ b/Engine/Source/Runtime/CoreUObject/Private/UObject/Class.cpp
@@ -3710,27 +3710,32 @@ void UClass::AddReferencedObjects(UObject* InThis, FReferenceCollector& Collecto
 		This->CallAddReferencedObjects(This->ClassDefaultObject, Collector);
 	}
 
-	Collector.AddReferencedObject( This->SparseClassDataStruct, This );
+	This->AddSparseClassDataReferencedObjects(Collector);
+	
+	Super::AddReferencedObjects( This, Collector );
+}
+
+void UClass::AddSparseClassDataReferencedObjects(FReferenceCollector& Collector)
+{
+	Collector.AddReferencedObject(SparseClassDataStruct, this);
 
 	// Add sparse class data
-	if (This->SparseClassDataStruct && This->SparseClassData)
+	if (SparseClassDataStruct && SparseClassData)
 	{
-		if (This->SparseClassDataStruct->StructFlags & STRUCT_AddStructReferencedObjects)
+		if (SparseClassDataStruct->StructFlags & STRUCT_AddStructReferencedObjects)
 		{
-			This->SparseClassDataStruct->GetCppStructOps()->AddStructReferencedObjects()(This->SparseClassData, Collector);
+			SparseClassDataStruct->GetCppStructOps()->AddStructReferencedObjects()(SparseClassData, Collector);
 		}
 		else
 		{
 			// The iterator will recursively loop through all structs in structs too.
-			for (TPropertyValueIterator<FObjectProperty> It(This->SparseClassDataStruct, This->SparseClassData); It; ++It)
+			for (TPropertyValueIterator<FObjectProperty> It(SparseClassDataStruct, SparseClassData); It; ++It)
 			{
 				UObject** ObjectPtr = static_cast<UObject**>(const_cast<void*>(It.Value()));
 				Collector.AddReferencedObject(*ObjectPtr);
 			}
 		}
 	}
-	
-	Super::AddReferencedObjects( This, Collector );
 }
 
 /**
diff --git a/Engine/Source/Runtime/CoreUObject/Public/UObject/Class.h b/Engine/Source/Runtime/CoreUObject/Public/UObject/Class.h
--- a/Engine/Source/Runtime/CoreUObject/Public/UObject/Class.h
+++ b/Engine/Source/Runtime/CoreUObject/Public/UObject/Class.h
@@ -2856,6 +2856,8 @@ protected:
 
 	void CleanupSparseClassData();
 
+	void AddSparseClassDataReferencedObjects(FReferenceCollector& Collector);
+
 private:
 #if WITH_EDITOR
 	/** Provides access to attributes of the underlying C++ class. Should never be unset. */
diff --git a/Engine/Source/Runtime/Engine/Private/BlueprintGeneratedClass.cpp b/Engine/Source/Runtime/Engine/Private/BlueprintGeneratedClass.cpp
--- a/Engine/Source/Runtime/Engine/Private/BlueprintGeneratedClass.cpp
+++ b/Engine/Source/Runtime/Engine/Private/BlueprintGeneratedClass.cpp
@@ -1708,6 +1708,7 @@ void UBlueprintGeneratedClass::AddReferencedObjectsInUbergraphFrame(UObject* InT
 	{
 		if (UBlueprintGeneratedClass* BPGC = Cast<UBlueprintGeneratedClass>(CurrentClass))
 		{
+			BPGC->AddSparseClassDataReferencedObjects(Collector);
 #if USE_UBER_GRAPH_PERSISTENT_FRAME
 			if (BPGC->UberGraphFramePointerProperty)
 			{
