Allow Blueprint-added component archetypes to load when the underlying class type has a non-default ClassWithin value.

diff --git a/Engine/Source/Runtime/CoreUObject/Private/UObject/Obj.cpp b/Engine/Source/Runtime/CoreUObject/Private/UObject/Obj.cpp
--- a/Engine/Source/Runtime/CoreUObject/Private/UObject/Obj.cpp
+++ b/Engine/Source/Runtime/CoreUObject/Private/UObject/Obj.cpp
@@ -175,7 +175,7 @@ bool UObject::Rename( const TCHAR* InName, UObject* NewOuter, ERenameFlags Flags
 
 	// Check that we are not renaming a within object into an Outer of the wrong type, unless we're renaming the CDO of a Blueprint.
 	// Moving objects to the transient package is commonly used halfway through destroying them so that is also fine
-	if( NewOuter && !NewOuter->IsA(GetClass()->ClassWithin) && !HasAnyFlags(RF_ClassDefaultObject) && NewOuter != GetTransientPackage())	
+	if( NewOuter && !NewOuter->IsA(GetClass()->ClassWithin) && !HasAnyFlags(RF_ClassDefaultObject|RF_ArchetypeObject) && NewOuter != GetTransientPackage())	
 	{
 		if (Flags & REN_Test)
 		{
diff --git a/Engine/Source/Runtime/CoreUObject/Private/UObject/UObjectGlobals.cpp b/Engine/Source/Runtime/CoreUObject/Private/UObject/UObjectGlobals.cpp
--- a/Engine/Source/Runtime/CoreUObject/Private/UObject/UObjectGlobals.cpp
+++ b/Engine/Source/Runtime/CoreUObject/Private/UObject/UObjectGlobals.cpp
@@ -2377,7 +2377,7 @@ bool StaticAllocateObjectErrorTests( const UClass* InClass, UObject* InOuter, FN
 		}
 	}
 
-	if ( (InFlags & RF_ClassDefaultObject) == 0 )
+	if ( (InFlags & (RF_ClassDefaultObject|RF_ArchetypeObject)) == 0 )
 	{
 		if ( InOuter != NULL && !InOuter->IsA(InClass->ClassWithin) )
 		{
@@ -2430,12 +2430,13 @@ UObject* StaticAllocateObject
 	}
 #endif // WITH_EDITOR
 	bool bCreatingCDO = (InFlags & RF_ClassDefaultObject) != 0;
+	const bool bCreatingArchetype = (InFlags & RF_ArchetypeObject) != 0;
 
 	check(InClass);
 	check(GIsEditor || !FScopedAllowAbstractClassAllocation::IsDisallowedAbstractClass(InClass, InFlags)); // this is a warning in the editor, otherwise it is illegal to create an abstract class, except the CDO
 	check(InOuter || (InClass == UPackage::StaticClass() && InName != NAME_None)); // only packages can not have an outer, and they must be named explicitly
 	//checkf(InClass != UPackage::StaticClass() || !InOuter || bCreatingCDO, TEXT("Creating nested packages is not allowed: Outer=%s, Package=%s"), *GetNameSafe(InOuter), *InName.ToString());
-	check(bCreatingCDO || !InOuter || InOuter->IsA(InClass->ClassWithin));
+	check(bCreatingCDO || bCreatingArchetype || !InOuter || InOuter->IsA(InClass->ClassWithin));
 	checkf(!IsGarbageCollecting(), TEXT("Unable to create new object: %s %s.%s. Creating UObjects while Collecting Garbage is not allowed!"),
 		*GetNameSafe(InClass), *GetPathNameSafe(InOuter), *InName.ToString());
 
@@ -2607,7 +2608,6 @@ UObject* StaticAllocateObject
 	}
 
 	// If class is transient, non-archetype objects must be transient.
-	bool const bCreatingArchetype = (InFlags & RF_ArchetypeObject) != 0;
 	if ( !bCreatingCDO && InClass->HasAnyClassFlags(CLASS_Transient) && !bCreatingArchetype )
 	{
 		InFlags |= RF_Transient;
