UE-141088 - Crash occurs when reading Real values from None objects

By default, implicit cast instructions expect their source expression to evaluate correctly and return a non-null source property to cast. If we read a variable from a context that's null (eg: a null actor's field), then the expression is never evaluated, and a null property is set as the most recent property, which trips an internal check.

While the attempted null access on the context triggers an internal Blueprint exception, we don't have any stack unwinding behavior. As a result, it continues evaluation. Our only protection is to check the most recent property's validity after stepping through the stack and before performing the cast.

Our solution uses the new StepAndCheckRecentProperties convenience function to check the state of the most recent property. If null, we simply skip the cast and log the error.

diff --git a/Engine/Source/Runtime/CoreUObject/Private/UObject/ScriptCore.cpp b/Engine/Source/Runtime/CoreUObject/Private/UObject/ScriptCore.cpp
--- a/Engine/Source/Runtime/CoreUObject/Private/UObject/ScriptCore.cpp
+++ b/Engine/Source/Runtime/CoreUObject/Private/UObject/ScriptCore.cpp
@@ -3665,9 +3665,14 @@ IMPLEMENT_VM_FUNCTION( EX_ObjToInterfaceCast, execInterfaceCast );
 
 DEFINE_FUNCTION(UObject::execDoubleToFloatCast)
 {
-	Stack.Step(Stack.Object, nullptr);
-
-	DoubleToFloatCast(nullptr, Stack.MostRecentPropertyAddress, RESULT_PARAM);
+	if (Stack.StepAndCheckMostRecentProperty(Stack.Object, nullptr))
+	{
+		DoubleToFloatCast(nullptr, Stack.MostRecentPropertyAddress, RESULT_PARAM);
+	}
+	else
+	{
+		UE_LOG(LogScript, Verbose, TEXT("Cast failed: recent properties were null!"));
+	}
 }
 IMPLEMENT_CAST_FUNCTION( CST_DoubleToFloat, execDoubleToFloatCast )
 
@@ -3685,9 +3690,14 @@ IMPLEMENT_CAST_FUNCTION( CST_DoubleToFloatSet, execDoubleToFloatSetCast )
 
 DEFINE_FUNCTION(UObject::execFloatToDoubleCast)
 {
-	Stack.Step(Stack.Object, nullptr);
-
-	FloatToDoubleCast(nullptr, Stack.MostRecentPropertyAddress, RESULT_PARAM);
+	if (Stack.StepAndCheckMostRecentProperty(Stack.Object, nullptr))
+	{
+		FloatToDoubleCast(nullptr, Stack.MostRecentPropertyAddress, RESULT_PARAM);
+	}
+	else
+	{
+		UE_LOG(LogScript, Verbose, TEXT("Cast failed: recent properties were null!"));
+	}
 }
 IMPLEMENT_CAST_FUNCTION( CST_FloatToDouble, execFloatToDoubleCast )
 
@@ -3705,17 +3715,27 @@ IMPLEMENT_CAST_FUNCTION( CST_FloatToDoubleSet, execFloatToDoubleSetCast )
 
 DEFINE_FUNCTION(UObject::execVectorToVector3fCast)
 {
-	Stack.Step(Stack.Object, nullptr);
-
-	FloatingPointCast<FVector, FVector3f>(nullptr, Stack.MostRecentPropertyAddress, RESULT_PARAM);
+	if (Stack.StepAndCheckMostRecentProperty(Stack.Object, nullptr))
+	{
+		FloatingPointCast<FVector, FVector3f>(nullptr, Stack.MostRecentPropertyAddress, RESULT_PARAM);
+	}
+	else
+	{
+		UE_LOG(LogScript, Verbose, TEXT("Cast failed: recent properties were null!"));
+	}
 }
 IMPLEMENT_CAST_FUNCTION( CST_VectorToVector3f, execVectorToVector3fCast )
 
 DEFINE_FUNCTION(UObject::execVector3fToVectorCast)
 {
-	Stack.Step(Stack.Object, nullptr);
-
-	FloatingPointCast<FVector3f, FVector>(nullptr, Stack.MostRecentPropertyAddress, RESULT_PARAM);
+	if (Stack.StepAndCheckMostRecentProperty(Stack.Object, nullptr))
+	{
+		FloatingPointCast<FVector3f, FVector>(nullptr, Stack.MostRecentPropertyAddress, RESULT_PARAM);
+	}
+	else
+	{
+		UE_LOG(LogScript, Verbose, TEXT("Cast failed: recent properties were null!"));
+	}
 }
 IMPLEMENT_CAST_FUNCTION( CST_Vector3fToVector, execVector3fToVectorCast )
 
diff --git a/Engine/Source/Runtime/CoreUObject/Public/UObject/ScriptCastingUtils.h b/Engine/Source/Runtime/CoreUObject/Public/UObject/ScriptCastingUtils.h
--- a/Engine/Source/Runtime/CoreUObject/Public/UObject/ScriptCastingUtils.h
+++ b/Engine/Source/Runtime/CoreUObject/Public/UObject/ScriptCastingUtils.h
@@ -111,12 +111,17 @@ void CopyAndCastArrayFromStack(FFrame& Stack, RESULT_DECL)
 	FArrayProperty* DestinationArrayProperty = ExactCastField<FArrayProperty>(Stack.MostRecentProperty);
 	void* DestinationAddress = RESULT_PARAM;
 
-	Stack.Step(Stack.Object, nullptr);
-	
-	FArrayProperty* SourceArrayProperty = ExactCastField<FArrayProperty>(Stack.MostRecentProperty);
-	void* SourceAddress = Stack.MostRecentPropertyAddress;
+	if (Stack.StepAndCheckMostRecentProperty(Stack.Object, nullptr))
+	{
+		FArrayProperty* SourceArrayProperty = ExactCastField<FArrayProperty>(Stack.MostRecentProperty);
+		void* SourceAddress = Stack.MostRecentPropertyAddress;
 
-	CopyAndCastArray<SourceType, DestinationType>(SourceArrayProperty, SourceAddress, DestinationArrayProperty, DestinationAddress);
+		CopyAndCastArray<SourceType, DestinationType>(SourceArrayProperty, SourceAddress, DestinationArrayProperty, DestinationAddress);
+	}
+	else
+	{
+		UE_LOG(LogScript, Verbose, TEXT("Cast failed: recent properties were null!"));
+	}
 }
 
 template <typename SourceType, typename DestinationType>
@@ -127,24 +132,29 @@ void CopyAndCastSetFromStack(FFrame& Stack, RESULT_DECL)
 	checkSlow(RESULT_PARAM);
 	FScriptSetHelper DestinationSetHelper(DestinationSetProperty, RESULT_PARAM);
 
-	Stack.Step(Stack.Object, nullptr);
-
-	checkSlow(Stack.MostRecentProperty);
-	checkSlow(Stack.MostRecentPropertyAddress);
-
-	FSetProperty* SourceSetProperty = ExactCastField<FSetProperty>(Stack.MostRecentProperty);
-	FScriptSetHelper SourceSetHelper(SourceSetProperty, Stack.MostRecentPropertyAddress);
-
-	DestinationSetHelper.EmptyElements(SourceSetHelper.Num());
-	for (int32 i = 0; i < SourceSetHelper.Num(); ++i)
+	if (Stack.StepAndCheckMostRecentProperty(Stack.Object, nullptr))
 	{
-		int32 NewIndex = DestinationSetHelper.AddDefaultValue_Invalid_NeedsRehash();
-		void* DestinationRawData = DestinationSetHelper.GetElementPtr(NewIndex);
-		const void* SourceRawData = SourceSetHelper.GetElementPtr(i);
-
-		FloatingPointCast<SourceType, DestinationType>(nullptr, SourceRawData, DestinationRawData);
+		checkSlow(Stack.MostRecentProperty);
+		checkSlow(Stack.MostRecentPropertyAddress);
+
+		FSetProperty* SourceSetProperty = ExactCastField<FSetProperty>(Stack.MostRecentProperty);
+		FScriptSetHelper SourceSetHelper(SourceSetProperty, Stack.MostRecentPropertyAddress);
+
+		DestinationSetHelper.EmptyElements(SourceSetHelper.Num());
+		for (int32 i = 0; i < SourceSetHelper.Num(); ++i)
+		{
+			int32 NewIndex = DestinationSetHelper.AddDefaultValue_Invalid_NeedsRehash();
+			void* DestinationRawData = DestinationSetHelper.GetElementPtr(NewIndex);
+			const void* SourceRawData = SourceSetHelper.GetElementPtr(i);
+
+			FloatingPointCast<SourceType, DestinationType>(nullptr, SourceRawData, DestinationRawData);
+		}
+		DestinationSetHelper.Rehash();
+	}
+	else
+	{
+		UE_LOG(LogScript, Verbose, TEXT("Cast failed: recent properties were null!"));
 	}
-	DestinationSetHelper.Rehash();
 }
 
 template <CastFunction KeyCastFunction, CastFunction ValueCastFunction>
@@ -155,33 +165,38 @@ void CopyAndCastMapFromStack(FFrame& Stack, RESULT_DECL)
 
 	FScriptMapHelper DestinationMapHelper(DestinationMapProperty, RESULT_PARAM);
 
-	Stack.Step(Stack.Object, nullptr);
+	if (Stack.StepAndCheckMostRecentProperty(Stack.Object, nullptr))
+	{
+		FMapProperty* SourceMapProperty = ExactCastField<FMapProperty>(Stack.MostRecentProperty);
+		checkSlow(SourceMapProperty);
 
-	FMapProperty* SourceMapProperty = ExactCastField<FMapProperty>(Stack.MostRecentProperty);
-	checkSlow(SourceMapProperty);
+		FScriptMapHelper SourceMapHelper(SourceMapProperty, Stack.MostRecentPropertyAddress);
 
-	FScriptMapHelper SourceMapHelper(SourceMapProperty, Stack.MostRecentPropertyAddress);
+		const FProperty* SourceKeyProperty = SourceMapProperty->KeyProp;
+		checkSlow(SourceKeyProperty);
 
-	const FProperty* SourceKeyProperty = SourceMapProperty->KeyProp;
-	checkSlow(SourceKeyProperty);
+		const FProperty* SourceValueProperty = SourceMapProperty->ValueProp;
+		checkSlow(SourceValueProperty);
 
-	const FProperty* SourceValueProperty = SourceMapProperty->ValueProp;
-	checkSlow(SourceValueProperty);
+		DestinationMapHelper.EmptyValues(SourceMapHelper.Num());
+		for (int32 i = 0; i < SourceMapHelper.Num(); ++i)
+		{
+			int32 NewIndex = DestinationMapHelper.AddDefaultValue_Invalid_NeedsRehash();
 
-	DestinationMapHelper.EmptyValues(SourceMapHelper.Num());
-	for (int32 i = 0; i < SourceMapHelper.Num(); ++i)
-	{
-		int32 NewIndex = DestinationMapHelper.AddDefaultValue_Invalid_NeedsRehash();
+			const void* SourceKeyRawData = SourceMapHelper.GetKeyPtr(i);
+			void* DestinationKeyRawData = DestinationMapHelper.GetKeyPtr(NewIndex);
 
-		const void* SourceKeyRawData = SourceMapHelper.GetKeyPtr(i);
-		void* DestinationKeyRawData = DestinationMapHelper.GetKeyPtr(NewIndex);
-		
-		KeyCastFunction(SourceKeyProperty, SourceKeyRawData, DestinationKeyRawData);
+			KeyCastFunction(SourceKeyProperty, SourceKeyRawData, DestinationKeyRawData);
 
-		const void* SourceValueRawData = SourceMapHelper.GetValuePtr(i);
-		void* DestinationValueRawData = DestinationMapHelper.GetValuePtr(NewIndex);
+			const void* SourceValueRawData = SourceMapHelper.GetValuePtr(i);
+			void* DestinationValueRawData = DestinationMapHelper.GetValuePtr(NewIndex);
 
-		ValueCastFunction(SourceValueProperty, SourceValueRawData, DestinationValueRawData);
+			ValueCastFunction(SourceValueProperty, SourceValueRawData, DestinationValueRawData);
+		}
+		DestinationMapHelper.Rehash();
+	}
+	else
+	{
+		UE_LOG(LogScript, Verbose, TEXT("Cast failed: recent properties were null!"));
 	}
-	DestinationMapHelper.Rehash();
 }
diff --git a/Engine/Source/Runtime/CoreUObject/Public/UObject/Stack.h b/Engine/Source/Runtime/CoreUObject/Public/UObject/Stack.h
--- a/Engine/Source/Runtime/CoreUObject/Public/UObject/Stack.h
+++ b/Engine/Source/Runtime/CoreUObject/Public/UObject/Stack.h
@@ -133,7 +133,10 @@ public:
 	}
 
 	// Functions.
-	COREUOBJECT_API void Step( UObject* Context, RESULT_DECL );
+	COREUOBJECT_API void Step(UObject* Context, RESULT_DECL);
+
+	/** Convenience function that calls Step, but also returns true if both MostRecentProperty and MostRecentPropertyAddress are non-null. */
+	FORCEINLINE_DEBUGGABLE bool StepAndCheckMostRecentProperty(UObject* Context, RESULT_DECL);
 
 	/** Replacement for Step that uses an explicitly specified property to unpack arguments **/
 	COREUOBJECT_API void StepExplicitProperty(void*const Result, FProperty* Property);
@@ -340,6 +343,13 @@ inline FName FFrame::ReadName()
 
 COREUOBJECT_API void GInitRunaway();
 
+FORCEINLINE_DEBUGGABLE bool FFrame::StepAndCheckMostRecentProperty(UObject* Context, RESULT_DECL)
+{
+	Step(Context, RESULT_PARAM);
+
+	return (MostRecentProperty && MostRecentPropertyAddress);
+}
+
 /**
  * Replacement for Step that checks the for byte code, and if none exists, then PropertyChainForCompiledIn is used.
  * Also makes an effort to verify that the params are in the correct order and the types are compatible.
