Fix some data truncation (C4244) warnings

#rnx

- Algo::CountIf returns SIZE_T but the container size type is int32, so it is safe to cast.

diff --git a/Engine/Source/Runtime/CoreUObject/Private/UObject/PackageTrailer.cpp b/Engine/Source/Runtime/CoreUObject/Private/UObject/PackageTrailer.cpp
--- a/Engine/Source/Runtime/CoreUObject/Private/UObject/PackageTrailer.cpp
+++ b/Engine/Source/Runtime/CoreUObject/Private/UObject/PackageTrailer.cpp
@@ -592,15 +592,15 @@ int32 FPackageTrailer::GetNumPayloads(EPayloadFilter Type) const
 			break;
 
 		case EPayloadFilter::Local:
-			Count = Algo::CountIf(Header.PayloadLookupTable, [](const Private::FLookupTableEntry& Entry) { return Entry.AccessMode == EPayloadAccessMode::Local; });
+			Count = (int32)Algo::CountIf(Header.PayloadLookupTable, [](const Private::FLookupTableEntry& Entry) { return Entry.AccessMode == EPayloadAccessMode::Local; });
 			break;
 
 		case EPayloadFilter::Referenced:
-			Count = Algo::CountIf(Header.PayloadLookupTable, [](const Private::FLookupTableEntry& Entry) { return Entry.AccessMode == EPayloadAccessMode::Referenced; });
+			Count = (int32)Algo::CountIf(Header.PayloadLookupTable, [](const Private::FLookupTableEntry& Entry) { return Entry.AccessMode == EPayloadAccessMode::Referenced; });
 			break;
 
 		case EPayloadFilter::Virtualized:
-			Count = Algo::CountIf(Header.PayloadLookupTable, [](const Private::FLookupTableEntry& Entry) { return Entry.AccessMode == EPayloadAccessMode::Virtualized; });
+			Count = (int32)Algo::CountIf(Header.PayloadLookupTable, [](const Private::FLookupTableEntry& Entry) { return Entry.AccessMode == EPayloadAccessMode::Virtualized; });
 			break;
 
 		default:
