Add unit tests for UE::Serialization::FEditorBulkData::FSharedBufferWithID

#rnx

- FSharedBufferWithID allows the caller to calculate the FIoHash of a FSharedBuffer at their own leisure (usually on a async task) and then updated a FEditorBulkData object with it. This allows the caller to control when the potentially expensive hashing is done.
- These tests should make sure that any use case that can work with FSharedBufferWithID will have the same results with FEditorBulkData as if they had just called ::UPdatePayload with the raw FSharedBuffer instead.

diff --git a/Engine/Source/Runtime/CoreUObject/Private/Tests/Serialization/EditorBulkDataTests.cpp b/Engine/Source/Runtime/CoreUObject/Private/Tests/Serialization/EditorBulkDataTests.cpp
--- a/Engine/Source/Runtime/CoreUObject/Private/Tests/Serialization/EditorBulkDataTests.cpp
+++ b/Engine/Source/Runtime/CoreUObject/Private/Tests/Serialization/EditorBulkDataTests.cpp
@@ -544,6 +544,72 @@ bool FEditorBulkDataIoHashToGuid::RunTest(const FString& Parameters)
 	return true;
 }
 
+/** This tests that updating a payload via a FSharedBufferWithID will have the same results as if the payload was applied to the bulkdata object directly */
+IMPLEMENT_SIMPLE_AUTOMATION_TEST(FEditorBulkDataSharedBufferWithID, TEXT("System.CoreUObject.Serialization.EditorBulkData.SharedBufferWithID"), TestFlags)
+bool FEditorBulkDataSharedBufferWithID::RunTest(const FString& Parameters)
+{
+	// Test FSharedBufferWithID/FEditorBulkData in their default states
+	{
+		FEditorBulkData::FSharedBufferWithID SharedId;
+		FEditorBulkData BulkDataFromSharedId;
+		BulkDataFromSharedId.UpdatePayload(MoveTemp(SharedId));
+
+		FEditorBulkData BulkData;
+
+		TestEqual(TEXT("0: PayloadId"), BulkDataFromSharedId.GetPayloadId(), BulkData.GetPayloadId());
+		TestEqual(TEXT("0: PayloadSize"), BulkDataFromSharedId.GetPayloadSize(), BulkData.GetPayloadSize());
+	}
+
+	// Test FSharedBufferWithID/FEditorBulkData updated with a null FSharedBuffer
+	{
+		FSharedBuffer NullBuffer;
+
+		FEditorBulkData::FSharedBufferWithID SharedId(NullBuffer);
+		FEditorBulkData BulkDataFromSharedId;
+		BulkDataFromSharedId.UpdatePayload(MoveTemp(SharedId));
+
+		FEditorBulkData BulkData;
+		BulkData.UpdatePayload(NullBuffer);
+
+		TestEqual(TEXT("1: PayloadId"), BulkDataFromSharedId.GetPayloadId(), BulkData.GetPayloadId());
+		TestEqual(TEXT("1: PayloadSize"), BulkDataFromSharedId.GetPayloadSize(), BulkData.GetPayloadSize());
+	}
+
+	// Test FSharedBufferWithID/FEditorBulkData updated with a zero length FSharedBuffer
+	{
+		FSharedBuffer ZeroLengthBuffer = FUniqueBuffer::Alloc(0).MoveToShared();
+
+		FEditorBulkData::FSharedBufferWithID SharedId(ZeroLengthBuffer);
+		FEditorBulkData BulkDataFromSharedId;
+		BulkDataFromSharedId.UpdatePayload(MoveTemp(SharedId));
+
+		FEditorBulkData BulkData;
+		BulkData.UpdatePayload(ZeroLengthBuffer);
+
+		TestEqual(TEXT("2: PayloadId"), BulkDataFromSharedId.GetPayloadId(), BulkData.GetPayloadId());
+		TestEqual(TEXT("2: PayloadSize"), BulkDataFromSharedId.GetPayloadSize(), BulkData.GetPayloadSize());
+	}
+
+	// Test FSharedBufferWithID/FEditorBulkData updated with a random set of data
+	{
+		const int64 BufferSize = 1024;
+		TUniquePtr<uint8[]> SourceBuffer = CreateRandomData(BufferSize);
+		FSharedBuffer RandomData = FSharedBuffer::MakeView(SourceBuffer.Get(), BufferSize);
+
+		FEditorBulkData::FSharedBufferWithID SharedId(RandomData);
+		FEditorBulkData BulkDataFromSharedId;
+		BulkDataFromSharedId.UpdatePayload(MoveTemp(SharedId));
+
+		FEditorBulkData BulkData;
+		BulkData.UpdatePayload(RandomData);
+
+		TestEqual(TEXT("3: PayloadId"), BulkDataFromSharedId.GetPayloadId(), BulkData.GetPayloadId());
+		TestEqual(TEXT("3: PayloadSize"), BulkDataFromSharedId.GetPayloadSize(), BulkData.GetPayloadSize());
+	}
+
+	return true;
+}
+
 
 } // namespace UE::Serialization
 
