Consider all properties within an editor-only class to be editor-only when gathering localization

diff --git a/Engine/Source/Runtime/CoreUObject/Private/Serialization/PropertyLocalizationDataGathering.cpp b/Engine/Source/Runtime/CoreUObject/Private/Serialization/PropertyLocalizationDataGathering.cpp
--- a/Engine/Source/Runtime/CoreUObject/Private/Serialization/PropertyLocalizationDataGathering.cpp
+++ b/Engine/Source/Runtime/CoreUObject/Private/Serialization/PropertyLocalizationDataGathering.cpp
@@ -177,21 +177,27 @@ bool FPropertyLocalizationDataGatherer::CanGatherFromInnerProperty(const FProper
 
 void FPropertyLocalizationDataGatherer::GatherLocalizationDataFromObjectWithCallbacks(const UObject* Object, const EPropertyLocalizationGathererTextFlags GatherTextFlags)
 {
+	EPropertyLocalizationGathererTextFlags ObjectGatherTextFlags = GatherTextFlags;
+	if (Object->GetClass()->GetPackage()->HasAnyPackageFlags(PKG_EditorOnly))
+	{
+		ObjectGatherTextFlags |= EPropertyLocalizationGathererTextFlags::ForceEditorOnlyProperties;
+	}
+
 	const FGatherableFieldsForType& GatherableFieldsForType = GetGatherableFieldsForType(Object->GetClass());
 	if (GatherableFieldsForType.CustomObjectCallback)
 	{
 		checkf(IsObjectValidForGather(Object), TEXT("Cannot gather for objects outside of the current package! Package: '%s'. Object: '%s'."), *Package->GetFullName(), *Object->GetFullName());
 
-		if (ShouldProcessObject(Object, GatherTextFlags))
+		if (ShouldProcessObject(Object, ObjectGatherTextFlags))
 		{
-			MarkObjectProcessed(Object, GatherTextFlags);
-			(*GatherableFieldsForType.CustomObjectCallback)(Object, *this, GatherTextFlags);
+			MarkObjectProcessed(Object, ObjectGatherTextFlags);
+			(*GatherableFieldsForType.CustomObjectCallback)(Object, *this, ObjectGatherTextFlags);
 		}
 	}
-	else if (ShouldProcessObject(Object, GatherTextFlags))
+	else if (ShouldProcessObject(Object, ObjectGatherTextFlags))
 	{
-		MarkObjectProcessed(Object, GatherTextFlags);
-		GatherLocalizationDataFromObject(Object, GatherTextFlags);
+		MarkObjectProcessed(Object, ObjectGatherTextFlags);
+		GatherLocalizationDataFromObject(Object, ObjectGatherTextFlags);
 	}
 }
 
