Removing overly verbose CSV event as it didn't give us much more detail.

diff --git a/Engine/Source/Runtime/CoreUObject/Public/UObject/FastReferenceCollector.h b/Engine/Source/Runtime/CoreUObject/Public/UObject/FastReferenceCollector.h
--- a/Engine/Source/Runtime/CoreUObject/Public/UObject/FastReferenceCollector.h
+++ b/Engine/Source/Runtime/CoreUObject/Public/UObject/FastReferenceCollector.h
@@ -700,8 +700,6 @@ private:
 	 */
 	void ProcessObjectArray(FGCArrayStruct& InObjectsToSerializeStruct, const FGraphEventRef& MyCompletionGraphEvent)
 	{
-		DECLARE_SCOPE_CYCLE_COUNTER(TEXT("TFastReferenceCollector::ProcessObjectArray"), STAT_FFastReferenceCollector_ProcessObjectArray, STATGROUP_GC);
-
 		UObject* CurrentObject = nullptr;
 
 		const int32 MinDesiredObjectsPerSubTask = ReferenceProcessor.GetMinDesiredObjectsPerSubTask(); // sometimes there will be less, a lot less
