Unshelved from Marc.Audy

Avoid IsInBlueprint returning true for Python objects

diff --git a/Engine/Source/Runtime/CoreUObject/Private/Blueprint/BlueprintSupport.cpp b/Engine/Source/Runtime/CoreUObject/Private/Blueprint/BlueprintSupport.cpp
--- a/Engine/Source/Runtime/CoreUObject/Private/Blueprint/BlueprintSupport.cpp
+++ b/Engine/Source/Runtime/CoreUObject/Private/Blueprint/BlueprintSupport.cpp
@@ -2388,13 +2388,13 @@ bool UObject::IsInBlueprint() const
  	while (TestObject)
  	{
  		const UClass *ClassObject = dynamic_cast<const UClass*>(TestObject);
-		if (ClassObject && !ClassObject->IsNative())
+		if (ClassObject && ClassObject->HasAnyClassFlags(CLASS_CompiledFromBlueprint))
  		{
  			return true;
  		}
 		else if (TestObject->HasAnyFlags(RF_ClassDefaultObject) 
 			&& TestObject->GetClass() 
-			&& !TestObject->GetClass()->IsNative())
+			&& TestObject->GetClass()->HasAnyClassFlags(CLASS_CompiledFromBlueprint))
  		{
  			return true;
  		}
