AsyncLoading2: Fixed Async and AsyncLoading flags sometimes not being cleared
#rnx

diff --git a/Engine/Source/Runtime/CoreUObject/Private/Serialization/AsyncLoading2.cpp b/Engine/Source/Runtime/CoreUObject/Private/Serialization/AsyncLoading2.cpp
--- a/Engine/Source/Runtime/CoreUObject/Private/Serialization/AsyncLoading2.cpp
+++ b/Engine/Source/Runtime/CoreUObject/Private/Serialization/AsyncLoading2.cpp
@@ -348,6 +348,7 @@ struct FExportObject
 	UObject* SuperObject = nullptr;
 	bool bFiltered = false;
 	bool bExportLoadFailed = false;
+	bool bWasFoundInMemory = false;
 };
 
 struct FPackageRequest
@@ -1612,7 +1613,6 @@ struct FAsyncPackage2
 		}
 		else
 		{
-			checkf(!ConstructedObjects.Contains(Object), TEXT("%s"), *Object->GetFullName());
 			ConstructedObjects.Add(Object);
 		}
 	}
@@ -3825,6 +3825,11 @@ void FAsyncPackage2::EventDrivenCreateExport(int32 LocalExportIndex)
 	// Object is found in memory.
 	if (Object)
 	{
+		// If it has the AsyncLoading flag set it was created during the current load of this package (likely as a subobject)
+		if (!Object->HasAnyInternalFlags(EInternalObjectFlags::AsyncLoading))
+		{
+			ExportObject.bWasFoundInMemory = true;
+		}
 		// If this object was allocated but never loaded (components created by a constructor, CDOs, etc) make sure it gets loaded
 		// Do this for all subobjects created in the native constructor.
 		const EObjectFlags ObjectFlags = Object->GetFlags();
@@ -5589,27 +5594,25 @@ void FAsyncPackage2::ClearConstructedObjects()
 
 	for (UObject* Object : ConstructedObjects)
 	{
-		if (Object->HasAnyFlags(RF_WasLoaded))
-		{
-			// exports and the upackage itself are are handled below
-			continue;
-		}
 		Object->AtomicallyClearInternalFlags(EInternalObjectFlags::AsyncLoading | EInternalObjectFlags::Async);
 	}
 	ConstructedObjects.Empty();
 
 	for (FExportObject& Export : Data.Exports)
 	{
-		if (Export.bFiltered | Export.bExportLoadFailed)
+		if (Export.bWasFoundInMemory)
 		{
-			continue;
+			check(Export.Object);
+			Export.Object->AtomicallyClearInternalFlags(EInternalObjectFlags::AsyncLoading | EInternalObjectFlags::Async);
+		}
+		else
+		{
+			checkf(!Export.Object || !Export.Object->HasAnyInternalFlags(EInternalObjectFlags::AsyncLoading | EInternalObjectFlags::Async),
+				TEXT("Export object: %s (ObjectFlags=%x, InternalObjectFlags=%x)"),
+					*Export.Object->GetFullName(),
+					Export.Object->GetFlags(),
+					Export.Object->GetInternalFlags());
 		}
-
-		UObject* Object = Export.Object;
-		check(Object);
-		checkf(Object->HasAnyFlags(RF_WasLoaded), TEXT("%s"), *Object->GetFullName());
-		checkf(Object->HasAnyInternalFlags(EInternalObjectFlags::Async), TEXT("%s"), *Object->GetFullName());
-		Object->AtomicallyClearInternalFlags(EInternalObjectFlags::AsyncLoading | EInternalObjectFlags::Async);
 	}
 
 	if (LinkerRoot)
