Fix assert when detaching a packages linker when the package file is still being used on another thread.

#rnx

- In the last year we have encountered an increasing number of problems where an asset is processing it's bulkdata on an async thread when it is garbage collected.
- This is caused because we detach the package's linker in UPackage::BeginDestroy but most async asset tasks are checked in UObject::IsReadyForFinishDestroy, meaning we can have async work trying to access the bulkdata in the package file up until UPackage::FinishDestroy is called
-- Usually we try and fix this problem by preventing  packages with outstanding async tasks from being garbage collected, but support for this is not universal and can lead to the above problems showing up in various commandlets.
- By moving the detachment of the packages linker to UPackage::FinishDestroy we should be able to prevent these problems in all cases.

diff --git a/Engine/Source/Runtime/CoreUObject/Private/UObject/Package.cpp b/Engine/Source/Runtime/CoreUObject/Private/UObject/Package.cpp
--- a/Engine/Source/Runtime/CoreUObject/Private/UObject/Package.cpp
+++ b/Engine/Source/Runtime/CoreUObject/Private/UObject/Package.cpp
@@ -302,9 +302,12 @@ bool UPackage::IsFullyLoaded() const
 	return false;
 }
 
-void UPackage::BeginDestroy()
+void UPackage::FinishDestroy()
 {
-	// Detach linker if still attached
+	// Detach linker if still attached, we do this in ::FinishDestroy rather than ::BeginDestroy so that the linker remains attached
+	// and valid for all UObjects in the package until they have all returned ::IsReadyForFinishDestroy as true. This means that 
+	// UObjects with ongoing asynchronous compilation work can safely cancel that work in ::BeginDestroy and wait for it to finish
+	// in ::IsReadyForFinishDestroy without worrying that the package file will be yanked out from under it.
 	if (FLinkerLoad* Linker = GetLinker())
 	{
 		// Detach() below will most likely null the LinkerLoad so keep a temp copy so that we can still call RemoveLinker on it
@@ -313,7 +316,7 @@ void UPackage::BeginDestroy()
 		SetLinker(nullptr);
 	}
 
-	Super::BeginDestroy();
+	Super::FinishDestroy();
 }
 
 bool UPackage::IsPostLoadThreadSafe() const
diff --git a/Engine/Source/Runtime/CoreUObject/Public/UObject/Package.h b/Engine/Source/Runtime/CoreUObject/Public/UObject/Package.h
--- a/Engine/Source/Runtime/CoreUObject/Public/UObject/Package.h
+++ b/Engine/Source/Runtime/CoreUObject/Public/UObject/Package.h
@@ -309,7 +309,7 @@ public:
 	*/
 	virtual void PostInitProperties() override;
 
-	virtual void BeginDestroy() override;	
+	virtual void FinishDestroy() override;
 
 	/** Serializer */
 	virtual void Serialize( FArchive& Ar ) override;
