Fix mismatch in memory allocation / deletion routines

#rnx

- The TUniquePtr would allocate the buffer with new[] and we were using FMemory::Free to delete it.
- This works under our default set up thanks to REPLACEMENT_OPERATOR_NEW_AND_DELETE but it is unsafe to leave it.
- Changed the call to FSharedBuffer::TakeOwnership so that the payload will be destroyed via delete []

diff --git a/Engine/Source/Runtime/CoreUObject/Private/Tests/Serialization/EditorBulkDataTests.cpp b/Engine/Source/Runtime/CoreUObject/Private/Tests/Serialization/EditorBulkDataTests.cpp
--- a/Engine/Source/Runtime/CoreUObject/Private/Tests/Serialization/EditorBulkDataTests.cpp
+++ b/Engine/Source/Runtime/CoreUObject/Private/Tests/Serialization/EditorBulkDataTests.cpp
@@ -172,7 +172,7 @@ bool FEditorBulkDataTestUpdatePayload::RunTest(const FString& Parameters)
 
 		// Update the bulkdata object with the original data but this time we give ownership of the buffer to
 		// the bulkdata object.
-		BulkData.UpdatePayload(FSharedBuffer::TakeOwnership(OriginalData.Release(), BufferSize, FMemory::Free));
+		BulkData.UpdatePayload(FSharedBuffer::TakeOwnership(OriginalData.Release(), BufferSize, [](void* Ptr, uint64) { delete[] (uint8*)Ptr; }));
 
 		FSharedBuffer Payload = BulkData.GetPayload().Get();
 		TestEqual(TEXT("Updated payload length"), (int64)Payload.GetSize(), BufferSize);
