Added support for override package in package text serialization:
- Text format now includes import and export tables instead of discovering import while serializating exports. This is required to have enough information about objects with an override package that have their Outer in another package.
- UObject Json formatters now export import/export indices instead of relying on the objects full paths, which index into the exported import/export tables.
- Merged code to serialize UObject* and FObjectPtr in Json formatters to avoid code duplication.

#rnx

diff --git a/Engine/Source/Runtime/CoreUObject/Private/Serialization/Formatters/JsonArchiveInputFormatter.cpp b/Engine/Source/Runtime/CoreUObject/Private/Serialization/Formatters/JsonArchiveInputFormatter.cpp
--- a/Engine/Source/Runtime/CoreUObject/Private/Serialization/Formatters/JsonArchiveInputFormatter.cpp
+++ b/Engine/Source/Runtime/CoreUObject/Private/Serialization/Formatters/JsonArchiveInputFormatter.cpp
@@ -14,9 +14,9 @@
 
 #if WITH_TEXT_ARCHIVE_SUPPORT
 
-FJsonArchiveInputFormatter::FJsonArchiveInputFormatter(FArchive& InInner, TFunction<UObject* (const FString&)> InResolveObjectName) 
+FJsonArchiveInputFormatter::FJsonArchiveInputFormatter(FArchive& InInner, TFunction<UObject* (const FPackageIndex)> InResolveObject) 
 	: Inner(InInner)
-	, ResolveObjectName(InResolveObjectName)
+	, ResolveObject(InResolveObject)
 {
 	Inner.SetIsTextFormat(true);
 	Inner.ArAllowLazyLoading = false;
@@ -378,10 +378,11 @@ void FJsonArchiveInputFormatter::Serialize(FName& Value)
 void FJsonArchiveInputFormatter::Serialize(UObject*& Value)
 {
 	FString StringValue;
-	const TCHAR Prefix[] = TEXT("Object:");
-	if (ValueStack.Top()->TryGetString(StringValue) && StringValue.StartsWith(Prefix))
+	if (ValueStack.Top()->TryGetString(StringValue))
 	{
-		Value = ResolveObjectName(*StringValue + UE_ARRAY_COUNT(Prefix) - 1);
+		FPackageIndex ObjectIndex;
+		LexFromString(ObjectIndex, *StringValue);
+		Value = ResolveObject(ObjectIndex);
 	}
 	else
 	{
@@ -442,16 +443,9 @@ void FJsonArchiveInputFormatter::Serialize(FLazyObjectPtr& Value)
 
 void FJsonArchiveInputFormatter::Serialize(FObjectPtr& Value)
 {
-	FString StringValue;
-	const TCHAR Prefix[] = TEXT("Object:");
-	if (ValueStack.Top()->TryGetString(StringValue) && StringValue.StartsWith(Prefix))
-	{
-		Value = ResolveObjectName(*StringValue + UE_ARRAY_COUNT(Prefix) - 1);
-	}
-	else
-	{
-		Value = nullptr;
-	}
+	UObject* Object;
+	Serialize(Object);
+	Value = Object;
 }
 
 void FJsonArchiveInputFormatter::Serialize(TArray<uint8>& Data)
diff --git a/Engine/Source/Runtime/CoreUObject/Private/Serialization/Formatters/JsonArchiveOutputFormatter.cpp b/Engine/Source/Runtime/CoreUObject/Private/Serialization/Formatters/JsonArchiveOutputFormatter.cpp
--- a/Engine/Source/Runtime/CoreUObject/Private/Serialization/Formatters/JsonArchiveOutputFormatter.cpp
+++ b/Engine/Source/Runtime/CoreUObject/Private/Serialization/Formatters/JsonArchiveOutputFormatter.cpp
@@ -345,8 +345,8 @@ void FJsonArchiveOutputFormatter::Serialize(UObject*& Value)
 {
 	if (Value != nullptr && IsObjectAllowed(Value))
 	{
-		FString FullObjectName = Value->GetFullName(nullptr, EObjectFullNameFlags::IncludeClassPackage);
-		SerializeStringInternal(FString::Printf(TEXT("Object:%s"), *FullObjectName));
+		FPackageIndex ObjectIndex = ObjectIndicesMap->FindChecked(Value);
+		SerializeStringInternal(LexToString(ObjectIndex));
 	}
 	else
 	{
@@ -405,15 +405,8 @@ void FJsonArchiveOutputFormatter::Serialize(FLazyObjectPtr& Value)
 
 void FJsonArchiveOutputFormatter::Serialize(FObjectPtr& Value)
 {
-	if (!Value.IsNull() && IsObjectAllowed(Value.Get()))
-	{
-		FString FullObjectName = Value->GetFullName(nullptr, EObjectFullNameFlags::IncludeClassPackage);
-		SerializeStringInternal(FString::Printf(TEXT("Object:%s"), *FullObjectName));
-	}
-	else
-	{
-		WriteValue(TEXT("null"));
-	}
+	UObject* Object = Value.Get();
+	Serialize(Object);
 }
 
 void FJsonArchiveOutputFormatter::Serialize(TArray<uint8>& Data)
@@ -615,7 +608,7 @@ void FJsonArchiveOutputFormatter::SerializeStringInternal(const FString& String)
 
 bool FJsonArchiveOutputFormatter::IsObjectAllowed(UObject* InObject) const
 {
-	return ObjectIndicesMap == nullptr || ObjectIndicesMap->Contains(InObject);
+	return ObjectIndicesMap && ObjectIndicesMap->Contains(InObject);
 }
 
 #endif
diff --git a/Engine/Source/Runtime/CoreUObject/Private/UObject/LinkerLoad.cpp b/Engine/Source/Runtime/CoreUObject/Private/UObject/LinkerLoad.cpp
--- a/Engine/Source/Runtime/CoreUObject/Private/UObject/LinkerLoad.cpp
+++ b/Engine/Source/Runtime/CoreUObject/Private/UObject/LinkerLoad.cpp
@@ -758,11 +758,11 @@ FLinkerLoad::ELinkerStatus FLinkerLoad::Tick( float InTimeLimit, bool bInUseTime
 			}
 
 #if WITH_TEXT_ARCHIVE_SUPPORT
-			// Reconstruct the import and export maps for text assets
+			// Construct the exports readers
 			if (Status == LINKER_Loaded)
 			{
-				SCOPED_LOADTIMER(LinkerLoad_ReconstructImportAndExportMap);
-				Status = ReconstructImportAndExportMap();
+				SCOPED_LOADTIMER(LinkerLoad_ConstructExportsReaders);
+				Status = ConstructExportsReaders();
 			}
 #endif
 
@@ -875,7 +875,7 @@ FLinkerLoad::FLinkerLoad(UPackage* InParent, const FPackagePath& InPackagePath,
 , ExportHashIndex(0)
 , bHasSerializedPackageFileSummary(false)
 , bHasSerializedPackageTrailer(false)
-, bHasReconstructedImportAndExportMap(false)
+, bHasConstructedExportsReaders(false)
 , bHasSerializedPreloadDependencies(false)
 , bHasFixedUpImportMap(false)
 , bHasPopulatedInstancingContext(false)
@@ -1087,9 +1087,8 @@ FLinkerLoad::ELinkerStatus FLinkerLoad::CreateLoader(
 				INC_DWORD_STAT(STAT_TextAssetLinkerCount);
 				DECLARE_SCOPE_CYCLE_COUNTER(TEXT("FLinkerLoad::CreateTextArchiveFormatter"), STAT_LinkerLoad_CreateTextArchiveFormatter, STATGROUP_LinkerLoad);
 				TRACE_CPUPROFILER_EVENT_SCOPE(FLinkerLoad::CreateTextArchiveFormatter);
-				StructuredArchiveFormatter = new FJsonArchiveInputFormatter(*this, [this](const FString& InFullPath)
+				StructuredArchiveFormatter = new FJsonArchiveInputFormatter(*this, [this](const FPackageIndex Index)
 				{
-					FPackageIndex Index = FindOrCreateImportOrExport(InFullPath);
 					if (Index.IsNull())
 					{
 						return (UObject*)nullptr;
@@ -1658,12 +1657,7 @@ FLinkerLoad::ELinkerStatus FLinkerLoad::SerializeImportMap()
 {
 	DECLARE_SCOPE_CYCLE_COUNTER( TEXT( "FLinkerLoad::SerializeImportMap" ), STAT_LinkerLoad_SerializeImportMap, STATGROUP_LinkerLoad );
 
-	if (IsTextFormat())
-	{
-		return LINKER_Loaded;
-	}
-
-	if(ImportMapIndex == 0 && Summary.ImportCount > 0 )
+	if(!IsTextFormat() && ImportMapIndex == 0 && Summary.ImportCount > 0 )
 	{
 		Seek( Summary.ImportOffset );
 	}
@@ -1938,12 +1932,7 @@ FLinkerLoad::ELinkerStatus FLinkerLoad::SerializeExportMap()
 {
 	DECLARE_SCOPE_CYCLE_COUNTER( TEXT( "FLinkerLoad::SerializeExportMap" ), STAT_LinkerLoad_SerializeExportMap, STATGROUP_LinkerLoad );
 
-	if (IsTextFormat())
-	{
-		return LINKER_Loaded;
-	}
-
-	if(ExportMapIndex == 0 && Summary.ExportCount > 0)
+	if(!IsTextFormat() && ExportMapIndex == 0 && Summary.ExportCount > 0)
 	{
 		Seek( Summary.ExportOffset );
 	}
@@ -1972,30 +1961,6 @@ FStructuredArchiveSlot FLinkerLoad::GetExportSlot(FPackageIndex InExportIndex)
 	return ExportReaders[Index]->GetRoot();
 }
 
-FPackageIndex FLinkerLoad::FindOrCreateImport(const FName InObjectName, const FName InClassName, const FName InClassPackageName)
-{
-	for (int32 ImportIndex = 0; ImportIndex < ImportMap.Num(); ++ImportIndex)
-	{
-		FObjectImport& ExistingImport = ImportMap[ImportIndex];
-		if (ExistingImport.ObjectName == InObjectName && ExistingImport.ClassPackage == InClassPackageName && ExistingImport.ClassName == InClassName)
-		{
-			return FPackageIndex::FromImport(ImportIndex);
-		}
-	}
-
-	FObjectImport& NewImport = ImportMap.Emplace_GetRef();
-	NewImport.ObjectName = InObjectName;
-	NewImport.ClassName = InClassName;
-	NewImport.ClassPackage = InClassPackageName;
-	NewImport.SourceIndex = INDEX_NONE;
-	NewImport.XObject = nullptr;
-	NewImport.SourceLinker = nullptr;
-
-	Summary.ImportCount++;
-
-	return FPackageIndex::FromImport(ImportMap.Num() - 1);
-}
-
 #if UE_WITH_OBJECT_HANDLE_LATE_RESOLVE
 FLinkerLoad::EImportLoadBehavior FLinkerLoad::ParseImportLoadBehavior(const FString* LoadBehaviorMeta)
 {
@@ -2102,127 +2067,23 @@ FString ExtractObjectName(const FString& InFullPath)
 	return ObjectName;
 }
 
-FPackageIndex FLinkerLoad::FindOrCreateImportOrExport(const FString& InFullPath)
-{
-	if (InFullPath.Len() == 0)
-	{
-		return FPackageIndex();
-	}
-
-	FString Class, Package, Object, SubObject;
-	FPackageName::SplitFullObjectPath(InFullPath, Class, Package, Object, SubObject);
-	FName ObjectName = *(Object + (SubObject.Len() ? TEXT(":") : TEXT("")) + SubObject);
-
-	bool bIsExport = Package == LinkerRoot->GetName();
-
-	if (bIsExport)
-	{
-		for (int32 ExportIndex = 0; ExportIndex < ExportMap.Num(); ++ExportIndex)
-		{
-			if (OriginalExportNames[ExportIndex] == ObjectName)
-			{
-				return FPackageIndex::FromExport(ExportIndex);
-			}
-		}
-
-		return FPackageIndex();
-	}
-	else
-	{
-		FName ClassName = *Class;
-		FName ClassPackageName;
-		FName PackageName = *Package;
-
-		if (Class.StartsWith(TEXT("/")))
-		{
-			int32 EndOfPackage = INDEX_NONE;
-			if (Class.FindChar('.', EndOfPackage))
-			{
-				ClassPackageName = *FString(EndOfPackage, *Class);
-				ClassName = *FString(Class.Len() - EndOfPackage, *Class + EndOfPackage + 1);
-			}
-		}
-		else
-		{
-			if (UClass * ObjectClass = FindObjectFast<UClass>(nullptr, ClassName, false, true))
-			{
-				ClassPackageName = *ObjectClass->GetOuterUPackage()->GetPathName();
-			}
-		}
-
-		if (ClassPackageName != NAME_None)
-		{
-			FPackageIndex ImportOuterIndex = FindOrCreateImport(PackageName, NAME_Package, FName(TEXT("/Script/CoreUObject")));
-			FPackageIndex ImportIndex = FindOrCreateImport(ObjectName, ClassName, ClassPackageName);
-			ImportMap[ImportIndex.ToImport()].OuterIndex = ImportOuterIndex;
-
-			return ImportIndex;
-		}
-		else
-		{
-			UE_LOG(LogLinker, Warning, TEXT("Failed to find class '%s' while trying to resolve full path '%s'"), *ClassName.ToString(), *InFullPath);
-		}
-	}
-
-	return FPackageIndex();
-}
-
-FLinkerLoad::ELinkerStatus FLinkerLoad::ReconstructImportAndExportMap()
+FLinkerLoad::ELinkerStatus FLinkerLoad::ConstructExportsReaders()
 {
-	DECLARE_SCOPE_CYCLE_COUNTER(TEXT("FLinkerLoad::ReconstructImportAndExportMap"), STAT_LinkerLoad_ReconstructImportAndExportMap, STATGROUP_LinkerLoad);
+	DECLARE_SCOPE_CYCLE_COUNTER(TEXT("FLinkerLoad::ConstructExportsReaders"), STAT_LinkerLoad_ConstructExportsReaders, STATGROUP_LinkerLoad);
 
-	if (!bHasReconstructedImportAndExportMap && IsTextFormat())
+	if (!bHasConstructedExportsReaders && IsTextFormat())
 	{
 		int32 NumExports = 0;
 		FStructuredArchiveMap PackageExports = StructuredArchiveRootRecord->EnterMap(SA_FIELD_NAME(TEXT("Exports")), NumExports);
 
-		TArray<FObjectTextExport> ExportRecords;
-		ExportRecords.Reserve(NumExports);
-		ExportMap.SetNum(NumExports);
-		OriginalExportNames.SetNum(NumExports);
 		ExportReaders.AddDefaulted(NumExports);
-
-		Summary.ExportCount = ExportMap.Num();
-		Summary.ImportCount = 0;
-
 		for (int32 ExportIndex = 0; ExportIndex < NumExports; ++ExportIndex)
 		{
-			FObjectTextExport& TextExport = ExportRecords.Emplace_GetRef(ExportMap[ExportIndex], nullptr);
 			FString ExportName;
 			ExportReaders[ExportIndex] = new FStructuredArchiveChildReader(PackageExports.EnterElement(ExportName));
-			ExportReaders[ExportIndex]->GetRoot() << TextExport;
-			OriginalExportNames[ExportIndex] = *ExportName;
-			
-			if (TextExport.OuterName.Len())
-			{
-				FString Class, Package, Object, SubObject;
-				FPackageName::SplitFullObjectPath(TextExport.OuterName, Class, Package, Object, SubObject);
-				FString OuterPath = Object;
-				if (SubObject.Len())
-				{
-					OuterPath += TEXT(":") + SubObject;
-				}
-				if (ExportName.StartsWith(OuterPath))
-				{
-					ExportName = ExportName.Right(ExportName.Len() - OuterPath.Len() - 1);
-				}
-			}
-
-			ExportMap[ExportIndex].ObjectName = *ExportName;
 		}
 		
-		// Now pass over all the exports and rebuild the export/import records
-		for (int32 ExportIndex = 0; ExportIndex < NumExports; ++ExportIndex)
-		{
-			FObjectTextExport& TextExport = ExportRecords[ExportIndex];
-			
-			TextExport.Export.ThisIndex = FPackageIndex::FromExport(ExportIndex);
-			TextExport.Export.ClassIndex = FindOrCreateImportOrExport(TextExport.ClassName);
-			TextExport.Export.SuperIndex = FindOrCreateImportOrExport(TextExport.SuperStructName);
-			TextExport.Export.OuterIndex = FindOrCreateImportOrExport(TextExport.OuterName);
-		}
-
-		bHasReconstructedImportAndExportMap = true;
+		bHasConstructedExportsReaders = true;
 		return LINKER_Loaded;
 	}
 	else
diff --git a/Engine/Source/Runtime/CoreUObject/Private/UObject/SavePackage.cpp b/Engine/Source/Runtime/CoreUObject/Private/UObject/SavePackage.cpp
--- a/Engine/Source/Runtime/CoreUObject/Private/UObject/SavePackage.cpp
+++ b/Engine/Source/Runtime/CoreUObject/Private/UObject/SavePackage.cpp
@@ -2109,6 +2109,7 @@ FSavePackageResultStruct UPackage::Save(UPackage* InOuter, UObject* InAsset, con
 				FArchiveFormatterType* Formatter = nullptr;
 				FArchive* TextFormatArchive = nullptr;
 				const bool bTextFormat = FString(Filename).EndsWith(FPackageName::GetTextAssetPackageExtension()) || FString(Filename).EndsWith(FPackageName::GetTextMapPackageExtension());
+				checkf(!bTextFormat, TEXT("Text format is not supported in old Save, please enable SavePackage2 instead"));
 
 				const FString BaseFilename = FPaths::GetBaseFilename(Filename);
 				// Make temp file. CreateTempFilename guarantees unique, non-existing filename.
diff --git a/Engine/Source/Runtime/CoreUObject/Private/UObject/SavePackage2.cpp b/Engine/Source/Runtime/CoreUObject/Private/UObject/SavePackage2.cpp
--- a/Engine/Source/Runtime/CoreUObject/Private/UObject/SavePackage2.cpp
+++ b/Engine/Source/Runtime/CoreUObject/Private/UObject/SavePackage2.cpp
@@ -1649,6 +1649,26 @@ ESavePackageResult WritePackageTextHeader(FStructuredArchive::FRecord& Structure
 		WriteGatherableText(StructuredArchiveRoot, SaveContext);
 	}
 
+	// Write ImportTable
+	{
+		SCOPED_SAVETIMER(UPackage_Save_WriteImportTable);
+		FStructuredArchive::FStream ImportTableStream = StructuredArchiveRoot.EnterStream(SA_FIELD_NAME(TEXT("ImportTable")));
+		for (FObjectImport& Import : Linker->ImportMap)
+		{
+			ImportTableStream.EnterElement() << Import;
+		}
+	}
+
+	// Write ExportTable
+	{
+		SCOPED_SAVETIMER(UPackage_Save_WriteExportTable);
+		FStructuredArchive::FStream ExportTableStream = StructuredArchiveRoot.EnterStream(SA_FIELD_NAME(TEXT("ExportTable")));
+		for (FObjectExport& Export : Linker->ExportMap)
+		{
+			ExportTableStream.EnterElement() << Export;
+		}
+	}
+
 	// Save thumbnails
 	{
 		SCOPED_SAVETIMER(UPackage_Save_SaveThumbnails);
@@ -1694,12 +1714,6 @@ ESavePackageResult WriteExports(FStructuredArchive::FRecord& StructuredArchiveRo
 			FString ObjectName = Export.Object->GetPathName(SaveContext.GetPackage());
 			FStructuredArchive::FSlot ExportSlot = ExportsRecord.EnterField(SA_FIELD_NAME(*ObjectName));
 
-			if (SaveContext.IsTextFormat())
-			{
-				FObjectTextExport ObjectTextExport(Export, SaveContext.GetPackage());
-				ExportSlot << ObjectTextExport;
-			}
-
 #if WITH_EDITOR
 			bool bSupportsText = UClass::IsSafeToSerializeToStructuredArchives(Export.Object->GetClass());
 #else
diff --git a/Engine/Source/Runtime/CoreUObject/Public/Serialization/Formatters/JsonArchiveInputFormatter.h b/Engine/Source/Runtime/CoreUObject/Public/Serialization/Formatters/JsonArchiveInputFormatter.h
--- a/Engine/Source/Runtime/CoreUObject/Public/Serialization/Formatters/JsonArchiveInputFormatter.h
+++ b/Engine/Source/Runtime/CoreUObject/Public/Serialization/Formatters/JsonArchiveInputFormatter.h
@@ -4,6 +4,7 @@
 
 #include "CoreMinimal.h"
 #include "Serialization/StructuredArchiveFormatter.h"
+#include "UObject/ObjectResource.h"
 
 class FJsonObject;
 class FJsonValue;
@@ -13,7 +14,7 @@ class FJsonValue;
 class COREUOBJECT_API FJsonArchiveInputFormatter final : public FStructuredArchiveFormatter
 {
 public:
-	FJsonArchiveInputFormatter(FArchive& InInner, TFunction<UObject* (const FString&)> InResolveObjectName = nullptr);
+	FJsonArchiveInputFormatter(FArchive& InInner, TFunction<UObject* (const FPackageIndex)> InResolveObject = nullptr);
 	virtual ~FJsonArchiveInputFormatter();
 
 	virtual FArchive& GetUnderlyingArchive() override;
@@ -95,7 +96,7 @@ private:
 		int64 ValueCountOnCreation;	// For debugging purposes, so we can ensure all values have been consumed
 	};
 
-	TFunction<UObject* (const FString&)> ResolveObjectName;
+	TFunction<UObject* (const FPackageIndex)> ResolveObject;
 	TArray<FObjectRecord> ObjectStack;
 	TArray<TSharedPtr<FJsonValue>> ValueStack;
 	TArray<TMap<FString, TSharedPtr<FJsonValue>>::TIterator> MapIteratorStack;
diff --git a/Engine/Source/Runtime/CoreUObject/Public/UObject/LinkerLoad.h b/Engine/Source/Runtime/CoreUObject/Public/UObject/LinkerLoad.h
--- a/Engine/Source/Runtime/CoreUObject/Public/UObject/LinkerLoad.h
+++ b/Engine/Source/Runtime/CoreUObject/Public/UObject/LinkerLoad.h
@@ -349,8 +349,8 @@ private:
 	bool					bHasSerializedPackageFileSummary:1;
 	/** Whether we already serialized the package trailer.																	*/
 	bool					bHasSerializedPackageTrailer : 1;
-	/** Whether we have already reconstructed the import/export tables for a text asset */
-	bool					bHasReconstructedImportAndExportMap:1;
+	/** Whether we have already constructed the exports readers																*/
+	bool					bHasConstructedExportsReaders:1;
 	/** Whether we already serialized preload dependencies.																	*/
 	bool					bHasSerializedPreloadDependencies:1;
 	/** Whether we already fixed up import map.																				*/
@@ -1052,24 +1052,24 @@ private:
 
 #if WITH_TEXT_ARCHIVE_SUPPORT
 	/**
-	 * Create an import and export table when loading a text asset.
+	 * Create the export readers.
 	 */
-	ELinkerStatus ReconstructImportAndExportMap();
+	ELinkerStatus ConstructExportsReaders();
 #endif
 
+	/**
+	 * Serializes the depends map.
+	 */
 	ELinkerStatus SerializeDependsMap();
 
+	/**
+	 * Serializes the preload dependencies.
+	 */
 	ELinkerStatus SerializePreloadDependencies();
 
 	/** Sets the basic linker archive info */
 	void ResetStatusInfo();
 
-	/** For a given full object path, find or create the associated import or export table record. Used when loading text assets which only store object paths */
-	FPackageIndex FindOrCreateImportOrExport(const FString& InFullPath);
-
-	/** For the given object and class info, find or create an associated import record. Used when loading text assets which only store object paths */
-	FPackageIndex FindOrCreateImport(const FName InObjectName, const FName InClassName, const FName InClassPackageName);
-
 #if UE_WITH_OBJECT_HANDLE_LATE_RESOLVE
 	enum class EImportLoadBehavior
 	{
@@ -1330,9 +1330,6 @@ private:
 private:
 
 #if WITH_TEXT_ARCHIVE_SUPPORT
-	// Cache of the export names in a text asset. Allows us to enter those export slots by index rather than needing to reconstruct the name
-	TArray<FName> OriginalExportNames;
-
 	// Function to get a slot for a given export
 	FStructuredArchiveSlot GetExportSlot(FPackageIndex InExportIndex);
 #endif
diff --git a/Engine/Source/Runtime/CoreUObject/Public/UObject/ObjectResource.h b/Engine/Source/Runtime/CoreUObject/Public/UObject/ObjectResource.h
--- a/Engine/Source/Runtime/CoreUObject/Public/UObject/ObjectResource.h
+++ b/Engine/Source/Runtime/CoreUObject/Public/UObject/ObjectResource.h
@@ -393,6 +393,7 @@ PRAGMA_ENABLE_DEPRECATION_WARNINGS
 /**
  * Simple wrapper around a FObjectExport which does the text asset specific serialization of export data
  */
+// TODO SavePackageDeprecation: remove once SavePackage2 replaces SavePackage
 struct FObjectTextExport
 {
 	/**
