Move _Serialize per-class and per-object scopes into the LoadTime and AssetLoadTime channels.
#rnx

diff --git a/Engine/Source/Runtime/Core/Private/Serialization/LoadTimeTrace.cpp b/Engine/Source/Runtime/Core/Private/Serialization/LoadTimeTrace.cpp
--- a/Engine/Source/Runtime/Core/Private/Serialization/LoadTimeTrace.cpp
+++ b/Engine/Source/Runtime/Core/Private/Serialization/LoadTimeTrace.cpp
@@ -9,6 +9,7 @@
 #include "HAL/PlatformTLS.h"
 
 UE_TRACE_CHANNEL_DEFINE(LoadTimeChannel)
+UE_TRACE_CHANNEL_DEFINE(AssetLoadTimeChannel)
 
 UE_TRACE_EVENT_BEGIN(LoadTime, BeginRequestGroup)
 	UE_TRACE_EVENT_FIELD(UE::Trace::WideString, Format)
diff --git a/Engine/Source/Runtime/Core/Public/Serialization/LoadTimeTrace.h b/Engine/Source/Runtime/Core/Public/Serialization/LoadTimeTrace.h
--- a/Engine/Source/Runtime/Core/Public/Serialization/LoadTimeTrace.h
+++ b/Engine/Source/Runtime/Core/Public/Serialization/LoadTimeTrace.h
@@ -18,6 +18,7 @@
 #if LOADTIMEPROFILERTRACE_ENABLED
 
 UE_TRACE_CHANNEL_EXTERN(LoadTimeChannel, CORE_API);
+UE_TRACE_CHANNEL_EXTERN(AssetLoadTimeChannel, CORE_API);
 
 struct FLoadTimeProfilerTrace
 {
diff --git a/Engine/Source/Runtime/CoreUObject/Private/UObject/LinkerLoad.cpp b/Engine/Source/Runtime/CoreUObject/Private/UObject/LinkerLoad.cpp
--- a/Engine/Source/Runtime/CoreUObject/Private/UObject/LinkerLoad.cpp
+++ b/Engine/Source/Runtime/CoreUObject/Private/UObject/LinkerLoad.cpp
@@ -4218,8 +4218,7 @@ void FLinkerLoad::Preload( UObject* Object )
 						check(CurrentLoadContext);
 #if WITH_EDITOR
 						SCOPED_LOADTIMER_TEXT(*WriteToString<128>(Object->GetClass()->GetFName(), TEXTVIEW("_LoadSerialize")));
-						TRACE_CPUPROFILER_EVENT_SCOPE_TEXT(*WriteToString<128>(Object->GetClass()->GetFName(), TEXTVIEW("_Serialize")));
-						TRACE_CPUPROFILER_EVENT_SCOPE_TEXT(*Object->GetFullName());
+						TRACE_CPUPROFILER_EVENT_SCOPE_TEXT_ON_CHANNEL(*Object->GetFullName(), AssetLoadTimeChannel);
 #endif
 
 						// Maintain the current SerializedObjects.
