Fix FObjectPtrProperty identical not respecting deep comparisons

diff --git a/Engine/Source/Runtime/CoreUObject/Private/UObject/PropertyBaseObject.cpp b/Engine/Source/Runtime/CoreUObject/Private/UObject/PropertyBaseObject.cpp
--- a/Engine/Source/Runtime/CoreUObject/Private/UObject/PropertyBaseObject.cpp
+++ b/Engine/Source/Runtime/CoreUObject/Private/UObject/PropertyBaseObject.cpp
@@ -68,9 +68,22 @@ void FObjectPropertyBase::InstanceSubobjects(void* Data, void const* DefaultData
 
 bool FObjectPropertyBase::Identical( const void* A, const void* B, uint32 PortFlags ) const
 {
+	// We never return Identical when duplicating for PIE because we want to be sure to serialize everything. An example is the LevelScriptActor being serialized against its CDO,
+	// which contains actor references. We want to serialize those references so they are fixed up.
+	if ((PortFlags & PPF_DuplicateForPIE) != 0)
+	{
+		return false;
+	}
+
 	UObject* ObjectA = A ? GetObjectPropertyValue(A) : nullptr;
 	UObject* ObjectB = B ? GetObjectPropertyValue(B) : nullptr;
-	if (!ObjectA && !ObjectB)
+
+	return StaticIdentical(ObjectA, ObjectB, PortFlags);
+}
+
+bool FObjectPropertyBase::StaticIdentical(UObject* ObjectA, UObject* ObjectB, uint32 PortFlags)
+{
+	if (ObjectA == ObjectB)
 	{
 		return true;
 	}
@@ -78,17 +91,11 @@ bool FObjectPropertyBase::Identical( const void* A, const void* B, uint32 PortFl
 	{
 		return false;
 	}
-	// Compare actual pointers. We don't do this during PIE because we want to be sure to serialize everything. An example is the LevelScriptActor being serialized against its CDO,
-	// which contains actor references. We want to serialize those references so they are fixed up.
-	const bool bDuplicatingForPIE = (PortFlags&PPF_DuplicateForPIE) != 0;
-	bool bResult = !bDuplicatingForPIE ? (ObjectA == ObjectB) : false;
-	// always serialize the cross level references, because they could be nullptr
-	// @todo: okay, this is pretty hacky overall - we should have a PortFlag or something
-	// that is set during SavePackage. Other times, we don't want to immediately return false
-	// (instead of just this ExportDefProps case)
+
+	bool bResult = false;
 
 	// In order for a deep comparison of instanced objects to match both objects must have the same class and name
-	if (!bResult && ObjectA->GetClass() == ObjectB->GetClass() && ObjectA->GetFName() == ObjectB->GetFName())
+	if (ObjectA->GetClass() == ObjectB->GetClass() && ObjectA->GetFName() == ObjectB->GetFName())
 	{
 		bool bPerformDeepComparison = (PortFlags&PPF_DeepComparison) != 0;
 		if (((PortFlags&PPF_DeepCompareInstances) != 0) && !bPerformDeepComparison)
diff --git a/Engine/Source/Runtime/CoreUObject/Private/UObject/PropertyClassPtr.cpp b/Engine/Source/Runtime/CoreUObject/Private/UObject/PropertyClassPtr.cpp
--- a/Engine/Source/Runtime/CoreUObject/Private/UObject/PropertyClassPtr.cpp
+++ b/Engine/Source/Runtime/CoreUObject/Private/UObject/PropertyClassPtr.cpp
@@ -30,8 +30,12 @@ void FClassPtrProperty::SerializeItem(FStructuredArchive::FSlot Slot, void* Valu
 
 bool FClassPtrProperty::Identical(const void* A, const void* B, uint32 PortFlags) const
 {
-	// Share comparison code with FObjectPtrProperty
-	return FObjectPtrProperty::StaticIdentical(A, B, PortFlags);
+	FObjectPtr ObjectA = A ? *((FObjectPtr*)A) : FObjectPtr();
+	FObjectHandle ObjectAHandle = ObjectA.GetHandle();
+	FObjectPtr ObjectB = B ? *((FObjectPtr*)B) : FObjectPtr();
+	FObjectHandle ObjectBHandle = ObjectB.GetHandle();
+
+	return ObjectAHandle == ObjectBHandle;
 }
 
 bool FClassPtrProperty::SameType(const FProperty* Other) const
diff --git a/Engine/Source/Runtime/CoreUObject/Private/UObject/PropertyObjectPtr.cpp b/Engine/Source/Runtime/CoreUObject/Private/UObject/PropertyObjectPtr.cpp
--- a/Engine/Source/Runtime/CoreUObject/Private/UObject/PropertyObjectPtr.cpp
+++ b/Engine/Source/Runtime/CoreUObject/Private/UObject/PropertyObjectPtr.cpp
@@ -111,17 +111,30 @@ bool FObjectPtrProperty::SameType(const FProperty* Other) const
 
 bool FObjectPtrProperty::Identical(const void* A, const void* B, uint32 PortFlags) const
 {
-	return StaticIdentical(A, B, PortFlags);
-}
+	// We never return Identical when duplicating for PIE because we want to be sure to serialize everything. An example is the LevelScriptActor being serialized against its CDO,
+	// which contains actor references. We want to serialize those references so they are fixed up.
+	if ((PortFlags & PPF_DuplicateForPIE) != 0)
+	{
+		return false;
+	}
 
-bool FObjectPtrProperty::StaticIdentical(const void* A, const void* B, uint32 PortFlags)
-{
 	FObjectPtr ObjectA = A ? *((FObjectPtr*)A) : FObjectPtr();
 	FObjectHandle ObjectAHandle = ObjectA.GetHandle();
 	FObjectPtr ObjectB = B ? *((FObjectPtr*)B) : FObjectPtr();
 	FObjectHandle ObjectBHandle = ObjectB.GetHandle();
 
-	return ObjectAHandle == ObjectBHandle;
+	if (ObjectAHandle == ObjectBHandle)
+	{
+		return true;
+	}
+
+	// Resolve the object handles and run the deep comparison logic 
+	if ((PortFlags & (PPF_DeepCompareInstances | PPF_DeepComparison)) != 0)
+	{
+		return FObjectPropertyBase::StaticIdentical(ObjectA.Get(), ObjectB.Get(), PortFlags);
+	}
+
+	return false;
 }
 
 UObject* FObjectPtrProperty::GetObjectPropertyValue(const void* PropertyValueAddress) const
diff --git a/Engine/Source/Runtime/CoreUObject/Public/UObject/UnrealType.h b/Engine/Source/Runtime/CoreUObject/Public/UObject/UnrealType.h
--- a/Engine/Source/Runtime/CoreUObject/Public/UObject/UnrealType.h
+++ b/Engine/Source/Runtime/CoreUObject/Public/UObject/UnrealType.h
@@ -2444,6 +2444,9 @@ public:
 	// Returns the qualified export path for a given object, parent, and export root scope
 	static FString GetExportPath(const UObject* Object, const UObject* Parent, const UObject* ExportRootScope, const uint32 PortFlags);
 
+	// Helper method for sharing code with FObjectPtrProperty even though one doesn't inherit from the other
+	static bool StaticIdentical(UObject* A, UObject* B, uint32 PortFlags);
+
 	virtual UObject* LoadObjectPropertyValue(const void* PropertyValueAddress) const
 	{
 		return GetObjectPropertyValue(PropertyValueAddress);
@@ -2646,9 +2649,8 @@ class COREUOBJECT_API FObjectPtrProperty : public FObjectProperty
 	virtual void SerializeItem(FStructuredArchive::FSlot Slot, void* Value, void const* Defaults) const override;
 	// End of FProperty interface
 
-	// Helper methods for sharing code with FClassPtrProperty even though one doesn't inherit from the other
+	// Helper method for sharing code with FClassPtrProperty even though one doesn't inherit from the other
 	static void StaticSerializeItem(const FObjectPropertyBase* ObjectProperty, FStructuredArchive::FSlot Slot, void* Value, void const* Defaults);
-	static bool StaticIdentical(const void* A, const void* B, uint32 PortFlags);
 
 	// FObjectProperty interface
 	virtual UObject* GetObjectPropertyValue(const void* PropertyValueAddress) const override;
