Defensive Crash Fix for when ReplaceWith UStruct pointer is null

#rnx

diff --git a/Engine/Source/Runtime/CoreUObject/Public/Serialization/ArchiveReplaceObjectAndStructPropertyRef.h b/Engine/Source/Runtime/CoreUObject/Public/Serialization/ArchiveReplaceObjectAndStructPropertyRef.h
--- a/Engine/Source/Runtime/CoreUObject/Public/Serialization/ArchiveReplaceObjectAndStructPropertyRef.h
+++ b/Engine/Source/Runtime/CoreUObject/Public/Serialization/ArchiveReplaceObjectAndStructPropertyRef.h
@@ -74,8 +74,8 @@ public:
 				{
 					// We want to replace the property's owner but since that would be even worse than replacing UObject's Outer
 					// we need to replace the entire property instead. We need to find the new property on the object we want to replace the Owner with
-					UStruct* NewOwnerStruct = CastChecked<UStruct>(*ReplaceWith);
-					FField* ReplaceWithField = NewOwnerStruct->FindPropertyByName(InField->GetFName());
+					UStruct* NewOwnerStruct = *ReplaceWith ? CastChecked<UStruct>(*ReplaceWith) : nullptr;
+					FField* ReplaceWithField = NewOwnerStruct ? NewOwnerStruct->FindPropertyByName(InField->GetFName()) : nullptr;
 					// Do we need to verify the existence of ReplaceWithField? Theoretically it could be missing on the new version
 					// of the owner struct and in this case we still don't want to keep the stale old property pointer around so it's safer to null it
 					InField = ReplaceWithField;
