OnEndLoadPackage: Change the arguments to a context object to avoid future need to change the function prototype. Add RecursiveDepth and bSynchronous arguments so that callers can take action only on the top-level EndLoadPackage if necessary.
BulkDataRegistry: Wait until the top-level EndLoadPackage, and wait until after OnEnginePostInit, before deleting bulkdatasize records. Deleting them sooner can make them be deleted before they're used by the EditorDomain if the EditorDomain's TrySavePackage triggers package loading and a recursive OnEndLoadPackage.
#rnx

diff --git a/Engine/Plugins/Animation/ControlRig/Source/ControlRig/Private/Sequencer/MovieSceneControlRigParameterTrack.cpp b/Engine/Plugins/Animation/ControlRig/Source/ControlRig/Private/Sequencer/MovieSceneControlRigParameterTrack.cpp
--- a/Engine/Plugins/Animation/ControlRig/Source/ControlRig/Private/Sequencer/MovieSceneControlRigParameterTrack.cpp
+++ b/Engine/Plugins/Animation/ControlRig/Source/ControlRig/Private/Sequencer/MovieSceneControlRigParameterTrack.cpp
@@ -404,7 +404,7 @@ void UMovieSceneControlRigParameterTrack::PostLoad()
 }
 
 #if WITH_EDITOR
-void UMovieSceneControlRigParameterTrack::HandlePackageDone(TConstArrayView<UPackage*> InPackages)
+void UMovieSceneControlRigParameterTrack::HandlePackageDone(const FEndLoadPackageContext& Context)
 {
 	if (!GetPackage()->GetHasBeenEndLoaded())
 	{
diff --git a/Engine/Plugins/Animation/ControlRig/Source/ControlRig/Public/Sequencer/MovieSceneControlRigParameterTrack.h b/Engine/Plugins/Animation/ControlRig/Source/ControlRig/Public/Sequencer/MovieSceneControlRigParameterTrack.h
--- a/Engine/Plugins/Animation/ControlRig/Source/ControlRig/Public/Sequencer/MovieSceneControlRigParameterTrack.h
+++ b/Engine/Plugins/Animation/ControlRig/Source/ControlRig/Public/Sequencer/MovieSceneControlRigParameterTrack.h
@@ -12,6 +12,8 @@
 #include "MovieSceneControlRigParameterSection.h"
 #include "MovieSceneControlRigParameterTrack.generated.h"
 
+struct FEndLoadPackageContext;
+
 /**
  * Handles animation of skeletal mesh actors using animation ControlRigs
  */
@@ -50,7 +52,7 @@ public:
 	virtual void GetSelectedNodes(TArray<FName>& OutSelectedNodes) override;
 
 #if WITH_EDITOR
-	void HandlePackageDone(TConstArrayView<UPackage*> InPackages);
+	void HandlePackageDone(const FEndLoadPackageContext& Context);
 	// control Rigs are ready only after its package is fully end-loaded
 	void HandleControlRigPackageDone(UControlRig* InControlRig);
 #endif
diff --git a/Engine/Plugins/Animation/ControlRig/Source/ControlRigDeveloper/Private/ControlRigBlueprint.cpp b/Engine/Plugins/Animation/ControlRig/Source/ControlRigDeveloper/Private/ControlRigBlueprint.cpp
--- a/Engine/Plugins/Animation/ControlRig/Source/ControlRigDeveloper/Private/ControlRigBlueprint.cpp
+++ b/Engine/Plugins/Animation/ControlRig/Source/ControlRigDeveloper/Private/ControlRigBlueprint.cpp
@@ -591,9 +591,9 @@ void UControlRigBlueprint::PostLoad()
 }
 
 #if WITH_EDITOR
-void UControlRigBlueprint::HandlePackageDone(TConstArrayView<UPackage*> InPackages)
+void UControlRigBlueprint::HandlePackageDone(const FEndLoadPackageContext& Context)
 {
-	if (!InPackages.Contains(GetPackage()))
+	if (!Context.LoadedPackages.Contains(GetPackage()))
 	{
 		return;
 	}
diff --git a/Engine/Plugins/Animation/ControlRig/Source/ControlRigDeveloper/Public/ControlRigBlueprint.h b/Engine/Plugins/Animation/ControlRig/Source/ControlRigDeveloper/Public/ControlRigBlueprint.h
--- a/Engine/Plugins/Animation/ControlRig/Source/ControlRigDeveloper/Public/ControlRigBlueprint.h
+++ b/Engine/Plugins/Animation/ControlRig/Source/ControlRigDeveloper/Public/ControlRigBlueprint.h
@@ -30,6 +30,7 @@
 class UControlRigBlueprintGeneratedClass;
 class USkeletalMesh;
 class UControlRigGraph;
+struct FEndLoadPackageContext;
 
 DECLARE_EVENT_TwoParams(UControlRigBlueprint, FOnVMCompiledEvent, UBlueprint*, URigVM*);
 DECLARE_EVENT_OneParam(UControlRigBlueprint, FOnRefreshEditorEvent, UControlRigBlueprint*);
@@ -580,7 +581,7 @@ private:
 	void HandleHierarchyModified(ERigHierarchyNotification InNotification, URigHierarchy* InHierarchy, const FRigBaseElement* InElement);
 
 #if WITH_EDITOR
-	void HandlePackageDone(TConstArrayView<UPackage*> InPackages);
+	void HandlePackageDone(const FEndLoadPackageContext& Context);
 	// ControlRigBP, once end-loaded, will inform other ControlRig-Dependent systems that ControlRig instances are ready.
 	void BroadcastControlRigPackageDone();
 #endif
diff --git a/Engine/Source/Editor/UnrealEd/Private/EditorDomain/BulkDataRegistryEditorDomain.cpp b/Engine/Source/Editor/UnrealEd/Private/EditorDomain/BulkDataRegistryEditorDomain.cpp
--- a/Engine/Source/Editor/UnrealEd/Private/EditorDomain/BulkDataRegistryEditorDomain.cpp
+++ b/Engine/Source/Editor/UnrealEd/Private/EditorDomain/BulkDataRegistryEditorDomain.cpp
@@ -116,13 +116,13 @@ FBulkDataRegistryEditorDomain::~FBulkDataRegistryEditorDomain()
 	}
 }
 
-void FBulkDataRegistryEditorDomain::OnEndLoadPackage(TConstArrayView<UPackage*> LoadedPackages)
+void FBulkDataRegistryEditorDomain::OnEndLoadPackage(const FEndLoadPackageContext& Context)
 {
 	TArray<TUniquePtr<FPendingPackage>> PackagesToWrite;
 	{
 		FScopeLock PendingPackageScopeLock(&PendingPackageLock);
 
-		for (UPackage* LoadedPackage : LoadedPackages)
+		for (UPackage* LoadedPackage : Context.LoadedPackages)
 		{
 			FName PackageName = LoadedPackage->GetFName();
 			uint32 KeyHash = GetTypeHash(PackageName);
diff --git a/Engine/Source/Editor/UnrealEd/Private/EditorDomain/BulkDataRegistryEditorDomain.h b/Engine/Source/Editor/UnrealEd/Private/EditorDomain/BulkDataRegistryEditorDomain.h
--- a/Engine/Source/Editor/UnrealEd/Private/EditorDomain/BulkDataRegistryEditorDomain.h
+++ b/Engine/Source/Editor/UnrealEd/Private/EditorDomain/BulkDataRegistryEditorDomain.h
@@ -182,7 +182,7 @@ private:
 	void PollPendingPackages(bool bWaitForCooldown);
 	void AddTempLoadedPayload(const FGuid& RegistryKey, uint64 PayloadSize);
 	void PruneTempLoadedPayloads();
-	void OnEndLoadPackage(TConstArrayView<UPackage*> LoadedPackages);
+	void OnEndLoadPackage(const FEndLoadPackageContext& Context);
 	void WritePayloadIdToCache(FName PackageName, const UE::Serialization::FEditorBulkData& BulkData) const;
 	void ReadPayloadIdsFromCache(FName PackageName, TArray<TRefCountPtr<FPendingPayloadId>>&& OldPendings,
 		TArray<TRefCountPtr<FPendingPayloadId>>&& NewPendings);
diff --git a/Engine/Source/Editor/UnrealEd/Private/EditorDomain/EditorDomain.cpp b/Engine/Source/Editor/UnrealEd/Private/EditorDomain/EditorDomain.cpp
--- a/Engine/Source/Editor/UnrealEd/Private/EditorDomain/EditorDomain.cpp
+++ b/Engine/Source/Editor/UnrealEd/Private/EditorDomain/EditorDomain.cpp
@@ -500,7 +500,7 @@ void FEditorDomain::Tick(float DeltaTime)
 	}
 }
 
-void FEditorDomain::OnEndLoadPackage(TConstArrayView<UPackage*> LoadedPackages)
+void FEditorDomain::OnEndLoadPackage(const FEndLoadPackageContext& Context)
 {
 	if (bExternalSave)
 	{
@@ -513,8 +513,8 @@ void FEditorDomain::OnEndLoadPackage(TConstArrayView<UPackage*> LoadedPackages)
 		{
 			return;
 		}
-		PackagesToSave.Reserve(LoadedPackages.Num());
-		for (UPackage* Package : LoadedPackages)
+		PackagesToSave.Reserve(Context.LoadedPackages.Num());
+		for (UPackage* Package : Context.LoadedPackages)
 		{
 			PackagesToSave.Add(Package);
 		}
diff --git a/Engine/Source/Editor/UnrealEd/Private/EditorDomain/EditorDomain.h b/Engine/Source/Editor/UnrealEd/Private/EditorDomain/EditorDomain.h
--- a/Engine/Source/Editor/UnrealEd/Private/EditorDomain/EditorDomain.h
+++ b/Engine/Source/Editor/UnrealEd/Private/EditorDomain/EditorDomain.h
@@ -205,7 +205,7 @@ private:
 	/** Callback for PostEngineInit, to handle saving of packages which we could not save before then. */
 	void OnPostEngineInit();
 	/** EndLoad callback to handle saving the EditorDomain version of the package. */
-	void OnEndLoadPackage(TConstArrayView<UPackage*> LoadedPackages);
+	void OnEndLoadPackage(const FEndLoadPackageContext& Context);
 	/** For each of the now-loaded packages, if we had to load from workspace domain, save into the editor domain. */
 	void FilterKeepPackagesToSave(TArray<UPackage*>& InOutLoadedPackages);
 	/** PackageSaved context to invalidate our information about where it should be loaded from. */
diff --git a/Engine/Source/Runtime/CoreUObject/Private/Serialization/AsyncLoading.cpp b/Engine/Source/Runtime/CoreUObject/Private/Serialization/AsyncLoading.cpp
--- a/Engine/Source/Runtime/CoreUObject/Private/Serialization/AsyncLoading.cpp
+++ b/Engine/Source/Runtime/CoreUObject/Private/Serialization/AsyncLoading.cpp
@@ -4316,7 +4316,8 @@ EAsyncPackageState::Type FAsyncLoadingThread::ProcessLoadedPackages(bool bUseTim
 		{
 			CompletedUPackage->SetHasBeenEndLoaded(true);
 		}
-		FCoreUObjectDelegates::OnEndLoadPackage.Broadcast(CompletedUPackages.Array());
+		FCoreUObjectDelegates::OnEndLoadPackage.Broadcast(
+			FEndLoadPackageContext{ CompletedUPackages.Array(), 0, false /* bSynchronous */ });
 	}
 #endif
 
diff --git a/Engine/Source/Runtime/CoreUObject/Private/Serialization/AsyncLoading2.cpp b/Engine/Source/Runtime/CoreUObject/Private/Serialization/AsyncLoading2.cpp
--- a/Engine/Source/Runtime/CoreUObject/Private/Serialization/AsyncLoading2.cpp
+++ b/Engine/Source/Runtime/CoreUObject/Private/Serialization/AsyncLoading2.cpp
@@ -4751,7 +4751,8 @@ EAsyncPackageState::Type FAsyncLoadingThread2::ProcessLoadedPackagesFromGameThre
 		{
 			CompletedUPackage->SetHasBeenEndLoaded(true);
 		}
-		FCoreUObjectDelegates::OnEndLoadPackage.Broadcast(CompletedUPackages.Array());
+		FCoreUObjectDelegates::OnEndLoadPackage.Broadcast(
+			FEndLoadPackageContext{ CompletedUPackages.Array(), 0, false /* bSynchronous */ });
 #endif
 		
 		{
diff --git a/Engine/Source/Runtime/CoreUObject/Private/Serialization/BulkDataRegistry.cpp b/Engine/Source/Runtime/CoreUObject/Private/Serialization/BulkDataRegistry.cpp
--- a/Engine/Source/Runtime/CoreUObject/Private/Serialization/BulkDataRegistry.cpp
+++ b/Engine/Source/Runtime/CoreUObject/Private/Serialization/BulkDataRegistry.cpp
@@ -143,11 +143,14 @@ FResaveSizeTracker::FResaveSizeTracker()
 	ELoadingPhase::Type CurrentPhase = IPluginManager::Get().GetLastCompletedLoadingPhase();
 	if (CurrentPhase == ELoadingPhase::None || CurrentPhase < ELoadingPhase::PostEngineInit)
 	{
-		FCoreDelegates::OnPostEngineInit.AddRaw(this, &FResaveSizeTracker::OnPostEngineInit);
+		// Our contract says that we need to keep information up until OnEnginePostInit is called on all
+		// subscribers to it, so we need to subscribe to the first event we can find that occurs after OnEnginePostInit
+		// is done. OnAllModuleLoadingPhasesComplete seems to be it.
+		FCoreDelegates::OnAllModuleLoadingPhasesComplete.AddRaw(this, &FResaveSizeTracker::OnAllModuleLoadingPhasesComplete);
 	}
 	else
 	{
-		OnPostEngineInit();
+		OnAllModuleLoadingPhasesComplete();
 	}
 	FCoreUObjectDelegates::OnEndLoadPackage.AddRaw(this, &FResaveSizeTracker::OnEndLoadPackage);
 }
@@ -155,10 +158,10 @@ FResaveSizeTracker::FResaveSizeTracker()
 FResaveSizeTracker::~FResaveSizeTracker()
 {
 	FCoreUObjectDelegates::OnEndLoadPackage.RemoveAll(this);
-	FCoreDelegates::OnPostEngineInit.RemoveAll(this);
+	FCoreDelegates::OnAllModuleLoadingPhasesComplete.RemoveAll(this);
 }
 
-void FResaveSizeTracker::OnPostEngineInit()
+void FResaveSizeTracker::OnAllModuleLoadingPhasesComplete()
 {
 	bPostEngineInitComplete = true;
 
@@ -195,7 +198,7 @@ uint64 FResaveSizeTracker::GetBulkDataResaveSize(FName PackageName)
 	return PackageBulkResaveSize.FindRef(PackageName);
 }
 
-void FResaveSizeTracker::OnEndLoadPackage(TConstArrayView<UPackage*> LoadedPackages)
+void FResaveSizeTracker::OnEndLoadPackage(const FEndLoadPackageContext& Context)
 {
 	if (!bPostEngineInitComplete)
 	{
@@ -203,8 +206,8 @@ void FResaveSizeTracker::OnEndLoadPackage(TConstArrayView<UPackage*> LoadedPacka
 	}
 
 	TArray<FName> PackageNames;
-	PackageNames.Reserve(LoadedPackages.Num());
-	for (UPackage* LoadedPackage : LoadedPackages)
+	PackageNames.Reserve(Context.LoadedPackages.Num());
+	for (UPackage* LoadedPackage : Context.LoadedPackages)
 	{
 		PackageNames.Add(LoadedPackage->GetFName());
 	}
@@ -212,15 +215,18 @@ void FResaveSizeTracker::OnEndLoadPackage(TConstArrayView<UPackage*> LoadedPacka
 	FWriteScopeLock ScopeLock(Lock);
 	// The contract for GetBulkDataResaveSize specifies that we must answer correctly until
 	// OnEndLoadPackage is complete. This includes being called from other subscribers to OnEndLoadPackage
-	// that might run after us. So we defer the removals from PackageBulkResaveSize until the next call 
+	// that might run after us. So we defer the removals from PackageBulkResaveSize until the next top-level call 
 	// to OnEndLoadPackage.
-	for (FName PackageName : DeferredRemove)
+	if (Context.RecursiveDepth == 0)
 	{
-		PackageBulkResaveSize.Remove(PackageName);
+		for (FName PackageName : DeferredRemove)
+		{
+			PackageBulkResaveSize.Remove(PackageName);
+		}
+		DeferredRemove.Reset();
 	}
-	DeferredRemove.Reset();
 	DeferredRemove.Append(MoveTemp(PackageNames));
 }
 
 }
-#endif // WITH_EDITOR
\ No newline at end of file
+#endif // WITH_EDITOR
diff --git a/Engine/Source/Runtime/CoreUObject/Private/UObject/UObjectGlobals.cpp b/Engine/Source/Runtime/CoreUObject/Private/UObject/UObjectGlobals.cpp
--- a/Engine/Source/Runtime/CoreUObject/Private/UObject/UObjectGlobals.cpp
+++ b/Engine/Source/Runtime/CoreUObject/Private/UObject/UObjectGlobals.cpp
@@ -1100,6 +1100,7 @@ public:
 // Temporary load counter for the game thread, used mostly for checking if we're still loading
 // @todo: remove this in the new loader
 static int32 GGameThreadLoadCounter = 0;
+static int32 GGameThreadEndLoadCounter = -1;
 
 /** Notify delegate listeners of all the packages that loaded; called only once per explicit call to LoadPackage. */
 void BroadcastEndLoad(TArray<UPackage*>&& LoadedPackages)
@@ -1116,7 +1117,11 @@ void BroadcastEndLoad(TArray<UPackage*>&& LoadedPackages)
 		{
 			LoadedPackage->SetHasBeenEndLoaded(true);
 		}
-		FCoreUObjectDelegates::OnEndLoadPackage.Broadcast(LoadedPackages);
+		++GGameThreadEndLoadCounter; // Starts at -1, so the first increment takes it to 0
+		FCoreUObjectDelegates::OnEndLoadPackage.Broadcast(
+			FEndLoadPackageContext{ LoadedPackages, GGameThreadEndLoadCounter, true/* bSynchronous */ });
+		--GGameThreadEndLoadCounter;
+		ensure(GGameThreadEndLoadCounter >= -1);
 	}
 #endif
 }
diff --git a/Engine/Source/Runtime/CoreUObject/Public/Serialization/BulkDataRegistry.h b/Engine/Source/Runtime/CoreUObject/Public/Serialization/BulkDataRegistry.h
--- a/Engine/Source/Runtime/CoreUObject/Public/Serialization/BulkDataRegistry.h
+++ b/Engine/Source/Runtime/CoreUObject/Public/Serialization/BulkDataRegistry.h
@@ -16,6 +16,7 @@
 
 class IBulkDataRegistry;
 class UPackage;
+struct FEndLoadPackageContext;
 struct FGuid;
 namespace UE::Serialization { class FEditorBulkData; }
 
@@ -101,8 +102,8 @@ public:
 	COREUOBJECT_API uint64 GetBulkDataResaveSize(FName PackageName);
 
 private:
-	COREUOBJECT_API void OnEndLoadPackage(TConstArrayView<UPackage*> LoadedPackages);
-	COREUOBJECT_API void OnPostEngineInit();
+	COREUOBJECT_API void OnEndLoadPackage(const FEndLoadPackageContext& Context);
+	COREUOBJECT_API void OnAllModuleLoadingPhasesComplete();
 
 	FRWLock Lock;
 	TMap<FName, uint64> PackageBulkResaveSize;
diff --git a/Engine/Source/Runtime/CoreUObject/Public/UObject/UObjectGlobals.h b/Engine/Source/Runtime/CoreUObject/Public/UObject/UObjectGlobals.h
--- a/Engine/Source/Runtime/CoreUObject/Public/UObject/UObjectGlobals.h
+++ b/Engine/Source/Runtime/CoreUObject/Public/UObject/UObjectGlobals.h
@@ -2383,6 +2383,13 @@ enum class EReloadCompleteReason
 	HotReloadManual,
 };
 
+struct FEndLoadPackageContext
+{
+	TConstArrayView<UPackage*> LoadedPackages;
+	int32 RecursiveDepth;
+	bool bSynchronous;
+};
+
 /**
  * Global CoreUObject delegates
  */
@@ -2436,7 +2443,7 @@ struct COREUOBJECT_API FCoreUObjectDelegates
 	static FOnAssetLoaded OnAssetLoaded;
 
 	/** Callback when packages end loading in LoadPackage or AsyncLoadPackage. All packages loaded recursively due to imports are included in the single call of the explicitly-loaded package. */
-	DECLARE_MULTICAST_DELEGATE_OneParam(FOnEndLoadPackage, TConstArrayView<UPackage*>);
+	DECLARE_MULTICAST_DELEGATE_OneParam(FOnEndLoadPackage, const FEndLoadPackageContext&);
 	static FOnEndLoadPackage OnEndLoadPackage;
 
 	/** Delegate used by SavePackage() to create the package backup */
