1) Minor updates to code generation to make the checksums a bit more stable
2) Added new struct and class specifiers to account for settings being extracted from the engine code.
3) UHT no longer depends on pre-existing classes in the engine but will use them when it needs to create the engine instances of the classes.

#rnx

diff --git a/Engine/Source/Programs/UnrealHeaderTool/Private/CodeGenerator.cpp b/Engine/Source/Programs/UnrealHeaderTool/Private/CodeGenerator.cpp
--- a/Engine/Source/Programs/UnrealHeaderTool/Private/CodeGenerator.cpp
+++ b/Engine/Source/Programs/UnrealHeaderTool/Private/CodeGenerator.cpp
@@ -1831,6 +1831,12 @@ void FNativeClassHeaderGenerator::ExportGeneratedPackageInitCode(FOutputDevice&
 		return AName < BName;
 	});
 
+#if UHT_ENABLE_EXTRA_HASH_OUTPUT
+	Out.Log(TEXT("#if 0\r\n"));
+	Out.Log(InDeclarations);
+	Out.Log(TEXT("#endif\r\n"));
+#endif
+
 	for (FUnrealFieldDefinitionInfo* FieldDef : Singletons)
 	{
 		Out.Log(FieldDef->GetExternDecl(true));
@@ -2182,7 +2188,9 @@ void FNativeClassHeaderGenerator::ExportNativeGeneratedInitCode(FOutputDevice& O
 	{
 		BaseClassHash = SuperClassDef->GetHash(ClassDef);
 	}
-	GeneratedClassRegisterFunctionText.Logf(TEXT("\r\n// %u\r\n"), BaseClassHash);
+
+	FUHTStringBuilder HashBuilder;
+	HashBuilder.Logf(TEXT("\r\n// %u\r\n"), BaseClassHash);
 
 	// Append info for the sparse class data struct onto the text to be hashed
 	TArray<FString> SparseClassDataTypes;
@@ -2192,14 +2200,22 @@ void FNativeClassHeaderGenerator::ExportNativeGeneratedInitCode(FOutputDevice& O
 	{
 		if (FUnrealScriptStructDefinitionInfo* SparseScriptStructDef = GTypeDefinitionInfoMap.FindByName<FUnrealScriptStructDefinitionInfo>(*SparseClassDataString))
 		{
-			GeneratedClassRegisterFunctionText.Logf(TEXT("%s\r\n"), *SparseScriptStructDef->GetName());
+			HashBuilder.Logf(TEXT("%s\r\n"), *SparseScriptStructDef->GetName());
 			for (FUnrealPropertyDefinitionInfo* ChildDef : TUHTFieldRange<FUnrealPropertyDefinitionInfo>(*SparseScriptStructDef))
 			{
-				GeneratedClassRegisterFunctionText.Logf(TEXT("%s %s\r\n"), *ChildDef->GetCPPType(), *ChildDef->GetNameWithDeprecated());
+				HashBuilder.Logf(TEXT("%s %s\r\n"), *ChildDef->GetCPPType(), *ChildDef->GetNameWithDeprecated());
 			}
 		}
 	}
 
+#if UHT_ENABLE_EXTRA_HASH_OUTPUT
+	Out.Log(TEXT("#if 0\r\n"));
+	Out.Log(HashBuilder);
+	Out.Log(TEXT("#endif\r\n"));
+#endif
+
+	GeneratedClassRegisterFunctionText.Log(HashBuilder);
+
 	// Calculate generated class initialization code hash so that we know when it changes after hot-reload
 	uint32 ClassHash = GenerateTextHash(*GeneratedClassRegisterFunctionText);
 	ClassDef.SetHash(ClassHash);
@@ -2996,7 +3012,7 @@ void ExportConstructorDefinition(FOutputDevice& Out, FUnrealClassDefinitionInfo&
 		FUnrealClassDefinitionInfo* SuperClassDef = ClassDef.GetSuperClass();
 		if (SuperClassDef != nullptr)
 		{
-			if (SuperClassDef->HasSource()) // Don't consider the internal types generated from the engine
+			if (SuperClassDef->HasSource() && !SuperClassDef->GetUnrealSourceFile().IsNoExportTypes()) // Don't consider the internal types generated from the engine
 			{
 				// Since we are dependent on our SuperClass having determined which constructors are defined, 
 				// if it is not yet determined we will need to wait on it becoming available. 
@@ -3132,11 +3148,21 @@ bool FNativeClassHeaderGenerator::WriteHeader(FGeneratedFileInfo& FileInfo, cons
 	GeneratedHeaderTextWithCopyright.Log(LINE_TERMINATOR);
 	GeneratedHeaderTextWithCopyright.Log(DisableDeprecationWarnings);
 
-	for (const FString& FWDecl : ForwardDeclarations)
+	if (ForwardDeclarations.Num() > 0)
 	{
-		if (FWDecl.Len() > 0)
+		TArray<const FString*> Sorted;
+		Sorted.Reserve(ForwardDeclarations.Num());
+		for (const FString& Ref : ForwardDeclarations)
+		{
+			if (Ref.Len() > 0)
+			{
+				Sorted.Add(&Ref);
+			}
+		}
+		Sorted.Sort();
+		for (const FString* Ref : Sorted)
 		{
-			GeneratedHeaderTextWithCopyright.Logf(TEXT("%s\r\n"), *FWDecl);
+			GeneratedHeaderTextWithCopyright.Logf(TEXT("%s\r\n"), **Ref);
 		}
 	}
 
@@ -3323,15 +3349,15 @@ void FNativeClassHeaderGenerator::ExportGeneratedStructBodyMacros(FOutputDevice&
 
 		FString StructMembers = StructRigVMInfo.Members.Declarations(false, TEXT(", \\\r\n\t\t"), true, false);
 
-		OutGeneratedHeaderText.Log(TEXT("\n"));
+		OutGeneratedHeaderText.Log(TEXT("\r\n"));
 		for (const FRigVMMethodInfo& MethodInfo : StructRigVMInfo.Methods)
 		{
 			FString ParameterSuffix = MethodInfo.Parameters.Declarations(true, TEXT(", \\\r\n\t\t"));
 			FString RigVMParameterPrefix2 = RigVMParameterPrefix + FString((StructMembers.IsEmpty() && ParameterSuffix.IsEmpty()) ? TEXT("") : TEXT(", \\\r\n\t\t"));
 			OutGeneratedHeaderText.Logf(TEXT("#define %s_%s() \\\r\n"), *StructNameCPP, *MethodInfo.Name);
-			OutGeneratedHeaderText.Logf(TEXT("\t%s %s::Static%s( \\\r\n\t\t%s%s%s \\\r\n\t)\n"), *MethodInfo.ReturnType, *StructNameCPP, *MethodInfo.Name, *RigVMParameterPrefix2, *StructMembers, *ParameterSuffix);
+			OutGeneratedHeaderText.Logf(TEXT("\t%s %s::Static%s( \\\r\n\t\t%s%s%s \\\r\n\t)\r\n"), *MethodInfo.ReturnType, *StructNameCPP, *MethodInfo.Name, *RigVMParameterPrefix2, *StructMembers, *ParameterSuffix);
 		}
-		OutGeneratedHeaderText.Log(TEXT("\n"));
+		OutGeneratedHeaderText.Log(TEXT("\r\n"));
 	}
 
 	// Export struct.
@@ -3580,7 +3606,7 @@ void FNativeClassHeaderGenerator::ExportGeneratedStructBodyMacros(FOutputDevice&
 				Out.Log(TEXT("\t\r\n"));
 			}
 
-			Out.Logf(TEXT("    %sStatic%s(\r\n\t\t%s%s%s\r\n\t);\n"), *MethodInfo.ReturnPrefix(), *MethodInfo.Name, *RigVMParameterPrefix3, *StructMembersForVirtualFunc, *ParameterSuffix);
+			Out.Logf(TEXT("\t%sStatic%s(\r\n\t\t%s%s%s\r\n\t);\r\n"), *MethodInfo.ReturnPrefix(), *MethodInfo.Name, *RigVMParameterPrefix3, *StructMembersForVirtualFunc, *ParameterSuffix);
 
 			if (RigVMVirtualFuncEpilog.Num() > 0)
 			{
@@ -3919,7 +3945,7 @@ void FNativeClassHeaderGenerator::ExportDelegateDefinition(FOutputDevice& Out, F
 	WriteMacro(Out, MacroName, DelegateOutput);
 }
 
-void FNativeClassHeaderGenerator::ExportEventParm(FUHTStringBuilder& Out, TSet<FString>& PropertyFwd, FUnrealFunctionDefinitionInfo& FunctionDef, int32 Indent, bool bOutputConstructor, EExportingState ExportingState)
+void FNativeClassHeaderGenerator::ExportEventParm(FUHTStringBuilder& Out, TSet<FString>& ForwardDeclarations, FUnrealFunctionDefinitionInfo& FunctionDef, int32 Indent, bool bOutputConstructor, EExportingState ExportingState)
 {
 	if (!WillExportEventParms(FunctionDef))
 	{
@@ -3944,7 +3970,7 @@ void FNativeClassHeaderGenerator::ExportEventParm(FUHTStringBuilder& Out, TSet<F
 			continue;
 		}
 
-		PropertyFwd.Add(PropDef->GetCPPTypeForwardDeclaration());
+		ForwardDeclarations.Add(PropDef->GetCPPTypeForwardDeclaration());
 
 		FUHTStringBuilder PropertyText;
 		PropertyText.Log(FCString::Tab(Indent + 1));
@@ -4211,8 +4237,8 @@ void FNativeClassHeaderGenerator::CheckRPCFunctions(FReferenceGatherers& OutRefe
 
 void FNativeClassHeaderGenerator::ExportNativeFunctionHeader(
 	FOutputDevice&                   Out,
-	TSet<FString>&                   OutFwdDecls,
-	FUnrealFunctionDefinitionInfo&				 FunctionDef,
+	TSet<FString>&                   ForwardDeclarations,
+	FUnrealFunctionDefinitionInfo&	 FunctionDef,
 	const FFuncInfo&				 FunctionData,
 	EExportFunctionType::Type        FunctionType,
 	EExportFunctionHeaderStyle::Type FunctionHeaderStyle,
@@ -4276,7 +4302,7 @@ void FNativeClassHeaderGenerator::ExportNativeFunctionHeader(
 
 		FString ExtendedReturnType;
 		FString ReturnType = ReturnPropertyDef->GetCPPType(&ExtendedReturnType, (FunctionHeaderStyle == EExportFunctionHeaderStyle::Definition && (FunctionType != EExportFunctionType::Interface) ? CPPF_Implementation : 0) | CPPF_ArgumentOrReturnValue);
-		OutFwdDecls.Add(ReturnPropertyDef->GetCPPTypeForwardDeclaration());
+		ForwardDeclarations.Add(ReturnPropertyDef->GetCPPTypeForwardDeclaration());
 		FUHTStringBuilder ReplacementText;
 		ReplacementText += ReturnType;
 		ApplyAlternatePropertyExportText(*ReturnPropertyDef, ReplacementText, EExportingState::Normal);
@@ -4324,7 +4350,7 @@ void FNativeClassHeaderGenerator::ExportNativeFunctionHeader(
 			continue;
 		}
 
-		OutFwdDecls.Add(PropertyDef->GetCPPTypeForwardDeclaration());
+		ForwardDeclarations.Add(PropertyDef->GetCPPTypeForwardDeclaration());
 
 		if( ParmCount++ )
 		{
@@ -5012,7 +5038,7 @@ bool FNativeClassHeaderGenerator::WriteSource(const FManifestModule& Module, FGe
 		{
 			FUnrealClassDefinitionInfo& ClassDef = TypeDef->AsClassChecked();
 			FUnrealClassDefinitionInfo* ClassWithin = ClassDef.GetClassWithin();
-			if (ClassWithin && ClassWithin != GUObjectDef && ClassWithin->HasSource())
+			if (ClassWithin && ClassWithin->HasSource() && !ClassWithin->GetUnrealSourceFile().IsNoExportTypes())
 			{
 				FString Header = GetBuildPath(ClassWithin->GetUnrealSourceFile());
 				RelativeIncludes.AddUnique(MoveTemp(Header));
@@ -5049,10 +5075,17 @@ bool FNativeClassHeaderGenerator::WriteSource(const FManifestModule& Module, FGe
 
 	if (InCrossModuleReferences.Num() > 0)
 	{
-		FileText.Logf(TEXT("// Cross Module References\r\n"));
+		TArray<const FString*> Sorted;
+		Sorted.Reserve(InCrossModuleReferences.Num());
 		for (const FString& Ref : InCrossModuleReferences)
 		{
-			FileText.Log(*Ref);
+			Sorted.Add(&Ref);
+		}
+		Sorted.Sort();
+		FileText.Logf(TEXT("// Cross Module References\r\n"));
+		for (const FString* Ref : Sorted)
+		{
+			FileText.Log(**Ref);
 		}
 		FileText.Logf(TEXT("// End Cross Module References\r\n"));
 	}
@@ -5348,6 +5381,31 @@ void FNativeClassHeaderGenerator::GenerateSourceFile(FGeneratedCPP& GeneratedCPP
 			{
 				SourceFile.GetSingletons().Append(MoveTemp(Singletons));
 			}
+
+			// Sort the forward declarations to make them more stable
+			TArray<FString> Lines;
+			GeneratedFunctionDeclarations.ParseIntoArrayLines(Lines);
+			TSet<FString> Unique;
+			for (const FString& Line : Lines)
+			{
+				Unique.Add(Line);
+			}
+			Lines.Empty();
+			for (const FString& Line : Unique)
+			{
+				Lines.Add(Line);
+			}
+			Lines.Sort();
+			GeneratedFunctionDeclarations.Empty();
+			for (const FString& Line : Lines)
+			{
+				GeneratedFunctionDeclarations.Logf(TEXT("%s\r\n"), *Line);
+			}
+#if UHT_ENABLE_EXTRA_HASH_OUTPUT
+			GeneratedCPPText.Logf(TEXT("#if 0\r\n"));
+			GeneratedCPPText.Log(GeneratedFunctionDeclarations);
+			GeneratedCPPText.Logf(TEXT("#endif\r\n"));
+#endif
 		}
 	);
 }
@@ -5797,48 +5855,7 @@ bool FNativeClassHeaderGenerator::SaveHeaderIfChanged(FGeneratedFileInfo& FileIn
 		const FString& ProjectSavedDir = FPaths::ProjectSavedDir();
 		const FString CleanFilename = FPaths::GetCleanFilename(FileInfo.GetFilename());
 
-		FString PathPrefix;
-		{
-			FString StdFilename = FileInfo.GetFilename();
-
-			FPaths::MakeStandardFilename(StdFilename);
-
-			bool bRelativePath = FPaths::IsRelative(StdFilename);
-
-			if (!bRelativePath)
-			{
-				// If path is still absolute that means MakeStandardFilename has failed
-				// In this case make it relative to the current project. 
-				bRelativePath = FPaths::MakePathRelativeTo(StdFilename, *FPaths::GetPath(FPaths::GetProjectFilePath()));
-			}
-
-			// If the path has passed either MakeStandardFilename or MakePathRelativeTo it should be using internal path separators
-			if (bRelativePath)
-			{
-				// Remove any preceding parent directory paths
-				while (StdFilename.RemoveFromStart(TEXT("../")));
-			}
-
-			StdFilename = FPaths::GetPath(StdFilename);
-
-			FStringOutputDevice Out;
-
-			for (TCHAR Char : StdFilename)
-			{
-				if (FChar::IsAlnum(Char))
-				{
-					Out.AppendChar(Char);
-				}
-				else
-				{
-					Out.AppendChar(TEXT('_'));
-				}
-			}
-
-			PathPrefix = MoveTemp(Out);
-		}
-
-		const FString Ref = ProjectSavedDir / TEXT("ReferenceGeneratedCode") / (PathPrefix + TEXT("_") + CleanFilename);
+		const FString Ref = ProjectSavedDir / TEXT("ReferenceGeneratedCode") / CleanFilename;
 
 		if (bWriteContents)
 		{
@@ -5850,7 +5867,7 @@ bool FNativeClassHeaderGenerator::SaveHeaderIfChanged(FGeneratedFileInfo& FileIn
 		{
 			{
 				FScopeLock Lock(&WritePacer);
-				const FString Verify = ProjectSavedDir / TEXT("VerifyGeneratedCode") / (PathPrefix + TEXT("_") + CleanFilename);
+				const FString Verify = ProjectSavedDir / TEXT("VerifyGeneratedCode") / CleanFilename;
 				bool Written = FFileHelper::SaveStringToFile(InNewHeaderContents, *Verify);
 				check(Written);
 			}
@@ -6006,55 +6023,35 @@ void GetScriptPlugins(TArray<IScriptGeneratorPluginInterface*>& ScriptPlugins)
 	}
 }
 
-TArray<TSharedRef<FUnrealTypeDefinitionInfo>> GEngineClasses;
-
 /**
  * Tries to resolve super classes for classes defined in the given class
  */
 void ResolveSuperClasses(const TCHAR* PackageName, FUnrealClassDefinitionInfo& ClassDef)
 {
-	UClass* Class = ClassDef.GetClassSafe();
-
-	// Propagate any already set ClassWithin (this only happens for the temporary engine objects)
-	if (Class != nullptr && Class->ClassWithin != nullptr)
-	{
-		ClassDef.SetClassWithin(&GTypeDefinitionInfoMap.FindByNameChecked<FUnrealClassDefinitionInfo>(*Class->ClassWithin->GetFName().ToString()));
-	}
 
 	// Resolve the base class
 	{
 		FUnrealStructDefinitionInfo::FBaseStructInfo& SuperClassInfo = ClassDef.GetSuperStructInfo();
 
-		if (Class == nullptr || Class->GetSuperClass() == nullptr)
+		if (!SuperClassInfo.Name.IsEmpty())
 		{
-			if (!SuperClassInfo.Name.IsEmpty())
-			{
-				const FString& BaseClassName = SuperClassInfo.Name;
-				const FString& BaseClassNameStripped = GetClassNameWithPrefixRemoved(BaseClassName);
-
-				FUnrealClassDefinitionInfo* FoundBaseClassDef = GTypeDefinitionInfoMap.FindByName<FUnrealClassDefinitionInfo>(*BaseClassNameStripped);
-				if (FoundBaseClassDef == nullptr)
-				{
-					FoundBaseClassDef = GTypeDefinitionInfoMap.FindByName<FUnrealClassDefinitionInfo>(*BaseClassName);
-				}
+			const FString& BaseClassName = SuperClassInfo.Name;
+			const FString& BaseClassNameStripped = GetClassNameWithPrefixRemoved(BaseClassName);
 
-				if (FoundBaseClassDef == nullptr)
-				{
-					// Don't know its parent class. Raise error.
-					ClassDef.Throwf(TEXT("Couldn't find parent type for '%s' named '%s' in current module (Package: %s) or any other module parsed so far."), *ClassDef.GetName(), *BaseClassName, PackageName);
-				}
+			FUnrealClassDefinitionInfo* FoundBaseClassDef = GTypeDefinitionInfoMap.FindByName<FUnrealClassDefinitionInfo>(*BaseClassNameStripped);
+			if (FoundBaseClassDef == nullptr)
+			{
+				FoundBaseClassDef = GTypeDefinitionInfoMap.FindByName<FUnrealClassDefinitionInfo>(*BaseClassName);
+			}
 
-				SuperClassInfo.Struct = FoundBaseClassDef;
-				ClassDef.SetClassCastFlags(FoundBaseClassDef->GetClassCastFlags());
-				if (Class != nullptr)
-				{
-					Class->SetSuperStruct(FoundBaseClassDef->GetClass());
-				}
+			if (FoundBaseClassDef == nullptr)
+			{
+				// Don't know its parent class. Raise error.
+				ClassDef.Throwf(TEXT("Couldn't find parent type for '%s' named '%s' in current module (Package: %s) or any other module parsed so far."), *ClassDef.GetName(), *BaseClassName, PackageName);
 			}
-		}
-		else
-		{
-			SuperClassInfo.Struct = &GTypeDefinitionInfoMap.FindByNameChecked<FUnrealClassDefinitionInfo>(*Class->GetSuperClass()->GetFName().ToString());
+
+			SuperClassInfo.Struct = FoundBaseClassDef;
+			ClassDef.SetClassCastFlags(FoundBaseClassDef->GetClassCastFlags());
 		}
 	}
 
@@ -6400,19 +6397,6 @@ void ResolveParents(TArray<FUnrealPackageDefinitionInfo*>& PackageDefs)
 	GUClassDef = &GTypeDefinitionInfoMap.FindByNameChecked<FUnrealClassDefinitionInfo>(*UClass::StaticClass()->GetFName().ToString());
 	GUInterfaceDef = &GTypeDefinitionInfoMap.FindByNameChecked<FUnrealClassDefinitionInfo>(*UInterface::StaticClass()->GetFName().ToString());
 
-	for (TSharedRef<FUnrealTypeDefinitionInfo> TypeDef : GEngineClasses)
-	{
-		FUnrealClassDefinitionInfo& ClassDef = UHTCastChecked<FUnrealClassDefinitionInfo>(TypeDef);
-
-		// Some of the temporary classes are contained within the no export file.  When they are, a new
-		// instance of the definition will be created.  So ignore this one if the found class def doesn't
-		// match this one.
-		if (&GTypeDefinitionInfoMap.FindByNameChecked(*ClassDef.GetFName().ToString()) == &ClassDef)
-		{
-			ResolveSuperClasses(TEXT("EngineObject"), ClassDef);
-		}
-	}
-
 	for (FUnrealPackageDefinitionInfo* PackageDef : PackageDefs)
 	{
 		FResults::Try([PackageDef]() { ResolveSuperClasses(*PackageDef); });
@@ -6608,23 +6592,8 @@ void ParseSourceFiles(TArray<FUnrealSourceFile*>& OrderedSourceFiles)
 
 void PostParseFinalize(TArray<FUnrealPackageDefinitionInfo*>& PackageDefs)
 {
-	// Until we fix the issue with types not having sources, we need to do this
-	TArray<FUnrealTypeDefinitionInfo*> NoSourceTypeDefs;
-	GTypeDefinitionInfoMap.ForAllTypesByName([&NoSourceTypeDefs](FUnrealTypeDefinitionInfo& TypeDef)
+	auto PostParseFinalize = [&PackageDefs](EPostParseFinalizePhase Phase)
 	{
-		if (!TypeDef.HasSource())
-		{
-			NoSourceTypeDefs.Add(&TypeDef);
-		}
-	});
-
-	auto PostParseFinalize = [&PackageDefs, &NoSourceTypeDefs](EPostParseFinalizePhase Phase)
-	{
-		for (FUnrealTypeDefinitionInfo* TypeDef : NoSourceTypeDefs)
-		{
-			TypeDef->PostParseFinalize(Phase);
-		}
-
 		for (FUnrealPackageDefinitionInfo* PackageDef : PackageDefs)
 		{
 			FResults::Try([PackageDef, Phase]()
@@ -6634,22 +6603,9 @@ void PostParseFinalize(TArray<FUnrealPackageDefinitionInfo*>& PackageDefs)
 		}
 	};
 
-	for (FUnrealTypeDefinitionInfo* TypeDef : NoSourceTypeDefs)
-	{
-		TypeDef->ConcurrentPostParseFinalize();
-	}
-
 	PostParseFinalize(EPostParseFinalizePhase::Phase1);
 	PostParseFinalize(EPostParseFinalizePhase::Phase2);
 	FResults::WaitForErrorTasks();
-
-	for (FUnrealTypeDefinitionInfo* TypeDef : NoSourceTypeDefs)
-	{
-		if (FUnrealObjectDefinitionInfo* ObjectDef = UHTCast<FUnrealObjectDefinitionInfo>(TypeDef))
-		{
-			check(ObjectDef->GetObject());
-		}
-	}
 }
 
 void CreateEngineTypes(TArray<FUnrealPackageDefinitionInfo*>& PackageDefs)
@@ -6737,11 +6693,8 @@ void ExportToScriptPlugins(TArray<IScriptGeneratorPluginInterface*>& ScriptPlugi
 	{
 		for (TSharedRef<FUnrealTypeDefinitionInfo> TypeDef : PackageDef->GetAllClasses())
 		{
-			if (TypeDef->HasSource())
-			{
-				UClass* Class = UHTCastChecked<FUnrealClassDefinitionInfo>(TypeDef).GetClass();
-				SourceFileLookup.Add(Class, &TypeDef->GetUnrealSourceFile());
-			}
+			UClass* Class = UHTCastChecked<FUnrealClassDefinitionInfo>(TypeDef).GetClass();
+			SourceFileLookup.Add(Class, &TypeDef->GetUnrealSourceFile());
 		}
 	}
 
@@ -6837,22 +6790,6 @@ ECompilationResult::Type UnrealHeaderTool_Main(const FString& ModuleInfoFilename
 
 	double TotalPrepareModuleTime = FResults::TimedTry([&PackageDefs, &ModuleInfoPath]() { PrepareModules(PackageDefs, ModuleInfoPath); });
 
-	// TEMPORARY!!!
-	// At this time, there is a collection of classes that aren't listed in the NoExport
-	// file.  This means that we don't have an associated type definition for them.
-	// Force the creation of type data for them.  Once these are added to NoExport, then
-	// the HasSource can be looked at to be removed.  However, it might change output slightly.  
-	// Use -WRITEREF and -VERIFYREF to detect these changes.
-	FResults::Try([]()
-	{
-	for (TObjectIterator<UClass> ClassIt; ClassIt; ++ClassIt)
-	{
-		TSharedRef<FUnrealClassDefinitionInfo> ClassDefRef = MakeShared<FUnrealClassDefinitionInfo>(*ClassIt);
-		GTypeDefinitionInfoMap.AddNameLookup(*ClassDefRef);
-		GEngineClasses.Add(ClassDefRef);
-	}
-	});
-
 	FString ExternalDependencies;
 	TArray<IScriptGeneratorPluginInterface*> ScriptPlugins;
 
@@ -6863,6 +6800,16 @@ ECompilationResult::Type UnrealHeaderTool_Main(const FString& ModuleInfoFilename
 	double TotalTopologicalSortTime = FResults::TimedTry([&OrderedSourceFiles]() { TopologicalSort(OrderedSourceFiles); });
 	double TotalParseTime = FResults::TimedTry([&OrderedSourceFiles]() { ParseSourceFiles(OrderedSourceFiles); });
 	double TotalPostParseFinalizeTime = FResults::TimedTry([&PackageDefs]() { PostParseFinalize(PackageDefs); });
+
+	// Look for any core classes that don't have a definition
+	for (TObjectIterator<UClass> ClassIt; ClassIt; ++ClassIt)
+	{
+		if (GTypeDefinitionInfoMap.FindByName(*ClassIt->GetFName().ToString()) == nullptr)
+		{
+			UE_LOG(LogCompile, Log, TEXT("The core class '%s' doesn't have a matching entry in NoExportTypes.h"), *ClassIt->GetFName().ToString());
+		}
+	}
+
 	TotalTopologicalSortTime += FResults::TimedTry([&OrderedSourceFiles]() { TopologicalSort(OrderedSourceFiles); }); // Sort again to include new dependencies
 	double TotalCodeGenTime = FResults::TimedTry([&PackageDefs, &OrderedSourceFiles]() { Export(PackageDefs, OrderedSourceFiles); });
 	double TotalCheckForScriptPluginsTime = FResults::TimedTry([&ScriptPlugins]() { GetScriptPlugins(ScriptPlugins); });
@@ -7004,22 +6951,14 @@ void ProcessParsedClass(FUnrealClassDefinitionInfo& ClassDef)
 
 	const static bool bVerboseOutput = FParse::Param(FCommandLine::Get(), TEXT("VERBOSE"));
 
-	if (ResultClass == nullptr)
+	if (TSharedRef<FUnrealTypeDefinitionInfo>* Existing = GTypeDefinitionInfoMap.FindByName(*ClassNameStripped))
 	{
-		if (TSharedRef<FUnrealTypeDefinitionInfo>* Existing = GTypeDefinitionInfoMap.FindByName(*ClassNameStripped))
-		{
-			ClassDef.Throwf(TEXT("Duplicate class name: %s also exists in file %s"), *ClassName, *(*Existing)->GetFilename());
-		}
-
-		if (bVerboseOutput)
-		{
-			UE_LOG(LogCompile, Log, TEXT("Imported: %s"), *ClassDef.GetFullName());
-		}
+		ClassDef.Throwf(TEXT("Duplicate class name: %s also exists in file %s"), *ClassName, *(*Existing)->GetFilename());
 	}
-	else
+
+	if (bVerboseOutput)
 	{
-		ClassDef.InitializeFromExistingUObject(ResultClass);
-		ClassDef.SetObject(ResultClass);
+		UE_LOG(LogCompile, Log, TEXT("Imported: %s"), *ClassDef.GetFullName());
 	}
 }
 
diff --git a/Engine/Source/Programs/UnrealHeaderTool/Private/HeaderParser.cpp b/Engine/Source/Programs/UnrealHeaderTool/Private/HeaderParser.cpp
--- a/Engine/Source/Programs/UnrealHeaderTool/Private/HeaderParser.cpp
+++ b/Engine/Source/Programs/UnrealHeaderTool/Private/HeaderParser.cpp
@@ -1545,7 +1545,7 @@ FString FHeaderParser::FormatCommentForToolTip(const FString& Input)
 			int32 TemporaryMaxWhitespace = MaxNumWhitespaceToRemove;
 
 			// Allow eating an extra tab on subsequent lines if it's present
-			if ((Index > 0) && (Line.Len() > 0) && (Line[0] == '\t'))
+			if ((Index > FirstIndex) && (Line.Len() > 0) && (Line[0] == '\t'))
 			{
 				TemporaryMaxWhitespace++;
 			}
@@ -1562,7 +1562,7 @@ FString FHeaderParser::FormatCommentForToolTip(const FString& Input)
 				Line.RightChopInline(Pos, false);
 			}
 
-			if (Index > 0)
+			if (Index > FirstIndex)
 			{
 				Result += TEXT("\n");
 			}
@@ -1907,6 +1907,7 @@ FUnrealScriptStructDefinitionInfo& FHeaderParser::CompileStructDeclaration()
 	{
 		ClearComment();
 		GetToken( Token );
+		int StartingLine = InputLine;
 
 		if (EAccessSpecifier AccessSpecifier = ParseAccessProtectionSpecifier(Token))
 		{
@@ -1959,6 +1960,16 @@ FUnrealScriptStructDefinitionInfo& FHeaderParser::CompileStructDeclaration()
 		{
 			PopCompilerDirective();
 			// Do nothing and hope that the if code below worked out OK earlier
+
+			// skip it and skip over the text, it is not recorded or processed
+			if (StartingLine == InputLine)
+			{
+				TCHAR c;
+				while (!IsEOL(c = GetChar()))
+				{
+				}
+				ClearComment();
+			}
 		}
 		else if ( Token.IsSymbol(TEXT('#')) && MatchIdentifier(TEXT("if"), ESearchCase::CaseSensitive) )
 		{
@@ -2781,6 +2792,9 @@ void FHeaderParser::CompileDirective()
 		{
 			UngetChar();
 		}
+
+		// Remove any comments parsed to prevent any trailing comments from being picked up
+		ClearComment();
 	}
 }
 
@@ -4823,7 +4837,8 @@ bool FHeaderParser::CompileDeclaration(TArray<FUnrealFunctionDefinitionInfo*>& D
 		bEncounteredNewStyleClass_UnmatchedBrackets = false;
 
 		// Pop nesting here to allow other non UClass declarations in the header file.
-		if (CurrentClassDef.HasAllClassFlags(CLASS_Interface))
+		// Make sure we treat UInterface as a class and not an interface
+		if (CurrentClassDef.HasAllClassFlags(CLASS_Interface) && &CurrentClassDef != GUInterfaceDef)
 		{
 			checkf(TopNest->NestType == ENestType::Interface || TopNest->NestType == ENestType::NativeInterface, TEXT("Unexpected end of interface block."));
 			PopNest(TopNest->NestType, TEXT("'Interface'"));
@@ -4843,7 +4858,14 @@ bool FHeaderParser::CompileDeclaration(TArray<FUnrealFunctionDefinitionInfo*>& D
 		}
 		else
 		{
-			PopNest(ENestType::Class, TEXT("'Class'"));
+			if (&CurrentClassDef == GUInterfaceDef && TopNest->NestType == ENestType::NativeInterface)
+			{
+				PopNest(TopNest->NestType, TEXT("'Interface'"));
+			}
+			else
+			{
+				PopNest(ENestType::Class, TEXT("'Class'"));
+			}
 			PostPopNestClass(CurrentClassDef);
 
 			// Ensure classes have a GENERATED_BODY declaration
@@ -5326,6 +5348,7 @@ FUnrealClassDefinitionInfo& FHeaderParser::ParseClassNameDeclaration(FString& De
 	FUnrealClassDefinitionInfo* ClassDef = FUnrealClassDefinitionInfo::FindClass(*GetClassNameWithPrefixRemoved(*DeclaredClassName));
 	check(ClassDef);
 	ClassDef->SetClassCastFlags(ClassCastFlagMap::Get().GetCastFlag(DeclaredClassName));
+	ClassDef->SetClassCastFlags(ClassCastFlagMap::Get().GetCastFlag(FString(TEXT("F")) + DeclaredClassName.RightChop(1))); // For properties, check alternate name
 
 	// Skip optional final keyword
 	MatchIdentifier(TEXT("final"), ESearchCase::CaseSensitive);
@@ -5557,6 +5580,7 @@ bool FHeaderParser::TryParseIInterfaceClass()
 	// Push the interface class nesting again.
 	PushNest(ENestType::NativeInterface, FoundClassDef);
 
+	CurrentAccessSpecifier = ACCESS_Private;
 	return true;
 }
 
@@ -5569,6 +5593,8 @@ void FHeaderParser::CompileInterfaceDeclaration()
 	// whereever a class declaration is allowed, an interface declaration is also allowed.
 	CheckAllow( TEXT("'interface'"), ENestAllowFlags::Class );
 
+	CurrentAccessSpecifier = ACCESS_Private;
+
 	FString DeclaredInterfaceName;
 	FString RequiredAPIMacroIfPresent;
 	TMap<FName, FString> MetaData;
@@ -8409,6 +8435,34 @@ TSharedRef<FUnrealTypeDefinitionInfo> FHeaderPreParser::ParseStructDeclaration(c
 			StructDef->SetStructFlags(STRUCT_Atomic);
 		}
 		break;
+
+		case EStructSpecifier::HasDefaults:
+			if (!SourceFile.IsNoExportTypes())
+			{
+				LogError(TEXT("The 'HasDefaults' struct specifier is only valid in the NoExportTypes.h file"));
+			}
+			break;
+
+		case EStructSpecifier::HasNoOpConstructor:
+			if (!SourceFile.IsNoExportTypes())
+			{
+				LogError(TEXT("The 'HasNoOpConstructor' struct specifier is only valid in the NoExportTypes.h file"));
+			}
+			break;
+
+		case EStructSpecifier::IsAlwaysAccessible:
+			if (!SourceFile.IsNoExportTypes())
+			{
+				LogError(TEXT("The 'IsAlwaysAccessible' struct specifier is only valid in the NoExportTypes.h file"));
+			}
+			break;
+
+		case EStructSpecifier::IsCoreType:
+			if (!SourceFile.IsNoExportTypes())
+			{
+				LogError(TEXT("The 'IsCoreType' struct specifier is only valid in the NoExportTypes.h file"));
+			}
+			break;
 		}
 	}
 
@@ -8591,6 +8645,8 @@ void FHeaderParser::ResetClassData()
 {
 	FUnrealClassDefinitionInfo& CurrentClassDef = GetCurrentClassDef();
 
+	check(CurrentClassDef.GetClassWithin() == nullptr);
+
 	// Set class flags and within.
 	CurrentClassDef.ClearClassFlags(CLASS_RecompilerClear);
 
@@ -8600,10 +8656,7 @@ void FHeaderParser::ResetClassData()
 		CurrentClassDef.SetClassConfigName(SuperClassDef->GetClassConfigName());
 
 		check(SuperClassDef->GetClassWithin());
-		if (CurrentClassDef.GetClassWithin() == nullptr)
-		{
-			CurrentClassDef.SetClassWithin(SuperClassDef->GetClassWithin());
-		}
+		CurrentClassDef.SetClassWithin(SuperClassDef->GetClassWithin());
 
 		// Copy special categories from parent
 		if (SuperClassDef->HasMetaData(FHeaderParserNames::NAME_HideCategories))
@@ -8635,6 +8688,10 @@ void FHeaderParser::ResetClassData()
 			CurrentClassDef.SetMetaData(FHeaderParserNames::NAME_PrioritizeCategories, *SuperClassDef->GetMetaData(FHeaderParserNames::NAME_PrioritizeCategories));
 		}
 	}
+	else
+	{
+		CurrentClassDef.SetClassWithin(GUObjectDef);
+	}
 
 	check(CurrentClassDef.GetClassWithin());
 }
diff --git a/Engine/Source/Programs/UnrealHeaderTool/Private/PropertyTypes.cpp b/Engine/Source/Programs/UnrealHeaderTool/Private/PropertyTypes.cpp
--- a/Engine/Source/Programs/UnrealHeaderTool/Private/PropertyTypes.cpp
+++ b/Engine/Source/Programs/UnrealHeaderTool/Private/PropertyTypes.cpp
@@ -996,7 +996,6 @@ struct FPropertyTypeTraitsBooleanBase : public FPropertyTypeTraitsBase
 		{
 			return SizeText;
 		}
-		return TEXT("uint32");
 	}
 
 	static bool DefaultValueStringCppFormatToInnerFormat(const FUnrealPropertyDefinitionInfo& PropDef, const FString& CppForm, FString& OutForm)
diff --git a/Engine/Source/Programs/UnrealHeaderTool/Private/Specifiers/ClassMetadataSpecifiers.def b/Engine/Source/Programs/UnrealHeaderTool/Private/Specifiers/ClassMetadataSpecifiers.def
--- a/Engine/Source/Programs/UnrealHeaderTool/Private/Specifiers/ClassMetadataSpecifiers.def
+++ b/Engine/Source/Programs/UnrealHeaderTool/Private/Specifiers/ClassMetadataSpecifiers.def
@@ -29,7 +29,9 @@ CLASS_METADATA_SPECIFIER(Hidden)
 CLASS_METADATA_SPECIFIER(HideCategories)
 CLASS_METADATA_SPECIFIER(HideDropdown)
 CLASS_METADATA_SPECIFIER(HideFunctions)
+CLASS_METADATA_SPECIFIER(Interface)
 CLASS_METADATA_SPECIFIER(Intrinsic)
+CLASS_METADATA_SPECIFIER(MatchedSerializers)
 CLASS_METADATA_SPECIFIER(MinimalAPI)
 CLASS_METADATA_SPECIFIER(NeedsDeferredDependencyLoading)
 CLASS_METADATA_SPECIFIER(NoExport)
diff --git a/Engine/Source/Programs/UnrealHeaderTool/Private/Specifiers/StructSpecifiers.def b/Engine/Source/Programs/UnrealHeaderTool/Private/Specifiers/StructSpecifiers.def
--- a/Engine/Source/Programs/UnrealHeaderTool/Private/Specifiers/StructSpecifiers.def
+++ b/Engine/Source/Programs/UnrealHeaderTool/Private/Specifiers/StructSpecifiers.def
@@ -5,5 +5,9 @@
 // [[ IncludeTool: Inline ]] // Markup to tell IncludeTool that this file is state changing and cannot be optimized out.
 
 STRUCT_SPECIFIER(Atomic)
+STRUCT_SPECIFIER(HasDefaults)
+STRUCT_SPECIFIER(HasNoOpConstructor)
 STRUCT_SPECIFIER(Immutable)
+STRUCT_SPECIFIER(IsAlwaysAccessible)
+STRUCT_SPECIFIER(IsCoreType)
 STRUCT_SPECIFIER(NoExport)
diff --git a/Engine/Source/Programs/UnrealHeaderTool/Private/UnrealHeaderTool.h b/Engine/Source/Programs/UnrealHeaderTool/Private/UnrealHeaderTool.h
--- a/Engine/Source/Programs/UnrealHeaderTool/Private/UnrealHeaderTool.h
+++ b/Engine/Source/Programs/UnrealHeaderTool/Private/UnrealHeaderTool.h
@@ -22,6 +22,9 @@
 // Enable concurrent code generation
 #define UHT_ENABLE_CONCURRENT_CODE_GENERATION 1
 
+// Enable the output of extra code gen that contributes to hashes
+#define UHT_ENABLE_EXTRA_HASH_OUTPUT 0
+
 struct FHeaderParserNames
 {
 	static const FName NAME_IsConversionRoot;
diff --git a/Engine/Source/Programs/UnrealHeaderTool/Private/UnrealTypeDefinitionInfo.cpp b/Engine/Source/Programs/UnrealHeaderTool/Private/UnrealTypeDefinitionInfo.cpp
--- a/Engine/Source/Programs/UnrealHeaderTool/Private/UnrealTypeDefinitionInfo.cpp
+++ b/Engine/Source/Programs/UnrealHeaderTool/Private/UnrealTypeDefinitionInfo.cpp
@@ -3,6 +3,7 @@
 #include "UnrealTypeDefinitionInfo.h"
 #include "BaseParser.h"
 #include "ClassMaps.h"
+#include "EngineAPI.h"
 #include "HeaderParser.h"
 #include "NativeClassExporter.h"
 #include "PropertyTypes.h"
@@ -1765,12 +1766,17 @@ void FUnrealStructDefinitionInfo::CreateUObjectEngineTypesInternal(ECreateEngine
 		GetStruct()->Bind();
 
 		// Internals will assert of we are relinking an intrinsic 
+		bool bStaticLink = true;
 		bool bRelinkExistingProperties = true;
 		if (FUnrealClassDefinitionInfo* ClassDef = UHTCast<FUnrealClassDefinitionInfo>(this))
 		{
 			bRelinkExistingProperties = !ClassDef->HasAnyClassFlags(CLASS_Intrinsic);
+			bStaticLink = !ClassDef->IsCompiledIn();
+		}
+		if (bStaticLink)
+		{
+			GetStruct()->StaticLink(bRelinkExistingProperties);
 		}
-		GetStruct()->StaticLink(bRelinkExistingProperties);
 		break;
 	}
 }
@@ -2111,14 +2117,21 @@ void FUnrealClassDefinitionInfo::CreateUObjectEngineTypesInternal(ECreateEngineT
 			Throwf(TEXT("ClassWithin is not set.  This can be the result of the class's UCLASS not being parsed due to statement skipping preceding the UCLASS keyword."));
 		}
 		check(ClassWithin);
-		if (GetObjectSafe() == nullptr)
+		check(GetObjectSafe() == nullptr);
+
+		const FString& ClassName = GetNameCPP();
+		FString ClassNameStripped = GetClassNameWithPrefixRemoved(*ClassName);
+
+		// See if we have an existing engine object we should use
+		UClass* Class = FEngineAPI::FindObject<UClass>(ANY_PACKAGE, *ClassNameStripped);
+
+		if (Class == nullptr)
 		{
 			UPackage* Package = GetPackageDef().GetPackage();
-			const FString& ClassName = GetNameCPP();
-			FString ClassNameStripped = GetClassNameWithPrefixRemoved(*ClassName);
 
 			// Create new class.
-			UClass* Class = new(EC_InternalUseOnlyConstructor, Package, *ClassNameStripped, RF_Public | RF_Standalone) UClass(FObjectInitializer(), nullptr);
+			Class = new(EC_InternalUseOnlyConstructor, Package, *ClassNameStripped, RF_Public | RF_Standalone) UClass(FObjectInitializer(), nullptr);
+
 
 			Class->ClassFlags = ClassFlags;
 			Class->ClassCastFlags = ClassCastFlags;
@@ -2145,54 +2158,66 @@ void FUnrealClassDefinitionInfo::CreateUObjectEngineTypesInternal(ECreateEngineT
 					}
 				}
 			}
-			SetObject(Class);
 		}
+		else
+		{
+			bIsCompiledIn = true;
+		}
+		SetObject(Class);
 		break;
 	}
 
 	case ECreateEngineTypesPhase::Phase2:
 	{
 
-		// Initialize the class object
-		UClass* Class = GetClass();
+		// Only do this for UHT types
+		if (!bIsCompiledIn)
+		{
+			// Initialize the class object
+			UClass* Class = GetClass();
 
-		// Clear the property size
-		Class->PropertiesSize = 0;
+			// Clear the property size
+			Class->PropertiesSize = 0;
 
-		// Make the visible to the package
-		Class->ClearFlags(RF_Transient);
-		check(Class->HasAnyFlags(RF_Public));
-		check(Class->HasAnyFlags(RF_Standalone));
+			// Make the visible to the package
+			Class->ClearFlags(RF_Transient);
+			check(Class->HasAnyFlags(RF_Public));
+			check(Class->HasAnyFlags(RF_Standalone));
 
-		// Finalize all of the children introduced in this class
-		for (TSharedRef<FUnrealFunctionDefinitionInfo> FunctionDef : GetFunctions())
-		{
-			UFunction* Function = FunctionDef->GetFunction();
-			Class->AddFunctionToFunctionMap(Function, Function->GetFName());
+			// Finalize all of the children introduced in this class
+			for (TSharedRef<FUnrealFunctionDefinitionInfo> FunctionDef : GetFunctions())
+			{
+				UFunction* Function = FunctionDef->GetFunction();
+				Class->AddFunctionToFunctionMap(Function, Function->GetFName());
+			}
 		}
 		break;
 	}
 
 	case ECreateEngineTypesPhase::Phase3:
 	{
+		// Only do this for UHT types
+		if (!bIsCompiledIn)
+		{
 
-		// Initialize the class object
-		UClass* Class = GetClass();
+			// Initialize the class object
+			UClass* Class = GetClass();
 
-		if (!HasAnyClassFlags(CLASS_Native))
-		{
-			//Class->UnMark(EObjectMark(OBJECTMARK_TagImp | OBJECTMARK_TagExp));
-		}
-		else if (!HasAnyClassFlags(CLASS_NoExport | CLASS_Intrinsic))
-		{
-			//Class->UnMark(OBJECTMARK_TagImp);
-			//Class->Mark(OBJECTMARK_TagExp);
-		}
+			if (!HasAnyClassFlags(CLASS_Native))
+			{
+				//Class->UnMark(EObjectMark(OBJECTMARK_TagImp | OBJECTMARK_TagExp));
+			}
+			else if (!HasAnyClassFlags(CLASS_NoExport | CLASS_Intrinsic))
+			{
+				//Class->UnMark(OBJECTMARK_TagImp);
+				//Class->Mark(OBJECTMARK_TagExp);
+			}
 
-		// This needs to be done outside of parallel blocks because it will modify UClass memory.
-		// Later calls to SetUpUhtReplicationData inside parallel blocks should be fine, because
-		// they will see the memory has already been set up, and just return the parent pointer.
-		Class->SetUpUhtReplicationData();
+			// This needs to be done outside of parallel blocks because it will modify UClass memory.
+			// Later calls to SetUpUhtReplicationData inside parallel blocks should be fine, because
+			// they will see the memory has already been set up, and just return the parent pointer.
+			Class->SetUpUhtReplicationData();
+		}
 		break;
 	}
 	}
@@ -2271,6 +2296,52 @@ void FUnrealClassDefinitionInfo::PostParseFinalizeInternal(EPostParseFinalizePha
 		{
 			GetPackageDef().SetWriteClassesH(true);
 		}
+
+		// See if we have an existing engine object we should use
+		if (UClass* Class = FEngineAPI::FindObject<UClass>(ANY_PACKAGE, *GetFName().ToString()))
+		{
+
+			// Check for differences in ClassFlags
+			const EClassFlags CheckFlags = CLASS_SaveInCompiledInClasses & ~CLASS_NoExport & ~CLASS_Intrinsic;
+			if ((ClassFlags & CheckFlags) != (Class->ClassFlags & CheckFlags))
+			{
+				LogError(TEXT("Class flags in core class '%s', is '0x%08x', should be '0x%08x'"), *GetNameCPP(), ClassFlags & CheckFlags, Class->ClassFlags & CheckFlags);
+			}
+
+			// Check for differences in super class
+			{
+				FName SuperClassName = Class->GetSuperClass() ? Class->GetSuperClass()->GetFName() : NAME_None;
+				FName SuperClassDefName = GetSuperClass() ? GetSuperClass()->GetFName() : NAME_None;
+				if (SuperClassName != SuperClassDefName)
+				{
+					LogError(TEXT("Super class mismatch in core class '%s', is '%s', should be '%s'"), *GetNameCPP(), *SuperClassDefName.ToString(), *SuperClassName.ToString());
+				}
+			}
+
+			// Check package
+			{
+				FName PackageName = Class->GetOutermost()->GetFName();
+				FName PackageDefName = GetPackageDef().GetFName();
+				if (PackageName != PackageDefName)
+				{
+					LogError(TEXT("Package mismatch in core class '%s', is '%s', should be '%s'"), *GetNameCPP(), *PackageDefName.ToString(), *PackageName.ToString());
+				}
+			}
+
+			// Check class config name
+			{
+				if (GetClassConfigName() != Class->ClassConfigName && GetNameCPP() != TEXT("UObject"))
+				{
+					LogError(TEXT("Class config name in core class '%s', is '%s', should be '%s'"), *GetNameCPP(), *GetClassConfigName().ToString(), *Class->ClassConfigName.ToString());
+				}
+			}
+
+			// Check ClassCastFlags
+			if (ClassCastFlags != Class->ClassCastFlags)
+			{
+				LogError(TEXT("Class cast flags in core class '%s', is '0x%016x', should be '0x%016x'"), *GetNameCPP(), ClassCastFlags, Class->ClassCastFlags);
+			}
+		}
 		break;
 	}
 	}
@@ -2298,24 +2369,6 @@ bool FUnrealClassDefinitionInfo::ImplementsInterface(const FUnrealClassDefinitio
 	return false;
 }
 
-void FUnrealClassDefinitionInfo::InitializeFromExistingUObject(UClass* Class)
-{
-	ClassFlags = Class->ClassFlags;
-	ClassCastFlags = Class->ClassCastFlags;
-	InitialEngineClassFlags = Class->ClassFlags;
-	PropertiesSize = Class->PropertiesSize;
-	ClassConfigName = Class->ClassConfigName;
-	SetInternalFlags(Class->GetInternalFlags());
-
-	UPackage* Package = Class->GetOutermost();
-	FUnrealPackageDefinitionInfo* PackageDef = GTypeDefinitionInfoMap.FindByName<FUnrealPackageDefinitionInfo>(*Package->GetName());
-	if (PackageDef == nullptr)
-	{
-		Throwf(TEXT("Unable to find package %s"), *Package->GetName());
-	}
-	SetOuter(PackageDef);
-}
-
 void FUnrealClassDefinitionInfo::ParseClassProperties(TArray<FPropertySpecifier>&& InClassSpecifiers, const FString& InRequiredAPIMacroIfPresent)
 {
 	ParsedClassFlags = CLASS_None;
@@ -2615,6 +2668,23 @@ void FUnrealClassDefinitionInfo::ParseClassProperties(TArray<FPropertySpecifier>
 			ParsedClassFlags |= CLASS_NeedsDeferredDependencyLoading;
 			break;
 
+		case EClassMetadataSpecifier::MatchedSerializers:
+			if (!GetUnrealSourceFile().IsNoExportTypes())
+			{
+				LogError(TEXT("The 'MatchedSerializers' class specifier is only valid in the NoExportTypes.h file"));
+			}
+			ParsedClassFlags |= CLASS_MatchedSerializers;
+			break;
+
+		case EClassMetadataSpecifier::Interface:
+			if (!GetUnrealSourceFile().IsNoExportTypes())
+			{
+				LogError(TEXT("The 'Interface' class specifier is only valid in the NoExportTypes.h file"));
+			}
+			ParsedClassFlags |= CLASS_Interface;
+			ParsedInterface = EParsedInterface::ParsedUInterface;
+			break;
+
 		default:
 			Throwf(TEXT("Unknown class specifier '%s'"), *PropSpecifier.Key);
 		}
@@ -2733,35 +2803,6 @@ void FUnrealClassDefinitionInfo::MergeAndValidateClassFlags(const FString& Decla
 	{
 		Throwf(TEXT("Class name '%s' is invalid, should be identified as '%s'"), *DeclaredClassName, *ExpectedClassName);
 	}
-
-	// This check only works if we already have an object.  This has to be moved to the engine type create code
-	if (InitialEngineClassFlags != CLASS_None)
-	{
-
-		// if the class's class flags didn't contain CLASS_NoExport before it was parsed, it means either:
-		// a) the DECLARE_CLASS macro for this native class doesn't contain the CLASS_NoExport flag (this is an error)
-		// b) this is a new native class, which isn't yet hooked up to static registration (this is OK)
-		if (HasAnyClassFlags(CLASS_NoExport) && !EnumHasAnyFlags(InitialEngineClassFlags, CLASS_NoExport))
-		{
-			if (!HasAnyClassFlags(CLASS_Intrinsic) && EnumHasAnyFlags(InitialEngineClassFlags, CLASS_Native))
-			{
-				Throwf(TEXT("'noexport': Must include CLASS_NoExport in native class declaration"));
-			}
-		}
-
-		if (!HasAnyClassFlags(CLASS_Abstract) && EnumHasAnyFlags(InitialEngineClassFlags, CLASS_Abstract))
-		{
-			if (HasAnyClassFlags(CLASS_NoExport))
-			{
-				Throwf(TEXT("'abstract': NoExport class missing abstract keyword from class declaration (must change C++ version first)"));
-				SetClassFlags(CLASS_Abstract);
-			}
-			else if (IsNative())
-			{
-				Throwf(TEXT("'abstract': missing abstract keyword from class declaration - class will no longer be exported as abstract"));
-			}
-		}
-	}
 }
 
 void FUnrealClassDefinitionInfo::SetAndValidateConfigName()
diff --git a/Engine/Source/Programs/UnrealHeaderTool/Private/UnrealTypeDefinitionInfo.h b/Engine/Source/Programs/UnrealHeaderTool/Private/UnrealTypeDefinitionInfo.h
--- a/Engine/Source/Programs/UnrealHeaderTool/Private/UnrealTypeDefinitionInfo.h
+++ b/Engine/Source/Programs/UnrealHeaderTool/Private/UnrealTypeDefinitionInfo.h
@@ -2300,10 +2300,6 @@ public:
 	virtual void ConcurrentPostParseFinalize() override;
 
 protected:
-	explicit FUnrealStructDefinitionInfo(UObject* Object)
-		: FUnrealFieldDefinitionInfo(Object)
-	{}
-
 	FUnrealStructDefinitionInfo(FUnrealSourceFile& InSourceFile, int32 InLineNumber, FString&& InNameCPP, FName InName, FUnrealObjectDefinitionInfo& InOuter);
 
 	/**
@@ -2765,12 +2761,6 @@ class FUnrealClassDefinitionInfo
 public:
 	FUnrealClassDefinitionInfo(FUnrealSourceFile& InSourceFile, int32 InLineNumber, FString&& InNameCPP, FName InName, bool bInIsInterface);
 
-	explicit FUnrealClassDefinitionInfo(UClass* Class)
-		: FUnrealStructDefinitionInfo(Class)
-	{
-		InitializeFromExistingUObject(Class);
-	}
-
 	/**
 	 * Attempts to find a class definition based on the given name
 	 */
@@ -3338,6 +3328,14 @@ public:
 		bCustomConstructor = true;
 	}
 
+	/**
+	 * Return true if the class is compiled in
+	 */
+	bool IsCompiledIn() const
+	{
+		return bIsCompiledIn;
+	}
+
 	/**
 	 * Get the generated body access specifier
 	 */
@@ -3375,12 +3373,6 @@ public:
 	 */
 	bool ImplementsInterface(const FUnrealClassDefinitionInfo& SomeInterface) const;
 
-	/**
-	 * Invoked when the definition has an existing engine class
-	 */
-	void InitializeFromExistingUObject(UClass* Class);
-
-
 	FString GetNameWithPrefix(EEnforceInterfacePrefix EnforceInterfacePrefix = EEnforceInterfacePrefix::None) const;
 
 protected:
@@ -3423,7 +3415,6 @@ private:
 	EClassFlags ClassFlags = CLASS_None; // Current class flags
 	EClassFlags ParsedClassFlags = CLASS_None; // Class flags from the preparse phase
 	EClassFlags InheritClassFlags = CLASS_ScriptInherit; // Flags to inherit from the super class
-	EClassFlags InitialEngineClassFlags = CLASS_None; // Class flags from a preexisting UObject
 	EClassCastFlags ClassCastFlags = CASTCLASS_None;
 	FUnrealClassDefinitionInfo* ClassWithin = nullptr;
 	ESerializerArchiveType ArchiveType = ESerializerArchiveType::None;
@@ -3466,6 +3457,9 @@ private:
 
 	/** True if the class is an interface */
 	bool bIsInterface = false;
+
+	/** True if the class is a compiled-in type and not created by UHT */
+	bool bIsCompiledIn = false;
 };
 
 template <typename To>
diff --git a/Engine/Source/Runtime/CoreUObject/Public/UObject/NoExportTypes.h b/Engine/Source/Runtime/CoreUObject/Public/UObject/NoExportTypes.h
--- a/Engine/Source/Runtime/CoreUObject/Public/UObject/NoExportTypes.h
+++ b/Engine/Source/Runtime/CoreUObject/Public/UObject/NoExportTypes.h
@@ -440,13 +440,13 @@ namespace EAppMsgType
 /**
  * A struct used as stub for deleted ones. 
  */
-USTRUCT(noexport)
+USTRUCT(noexport, IsAlwaysAccessible, HasDefaults)
 struct FFallbackStruct
 {
 };
 
 /** A globally unique identifier (mirrored from Guid.h) */
-USTRUCT(immutable, noexport, BlueprintType)
+USTRUCT(immutable, noexport, BlueprintType, IsAlwaysAccessible, HasDefaults)
 struct FGuid
 {
 	UPROPERTY(EditAnywhere, SaveGame, Category=Guid)
@@ -466,7 +466,7 @@ struct FGuid
  * A point or direction FVector in 3d space.
  * @note The full C++ class is located here: Engine\Source\Runtime\Core\Public\Math\Vector.h
  */
-USTRUCT(immutable, noexport, BlueprintType, meta = (HasNativeBreak = "Engine.KismetMathLibrary.BreakVector3f"))
+USTRUCT(immutable, noexport, BlueprintType, IsAlwaysAccessible, HasDefaults, HasNoOpConstructor, IsCoreType, meta = (HasNativeBreak = "Engine.KismetMathLibrary.BreakVector3f"))
 struct FVector3f
 {
 	UPROPERTY(EditAnywhere, Category = Vector, SaveGame)
@@ -483,7 +483,7 @@ struct FVector3f
  * A point or direction FVector in 3d space.
  * @note The full C++ class is located here: Engine\Source\Runtime\Core\Public\Math\Vector.h
  */
-USTRUCT(immutable, noexport)
+USTRUCT(immutable, noexport, IsAlwaysAccessible, HasDefaults, HasNoOpConstructor, IsCoreType)
 struct FVector3d
 {
 	UPROPERTY(EditAnywhere, Category = Vector, SaveGame)
@@ -500,7 +500,7 @@ struct FVector3d
  * A point or direction FVector in 3d space.
  * @note The full C++ class is located here: Engine\Source\Runtime\Core\Public\Math\Vector.h
  */
-USTRUCT(immutable, noexport, BlueprintType, meta = (HasNativeMake = "Engine.KismetMathLibrary.MakeVector", HasNativeBreak = "Engine.KismetMathLibrary.BreakVector"))
+USTRUCT(immutable, noexport, BlueprintType, IsAlwaysAccessible, HasDefaults, HasNoOpConstructor, IsCoreType, meta = (HasNativeMake = "Engine.KismetMathLibrary.MakeVector", HasNativeBreak = "Engine.KismetMathLibrary.BreakVector"))
 struct FVector
 {
 	UPROPERTY(EditAnywhere, BlueprintReadWrite, Category = Vector, SaveGame)
@@ -518,7 +518,7 @@ struct FVector
 * A 4-D homogeneous vector.
 * @note The full C++ class is located here: Engine\Source\Runtime\Core\Public\Math\Vector4.h
 */
-USTRUCT(immutable, noexport)
+USTRUCT(immutable, noexport, IsAlwaysAccessible, HasDefaults, HasNoOpConstructor, IsCoreType)
 struct FVector4f
 {
 	UPROPERTY(EditAnywhere, Category = Vector4, SaveGame)
@@ -538,7 +538,7 @@ struct FVector4f
 * A 4-D homogeneous vector.
 * @note The full C++ class is located here: Engine\Source\Runtime\Core\Public\Math\Vector4.h
 */
-USTRUCT(immutable, noexport)
+USTRUCT(immutable, noexport, IsAlwaysAccessible, HasDefaults, HasNoOpConstructor, IsCoreType)
 struct FVector4d
 {
 	UPROPERTY(EditAnywhere, Category = Vector4, SaveGame)
@@ -559,7 +559,7 @@ struct FVector4d
 * A 4-D homogeneous vector.
 * @note The full C++ class is located here: Engine\Source\Runtime\Core\Public\Math\Vector4.h
 */
-USTRUCT(immutable, noexport, BlueprintType, meta = (HasNativeMake = "Engine.KismetMathLibrary.MakeVector4", HasNativeBreak = "Engine.KismetMathLibrary.BreakVector4"))
+USTRUCT(immutable, noexport, BlueprintType, IsAlwaysAccessible, HasDefaults, HasNoOpConstructor, IsCoreType, meta = (HasNativeMake = "Engine.KismetMathLibrary.MakeVector4", HasNativeBreak = "Engine.KismetMathLibrary.BreakVector4"))
 struct FVector4
 {
 	UPROPERTY(EditAnywhere, BlueprintReadWrite, Category=Vector4, SaveGame)
@@ -580,7 +580,7 @@ struct FVector4
 * A vector in 2-D space composed of components (X, Y) with floating point precision.
 * @note The full C++ class is located here: Engine\Source\Runtime\Core\Public\Math\Vector2D.h
 */
-USTRUCT(immutable, noexport)
+USTRUCT(immutable, noexport, IsAlwaysAccessible, HasDefaults, HasNoOpConstructor, IsCoreType)
 struct FVector2f
 {
 	UPROPERTY(EditAnywhere, Category=Vector2D, SaveGame)
@@ -595,7 +595,7 @@ struct FVector2f
 * @note The full C++ class is located here: Engine\Source\Runtime\Core\Public\Math\Vector2D.h
 */
 // LWC_TODO: CRITICAL! Name collision in UHT with FVector2D due to case insensitive FNames!
-// USTRUCT(immutable, noexport)
+// USTRUCT(immutable, noexport, IsAlwaysAccessible, HasDefaults, HasNoOpConstructor, IsCoreType)
 // struct FVector2d
 // {
 // 	UPROPERTY(EditAnywhere, Category=Vector2D, SaveGame)
@@ -609,7 +609,7 @@ struct FVector2f
  * A vector in 2-D space composed of components (X, Y) with floating point precision.
  * @note The full C++ class is located here: Engine\Source\Runtime\Core\Public\Math\Vector2D.h
  */
-USTRUCT(immutable, noexport, BlueprintType, meta=(HasNativeMake="Engine.KismetMathLibrary.MakeVector2D", HasNativeBreak="Engine.KismetMathLibrary.BreakVector2D"))
+USTRUCT(immutable, noexport, BlueprintType, IsAlwaysAccessible, HasDefaults, HasNoOpConstructor, IsCoreType, meta=(HasNativeMake="Engine.KismetMathLibrary.MakeVector2D", HasNativeBreak="Engine.KismetMathLibrary.BreakVector2D"))
 struct FVector2D
 {
 	UPROPERTY(EditAnywhere, BlueprintReadWrite, Category=Vector2D, SaveGame)
@@ -620,7 +620,7 @@ struct FVector2D
 };
 
 /** A pair of 3D vectors (mirrored from TwoVectors.h). */
-USTRUCT(immutable, BlueprintType, noexport)
+USTRUCT(immutable, BlueprintType, noexport, IsAlwaysAccessible, HasDefaults)
 struct FTwoVectors
 {
 	UPROPERTY(EditAnywhere, BlueprintReadWrite, Category=TwoVectors, SaveGame)
@@ -634,7 +634,7 @@ struct FTwoVectors
  * A plane definition in 3D space.
  * @note The full C++ class is located here: Engine\Source\Runtime\Core\Public\Math\Plane.h
  */
-USTRUCT(immutable, noexport)
+USTRUCT(immutable, noexport, IsAlwaysAccessible, HasDefaults, HasNoOpConstructor, IsCoreType)
 struct FPlane4f : public FVector3f
 {
 	UPROPERTY(EditAnywhere, Category=Plane, SaveGame)
@@ -645,7 +645,7 @@ struct FPlane4f : public FVector3f
  * A plane definition in 3D space.
  * @note The full C++ class is located here: Engine\Source\Runtime\Core\Public\Math\Plane.h
  */
-USTRUCT(immutable, noexport)
+USTRUCT(immutable, noexport, IsAlwaysAccessible, HasDefaults, HasNoOpConstructor, IsCoreType)
 struct FPlane4d : public FVector3d
 {
 	UPROPERTY(EditAnywhere, Category = Plane, SaveGame)
@@ -656,7 +656,7 @@ struct FPlane4d : public FVector3d
  * A plane definition in 3D space.
  * @note The full C++ class is located here: Engine\Source\Runtime\Core\Public\Math\Plane.h
  */
-USTRUCT(immutable, noexport, BlueprintType)
+USTRUCT(immutable, noexport, BlueprintType, IsAlwaysAccessible, HasDefaults, HasNoOpConstructor, IsCoreType)
 struct FPlane : public FVector
 {
 	UPROPERTY(EditAnywhere, BlueprintReadWrite, Category = Plane, SaveGame)
@@ -668,7 +668,7 @@ struct FPlane : public FVector
  * An orthogonal rotation in 3d space.
  * @note The full C++ class is located here: Engine\Source\Runtime\Core\Public\Math\Rotator.h
  */
-USTRUCT(immutable, noexport)
+USTRUCT(immutable, noexport, IsAlwaysAccessible, HasDefaults, HasNoOpConstructor, IsCoreType)
 struct FRotator3f
 {
 	/** Pitch (degrees) around Y axis */
@@ -688,7 +688,7 @@ struct FRotator3f
  * An orthogonal rotation in 3d space.
  * @note The full C++ class is located here: Engine\Source\Runtime\Core\Public\Math\Rotator.h
  */
-USTRUCT(immutable, noexport)
+USTRUCT(immutable, noexport, IsAlwaysAccessible, HasDefaults, HasNoOpConstructor, IsCoreType)
 struct FRotator3d
 {
 	/** Pitch (degrees) around Y axis */
@@ -708,7 +708,7 @@ struct FRotator3d
  * An orthogonal rotation in 3d space.
  * @note The full C++ class is located here: Engine\Source\Runtime\Core\Public\Math\Rotator.h
  */
-USTRUCT(immutable, noexport, BlueprintType, meta=(HasNativeMake="Engine.KismetMathLibrary.MakeRotator", HasNativeBreak="Engine.KismetMathLibrary.BreakRotator"))
+USTRUCT(immutable, noexport, BlueprintType, IsAlwaysAccessible, HasDefaults, HasNoOpConstructor, IsCoreType, meta=(HasNativeMake="Engine.KismetMathLibrary.MakeRotator", HasNativeBreak="Engine.KismetMathLibrary.BreakRotator"))
 struct FRotator
 {
 	/** Pitch (degrees) around Y axis */
@@ -728,7 +728,7 @@ struct FRotator
  * Quaternion.
  * @note The full C++ class is located here: Engine\Source\Runtime\Core\Public\Math\Quat.h
  */
-USTRUCT(immutable, noexport)
+USTRUCT(immutable, noexport, IsAlwaysAccessible, HasDefaults, HasNoOpConstructor, IsCoreType)
 struct FQuat4f
 {
 	UPROPERTY(EditAnywhere, Category=Quat, SaveGame)
@@ -750,7 +750,7 @@ struct FQuat4f
  * Quaternion.
  * @note The full C++ class is located here: Engine\Source\Runtime\Core\Public\Math\Quat.h
  */
-USTRUCT(immutable, noexport)
+USTRUCT(immutable, noexport, IsAlwaysAccessible, HasDefaults, HasNoOpConstructor, IsCoreType)
 struct FQuat4d
 {
 	UPROPERTY(EditAnywhere, Category = Quat, SaveGame)
@@ -772,7 +772,7 @@ struct FQuat4d
  * Quaternion.
  * @note The full C++ class is located here: Engine\Source\Runtime\Core\Public\Math\Quat.h
  */
-USTRUCT(immutable, noexport, BlueprintType, meta=(HasNativeMake ="Engine.KismetMathLibrary.MakeQuat", HasNativeBreak="Engine.KismetMathLibrary.BreakQuat"))
+USTRUCT(immutable, noexport, BlueprintType, IsAlwaysAccessible, HasDefaults, HasNoOpConstructor, IsCoreType, meta=(HasNativeMake ="Engine.KismetMathLibrary.MakeQuat", HasNativeBreak="Engine.KismetMathLibrary.BreakQuat"))
 struct FQuat
 {
 	UPROPERTY(EditAnywhere, BlueprintReadWrite, Category = Quat, SaveGame)
@@ -839,7 +839,7 @@ struct FPackedRGBA16N
  * Screen coordinates.
  * @note The full C++ class is located here: Engine\Source\Runtime\Core\Public\Math\IntPoint.h
  */
-USTRUCT(immutable, noexport, BlueprintType)
+USTRUCT(immutable, noexport, BlueprintType, IsAlwaysAccessible, HasDefaults, HasNoOpConstructor)
 struct FIntPoint
 {
 	UPROPERTY(EditAnywhere, BlueprintReadWrite, Category=IntPoint, SaveGame)
@@ -854,7 +854,7 @@ struct FIntPoint
  * An integer vector in 3D space.
  * @note The full C++ class is located here: Engine\Source\Runtime\Core\Public\Math\IntVector.h
  */
-USTRUCT(immutable, noexport, BlueprintType)
+USTRUCT(immutable, noexport, BlueprintType, IsAlwaysAccessible, HasDefaults, HasNoOpConstructor)
 struct FIntVector
 {
 	UPROPERTY(EditAnywhere, BlueprintReadWrite, Category=IntVector, SaveGame)
@@ -871,7 +871,7 @@ struct FIntVector
  * Stores a color with 8 bits of precision per channel. (BGRA).
  * @note The full C++ class is located here: Engine\Source\Runtime\Core\Public\Math\Color.h
  */
-USTRUCT(immutable, noexport, BlueprintType)
+USTRUCT(immutable, noexport, BlueprintType, IsAlwaysAccessible, HasDefaults, HasNoOpConstructor)
 struct FColor
 {
 	UPROPERTY(EditAnywhere, BlueprintReadWrite, Category=Color, SaveGame, meta=(ClampMin="0", ClampMax="255"))
@@ -892,7 +892,7 @@ struct FColor
  * A linear, 32-bit/component floating point RGBA color.
  * @note The full C++ class is located here: Engine\Source\Runtime\Core\Public\Math\Color.h
  */
-USTRUCT(immutable, noexport, BlueprintType)
+USTRUCT(immutable, noexport, BlueprintType, IsAlwaysAccessible, HasDefaults, HasNoOpConstructor)
 struct FLinearColor
 {
 	UPROPERTY(EditAnywhere, BlueprintReadWrite, Category=LinearColor, SaveGame)
@@ -913,7 +913,7 @@ struct FLinearColor
  * A point or direction FVector in 3d space.
  * @note The full C++ class is located here: Engine\Source\Runtime\Core\Public\Math\Box.h
  */
-USTRUCT(immutable, noexport)
+USTRUCT(immutable, noexport, IsAlwaysAccessible, HasDefaults, HasNoOpConstructor, IsCoreType)
 struct FBox3f
 {
 	UPROPERTY(EditAnywhere, Category = Box, SaveGame)
@@ -931,7 +931,7 @@ struct FBox3f
  * A point or direction FVector in 3d space.
  * @note The full C++ class is located here: Engine\Source\Runtime\Core\Public\Math\Box.h
  */
-USTRUCT(immutable, noexport)
+USTRUCT(immutable, noexport, IsAlwaysAccessible, HasDefaults, HasNoOpConstructor, IsCoreType)
 struct FBox3d
 {
 	UPROPERTY(EditAnywhere, Category = Box, SaveGame)
@@ -948,7 +948,7 @@ struct FBox3d
  * A bounding box.
  * @note The full C++ class is located here: Engine\Source\Runtime\Core\Public\Math\Box.h
  */
-USTRUCT(immutable, noexport, BlueprintType, meta=(HasNativeMake="Engine.KismetMathLibrary.MakeBox"))
+USTRUCT(immutable, noexport, BlueprintType, IsAlwaysAccessible, HasDefaults, HasNoOpConstructor, IsCoreType, meta=(HasNativeMake="Engine.KismetMathLibrary.MakeBox"))
 struct FBox
 {
 	UPROPERTY(EditAnywhere, BlueprintReadWrite, Category=Box, SaveGame)
@@ -965,7 +965,7 @@ struct FBox
  * A rectangular 2D Box.
  * @note The full C++ class is located here: Engine\Source\Runtime\Core\Public\Math\Box2D.h
  */
-USTRUCT(immutable, noexport)
+USTRUCT(immutable, noexport, IsAlwaysAccessible, HasDefaults, HasNoOpConstructor, IsCoreType)
 struct FBox2f
 {
 	UPROPERTY(EditAnywhere, Category=Box2D, SaveGame)
@@ -983,7 +983,7 @@ struct FBox2f
 * @note The full C++ class is located here: Engine\Source\Runtime\Core\Public\Math\Box2D.h
 */
 // LWC_TODO: CRITICAL! Name collision in UHT with FBox2D due to case insensitive FNames!
-// USTRUCT(immutable, noexport)
+// USTRUCT(immutable, noexport, IsAlwaysAccessible, HasDefaults, HasNoOpConstructor, IsCoreType)
 // struct FBox2d
 // {
 // 	UPROPERTY(EditAnywhere, Category=Box2D, SaveGame)
@@ -1000,7 +1000,7 @@ struct FBox2f
 * A rectangular 2D Box.
 * @note The full C++ class is located here: Engine\Source\Runtime\Core\Public\Math\Box2D.h
 */
-USTRUCT(immutable, noexport, BlueprintType, meta=(HasNativeMake="Engine.KismetMathLibrary.MakeBox2D"))
+USTRUCT(immutable, noexport, BlueprintType, IsAlwaysAccessible, HasDefaults, HasNoOpConstructor, IsCoreType, meta=(HasNativeMake="Engine.KismetMathLibrary.MakeBox2D"))
 struct FBox2D
 {
 	UPROPERTY(EditAnywhere, BlueprintReadWrite, Category=Box2D, SaveGame)
@@ -1017,7 +1017,7 @@ struct FBox2D
  * A bounding box and bounding sphere with the same origin.
  * @note The full C++ class is located here : Engine\Source\Runtime\Core\Public\Math\BoxSphereBounds.h
  */
-USTRUCT(noexport)
+USTRUCT(noexport, IsAlwaysAccessible, HasDefaults, HasNoOpConstructor, IsCoreType)
 struct FBoxSphereBounds3f
 {
 	/** Holds the origin of the bounding box and sphere. */
@@ -1037,7 +1037,7 @@ struct FBoxSphereBounds3f
  * A bounding box and bounding sphere with the same origin.
  * @note The full C++ class is located here : Engine\Source\Runtime\Core\Public\Math\BoxSphereBounds.h
  */
-USTRUCT(noexport)
+USTRUCT(noexport, IsAlwaysAccessible, HasDefaults, HasNoOpConstructor, IsCoreType)
 struct FBoxSphereBounds3d
 {
 	/** Holds the origin of the bounding box and sphere. */
@@ -1057,7 +1057,7 @@ struct FBoxSphereBounds3d
  * A bounding box and bounding sphere with the same origin.
  * @note The full C++ class is located here : Engine\Source\Runtime\Core\Public\Math\BoxSphereBounds.h
  */
-USTRUCT(noexport, BlueprintType, meta = (HasNativeMake = "Engine.KismetMathLibrary.MakeBoxSphereBounds", HasNativeBreak = "Engine.KismetMathLibrary.BreakBoxSphereBounds"))
+USTRUCT(noexport, BlueprintType, IsAlwaysAccessible, HasDefaults, HasNoOpConstructor, IsCoreType, meta = (HasNativeMake = "Engine.KismetMathLibrary.MakeBoxSphereBounds", HasNativeBreak = "Engine.KismetMathLibrary.BreakBoxSphereBounds"))
 struct FBoxSphereBounds
 {
 	/** Holds the origin of the bounding box and sphere. */
@@ -1077,7 +1077,7 @@ struct FBoxSphereBounds
  * Structure for arbitrarily oriented boxes (i.e. not necessarily axis-aligned).
  * @note The full C++ class is located here: Engine\Source\Runtime\Core\Public\Math\OrientedBox.h
  */
-USTRUCT(immutable, noexport)
+USTRUCT(immutable, noexport, IsAlwaysAccessible, HasDefaults)
 struct FOrientedBox
 {
 	/** Holds the center of the box. */
@@ -1113,7 +1113,7 @@ struct FOrientedBox
  * A 4x4 matrix.
  * @note The full C++ class is located here: Engine\Source\Runtime\Core\Public\Math\Matrix.h
  */
-USTRUCT(immutable, noexport)
+USTRUCT(immutable, noexport, HasDefaults, HasNoOpConstructor, IsCoreType)
 struct FMatrix44f
 {
 	UPROPERTY(EditAnywhere, Category=Matrix, SaveGame)
@@ -1134,7 +1134,7 @@ struct FMatrix44f
  * A 4x4 matrix.
  * @note The full C++ class is located here: Engine\Source\Runtime\Core\Public\Math\Matrix.h
  */
-USTRUCT(immutable, noexport)
+USTRUCT(immutable, noexport, HasDefaults, HasNoOpConstructor, IsCoreType)
 struct FMatrix44d
 {
 	UPROPERTY(EditAnywhere, Category = Matrix, SaveGame)
@@ -1155,7 +1155,7 @@ struct FMatrix44d
  * A 4x4 matrix.
  * @note The full C++ class is located here: Engine\Source\Runtime\Core\Public\Math\Matrix.h
  */
-USTRUCT(immutable, noexport, BlueprintType)
+USTRUCT(immutable, noexport, BlueprintType, HasDefaults, HasNoOpConstructor, IsCoreType)
 struct FMatrix
 {
 	UPROPERTY(EditAnywhere, BlueprintReadWrite, Category = Matrix, SaveGame)
@@ -1176,7 +1176,7 @@ struct FMatrix
  * Describes one specific point on an interpolation curve.
  * @note This is a mirror of TInterpCurvePoint<float>, defined in InterpCurvePoint.h
  */
-USTRUCT(noexport, BlueprintType)
+USTRUCT(noexport, BlueprintType, IsAlwaysAccessible, HasDefaults, HasNoOpConstructor)
 struct FInterpCurvePointFloat
 {
 	/** Float input value that corresponds to this key (eg. time). */
@@ -1225,7 +1225,7 @@ struct FInterpCurveFloat
  * Describes one specific point on an interpolation curve.
  * @note This is a mirror of TInterpCurvePoint<FVector2D>, defined in InterpCurvePoint.h
  */
-USTRUCT(noexport, BlueprintType)
+USTRUCT(noexport, BlueprintType, IsAlwaysAccessible, HasDefaults, HasNoOpConstructor)
 struct FInterpCurvePointVector2D
 {
 	/** Float input value that corresponds to this key (eg. time). */
@@ -1273,7 +1273,7 @@ struct FInterpCurveVector2D
  * Describes one specific point on an interpolation curve.
  * @note This is a mirror of TInterpCurvePoint<FVector>, defined in InterpCurvePoint.h
  */
-USTRUCT(noexport, BlueprintType)
+USTRUCT(noexport, BlueprintType, IsAlwaysAccessible, HasDefaults, HasNoOpConstructor)
 struct FInterpCurvePointVector
 {
 	/** Float input value that corresponds to this key (eg. time). */
@@ -1321,7 +1321,7 @@ struct FInterpCurveVector
  * Describes one specific point on an interpolation curve.
  * @note This is a mirror of TInterpCurvePoint<FQuat>, defined in InterpCurvePoint.h
  */
-USTRUCT(noexport, BlueprintType)
+USTRUCT(noexport, BlueprintType, IsAlwaysAccessible, HasDefaults, HasNoOpConstructor)
 struct FInterpCurvePointQuat
 {
 	/** Float input value that corresponds to this key (eg. time). */
@@ -1369,7 +1369,7 @@ struct FInterpCurveQuat
  * Describes one specific point on an interpolation curve.
  * @note This is a mirror of TInterpCurvePoint<FTwoVectors>, defined in InterpCurvePoint.h
  */
-USTRUCT(noexport, BlueprintType)
+USTRUCT(noexport, BlueprintType, IsAlwaysAccessible, HasDefaults, HasNoOpConstructor)
 struct FInterpCurvePointTwoVectors
 {
 	/** Float input value that corresponds to this key (eg. time). */
@@ -1417,7 +1417,7 @@ struct FInterpCurveTwoVectors
  * Describes one specific point on an interpolation curve.
  * @note This is a mirror of TInterpCurvePoint<FLinearColor>, defined in InterpCurvePoint.h
  */
-USTRUCT(noexport, BlueprintType)
+USTRUCT(noexport, BlueprintType, IsAlwaysAccessible, HasDefaults, HasNoOpConstructor)
 struct FInterpCurvePointLinearColor
 {
 	/** Float input value that corresponds to this key (eg. time). */
@@ -1465,7 +1465,7 @@ struct FInterpCurveLinearColor
  * Transform composed of Quat/Translation/Scale.
  * @note This is implemented in either TransformVectorized.h or TransformNonVectorized.h depending on the platform.
  */
-USTRUCT(immutable, noexport)
+USTRUCT(immutable, noexport, IsAlwaysAccessible, HasDefaults, IsCoreType)
 struct FTransform3f
 {
 	/** Rotation of this transformation, as a quaternion. */
@@ -1485,7 +1485,7 @@ struct FTransform3f
  * Transform composed of Quat/Translation/Scale.
  * @note This is implemented in either TransformVectorized.h or TransformNonVectorized.h depending on the platform.
  */
-USTRUCT(immutable, noexport)
+USTRUCT(immutable, noexport, IsAlwaysAccessible, HasDefaults, IsCoreType)
 struct FTransform3d
 {
 	/** Rotation of this transformation, as a quaternion. */
@@ -1505,7 +1505,7 @@ struct FTransform3d
  * Transform composed of Quat/Translation/Scale.
  * @note This is implemented in either TransformVectorized.h or TransformNonVectorized.h depending on the platform.
  */
-USTRUCT(noexport, BlueprintType, meta=(HasNativeMake="Engine.KismetMathLibrary.MakeTransform", HasNativeBreak="Engine.KismetMathLibrary.BreakTransform"))
+USTRUCT(noexport, BlueprintType, IsAlwaysAccessible, HasDefaults, IsCoreType, meta=(HasNativeMake="Engine.KismetMathLibrary.MakeTransform", HasNativeBreak="Engine.KismetMathLibrary.BreakTransform"))
 struct FTransform
 {
 	/** Rotation of this transformation, as a quaternion. */
@@ -1525,7 +1525,7 @@ struct FTransform
  * Thread-safe random number generator that can be manually seeded.
  * @note The full C++ class is located here: Engine\Source\Runtime\Core\Public\Math\RandomStream.h
  */
-USTRUCT(noexport, BlueprintType, meta = (HasNativeMake = "Engine.KismetMathLibrary.MakeRandomStream", HasNativeBreak = "Engine.KismetMathLibrary.BreakRandomStream"))
+USTRUCT(noexport, BlueprintType, IsAlwaysAccessible, HasDefaults, HasNoOpConstructor, meta = (HasNativeMake = "Engine.KismetMathLibrary.MakeRandomStream", HasNativeBreak = "Engine.KismetMathLibrary.BreakRandomStream"))
 struct FRandomStream
 {
 public:
@@ -1542,7 +1542,7 @@ public:
  * A value representing a specific point date and time over a wide range of years.
  * @note The full C++ class is located here: Engine\Source\Runtime\Core\Public\Misc\DateTime.h
  */
-USTRUCT(immutable, noexport, BlueprintType, meta=(HasNativeMake="Engine.KismetMathLibrary.MakeDateTime", HasNativeBreak="Engine.KismetMathLibrary.BreakDateTime"))
+USTRUCT(immutable, noexport, BlueprintType, IsAlwaysAccessible, HasDefaults, meta=(HasNativeMake="Engine.KismetMathLibrary.MakeDateTime", HasNativeBreak="Engine.KismetMathLibrary.BreakDateTime"))
 struct FDateTime
 {
 	int64 Ticks;
@@ -1552,7 +1552,7 @@ struct FDateTime
  * A frame number value, representing discrete frames since the start of timing.
  * @note The full C++ class is located here: Engine\Source\Runtime\Core\Public\Misc\FrameNumber.h
  */
-USTRUCT(noexport, BlueprintType)
+USTRUCT(noexport, BlueprintType, IsAlwaysAccessible, HasDefaults)
 struct FFrameNumber
 {
 	UPROPERTY(BlueprintReadWrite, EditAnywhere, Category=FrameNumber)
@@ -1637,7 +1637,7 @@ struct FTimecode
  * A time span value, which is the difference between two dates and times.
  * @note The full C++ class is located here: Engine\Source\Runtime\Core\Public\Misc\Timespan.h
  */
-USTRUCT(immutable, noexport, BlueprintType, meta=(HasNativeMake="Engine.KismetMathLibrary.MakeTimespan", HasNativeBreak="Engine.KismetMathLibrary.BreakTimespan"))
+USTRUCT(immutable, noexport, BlueprintType, IsAlwaysAccessible, HasDefaults, meta=(HasNativeMake="Engine.KismetMathLibrary.MakeTimespan", HasNativeBreak="Engine.KismetMathLibrary.BreakTimespan"))
 struct FTimespan
 {
 	int64 Ticks;
@@ -1647,7 +1647,7 @@ struct FTimespan
  * A struct that contains a string reference to an object, either a top level asset or a subobject.
  * @note The full C++ class is located here: Engine\Source\Runtime\CoreUObject\Public\UObject\SoftObjectPath.h
  */
-USTRUCT(noexport, BlueprintType, meta=(HasNativeMake="Engine.KismetSystemLibrary.MakeSoftObjectPath", HasNativeBreak="Engine.KismetSystemLibrary.BreakSoftObjectPath"))
+USTRUCT(noexport, BlueprintType, IsAlwaysAccessible, HasDefaults, meta=(HasNativeMake="Engine.KismetSystemLibrary.MakeSoftObjectPath", HasNativeBreak="Engine.KismetSystemLibrary.BreakSoftObjectPath"))
 struct FSoftObjectPath
 {
 	/** Asset path, patch to a top level object in a package */
@@ -1663,7 +1663,7 @@ struct FSoftObjectPath
  * A struct that contains a string reference to a class, can be used to make soft references to classes.
  * @note The full C++ class is located here: Engine\Source\Runtime\CoreUObject\Public\UObject\SoftObjectPath.h
  */
-USTRUCT(noexport, BlueprintType, meta=(HasNativeMake="Engine.KismetSystemLibrary.MakeSoftClassPath", HasNativeBreak="Engine.KismetSystemLibrary.BreakSoftClassPath"))
+USTRUCT(noexport, BlueprintType, IsAlwaysAccessible, HasDefaults, meta=(HasNativeMake="Engine.KismetSystemLibrary.MakeSoftClassPath", HasNativeBreak="Engine.KismetSystemLibrary.BreakSoftClassPath"))
 struct FSoftClassPath : public FSoftObjectPath
 {
 };
@@ -1672,7 +1672,7 @@ struct FSoftClassPath : public FSoftObjectPath
  * A type of primary asset, used by the Asset Manager system.
  * @note The full C++ class is located here: Engine\Source\Runtime\CoreUObject\Public\UObject\PrimaryAssetId.h
  */
-USTRUCT(noexport, BlueprintType)
+USTRUCT(noexport, BlueprintType, IsAlwaysAccessible, HasDefaults)
 struct FPrimaryAssetType
 {
 	/** The Type of this object, by default its base class's name */
@@ -1684,7 +1684,7 @@ struct FPrimaryAssetType
  * This identifies an object as a "primary" asset that can be searched for by the AssetManager and used in various tools
  * @note The full C++ class is located here: Engine\Source\Runtime\CoreUObject\Public\UObject\PrimaryAssetId.h
  */
-USTRUCT(noexport, BlueprintType)
+USTRUCT(noexport, BlueprintType, IsAlwaysAccessible, HasDefaults)
 struct FPrimaryAssetId
 {
 	/** The Type of this object, by default its base class's name */
@@ -1982,7 +1982,7 @@ enum class EDataValidationResult : uint8
 	NotValidated
 };
 
-USTRUCT(noexport, BlueprintType)
+USTRUCT(noexport, BlueprintType, IsAlwaysAccessible, HasDefaults)
 struct FARFilter
 {
 	/** The filter component for package names */
@@ -2040,7 +2040,7 @@ struct COREUOBJECT_API FAssetBundleEntry
 };
 
 /** A struct with a list of asset bundle entries. If one of these is inside a UObject it will get automatically exported as the asset registry tag AssetBundleData */
-USTRUCT(noexport)
+USTRUCT(noexport, IsAlwaysAccessible, HasDefaults)
 struct COREUOBJECT_API FAssetBundleData
 {
 	/** List of bundles defined */
@@ -2052,7 +2052,7 @@ struct COREUOBJECT_API FAssetBundleData
  * A struct to hold important information about an assets found by the Asset Registry
  * This struct is transient and should never be serialized
  */
-USTRUCT(noexport, BlueprintType)
+USTRUCT(noexport, BlueprintType, IsAlwaysAccessible, HasDefaults)
 struct FAssetData
 {
 	/** The object path for the asset in the form PackageName.AssetName. Only top level objects in a package can have AssetData */
@@ -2079,7 +2079,7 @@ struct FAssetData
 	uint32 PackageFlags = 0;
 };
 
-USTRUCT(noexport)
+USTRUCT(noexport, IsAlwaysAccessible, HasDefaults)
 struct FTestUninitializedScriptStructMembersTest
 {
 	UPROPERTY(Transient)
@@ -2096,7 +2096,7 @@ struct FTestUninitializedScriptStructMembersTest
  * Direct base class for all UE objects
  * @note The full C++ class is located here: Engine\Source\Runtime\CoreUObject\Public\UObject\Object.h
  */
-UCLASS(abstract, noexport)
+UCLASS(abstract, noexport, MatchedSerializers)
 class UObject
 {
 	GENERATED_BODY()
@@ -2114,6 +2114,382 @@ public:
 	void ExecuteUbergraph(int32 EntryPoint);
 };
 
+/**
+ * @note The full C++ class is located here: Engine\Source\Runtime\CoreUObject\Public\UObject\Class.h
+ */
+
+UCLASS(abstract, noexport, intrinsic, Config = Engine)
+class UField : public UObject
+{
+	GENERATED_BODY()
+public:
+	UField(const FObjectInitializer& ObjectInitializer = FObjectInitializer::Get());
+};
+
+UCLASS(noexport, intrinsic, Config = Engine)
+class UEnum : public UField
+{
+	GENERATED_BODY()
+public:
+	UEnum(const FObjectInitializer& ObjectInitialzer);
+};
+
+UCLASS(noexport, intrinsic, MatchedSerializers, Config = Engine)
+class UStruct : public UField
+{
+	GENERATED_BODY()
+public:
+	UStruct(const FObjectInitializer& ObjectInitializer = FObjectInitializer::Get());
+};
+
+UCLASS(noexport, intrinsic, MatchedSerializers, Config = Engine)
+class UScriptStruct : public UStruct
+{
+	GENERATED_BODY()
+public:
+	UScriptStruct(const FObjectInitializer& ObjectInitializer = FObjectInitializer::Get());
+};
+
+UCLASS(noexport, intrinsic, Config = Engine)
+class UFunction : public UStruct
+{
+	GENERATED_BODY()
+public:
+	UFunction(const FObjectInitializer& ObjectInitializer = FObjectInitializer::Get());
+};
+
+UCLASS(noexport, intrinsic, within=Package, Config = Engine)
+class UClass : public UStruct
+{
+	GENERATED_BODY()
+public:
+	UClass(const FObjectInitializer& ObjectInitializer = FObjectInitializer::Get());
+};
+
+/**
+ * @note The full C++ class is located here: Engine\Source\Runtime\CoreUObject\Public\UObject\Interface.h
+ */
+
+UCLASS(abstract, noexport, intrinsic, interface, Config = Engine)
+class UInterface : public UObject
+{
+	GENERATED_BODY()
+public:
+	UInterface(const FObjectInitializer& ObjectInitializer = FObjectInitializer::Get());
+};
+
+class IInterface
+{
+	GENERATED_BODY()
+};
+
+
+/**
+ * @note The full C++ class is located here: Engine\Source\Runtime\CoreUObject\Public\UObject\Package.h
+ */
+
+UCLASS(noexport, intrinsic, Config = Engine)
+class UPackage : public UObject
+{
+	GENERATED_BODY()
+};
+
+/**
+ * @note The full C++ class is located here: Engine\Source\Runtime\CoreUObject\Public\UObject\CoreNet.h
+ */
+
+UCLASS(noexport, intrinsic, abstract, transient, Config = Engine)
+class UPackageMap : public UObject
+{
+	GENERATED_BODY()
+public:
+	UPackageMap(const FObjectInitializer& ObjectInitializer = FObjectInitializer::Get());
+};
+
+/**
+ * @note The full C++ class is located here: Engine\Source\Runtime\CoreUObject\Public\UObject\MetaData.h
+ */
+
+UCLASS(noexport, intrinsic, MatchedSerializers, Config = Engine)
+class UMetaData : public UObject
+{
+	GENERATED_BODY()
+public:
+	UMetaData(const FObjectInitializer& ObjectInitializer = FObjectInitializer::Get());
+};
+
+/**
+ * @note The full C++ class is located here: Engine\Source\Runtime\CoreUObject\Public\UObject\UnrealTypePrivate.h
+ */
+
+UCLASS(noexport, intrinsic, abstract, Config = Engine)
+class UProperty : public UField
+{
+	GENERATED_BODY()
+public:
+	UProperty(const FObjectInitializer& ObjectInitializer = FObjectInitializer::Get());
+};
+
+UCLASS(noexport, intrinsic, Config = Engine)
+class UStructProperty : public UProperty
+{
+	GENERATED_BODY()
+public:
+	UStructProperty(const FObjectInitializer& ObjectInitializer = FObjectInitializer::Get());
+};
+
+UCLASS(noexport, intrinsic, Config = Engine)
+class UEnumProperty : public UProperty
+{
+	GENERATED_BODY()
+public:
+	UEnumProperty(const FObjectInitializer& ObjectInitializer = FObjectInitializer::Get());
+};
+
+UCLASS(noexport, intrinsic, Config = Engine)
+class UBoolProperty : public UProperty
+{
+	GENERATED_BODY()
+public:
+	UBoolProperty(const FObjectInitializer& ObjectInitializer = FObjectInitializer::Get());
+};
+
+UCLASS(abstract, noexport, intrinsic, Config = Engine)
+class UNumericProperty : public UProperty
+{
+	GENERATED_BODY()
+public:
+	UNumericProperty(const FObjectInitializer& ObjectInitializer = FObjectInitializer::Get());
+};
+
+UCLASS(noexport, intrinsic, Config = Engine)
+class UIntProperty : public UNumericProperty
+{
+	GENERATED_BODY()
+public:
+	UIntProperty(const FObjectInitializer& ObjectInitializer = FObjectInitializer::Get());
+};
+
+UCLASS(noexport, intrinsic, Config = Engine)
+class UInt8Property : public UNumericProperty
+{
+	GENERATED_BODY()
+public:
+	UInt8Property(const FObjectInitializer& ObjectInitializer = FObjectInitializer::Get());
+};
+
+UCLASS(noexport, intrinsic, Config = Engine)
+class UInt16Property : public UNumericProperty
+{
+	GENERATED_BODY()
+public:
+	UInt16Property(const FObjectInitializer& ObjectInitializer = FObjectInitializer::Get());
+};
+
+UCLASS(noexport, intrinsic, Config = Engine)
+class UInt64Property : public UNumericProperty
+{
+	GENERATED_BODY()
+public:
+	UInt64Property(const FObjectInitializer& ObjectInitializer = FObjectInitializer::Get());
+};
+
+UCLASS(noexport, intrinsic, Config = Engine)
+class UByteProperty : public UNumericProperty
+{
+	GENERATED_BODY()
+public:
+	UByteProperty(const FObjectInitializer& ObjectInitializer = FObjectInitializer::Get());
+};
+
+UCLASS(noexport, intrinsic, Config = Engine)
+class UUint16Property : public UNumericProperty
+{
+	GENERATED_BODY()
+public:
+	UUint16Property(const FObjectInitializer& ObjectInitializer = FObjectInitializer::Get());
+};
+
+UCLASS(noexport, intrinsic, Config = Engine)
+class UUint32Property : public UNumericProperty
+{
+	GENERATED_BODY()
+public:
+	UUint32Property(const FObjectInitializer& ObjectInitializer = FObjectInitializer::Get());
+};
+
+UCLASS(noexport, intrinsic, Config = Engine)
+class UUint64Property : public UNumericProperty
+{
+	GENERATED_BODY()
+public:
+	UUint64Property(const FObjectInitializer& ObjectInitializer = FObjectInitializer::Get());
+};
+
+UCLASS(noexport, intrinsic, Config = Engine)
+class UFloatProperty : public UNumericProperty
+{
+	GENERATED_BODY()
+public:
+	UFloatProperty(const FObjectInitializer& ObjectInitializer = FObjectInitializer::Get());
+};
+
+UCLASS(noexport, intrinsic, Config = Engine)
+class UDoubleProperty : public UNumericProperty
+{
+	GENERATED_BODY()
+public:
+	UDoubleProperty(const FObjectInitializer& ObjectInitializer = FObjectInitializer::Get());
+};
+
+UCLASS(abstract, noexport, intrinsic, Config = Engine)
+class UObjectPropertyBase : public UProperty
+{
+	GENERATED_BODY()
+public:
+	UObjectPropertyBase(const FObjectInitializer& ObjectInitializer = FObjectInitializer::Get());
+};
+
+UCLASS(noexport, intrinsic, Config = Engine)
+class UObjectProperty : public UObjectPropertyBase
+{
+	GENERATED_BODY()
+public:
+	UObjectProperty(const FObjectInitializer& ObjectInitializer = FObjectInitializer::Get());
+};
+
+UCLASS(noexport, intrinsic, Config = Engine)
+class UClassProperty : public UObjectProperty
+{
+	GENERATED_BODY()
+public:
+	UClassProperty(const FObjectInitializer& ObjectInitializer = FObjectInitializer::Get());
+};
+
+UCLASS(noexport, intrinsic, Config = Engine)
+class USoftObjectProperty : public UObjectPropertyBase
+{
+	GENERATED_BODY()
+public:
+	USoftObjectProperty(const FObjectInitializer& ObjectInitializer = FObjectInitializer::Get());
+};
+
+UCLASS(noexport, intrinsic, Config = Engine)
+class USoftClassProperty : public USoftObjectProperty
+{
+	GENERATED_BODY()
+public:
+	USoftClassProperty(const FObjectInitializer& ObjectInitializer = FObjectInitializer::Get());
+};
+
+UCLASS(noexport, intrinsic, Config = Engine)
+class UWeakObjectProperty : public UObjectPropertyBase
+{
+	GENERATED_BODY()
+public:
+	UWeakObjectProperty(const FObjectInitializer& ObjectInitializer = FObjectInitializer::Get());
+};
+
+UCLASS(noexport, intrinsic, Config = Engine)
+class ULazyObjectProperty : public UObjectPropertyBase
+{
+	GENERATED_BODY()
+public:
+	ULazyObjectProperty(const FObjectInitializer& ObjectInitializer = FObjectInitializer::Get());
+};
+
+UCLASS(noexport, intrinsic, Config = Engine)
+class UInterfaceProperty : public UProperty
+{
+	GENERATED_BODY()
+public:
+	UInterfaceProperty(const FObjectInitializer& ObjectInitializer = FObjectInitializer::Get());
+};
+
+UCLASS(noexport, intrinsic, Config = Engine)
+class UDelegateProperty : public UProperty
+{
+	GENERATED_BODY()
+public:
+	UDelegateProperty(const FObjectInitializer& ObjectInitializer = FObjectInitializer::Get());
+};
+
+UCLASS(abstract, noexport, intrinsic, Config = Engine)
+class UMulticastDelegateProperty : public UProperty
+{
+	GENERATED_BODY()
+public:
+	UMulticastDelegateProperty(const FObjectInitializer& ObjectInitializer = FObjectInitializer::Get());
+};
+
+UCLASS(noexport, intrinsic, Config = Engine)
+class UMulticastInlineDelegateProperty : public UMulticastDelegateProperty
+{
+	GENERATED_BODY()
+public:
+	UMulticastInlineDelegateProperty(const FObjectInitializer& ObjectInitializer = FObjectInitializer::Get());
+};
+
+UCLASS(noexport, intrinsic, Config = Engine)
+class UMulticastSparseDelegateProperty : public UMulticastDelegateProperty
+{
+	GENERATED_BODY()
+public:
+	UMulticastSparseDelegateProperty(const FObjectInitializer& ObjectInitializer = FObjectInitializer::Get());
+};
+
+UCLASS(noexport, intrinsic, Config = Engine)
+class UNameProperty : public UProperty
+{
+	GENERATED_BODY()
+public:
+	UNameProperty(const FObjectInitializer& ObjectInitializer = FObjectInitializer::Get());
+};
+
+UCLASS(noexport, intrinsic, Config = Engine)
+class UStrProperty : public UProperty
+{
+	GENERATED_BODY()
+public:
+	UStrProperty(const FObjectInitializer& ObjectInitializer = FObjectInitializer::Get());
+};
+
+UCLASS(noexport, intrinsic, Config = Engine)
+class UTextProperty : public UProperty
+{
+	GENERATED_BODY()
+public:
+	UTextProperty(const FObjectInitializer& ObjectInitializer = FObjectInitializer::Get());
+};
+
+UCLASS(noexport, intrinsic, Config = Engine)
+class UArrayProperty : public UProperty
+{
+	GENERATED_BODY()
+public:
+	UArrayProperty(const FObjectInitializer& ObjectInitializer = FObjectInitializer::Get());
+};
+
+UCLASS(noexport, intrinsic, Config = Engine)
+class USetProperty : public UProperty
+{
+	GENERATED_BODY()
+public:
+	USetProperty(const FObjectInitializer& ObjectInitializer = FObjectInitializer::Get());
+};
+
+UCLASS(noexport, intrinsic, Config = Engine)
+class UMapProperty : public UProperty
+{
+	GENERATED_BODY()
+public:
+	UMapProperty(const FObjectInitializer& ObjectInitializer = FObjectInitializer::Get());
+};
+
+/**
+ * @note The full C++ class is located here: Engine\Source\Runtime\CoreUObject\Public\UObject\SoftObjectPath.h
+ */
+
 /** Structure for file paths that are displayed in the editor with a picker UI. */
 USTRUCT(noexport, BlueprintType)
 struct FFilePath
