do clause support

diff --git a/Engine/Plugins/Animation/ControlRig/Source/ControlRig/Private/ControlRigBlueprintGeneratedClass.cpp b/Engine/Plugins/Animation/ControlRig/Source/ControlRig/Private/ControlRigBlueprintGeneratedClass.cpp
--- a/Engine/Plugins/Animation/ControlRig/Source/ControlRig/Private/ControlRigBlueprintGeneratedClass.cpp
+++ b/Engine/Plugins/Animation/ControlRig/Source/ControlRig/Private/ControlRigBlueprintGeneratedClass.cpp
@@ -20,7 +20,7 @@ uint8* UControlRigBlueprintGeneratedClass::GetPersistentUberGraphFrame(UObject*
 	return Super::GetPersistentUberGraphFrame(Obj, FuncToCheck);
 }
 
-void UControlRigBlueprintGeneratedClass::PostInitInstance(UObject* InObj)
+void UControlRigBlueprintGeneratedClass::PostInitInstance(UObject* InObj, FObjectInstancingGraph* InstanceGraph)
 {
 	UControlRig* ControlRig = Cast<UControlRig>(InObj);
 	check(ControlRig);
diff --git a/Engine/Plugins/Animation/ControlRig/Source/ControlRig/Public/ControlRigBlueprintGeneratedClass.h b/Engine/Plugins/Animation/ControlRig/Source/ControlRig/Public/ControlRigBlueprintGeneratedClass.h
--- a/Engine/Plugins/Animation/ControlRig/Source/ControlRig/Public/ControlRigBlueprintGeneratedClass.h
+++ b/Engine/Plugins/Animation/ControlRig/Source/ControlRig/Public/ControlRigBlueprintGeneratedClass.h
@@ -18,7 +18,7 @@ public:
 
 	// UClass interface
 	virtual uint8* GetPersistentUberGraphFrame(UObject* Obj, UFunction* FuncToCheck) const override;
-	virtual void PostInitInstance(UObject* InObj) override;
+	virtual void PostInitInstance(UObject* InObj, FObjectInstancingGraph* InstanceGraph) override;
 
 	// UObject interface
 	void Serialize(FArchive& Ar);
diff --git a/Engine/Plugins/Experimental/PythonScriptPlugin/Source/PythonScriptPlugin/Private/PyWrapperObject.cpp b/Engine/Plugins/Experimental/PythonScriptPlugin/Source/PythonScriptPlugin/Private/PyWrapperObject.cpp
--- a/Engine/Plugins/Experimental/PythonScriptPlugin/Source/PythonScriptPlugin/Private/PyWrapperObject.cpp
+++ b/Engine/Plugins/Experimental/PythonScriptPlugin/Source/PythonScriptPlugin/Private/PyWrapperObject.cpp
@@ -1821,9 +1821,9 @@ void UPythonGeneratedClass::PostRename(UObject* OldOuter, const FName OldName)
 	}
 }
 
-void UPythonGeneratedClass::PostInitInstance(UObject* InObj)
+void UPythonGeneratedClass::PostInitInstance(UObject* InObj, FObjectInstancingGraph* InstanceGraph)
 {
-	Super::PostInitInstance(InObj);
+	Super::PostInitInstance(InObj, InstanceGraph);
 
 	// Execute Python code within this block
 	{
diff --git a/Engine/Plugins/Experimental/PythonScriptPlugin/Source/PythonScriptPlugin/Private/PyWrapperObject.h b/Engine/Plugins/Experimental/PythonScriptPlugin/Source/PythonScriptPlugin/Private/PyWrapperObject.h
--- a/Engine/Plugins/Experimental/PythonScriptPlugin/Source/PythonScriptPlugin/Private/PyWrapperObject.h
+++ b/Engine/Plugins/Experimental/PythonScriptPlugin/Source/PythonScriptPlugin/Private/PyWrapperObject.h
@@ -193,7 +193,7 @@ public:
 	virtual void BeginDestroy() override;
 
 	//~ UClass interface
-	virtual void PostInitInstance(UObject* InObj) override;
+	virtual void PostInitInstance(UObject* InObj, FObjectInstancingGraph* InstanceGraph) override;
 
 	//~ IPythonResourceOwner interface
 	virtual void ReleasePythonResources() override;
diff --git a/Engine/Source/Runtime/CoreUObject/Private/Serialization/AsyncLoading.cpp b/Engine/Source/Runtime/CoreUObject/Private/Serialization/AsyncLoading.cpp
--- a/Engine/Source/Runtime/CoreUObject/Private/Serialization/AsyncLoading.cpp
+++ b/Engine/Source/Runtime/CoreUObject/Private/Serialization/AsyncLoading.cpp
@@ -5036,6 +5036,7 @@ FAsyncPackage::FAsyncPackage(FAsyncLoadingThread& InThread, const FAsyncPackageD
 , FinishExternalReadDependenciesIndex(0)
 , PostLoadIndex(0)
 , DeferredPostLoadIndex(0)
+, PostLoadInstanceIndex(0)
 , DeferredFinalizeIndex(0)
 , DeferredClusterIndex(0)
 , TimeLimit(FLT_MAX)
@@ -6420,10 +6421,29 @@ EAsyncPackageState::Type FAsyncPackage::PostLoadDeferredObjects(double InTickSta
 		UpdateLoadPercentage();
 	}
 
-	INC_FLOAT_STAT_BY(STAT_FAsyncPackage_TotalPostLoadGameThread, (float)(FPlatformTime::Seconds() - PostLoadStartTime));
-
 	// New objects might have been loaded during PostLoad.
 	Result = (DeferredPostLoadIndex == DeferredPostLoadObjects.Num()) ? EAsyncPackageState::Complete : EAsyncPackageState::TimeOut;
+
+	if (Result == EAsyncPackageState::Complete)
+	{
+		// Class specific object non-native initialization
+		while (PostLoadInstanceIndex < PackageObjLoaded.Num() &&
+			!AsyncLoadingThread.IsAsyncLoadingSuspendedInternal() &&
+			!::IsTimeLimitExceeded(InTickStartTime, bInUseTimeLimit, InOutTimeLimit, LastTypeOfWorkPerformed, LastObjectWorkWasPerformedOn))
+		{
+			UObject* Object = PackageObjLoaded[PostLoadInstanceIndex++];
+			check(Object);
+
+			LastObjectWorkWasPerformedOn = Object;
+			LastTypeOfWorkPerformed = TEXT("postloadinginstance_gamethread");
+
+			UClass* ObjectClass = Object->GetClass();
+			ObjectClass->PostLoadInstance(Object);
+		}
+	}
+
+	INC_FLOAT_STAT_BY(STAT_FAsyncPackage_TotalPostLoadGameThread, (float)(FPlatformTime::Seconds() - PostLoadStartTime));
+	Result = (PostLoadInstanceIndex == PackageObjLoaded.Num()) ? EAsyncPackageState::Complete : EAsyncPackageState::TimeOut;
 	if (Result == EAsyncPackageState::Complete)
 	{
 		LastObjectWorkWasPerformedOn = nullptr;
diff --git a/Engine/Source/Runtime/CoreUObject/Private/Serialization/AsyncLoading.h b/Engine/Source/Runtime/CoreUObject/Private/Serialization/AsyncLoading.h
--- a/Engine/Source/Runtime/CoreUObject/Private/Serialization/AsyncLoading.h
+++ b/Engine/Source/Runtime/CoreUObject/Private/Serialization/AsyncLoading.h
@@ -568,6 +568,8 @@ private:
 	int32							PostLoadIndex;
 	/** Current index into DeferredPostLoadObjects array used to spread routing PostLoad over several frames			*/
 	int32						DeferredPostLoadIndex;
+	/** Current index into DeferredPostLoadObjects array used to spread routing PostLoad over several frames			*/
+	int32						PostLoadInstanceIndex;
 	/** Current index into DeferredFinalizeObjects array used to spread routing PostLoad over several frames			*/
 	int32						DeferredFinalizeIndex;
 	/** Current index into DeferredClusterObjects array used to spread routing CreateClusters over several frames			*/
diff --git a/Engine/Source/Runtime/CoreUObject/Private/Serialization/AsyncLoading2.cpp b/Engine/Source/Runtime/CoreUObject/Private/Serialization/AsyncLoading2.cpp
--- a/Engine/Source/Runtime/CoreUObject/Private/Serialization/AsyncLoading2.cpp
+++ b/Engine/Source/Runtime/CoreUObject/Private/Serialization/AsyncLoading2.cpp
@@ -1301,6 +1301,7 @@ enum class EAsyncPackageLoadingState2 : uint8
 	DeferredPostLoad,
 	DeferredPostLoadDone,
 	Finalize,
+	PostLoadInstances,
 	CreateClusters,
 	Complete,
 	DeferredDelete,
@@ -1622,6 +1623,9 @@ struct FAsyncPackage2
 	/** Returns the UPackage wrapped by this, if it is valid */
 	UPackage* GetLoadedPackage();
 
+	/** Class specific callback for initializing non-native objects */
+	EAsyncPackageState::Type PostLoadInstances(FAsyncLoadingThreadState2& ThreadState);
+
 	/** Creates GC clusters from loaded objects */
 	EAsyncPackageState::Type CreateClusters(FAsyncLoadingThreadState2& ThreadState);
 
@@ -1660,6 +1664,8 @@ private:
 	int32						ExternalReadIndex = 0;
 	/** Current index into DeferredClusterObjects array used to spread routing CreateClusters over several frames			*/
 	int32						DeferredClusterIndex = 0;
+	/** Current index into Export objects array used to spread routing PostLoadInstances over several frames			*/
+	int32						PostLoadInstanceIndex = 0;
 	EAsyncPackageLoadingState2	AsyncPackageLoadingState = EAsyncPackageLoadingState2::NewPackage;
 
 	struct FAllDependenciesState
@@ -1759,7 +1765,8 @@ private:
 	bool						bLoadHasFailed = false;
 	/** True if this package was created by this async package */
 	bool						bCreatedLinkerRoot = false;
-
+	/** True if package contains objects that can be put into GC clusters */
+	bool						bHasClusterObjects = false;
 	/** List of all request handles */
 	TArray<int32, TInlineAllocator<2>> RequestIDs;
 	/** List of ConstructedObjects = Exports + UPackage + ObjectsCreatedFromExports */
@@ -4643,7 +4650,6 @@ EAsyncPackageState::Type FAsyncLoadingThread2::ProcessLoadedPackagesFromGameThre
 			UE_ASYNC_PACKAGE_DEBUG(Package->Desc);
 			check(Package->AsyncPackageLoadingState == EAsyncPackageLoadingState2::Finalize);
 
-			bool bHasClusterObjects = false;
 			TArray<UObject*> CDODefaultSubobjects;
 			// Clear async loading flags (we still want RF_Async, but EInternalObjectFlags::AsyncLoading can be cleared)
 			for (int32 FinalizeIndex = 0, ExportCount = Package->Data.ExportInfo.ExportCount; FinalizeIndex < ExportCount; ++FinalizeIndex)
@@ -4682,7 +4688,7 @@ EAsyncPackageState::Type FAsyncLoadingThread2::ProcessLoadedPackagesFromGameThre
 				{
 					if (!(Export.bFiltered | Export.bExportLoadFailed) && Export.Object->CanBeClusterRoot())
 					{
-						bHasClusterObjects = true;
+						Package->bHasClusterObjects = true;
 						break;
 					}
 				}
@@ -4714,14 +4720,7 @@ EAsyncPackageState::Type FAsyncLoadingThread2::ProcessLoadedPackagesFromGameThre
 			TRACE_LOADTIME_END_LOAD_ASYNC_PACKAGE(Package);
 
 			check(Package->AsyncPackageLoadingState == EAsyncPackageLoadingState2::Finalize);
-			if (bHasClusterObjects)
-			{
-				Package->AsyncPackageLoadingState = EAsyncPackageLoadingState2::CreateClusters;
-			}
-			else
-			{
-				Package->AsyncPackageLoadingState = EAsyncPackageLoadingState2::Complete;
-			}
+			Package->AsyncPackageLoadingState = EAsyncPackageLoadingState2::PostLoadInstances;
 			PackagesReadyForCallback.Add(Package);
 		}
 
@@ -4777,7 +4776,30 @@ EAsyncPackageState::Type FAsyncLoadingThread2::ProcessLoadedPackagesFromGameThre
 			FAsyncPackage2* Package = CompletedPackages[PackageIndex];
 			{
 				bool bSafeToDelete = false;
-				if (Package->AsyncPackageLoadingState == EAsyncPackageLoadingState2::CreateClusters)
+				if (Package->AsyncPackageLoadingState == EAsyncPackageLoadingState2::PostLoadInstances)
+				{
+					SCOPE_CYCLE_COUNTER(STAT_FAsyncPackage_PostLoadInstancesGameThread);
+					// This package will create GC clusters but first check if all dependencies of this package have been fully loaded
+					if (Package->PostLoadInstances(ThreadState) == EAsyncPackageState::Complete)
+					{
+						// PostLoadInstances routed to all objects, safe to proceed to creating clusters or directly to completion
+						if (Package->bHasClusterObjects)
+						{
+							Package->AsyncPackageLoadingState = EAsyncPackageLoadingState2::CreateClusters;
+						}
+						else
+						{
+							Package->AsyncPackageLoadingState = EAsyncPackageLoadingState2::Complete;
+						}
+					}
+					else
+					{
+						// PostLoadInstances timed out
+						Result = EAsyncPackageState::TimeOut;
+						break;
+					}
+				}
+				else if (Package->AsyncPackageLoadingState == EAsyncPackageLoadingState2::CreateClusters)
 				{
 					SCOPE_CYCLE_COUNTER(STAT_FAsyncPackage_CreateClustersGameThread);
 					// This package will create GC clusters but first check if all dependencies of this package have been fully loaded
@@ -5829,6 +5851,21 @@ EAsyncPackageState::Type FAsyncPackage2::ProcessExternalReads(EExternalReadActio
 	return EAsyncPackageState::Complete;
 }
 
+EAsyncPackageState::Type FAsyncPackage2::PostLoadInstances(FAsyncLoadingThreadState2& ThreadState)
+{
+	while (PostLoadInstanceIndex < Data.ExportInfo.ExportCount && !ThreadState.IsTimeLimitExceeded(TEXT("PostLoadInstances")))
+	{
+		const FExportObject& Export = Data.Exports[PostLoadInstanceIndex++];
+
+		if (!(Export.bFiltered | Export.bExportLoadFailed))
+		{
+			UClass* ObjClass = Export.Object->GetClass();
+			ObjClass->PostLoadInstance(Export.Object);
+		}
+	}
+	return PostLoadInstanceIndex == Data.ExportInfo.ExportCount ? EAsyncPackageState::Complete : EAsyncPackageState::TimeOut;
+}
+
 EAsyncPackageState::Type FAsyncPackage2::CreateClusters(FAsyncLoadingThreadState2& ThreadState)
 {
 	while (DeferredClusterIndex < Data.ExportInfo.ExportCount && !ThreadState.IsTimeLimitExceeded(TEXT("CreateClusters")))
diff --git a/Engine/Source/Runtime/CoreUObject/Private/UObject/Class.cpp b/Engine/Source/Runtime/CoreUObject/Private/UObject/Class.cpp
--- a/Engine/Source/Runtime/CoreUObject/Private/UObject/Class.cpp
+++ b/Engine/Source/Runtime/CoreUObject/Private/UObject/Class.cpp
@@ -2152,7 +2152,7 @@ void UStruct::InstanceSubobjectTemplates( void* Data, void const* DefaultData, U
 
 	for ( FProperty* Property = RefLink; Property != NULL; Property = Property->NextRef )
 	{
-		if (Property->ContainsInstancedObjectProperty())
+		if (Property->ContainsInstancedObjectProperty() && (!InstanceGraph || !InstanceGraph->IsPropertyInSubobjectExclusionList(Property)))
 		{
 			Property->InstanceSubobjects( Property->ContainerPtrToValuePtr<uint8>(Data), (uint8*)Property->ContainerPtrToValuePtrForDefaults<uint8>(DefaultStruct, DefaultData), Owner, InstanceGraph );
 		}
diff --git a/Engine/Source/Runtime/CoreUObject/Private/UObject/UObjectGlobals.cpp b/Engine/Source/Runtime/CoreUObject/Private/UObject/UObjectGlobals.cpp
--- a/Engine/Source/Runtime/CoreUObject/Private/UObject/UObjectGlobals.cpp
+++ b/Engine/Source/Runtime/CoreUObject/Private/UObject/UObjectGlobals.cpp
@@ -1726,6 +1726,17 @@ void EndLoad(FUObjectSerializeContext* LoadContext, TArray<UPackage*>* OutLoaded
 				}
 			}
 
+			{
+				// Additional operation performed by classes (used for non-native initialization)
+				SCOPED_LOADTIMER(PostLoadInstance);
+				for (UObject* Obj : ObjLoaded)
+				{
+					UClass* ObjClass = Obj->GetClass();
+					check(ObjClass);
+					ObjClass->PostLoadInstance(Obj);
+				}
+			}
+
 			// Create clusters after all objects have been loaded
 			if (FPlatformProperties::RequiresCookedData() && !GIsInitialLoad && GCreateGCClusters && GAssetClustreringEnabled && !GUObjectArray.IsOpenForDisregardForGC())
 			{
@@ -3051,7 +3062,7 @@ void FObjectInitializer::PostConstructInit()
 		Obj->PostInitProperties();
 	}
 
-	Class->PostInitInstance(Obj);
+	Class->PostInitInstance(Obj, InstanceGraph);
 
 #if !(UE_BUILD_SHIPPING || UE_BUILD_TEST)
 	if (!FUObjectThreadContext::Get().PostInitPropertiesCheck.Num() || (FUObjectThreadContext::Get().PostInitPropertiesCheck.Pop(false) != Obj))
diff --git a/Engine/Source/Runtime/CoreUObject/Public/Serialization/AsyncPackageLoader.h b/Engine/Source/Runtime/CoreUObject/Public/Serialization/AsyncPackageLoader.h
--- a/Engine/Source/Runtime/CoreUObject/Public/Serialization/AsyncPackageLoader.h
+++ b/Engine/Source/Runtime/CoreUObject/Public/Serialization/AsyncPackageLoader.h
@@ -16,6 +16,7 @@ DECLARE_STATS_GROUP(TEXT("Async Load Game Thread"), STATGROUP_AsyncLoadGameThrea
 DECLARE_CYCLE_STAT(TEXT("PostLoadObjects GT"), STAT_FAsyncPackage_PostLoadObjectsGameThread, STATGROUP_AsyncLoadGameThread);
 DECLARE_CYCLE_STAT(TEXT("TickAsyncLoading GT"), STAT_FAsyncPackage_TickAsyncLoadingGameThread, STATGROUP_AsyncLoadGameThread);
 DECLARE_CYCLE_STAT(TEXT("Flush Async Loading GT"), STAT_FAsyncPackage_FlushAsyncLoadingGameThread, STATGROUP_AsyncLoadGameThread);
+DECLARE_CYCLE_STAT_WITH_FLAGS(TEXT("PostLoadInstances GT"), STAT_FAsyncPackage_PostLoadInstancesGameThread, STATGROUP_AsyncLoadGameThread, EStatFlags::Verbose);
 DECLARE_CYCLE_STAT_WITH_FLAGS(TEXT("CreateClusters GT"), STAT_FAsyncPackage_CreateClustersGameThread, STATGROUP_AsyncLoadGameThread, EStatFlags::Verbose);
 
 enum class ENotifyRegistrationType;
diff --git a/Engine/Source/Runtime/CoreUObject/Public/UObject/Class.h b/Engine/Source/Runtime/CoreUObject/Public/UObject/Class.h
--- a/Engine/Source/Runtime/CoreUObject/Public/UObject/Class.h
+++ b/Engine/Source/Runtime/CoreUObject/Public/UObject/Class.h
@@ -3069,7 +3069,12 @@ public:
 	/**
 	 * Called after PostInitProperties during object construction to allow class specific initialization of an object instance.
 	 */
-	virtual void PostInitInstance(UObject* InObj) {}
+	virtual void PostInitInstance(UObject* InObj, FObjectInstancingGraph* InstanceGraph) {}
+
+	/**
+	 * Called during PostLoad of an object to allow class specific PostLoad operations of an object instance.
+	 */
+	virtual void PostLoadInstance(UObject* InObj) {}
 
 	/**
 	 * Helper method to assist with initializing object properties from an explicit list.
@@ -3661,6 +3666,22 @@ public:
 		bLoadingObject = bIsLoading;
 	}
 
+	/**
+	 * Adds a member variable property that should not instantiate subobjects
+	 */
+	void AddPropertyToSubobjectExclusionList(const FProperty* Property)
+	{
+		SubobjectInstantiationExclusionList.Add(Property);
+	}
+
+	/**
+	 * Checks if a member variable property is in subobject instantiation exclusion list
+	 */
+	bool IsPropertyInSubobjectExclusionList(const FProperty* Property) const
+	{
+		return SubobjectInstantiationExclusionList.Contains(Property);
+	}
+
 private:
 	/**
 	 * Returns whether this instancing graph has a valid destination root.
@@ -3748,6 +3769,9 @@ private:
 	 * Maps the source (think archetype) to the destination (think instance)
 	 */
 	TMap<class UObject*,class UObject*>			SourceToDestinationMap;
+
+	/** List of member variable properties that should not instantiate subobjects */
+	TSet<const FProperty*>						SubobjectInstantiationExclusionList;
 };
 
 // UFunction interface.
