Added context to PostCDOCompiled and also call it for skeleton-only compiles

diff --git a/Engine/Plugins/Runtime/StateTree/Source/StateTreeModule/Private/Blueprint/StateTreeItemBlueprintBase.cpp b/Engine/Plugins/Runtime/StateTree/Source/StateTreeModule/Private/Blueprint/StateTreeItemBlueprintBase.cpp
--- a/Engine/Plugins/Runtime/StateTree/Source/StateTreeModule/Private/Blueprint/StateTreeItemBlueprintBase.cpp
+++ b/Engine/Plugins/Runtime/StateTree/Source/StateTreeModule/Private/Blueprint/StateTreeItemBlueprintBase.cpp
@@ -1,4 +1,4 @@
-﻿// Copyright Epic Games, Inc. All Rights Reserved.
+// Copyright Epic Games, Inc. All Rights Reserved.
 
 #include "Blueprint/StateTreeItemBlueprintBase.h"
 #include "CoreMinimal.h"
@@ -31,8 +31,13 @@ AActor* UStateTreeItemBlueprintBase::GetOwnerActor(const FStateTreeExecutionCont
 }
 
 #if WITH_EDITOR
-void UStateTreeItemBlueprintBase::PostCDOCompiled()
+void UStateTreeItemBlueprintBase::PostCDOCompiled(const FPostCDOCompiledContext& Context)
 {
+	if (Context.bIsSkeletonOnly)
+	{
+		return;
+	}
+
 	// Save some property data for runtime use.
 	PropertyInfos.Reset();
 
diff --git a/Engine/Plugins/Runtime/StateTree/Source/StateTreeModule/Public/Blueprint/StateTreeItemBlueprintBase.h b/Engine/Plugins/Runtime/StateTree/Source/StateTreeModule/Public/Blueprint/StateTreeItemBlueprintBase.h
--- a/Engine/Plugins/Runtime/StateTree/Source/StateTreeModule/Public/Blueprint/StateTreeItemBlueprintBase.h
+++ b/Engine/Plugins/Runtime/StateTree/Source/StateTreeModule/Public/Blueprint/StateTreeItemBlueprintBase.h
@@ -1,4 +1,4 @@
-﻿// Copyright Epic Games, Inc. All Rights Reserved.
+// Copyright Epic Games, Inc. All Rights Reserved.
 
 #pragma once
 
@@ -61,7 +61,7 @@ protected:
 	AActor* GetOwnerActor(const FStateTreeExecutionContext& Context) const;
 
 #if WITH_EDITOR
-	virtual void PostCDOCompiled() override;
+	virtual void PostCDOCompiled(const FPostCDOCompiledContext& Context) override;
 #endif
 
 	/** Metadata for properties */
diff --git a/Engine/Source/Editor/AnimGraph/Private/AnimBlueprintCompiler.cpp b/Engine/Source/Editor/AnimGraph/Private/AnimBlueprintCompiler.cpp
--- a/Engine/Source/Editor/AnimGraph/Private/AnimBlueprintCompiler.cpp
+++ b/Engine/Source/Editor/AnimGraph/Private/AnimBlueprintCompiler.cpp
@@ -1335,8 +1335,13 @@ void FAnimBlueprintCompilerContext::SpawnNewClass(const FString& NewClassName)
 	});
 }
 
-void FAnimBlueprintCompilerContext::OnPostCDOCompiled()
+void FAnimBlueprintCompilerContext::OnPostCDOCompiled(const UObject::FPostCDOCompiledContext& Context)
 {
+	if (Context.bIsSkeletonOnly)
+	{
+		return;
+	}
+
 	UAnimBlueprintGeneratedClass* NewAnimBlueprintClass = GetNewAnimBlueprintClass();
 	NewAnimBlueprintClass->OnPostLoadDefaults(NewAnimBlueprintClass->ClassDefaultObject);
 }
diff --git a/Engine/Source/Editor/AnimGraph/Private/AnimBlueprintCompiler.h b/Engine/Source/Editor/AnimGraph/Private/AnimBlueprintCompiler.h
--- a/Engine/Source/Editor/AnimGraph/Private/AnimBlueprintCompiler.h
+++ b/Engine/Source/Editor/AnimGraph/Private/AnimBlueprintCompiler.h
@@ -63,7 +63,7 @@ protected:
 	virtual void ProcessOneFunctionGraph(UEdGraph* SourceGraph, bool bInternalFunction = false) override;
 	virtual void SpawnNewClass(const FString& NewClassName) override;
 	virtual void OnNewClassSet(UBlueprintGeneratedClass* ClassToUse) override;
-	virtual void OnPostCDOCompiled() override;
+	virtual void OnPostCDOCompiled(const UObject::FPostCDOCompiledContext& Context) override;
 	virtual void CopyTermDefaultsToDefaultObject(UObject* DefaultObject) override;
 	virtual void PostCompile() override;
 	virtual void PostCompileDiagnostics() override;
diff --git a/Engine/Source/Editor/Kismet/Private/BlueprintCompilationManager.cpp b/Engine/Source/Editor/Kismet/Private/BlueprintCompilationManager.cpp
--- a/Engine/Source/Editor/Kismet/Private/BlueprintCompilationManager.cpp
+++ b/Engine/Source/Editor/Kismet/Private/BlueprintCompilationManager.cpp
@@ -1478,9 +1478,12 @@ void FBlueprintCompilationManagerImpl::FlushCompilationQueueImpl(bool bSuppressB
 		{
 			DECLARE_SCOPE_HIERARCHICAL_COUNTER(PostCDOCompiled)
 
-			if(!CompilerData.IsSkeletonOnly() && CompilerData.Compiler.IsValid())
+			if (CompilerData.Compiler.IsValid())
 			{
-				CompilerData.Compiler->PostCDOCompiled();
+				UObject::FPostCDOCompiledContext PostCDOCompiledContext;
+				PostCDOCompiledContext.bIsSkeletonOnly = CompilerData.IsSkeletonOnly();
+
+				CompilerData.Compiler->PostCDOCompiled(PostCDOCompiledContext);
 			}
 		}
 
diff --git a/Engine/Source/Editor/KismetCompiler/Private/KismetCompiler.cpp b/Engine/Source/Editor/KismetCompiler/Private/KismetCompiler.cpp
--- a/Engine/Source/Editor/KismetCompiler/Private/KismetCompiler.cpp
+++ b/Engine/Source/Editor/KismetCompiler/Private/KismetCompiler.cpp
@@ -4741,7 +4741,7 @@ void FKismetCompilerContext::CompileFunctions(EInternalCompilerFlags InternalFla
 			PropagateValuesToCDO(NewCDO, OldCDO);
 
 			// Perform any fixup or caching based on the new CDO.
-			PostCDOCompiled();
+			PostCDOCompiled(UObject::FPostCDOCompiledContext());
 		}
 
 		// Note: The old->new CDO copy is deferred when regenerating, so we skip this step in that case.
@@ -5037,13 +5037,17 @@ void FKismetCompilerContext::CompileFunctions(EInternalCompilerFlags InternalFla
 	PostCompile();
 }
 
-void FKismetCompilerContext::PostCDOCompiled()
+void FKismetCompilerContext::PostCDOCompiled(const UObject::FPostCDOCompiledContext& Context)
 {
 	// Notify the CDO that it has finished compiling
-	NewClass->ClassDefaultObject->PostCDOCompiled();
+	NewClass->ClassDefaultObject->PostCDOCompiled(Context);
+	if (!Context.bIsSkeletonOnly)
+	{
+		NewClass->ClassDefaultObject->PostCDOCompiled();
+	}
 
 	// Allow children to customize PostCDOCompile:
-	OnPostCDOCompiled();
+	OnPostCDOCompiled(Context);
 }
 
 void FKismetCompilerContext::Compile()
diff --git a/Engine/Source/Editor/KismetCompiler/Public/KismetCompiler.h b/Engine/Source/Editor/KismetCompiler/Public/KismetCompiler.h
--- a/Engine/Source/Editor/KismetCompiler/Public/KismetCompiler.h
+++ b/Engine/Source/Editor/KismetCompiler/Public/KismetCompiler.h
@@ -156,7 +156,7 @@ public:
 	void CompileFunctions(EInternalCompilerFlags InternalFlags);
 
 	/** Called after the CDO has been generated, allows assignment of cached/derived data: */
-	void PostCDOCompiled();
+	void PostCDOCompiled(const UObject::FPostCDOCompiledContext& Context);
 
 	/** Compile a blueprint into a class and a set of functions */
 	void Compile();
@@ -289,7 +289,7 @@ protected:
 	virtual void PostCreateSchema();
 	virtual void SpawnNewClass(const FString& NewClassName);
 	virtual void OnNewClassSet(UBlueprintGeneratedClass* ClassToUse) {}
-	virtual void OnPostCDOCompiled() {}
+	virtual void OnPostCDOCompiled(const UObject::FPostCDOCompiledContext& Context) {}
 
 	/**
 	 * Backwards Compatability:  Ensures that the passed in TargetClass is of the proper type (e.g. BlueprintGeneratedClass, AnimBlueprintGeneratedClass), and NULLs the reference if it is not 
diff --git a/Engine/Source/Editor/UMGEditor/Private/WidgetBlueprintCompiler.cpp b/Engine/Source/Editor/UMGEditor/Private/WidgetBlueprintCompiler.cpp
--- a/Engine/Source/Editor/UMGEditor/Private/WidgetBlueprintCompiler.cpp
+++ b/Engine/Source/Editor/UMGEditor/Private/WidgetBlueprintCompiler.cpp
@@ -988,9 +988,14 @@ private:
 };
 
 
-void FWidgetBlueprintCompilerContext::OnPostCDOCompiled()
+void FWidgetBlueprintCompilerContext::OnPostCDOCompiled(const UObject::FPostCDOCompiledContext& Context)
 {
-	Super::OnPostCDOCompiled();
+	Super::OnPostCDOCompiled(Context);
+
+	if (Context.bIsSkeletonOnly)
+	{
+		return;
+	}
 
 	WidgetToMemberVariableMap.Empty();
 	WidgetAnimToMemberVariableMap.Empty();
diff --git a/Engine/Source/Editor/UMGEditor/Public/WidgetBlueprintCompiler.h b/Engine/Source/Editor/UMGEditor/Public/WidgetBlueprintCompiler.h
--- a/Engine/Source/Editor/UMGEditor/Public/WidgetBlueprintCompiler.h
+++ b/Engine/Source/Editor/UMGEditor/Public/WidgetBlueprintCompiler.h
@@ -69,7 +69,7 @@ protected:
 	virtual void CopyTermDefaultsToDefaultObject(UObject* DefaultObject);
 	virtual void FinishCompilingClass(UClass* Class) override;
 	virtual bool ValidateGeneratedClass(UBlueprintGeneratedClass* Class) override;
-	virtual void OnPostCDOCompiled() override;
+	virtual void OnPostCDOCompiled(const UObject::FPostCDOCompiledContext& Context) override;
 	// End FKismetCompilerContext
 
 	void SanitizeBindings(UBlueprintGeneratedClass* Class);
diff --git a/Engine/Source/Runtime/CoreUObject/Public/UObject/Object.h b/Engine/Source/Runtime/CoreUObject/Public/UObject/Object.h
--- a/Engine/Source/Runtime/CoreUObject/Public/UObject/Object.h
+++ b/Engine/Source/Runtime/CoreUObject/Public/UObject/Object.h
@@ -216,11 +216,25 @@ public:
 	}
 
 #if WITH_EDITOR
+	struct FPostCDOCompiledContext
+	{
+		/** True if this notification was from a 'skeleton-only' compile */
+		bool bIsSkeletonOnly = false;
+	};
+
 	/**
 	 * Called after the Blueprint compiler has finished generating the Class Default Object (CDO) for a class. This can only happen in the editor.
 	 * This is called when the CDO and its associated class structure have been fully generated and populated, and allows the assignment of cached/derived data, 
 	 * eg) caching the name/count of properties of a certain type, or inspecting the properties on the class and using their meta-data and CDO default values to derive game data.
 	 */
+	virtual void PostCDOCompiled(const FPostCDOCompiledContext& Context)
+	{
+	}
+
+	/**
+	 * Called after the Blueprint compiler has finished generating the Class Default Object (CDO) for a class. This can only happen in the editor.
+	 * @note This version is deprecated as it wasn't called for skeleton-only compilation. Use the version taking FPostCDOCompiledContext instead.
+	 */
 	virtual void PostCDOCompiled()
 	{
 	}
