Use %llu instead of %ull as cross platform printf specifier

#rnx

diff --git a/Engine/Source/Runtime/CoreUObject/Private/UObject/EnumProperty.cpp b/Engine/Source/Runtime/CoreUObject/Private/UObject/EnumProperty.cpp
--- a/Engine/Source/Runtime/CoreUObject/Private/UObject/EnumProperty.cpp
+++ b/Engine/Source/Runtime/CoreUObject/Private/UObject/EnumProperty.cpp
@@ -279,7 +279,7 @@ void FEnumProperty::ExportTextItem(FString& ValueStr, const void* PropertyValue,
 		if (GoodValue == MaxValue)
 		{
 			// not all native enums have Max value declared
-			ValueStr += FString::Printf(TEXT("(%s)(%ull)"), *FullyQualifiedEnumName, ActualValue);
+			ValueStr += FString::Printf(TEXT("(%s)(%llu)"), *FullyQualifiedEnumName, ActualValue);
 		}
 		else
 		{
diff --git a/Engine/Source/Runtime/Engine/Classes/Engine/EngineTypes.h b/Engine/Source/Runtime/Engine/Classes/Engine/EngineTypes.h
--- a/Engine/Source/Runtime/Engine/Classes/Engine/EngineTypes.h
+++ b/Engine/Source/Runtime/Engine/Classes/Engine/EngineTypes.h
@@ -3465,7 +3465,7 @@ struct FTimerHandle
 
 	FString ToString() const
 	{
-		return FString::Printf(TEXT("%ull"), Handle);
+		return FString::Printf(TEXT("%llu"), Handle);
 	}
 
 private:
