Fully deprecated the contextless PostCDOCompiled

#rnx

diff --git a/Engine/Source/Editor/KismetCompiler/Private/KismetCompiler.cpp b/Engine/Source/Editor/KismetCompiler/Private/KismetCompiler.cpp
--- a/Engine/Source/Editor/KismetCompiler/Private/KismetCompiler.cpp
+++ b/Engine/Source/Editor/KismetCompiler/Private/KismetCompiler.cpp
@@ -5043,7 +5043,9 @@ void FKismetCompilerContext::PostCDOCompiled(const UObject::FPostCDOCompiledCont
 	NewClass->ClassDefaultObject->PostCDOCompiled(Context);
 	if (!Context.bIsSkeletonOnly)
 	{
+		PRAGMA_DISABLE_DEPRECATION_WARNINGS;
 		NewClass->ClassDefaultObject->PostCDOCompiled();
+		PRAGMA_ENABLE_DEPRECATION_WARNINGS;
 	}
 
 	// Allow children to customize PostCDOCompile:
diff --git a/Engine/Source/Runtime/CoreUObject/Public/UObject/Object.h b/Engine/Source/Runtime/CoreUObject/Public/UObject/Object.h
--- a/Engine/Source/Runtime/CoreUObject/Public/UObject/Object.h
+++ b/Engine/Source/Runtime/CoreUObject/Public/UObject/Object.h
@@ -235,6 +235,7 @@ public:
 	 * Called after the Blueprint compiler has finished generating the Class Default Object (CDO) for a class. This can only happen in the editor.
 	 * @note This version is deprecated as it wasn't called for skeleton-only compilation. Use the version taking FPostCDOCompiledContext instead.
 	 */
+	UE_DEPRECATED(5.1, "Use version that takes FPostCDOCompiledContext instead.")
 	virtual void PostCDOCompiled()
 	{
 	}
