Adding support to writing platform specific tags to the asset registry during cook by the book.

FArchive now holds a "CookData" blob that contains relevant information for cooking. The TargetPlatform has been moved to this structure, and it has a reference to a CookContext that can hold data generated during the cook that doesn't get serialized to the uasset (in this case, the asset registry tags).

diff --git a/Engine/Source/Editor/UnrealEd/Private/Commandlets/AssetRegistryGenerator.cpp b/Engine/Source/Editor/UnrealEd/Private/Commandlets/AssetRegistryGenerator.cpp
--- a/Engine/Source/Editor/UnrealEd/Private/Commandlets/AssetRegistryGenerator.cpp
+++ b/Engine/Source/Editor/UnrealEd/Private/Commandlets/AssetRegistryGenerator.cpp
@@ -19,6 +19,7 @@
 #include "CollectionManagerModule.h"
 #include "Interfaces/ITargetPlatform.h"
 #include "AssetRegistry/AssetRegistryModule.h"
+#include "AssetRegistry/CookTagList.h"
 #include "GameDelegates.h"
 #include "Commandlets/IChunkDataGenerator.h"
 #include "Commandlets/ChunkDependencyInfo.h"
@@ -981,6 +982,41 @@ void FAssetRegistryGenerator::UpdateKeptPackages()
 	}
 }
 
+static void AddCookTagsToState(TMap<FSoftObjectPath, TArray<TPair<FName, FString>>>&& InCookTags, FAssetRegistryState& InState)
+{
+	TRACE_CPUPROFILER_EVENT_SCOPE(AddCookTagToState)
+
+	uint32 ObjectCount = 0;
+	uint32 TagCount = 0;
+
+	for (TPair<FSoftObjectPath, TArray<TPair<FName, FString>>>& ObjectToTags : InCookTags)
+	{
+		ObjectCount++;
+		const FAssetData* AssetData = InState.GetAssetByObjectPath(FName(WriteToString<256>(ObjectToTags.Key)));
+
+		// Migrate to FAssetDataTagMap
+		FAssetDataTagMap NewTags;				
+		for (TPair<FName, FString>& Tag : ObjectToTags.Value)
+		{			
+			// Don't add empty tags
+			if (!Tag.Key.IsNone() && !Tag.Value.IsEmpty())
+			{
+				// Prepend Cook_
+				// Do the accumulation in a stack buffer to avoid a bit of work, but we still
+				// eat the FName generation.
+				FName TagName(WriteToString<256>(TEXTVIEW("Cook_"), Tag.Key));
+				NewTags.Add(TagName, MoveTemp(Tag.Value));
+				TagCount++;
+			}
+		}
+
+		InState.AddTagsToAssetData(ObjectToTags.Key, MoveTemp(NewTags));
+	} // end for each object
+
+	UE_LOG(LogAssetRegistryGenerator, Verbose, TEXT("Added %d cook tags to %d objects"), TagCount, ObjectCount);
+	InCookTags.Reset();
+}
+
 void FAssetRegistryGenerator::UpdateCollectionAssetData()
 {
 	// Read out the per-platform settings use to build the list of collections to tag
@@ -1410,10 +1446,16 @@ bool FAssetRegistryGenerator::SaveAssetRegistry(const FString& SandboxPath, bool
 
 		if (bSerializeDevelopmentAssetRegistry)
 		{
+			// Make a copy of the state so we can add tags to only the development registry.
+			FAssetRegistryState DevelopmentState;
+			DevelopmentState.InitializeFromExisting(State, DevelopmentSaveOptions);
+
+			AddCookTagsToState(MoveTemp(CookTagsToAdd), DevelopmentState);
+
 			// Create development registry data, used for DLC cooks, iterative cooks, and editor viewing
 			FArrayWriter SerializedAssetRegistry;
 
-			State.Save(SerializedAssetRegistry, DevelopmentSaveOptions);
+			DevelopmentState.Save(SerializedAssetRegistry, DevelopmentSaveOptions);
 
 			// Save the generated registry
 			FFileHelper::SaveArrayToFile(SerializedAssetRegistry, *PlatformSandboxPath);
@@ -2357,7 +2399,9 @@ const FAssetData* FAssetRegistryGenerator::CreateOrFindAssetData(UObject& Object
 }
 
 void FAssetRegistryGenerator::UpdateAssetRegistryPackageData(const UPackage& Package,
-	FSavePackageResultStruct& SavePackageResult, TFuture<FMD5Hash>& CookedHash)
+	FSavePackageResultStruct& SavePackageResult, TFuture<FMD5Hash>& CookedHash,
+	FCookTagList&& InArchiveCookTagList
+	)
 {
 	const FName PackageName = Package.GetFName();
 	PreviousPackagesToUpdate.Remove(PackageName);
@@ -2370,6 +2414,10 @@ void FAssetRegistryGenerator::UpdateAssetRegistryPackageData(const UPackage& Pac
 		// Ensure all assets in the package are recorded in the registry
 		CreateOrFindAssetDatas(Package);
 
+		// Migrate cook tags over
+		CookTagsToAdd.Append(MoveTemp(InArchiveCookTagList.ObjectToTags));
+		InArchiveCookTagList.Reset();
+
 		// Set the PackageFlags to the recorded value from SavePackage
 		NewPackageFlags = SavePackageResult.SerializedPackageFlags;
 
diff --git a/Engine/Source/Editor/UnrealEd/Private/Commandlets/SavePackageUtilitiesCommandlet.cpp b/Engine/Source/Editor/UnrealEd/Private/Commandlets/SavePackageUtilitiesCommandlet.cpp
--- a/Engine/Source/Editor/UnrealEd/Private/Commandlets/SavePackageUtilitiesCommandlet.cpp
+++ b/Engine/Source/Editor/UnrealEd/Private/Commandlets/SavePackageUtilitiesCommandlet.cpp
@@ -11,6 +11,7 @@
 #include "Misc/FeedbackContext.h"
 #include "Misc/PackageName.h"
 #include "Misc/Paths.h"
+#include "UObject/ArchiveCookContext.h"
 #include "UObject/LinkerDiff.h"
 #include "UObject/LinkerSave.h"
 #include "UObject/Package.h"
@@ -34,10 +35,18 @@ int32 USavePackageUtilitiesCommandlet::Main(const FString& Params)
 		UObject* Asset = Package->FindAssetInPackage();
 		FString Filename = FPaths::CreateTempFilename(*FPaths::ProjectSavedDir());
 
+		// CookData should only be nonzero if we are cooking.
+		TOptional<FArchiveCookData> CookData;
+		FArchiveCookContext CookContext(Package, FArchiveCookContext::ECookTypeUnknown);
+		if (TargetPlatform != nullptr)
+		{
+			CookData.Emplace(*TargetPlatform, CookContext);
+		}
+
 		FSavePackageArgs SaveArgs;
 		SaveArgs.TopLevelFlags = RF_Public;
 		SaveArgs.SaveFlags = SAVE_CompareLinker;
-		SaveArgs.TargetPlatform = TargetPlatform;
+		SaveArgs.ArchiveCookData = CookData.GetPtrOrNull();
 		SaveArgs.bSlowTask = false;
 
 		// if not cooking add RF_Standalone to the top level flags
diff --git a/Engine/Source/Editor/UnrealEd/Private/CookOnTheFlyServer.cpp b/Engine/Source/Editor/UnrealEd/Private/CookOnTheFlyServer.cpp
--- a/Engine/Source/Editor/UnrealEd/Private/CookOnTheFlyServer.cpp
+++ b/Engine/Source/Editor/UnrealEd/Private/CookOnTheFlyServer.cpp
@@ -108,6 +108,7 @@
 #include "String/ParseLines.h"
 #include "TargetDomain/TargetDomainUtils.h"
 #include "UnrealEdGlobals.h"
+#include "UObject/ArchiveCookContext.h"
 #include "UObject/Class.h"
 #include "UObject/ConstructorHelpers.h"
 #include "UObject/GarbageCollection.h"
@@ -3039,6 +3040,10 @@ private: // Used only by UCookOnTheFlyServer, which has private access
 	bool bHasDelayLoaded = false;
 	bool bContainsMap = false;
 
+	// This holds cook data generated by Serialize() that isn't saved in the package - only valid in
+	// PerPlatform callbacks and is reset in SetupPlatform.
+	FArchiveCookContext ArchiveCookContext;
+
 	// Per-platform data, only valid in PerPlatform callbacks
 	const ITargetPlatform* TargetPlatform = nullptr;
 	FCookSavePackageContext* CookContext = nullptr;
@@ -4113,12 +4118,13 @@ void UCookOnTheFlyServer::SaveCookedPackage(UE::Cook::FSaveCookedPackageContext&
 			UE_SCOPED_HIERARCHICAL_COOKTIMER(GEditorSavePackage);
 			UE_TRACK_REFERENCING_PLATFORM_SCOPED(TargetPlatform);
 
+			FArchiveCookData CookData(*TargetPlatform, Context.ArchiveCookContext);
 			FSavePackageArgs SaveArgs;
 			SaveArgs.TopLevelFlags = Context.FlagsToCook;
 			SaveArgs.bForceByteSwapping = Context.bEndianSwap;
 			SaveArgs.bWarnOfLongFilename = false;
 			SaveArgs.SaveFlags = Context.SaveFlags;
-			SaveArgs.TargetPlatform = TargetPlatform;
+			SaveArgs.ArchiveCookData = &CookData;
 			SaveArgs.bSlowTask = false;
 			SaveArgs.SavePackageContext = Context.SavePackageContext;
 
@@ -4170,6 +4176,7 @@ FSaveCookedPackageContext::FSaveCookedPackageContext(UCookOnTheFlyServer& InCOTF
 	, Package(PackageData.GetPackage())
 	, PackageName(Package ? Package->GetName() : FString())
 	, Filename(PackageData.GetFileName().ToString())
+	, ArchiveCookContext(Package, COTFS.IsCookByTheBookMode() ? FArchiveCookContext::ECookByTheBook : FArchiveCookContext::ECookOnTheFly)
 {
 }
 
@@ -4204,6 +4211,8 @@ void FSaveCookedPackageContext::SetupPlatform(const ITargetPlatform* InTargetPla
 	PlatFilename = Filename.Replace(TEXT("[Platform]"), *TargetPlatform->PlatformName());
 	bPlatformSetupSuccessful = false;
 
+	ArchiveCookContext.Reset();
+
 	// don't save Editor resources from the Engine if the target doesn't have editoronly data
 	if (COTFS.IsCookFlagSet(ECookInitializationFlags::SkipEditorContent) &&
 		(PackageName.StartsWith(TEXT("/Engine/Editor")) || PackageName.StartsWith(TEXT("/Engine/VREditor"))) &&
@@ -4332,7 +4341,7 @@ void FSaveCookedPackageContext::FinishPlatform()
 	// Update asset registry
 	if (COTFS.CookByTheBookOptions)
 	{
-		Generator.UpdateAssetRegistryPackageData(*Package, SavePackageResult, CookedHash);
+		Generator.UpdateAssetRegistryPackageData(*Package, SavePackageResult, CookedHash, MoveTemp(*ArchiveCookContext.GetCookTagList()));
 	}
 
 	// In success or failure we want to mark the package as cooked, unless the failure was due to being referenced only
@@ -8854,7 +8863,10 @@ uint32 UCookOnTheFlyServer::FullLoadAndSave(uint32& CookedPackageCount)
 						{
 							Package->ClearPackageFlags(PKG_FilterEditorOnly);
 						}
-								
+						
+						FArchiveCookContext CookContext(Package, FArchiveCookContext::ECookByTheBook);
+						FArchiveCookData CookData(*Target, CookContext);
+
 						GIsCookerLoadingPackage = true;
 						check(SavePackageContexts.Num() > PlatformIndex);
 						FSavePackageContext& SavePackageContext = SavePackageContexts[PlatformIndex]->SaveContext;
@@ -8867,7 +8879,7 @@ uint32 UCookOnTheFlyServer::FullLoadAndSave(uint32& CookedPackageCount)
 						SaveArgs.bForceByteSwapping = bSwap;
 						SaveArgs.bWarnOfLongFilename = false;
 						SaveArgs.SaveFlags = SaveFlags;
-						SaveArgs.TargetPlatform = Target;
+						SaveArgs.ArchiveCookData = &CookData;
 						SaveArgs.bSlowTask = false;
 						SaveArgs.SavePackageContext = &SavePackageContext;
 						FSavePackageResultStruct SaveResult = GEditor->Save(Package, World, *PlatFilename, SaveArgs);
@@ -8882,7 +8894,7 @@ uint32 UCookOnTheFlyServer::FullLoadAndSave(uint32& CookedPackageCount)
 						}
 
 						// Update asset registry
-						Generator.UpdateAssetRegistryPackageData(*Package, SaveResult, SaveResult.CookedHash);
+						Generator.UpdateAssetRegistryPackageData(*Package, SaveResult, SaveResult.CookedHash, MoveTemp(*CookContext.GetCookTagList()));
 						FAssetPackageData* AssetPackageData = Generator.GetAssetPackageData(Package->GetFName());
 						check(AssetPackageData);
 
diff --git a/Engine/Source/Editor/UnrealEd/Private/Cooker/CookTypes.h b/Engine/Source/Editor/UnrealEd/Private/Cooker/CookTypes.h
--- a/Engine/Source/Editor/UnrealEd/Private/Cooker/CookTypes.h
+++ b/Engine/Source/Editor/UnrealEd/Private/Cooker/CookTypes.h
@@ -172,7 +172,10 @@ namespace UE::Cook
 		}
 	};
 
-	/** Context data passed into SavePackage for a given TargetPlatform. */
+	/** 
+	* Context data passed into SavePackage for a given TargetPlatform. Constant across packages, and internal
+	* to the cooker.
+	*/
 	struct FCookSavePackageContext
 	{
 		FCookSavePackageContext(const ITargetPlatform* InTargetPlatform,
diff --git a/Engine/Source/Editor/UnrealEd/Private/EditorEngine.cpp b/Engine/Source/Editor/UnrealEd/Private/EditorEngine.cpp
--- a/Engine/Source/Editor/UnrealEd/Private/EditorEngine.cpp
+++ b/Engine/Source/Editor/UnrealEd/Private/EditorEngine.cpp
@@ -92,6 +92,7 @@
 #include "Interfaces/IPluginManager.h"
 #include "UObject/PackageReload.h"
 #include "UObject/ReferenceChainSearch.h"
+#include "UObject/ArchiveCookContext.h"
 #include "HAL/PlatformApplicationMisc.h"
 #include "IMediaModule.h"
 #include "Scalability.h"
@@ -4308,7 +4309,14 @@ FSavePackageResultStruct UEditorEngine::Save(UPackage* InOuter, UObject* InBase,
 	const FDateTime& FinalTimeStamp, bool bSlowTask, FArchiveDiffMap* InOutDiffMap,
 	FSavePackageContext* SavePackageContext)
 {
-	FSavePackageArgs SaveArgs = { TargetPlatform, TopLevelFlags, SaveFlags, bForceByteSwapping,
+	// CookData can only be nonzero if we are cooking.
+	TOptional<FArchiveCookData> CookData;
+	FArchiveCookContext CookContext(InOuter, FArchiveCookContext::ECookTypeUnknown);
+	if (TargetPlatform != nullptr)
+	{
+		CookData.Emplace(*TargetPlatform, CookContext);
+	}
+	FSavePackageArgs SaveArgs = { nullptr /* deprecated target platform */, CookData.GetPtrOrNull(), TopLevelFlags, SaveFlags, bForceByteSwapping,
 		bWarnOfLongFilename, bSlowTask, FinalTimeStamp, Error, SavePackageContext };
 	return Save(InOuter, InBase, Filename, SaveArgs);
 }
@@ -4321,7 +4329,7 @@ FSavePackageResultStruct UEditorEngine::Save(UPackage* InOuter, UObject* InAsset
 
 	FSavePackageArgs SaveArgs(InSaveArgs);
 	FScopedSlowTask SlowTask(100, FText(), SaveArgs.bSlowTask);
-	bool bIsCooking = SaveArgs.TargetPlatform != nullptr;
+	bool bIsCooking = SaveArgs.ArchiveCookData != nullptr;
 	SaveArgs.TopLevelFlags = UE::SavePackageUtilities::NormalizeTopLevelFlags(SaveArgs.TopLevelFlags, bIsCooking);
 	UObject* Asset = InAsset;
 	if (!Asset && InOuter)
@@ -4345,7 +4353,7 @@ FSavePackageResultStruct UEditorEngine::Save(UPackage* InOuter, UObject* InAsset
 	bool bForceInitializedWorld = false;
 	const bool bSavingConcurrent = !!(SaveArgs.SaveFlags & ESaveFlags::SAVE_Concurrent);
 
-	FObjectSaveContextData ObjectSaveContext(InOuter, SaveArgs.TargetPlatform, Filename, SaveArgs.SaveFlags);
+	FObjectSaveContextData ObjectSaveContext(InOuter, SaveArgs.GetTargetPlatform(), Filename, SaveArgs.SaveFlags);
 	UWorld *OriginalOwningWorld = nullptr;
 	if ( World )
 	{
@@ -4455,7 +4463,14 @@ bool UEditorEngine::SavePackage(UPackage* InOuter, UObject* InBase, EObjectFlags
 	FOutputDevice* Error, FLinkerNull* Conform, bool bForceByteSwapping, bool bWarnOfLongFilename,
 	uint32 SaveFlags, const ITargetPlatform* TargetPlatform, const FDateTime& FinalTimeStamp, bool bSlowTask)
 {
-	FSavePackageArgs SaveArgs = { TargetPlatform, TopLevelFlags, SaveFlags, bForceByteSwapping,
+	// CookData should only be nonzero if we are cooking.
+	TOptional<FArchiveCookData> CookData;
+	FArchiveCookContext CookContext(InOuter, FArchiveCookContext::ECookTypeUnknown);
+	if (TargetPlatform != nullptr)
+	{
+		CookData.Emplace(*TargetPlatform, CookContext);
+	}
+	FSavePackageArgs SaveArgs = { nullptr /* deprecated target platform */, CookData.GetPtrOrNull(), TopLevelFlags, SaveFlags, bForceByteSwapping,
 		bWarnOfLongFilename, bSlowTask, FinalTimeStamp, Error };
 	return SavePackage(InOuter, InBase, Filename, SaveArgs);
 }
diff --git a/Engine/Source/Editor/UnrealEd/Public/Commandlets/AssetRegistryGenerator.h b/Engine/Source/Editor/UnrealEd/Public/Commandlets/AssetRegistryGenerator.h
--- a/Engine/Source/Editor/UnrealEd/Public/Commandlets/AssetRegistryGenerator.h
+++ b/Engine/Source/Editor/UnrealEd/Public/Commandlets/AssetRegistryGenerator.h
@@ -9,6 +9,7 @@
 #include "Misc/AssetRegistryInterface.h"
 #include "Misc/Paths.h"
 #include "Templates/UniquePtr.h"
+#include "UObject/Object.h"
 #include "UObject/UObjectHash.h"
 
 class FSandboxPlatformFile;
@@ -17,6 +18,7 @@ class ITargetPlatform;
 class IChunkDataGenerator;
 class UChunkDependencyInfo;
 struct FChunkDependencyTreeNode;
+struct FCookTagList;
 
 /**
  * Helper class for generating streaming install manifests
@@ -188,12 +190,13 @@ public:
 	void GetChunkAssignments(TArray<TSet<FName>>& OutAssignments) const;
 
 	/**
-	 * Attempts to update the metadata for a package in an asset registry generator
+	 * Attempts to update the metadata for a package in an asset registry generator.
+	 * This is only called for CookByTheBooks.
 	 *
 	 * @param Package The package to update info on
 	 * @param SavePackageResult The metadata to associate with the given package name
 	 */
-	void UpdateAssetRegistryPackageData(const UPackage& Package, FSavePackageResultStruct& SavePackageResult, TFuture<FMD5Hash>& CookedHash);
+	void UpdateAssetRegistryPackageData(const UPackage& Package, FSavePackageResultStruct& SavePackageResult, TFuture<FMD5Hash>& CookedHash, FCookTagList&& InArchiveCookTagList);
 
 private:
 	/**
@@ -221,6 +224,12 @@ private:
 
 	/** State of the asset registry that is being built for this platform */
 	FAssetRegistryState State;
+	
+	/** 
+	 * The list of tags to add for each asset. This is populated during cook by the books,
+	 * and is only added to development registries.
+	 */
+	TMap<FSoftObjectPath, TArray<TPair<FName, FString>>> CookTagsToAdd;
 
 	TMap<FName, TPair<TArray<FAssetData>, FAssetPackageData>> PreviousPackagesToUpdate;
 	/** List of packages that were loaded at startup */
diff --git a/Engine/Source/Runtime/Core/Private/Serialization/Archive.cpp b/Engine/Source/Runtime/Core/Private/Serialization/Archive.cpp
--- a/Engine/Source/Runtime/Core/Private/Serialization/Archive.cpp
+++ b/Engine/Source/Runtime/Core/Private/Serialization/Archive.cpp
@@ -192,7 +192,7 @@ void FArchiveState::Reset()
 	ArIsNetArchive						= false;
 	ArCustomPropertyList				= nullptr;
 	ArUseCustomPropertyList				= false;
-	CookingTargetPlatform				= nullptr;
+	CookData							= nullptr;
 	SerializedProperty					= nullptr;
 
 	delete SerializedPropertyChain;
@@ -252,7 +252,7 @@ void FArchiveState::CopyTrivialFArchiveStatusMembers(const FArchiveState& Archiv
 	ArIsNetArchive                       = ArchiveToCopy.ArIsNetArchive;
 	ArCustomPropertyList                 = ArchiveToCopy.ArCustomPropertyList;
 	ArUseCustomPropertyList              = ArchiveToCopy.ArUseCustomPropertyList;
-	CookingTargetPlatform                = ArchiveToCopy.CookingTargetPlatform;
+	CookData							 = ArchiveToCopy.CookData;
 	SerializedProperty					 = ArchiveToCopy.SerializedProperty;
 #if USE_STABLE_LOCALIZATION_KEYS
 	SetBaseLocalizationNamespace(ArchiveToCopy.GetBaseLocalizationNamespace());
@@ -911,10 +911,10 @@ void FArchive::SerializeCompressedNew(void* V, int64 Length, FName CompressionFo
 
 		// if there's a cooking target, and it wants to replace Zlib compression with another format, use it. When loading, 
 		// the platform will replace Zlib with that format above
-		if (CompressionFormatToEncode == NAME_Zlib && CookingTargetPlatform != nullptr)
+		if (CompressionFormatToEncode == NAME_Zlib && IsCooking())
 		{
 			// use the replacement format
-			CompressionFormatToEncode = CookingTargetPlatform->GetZlibReplacementFormat();
+			CompressionFormatToEncode = CookingTarget()->GetZlibReplacementFormat();
 
 			// with v2 headers, the modified CompressionFormatToEncode will be written in the archive
 		}
diff --git a/Engine/Source/Runtime/Core/Public/Serialization/Archive.h b/Engine/Source/Runtime/Core/Public/Serialization/Archive.h
--- a/Engine/Source/Runtime/Core/Public/Serialization/Archive.h
+++ b/Engine/Source/Runtime/Core/Public/Serialization/Archive.h
@@ -9,6 +9,7 @@
 #include "Misc/CompressionFlags.h"
 #include "Misc/EngineVersionBase.h"
 #include "Misc/VarArgs.h"
+#include "Serialization/ArchiveCookData.h"
 #include "Templates/AndOrNot.h"
 #include "Templates/EnableIf.h"
 #include "Templates/Function.h"
@@ -35,6 +36,7 @@ typedef TFunction<bool (double RemainingTime)> FExternalReadCallback;
 struct FUObjectSerializeContext;
 class FField;
 enum class EFileRegionType : uint8;
+struct FArchiveCookContext;
 
 // Temporary while we shake out the EDL at boot
 #define USE_EVENT_DRIVEN_ASYNC_LOAD_AT_BOOT_TIME (1)
@@ -582,30 +584,47 @@ public:
 	 * @return true if the archive is used for cooking, false otherwise.
 	 */
 	FORCEINLINE bool IsCooking() const
-	{
-		check(!CookingTargetPlatform || (!IsLoading() && !IsTransacting() && IsSaving()));
-
-		return !!CookingTargetPlatform;
+	{		
+		return CookData != nullptr;
 	}
 
 	/**
-	 * Returns the cooking target platform.
-	 *
-	 * @return Target platform.
+	 * Marks that this archive is "cooking" by providing the cook data bundle.
+	 * Must be set after loading/saving/transacting.
 	 */
-	FORCEINLINE const ITargetPlatform* CookingTarget() const
+	void SetCookData(FArchiveCookData* InCookData)
 	{
-		return CookingTargetPlatform;
+		checkf(!(InCookData == nullptr && CookData), TEXT("Can't turn off cooking once you turn it on!"));
+
+		if (InCookData)
+		{
+			check(!IsLoading() && !IsTransacting() && IsSaving());
+			CookData = InCookData;
+		}
+		
 	}
 
+	FArchiveCookData* GetCookData()
+	{
+		return CookData;
+	}
+
+	FORCEINLINE FArchiveCookContext* GetCookContext()
+	{
+		return CookData ? &CookData->CookContext : nullptr;
+	}
+
+	UE_DEPRECATED(5.1, "CookingTarget has been moved to FArchiveCookData.")
+	void SetCookingTarget(const ITargetPlatform*) {}
+
 	/**
-	 * Sets the cooking target platform.
+	 * Returns the cooking target platform.
 	 *
-	 * @param InCookingTarget The target platform to set.
+	 * @return Target platform.
 	 */
-	FORCEINLINE void SetCookingTarget(const ITargetPlatform* InCookingTarget)
+	FORCEINLINE const ITargetPlatform* CookingTarget() const
 	{
-		CookingTargetPlatform = InCookingTarget;
+		return CookData ? &CookData->TargetPlatform : nullptr;
 	}
 
 	/**
@@ -970,9 +989,9 @@ public:
 
 // These will be private in FArchive
 protected:
-	/** Holds the cooking target platform. */
-	const ITargetPlatform* CookingTargetPlatform;
-
+	/** Holds data for cooking. Required if cooking, nullptr means not cooking */
+	FArchiveCookData* CookData = nullptr;
+	
 	/** Holds the pointer to the property that is currently being serialized */
 	FProperty* SerializedProperty;
 
@@ -1880,7 +1899,9 @@ public:
 	using FArchiveState::IsNetArchive;
 	using FArchiveState::IsCooking;
 	using FArchiveState::CookingTarget;
-	using FArchiveState::SetCookingTarget;
+	using FArchiveState::SetCookData;
+	using FArchiveState::GetCookData;
+	using FArchiveState::GetCookContext;
 	using FArchiveState::UseToResolveEnumerators;
 	using FArchiveState::ShouldSkipProperty;
 	using FArchiveState::SetSerializedProperty;
@@ -2143,9 +2164,6 @@ public:
 	virtual void PopFileRegionType() { }
 
 private:
-	/** Holds the cooking target platform. */
-	using FArchiveState::CookingTargetPlatform;
-
 	/** Holds the pointer to the property that is currently being serialized */
 	using FArchiveState::SerializedProperty;
 
diff --git a/Engine/Source/Runtime/Core/Public/Serialization/ArchiveCookData.h b/Engine/Source/Runtime/Core/Public/Serialization/ArchiveCookData.h
new file mode 100644
--- /dev/null
+++ b/Engine/Source/Runtime/Core/Public/Serialization/ArchiveCookData.h
@@ -0,0 +1,18 @@
+// Copyright Epic Games, Inc. All Rights Reserved.
+
+#pragma once
+
+class ITargetPlatform;
+struct FArchiveCookContext;
+
+// Holds archive data only relevant for cooking archives. Fill this out
+// as part of FSavePackageArgs.
+struct CORE_API FArchiveCookData
+{
+	const ITargetPlatform& TargetPlatform;
+	FArchiveCookContext& CookContext;
+	FArchiveCookData(const ITargetPlatform& InTargetPlatform, FArchiveCookContext& InCookContext) :
+		TargetPlatform(InTargetPlatform),
+		CookContext(InCookContext)
+	{}
+};
diff --git a/Engine/Source/Runtime/CoreUObject/Private/UObject/SavePackage.cpp b/Engine/Source/Runtime/CoreUObject/Private/UObject/SavePackage.cpp
--- a/Engine/Source/Runtime/CoreUObject/Private/UObject/SavePackage.cpp
+++ b/Engine/Source/Runtime/CoreUObject/Private/UObject/SavePackage.cpp
@@ -23,6 +23,7 @@
 #include "Misc/ScopedSlowTask.h"
 #include "Misc/ObjectThumbnail.h"
 #include "UObject/ObjectMacros.h"
+#include "UObject/ArchiveCookContext.h"
 #include "UObject/UObjectGlobals.h"
 #include "UObject/UObjectHash.h"
 #include "UObject/Object.h"
@@ -372,7 +373,7 @@ public:
 		ArShouldSkipBulkData	= true;
 
 		ArPortFlags = Linker->GetPortFlags();
-		SetCookingTarget(Linker->CookingTarget());
+		SetCookData(Linker->GetCookData());
 	}
 
 	virtual FArchive& operator<<(UObject*& Obj) override;
@@ -1413,7 +1414,7 @@ struct FObjectExportSeekFreeSorter
 	 */
 	void SortExports(FLinkerSave* Linker)
 	{
-		SortArchive.SetCookingTarget(Linker->CookingTarget());
+		SortArchive.SetCookData(Linker->GetCookData());
 
 		TMap<UObject*,int32>	OriginalExportIndexes;
 
@@ -1775,7 +1776,15 @@ FSavePackageResultStruct UPackage::Save(UPackage* InOuter, UObject* Base, EObjec
 	const FDateTime& FinalTimeStamp, bool bSlowTask, FArchiveDiffMap* InOutDiffMap,
 	FSavePackageContext* SavePackageContext)
 {
-	FSavePackageArgs SaveArgs = { TargetPlatform, TopLevelFlags, SaveFlags, bForceByteSwapping,
+	// CookData should only be nonzero if we are cooking.
+	TOptional<FArchiveCookData> CookData;
+	FArchiveCookContext CookContext(InOuter, FArchiveCookContext::ECookTypeUnknown);
+	if (TargetPlatform != nullptr)
+	{		
+		CookData.Emplace(*TargetPlatform, CookContext);
+	}
+	
+	FSavePackageArgs SaveArgs = { nullptr /* deprecated target platform */, CookData.GetPtrOrNull(), TopLevelFlags, SaveFlags, bForceByteSwapping,
 		bWarnOfLongFilename, bSlowTask, FinalTimeStamp, Error, SavePackageContext };
 	return UPackage::Save(InOuter, Base, Filename, SaveArgs);
 }
@@ -1783,11 +1792,12 @@ FSavePackageResultStruct UPackage::Save(UPackage* InOuter, UObject* Base, EObjec
 FSavePackageResultStruct UPackage::Save(UPackage* InOuter, UObject* InAsset, const TCHAR* Filename,
 	const FSavePackageArgs& SaveArgs)
 {
-	const ITargetPlatform* TargetPlatform = SaveArgs.TargetPlatform;
-	if (SavePackageUtilities::IsNewSaveEnabled(TargetPlatform != nullptr))
+	const bool bIsCooking = SaveArgs.IsCooking();
+	if (SavePackageUtilities::IsNewSaveEnabled(bIsCooking))
 	{
 		return UPackage::Save2(InOuter, InAsset, Filename, SaveArgs);
 	}
+	const ITargetPlatform* TargetPlatform = SaveArgs.GetTargetPlatform();
 	UObject* Base = InAsset;
 	EObjectFlags TopLevelFlags = SaveArgs.TopLevelFlags;
 	FOutputDevice* Error = SaveArgs.Error;
@@ -1807,7 +1817,6 @@ FSavePackageResultStruct UPackage::Save(UPackage* InOuter, UObject* InAsset, con
 	check(InOuter);
 	check(Filename);
 
-	const bool bIsCooking = TargetPlatform != nullptr;
 	FPackagePath TargetPackagePath = FPackagePath::FromLocalPath(Filename);
 	if (TargetPackagePath.GetHeaderExtension() == EPackageExtension::Unspecified)
 	{
@@ -2022,7 +2031,7 @@ FSavePackageResultStruct UPackage::Save(UPackage* InOuter, UObject* InAsset, con
 			// Export objects (tags them as OBJECTMARK_TagExp).
 			FArchiveSaveTagExports ExportTaggerArchive( InOuter );
 			ExportTaggerArchive.SetPortFlags( ComparisonFlags );
-			ExportTaggerArchive.SetCookingTarget(TargetPlatform);
+			ExportTaggerArchive.SetCookData(SaveArgs.ArchiveCookData);
 			ExportTaggerArchive.SetSerializeContext(SaveContext);
 
 			check( ExportTaggerArchive.IsCooking() == !!TargetPlatform );
@@ -2236,7 +2245,7 @@ FSavePackageResultStruct UPackage::Save(UPackage* InOuter, UObject* InAsset, con
 
 				Linker->SetPortFlags(ComparisonFlags);
 				Linker->SetFilterEditorOnly( FilterEditorOnly );
-				Linker->SetCookingTarget(TargetPlatform);
+				Linker->SetCookData(SaveArgs.ArchiveCookData);
 
 				Linker->SetUseUnversionedPropertySerialization(bSaveUnversionedProperties);
 
@@ -4369,7 +4378,14 @@ bool UPackage::SavePackage(UPackage* InOuter, UObject* Base, EObjectFlags TopLev
 	FOutputDevice* Error, FLinkerNull* Conform, bool bForceByteSwapping, bool bWarnOfLongFilename, uint32 SaveFlags,
 	const ITargetPlatform* TargetPlatform, const FDateTime& FinalTimeStamp, bool bSlowTask)
 {
-	FSavePackageArgs SaveArgs = { TargetPlatform, TopLevelFlags, SaveFlags, bForceByteSwapping,
+	// CookData should only be nonzero if we are cooking.
+	TOptional<FArchiveCookData> CookData;
+	FArchiveCookContext CookContext(InOuter, FArchiveCookContext::ECookTypeUnknown);
+	if (TargetPlatform != nullptr)
+	{
+		CookData.Emplace(*TargetPlatform, CookContext);
+	}
+	FSavePackageArgs SaveArgs = { nullptr /* deprecated target platform */, CookData.GetPtrOrNull(), TopLevelFlags, SaveFlags, bForceByteSwapping,
 		bWarnOfLongFilename, bSlowTask, FinalTimeStamp, Error };
 	return SavePackage(InOuter, Base, Filename, SaveArgs);
 }
diff --git a/Engine/Source/Runtime/CoreUObject/Private/UObject/SavePackage/PackageHarvester.cpp b/Engine/Source/Runtime/CoreUObject/Private/UObject/SavePackage/PackageHarvester.cpp
--- a/Engine/Source/Runtime/CoreUObject/Private/UObject/SavePackage/PackageHarvester.cpp
+++ b/Engine/Source/Runtime/CoreUObject/Private/UObject/SavePackage/PackageHarvester.cpp
@@ -179,7 +179,7 @@ FPackageHarvester::FPackageHarvester(FSaveContext& InContext)
 
 	this->SetPortFlags(SaveContext.GetPortFlags());
 	this->SetFilterEditorOnly(SaveContext.IsFilterEditorOnly());
-	this->SetCookingTarget(SaveContext.GetTargetPlatform());
+	this->SetCookData(SaveContext.GetSaveArgs().ArchiveCookData);
 	this->SetSerializeContext(SaveContext.GetSerializeContext());
 	this->SetUseUnversionedPropertySerialization(SaveContext.IsSaveUnversionedProperties());
 }
diff --git a/Engine/Source/Runtime/CoreUObject/Private/UObject/SavePackage/SaveContext.h b/Engine/Source/Runtime/CoreUObject/Private/UObject/SavePackage/SaveContext.h
--- a/Engine/Source/Runtime/CoreUObject/Private/UObject/SavePackage/SaveContext.h
+++ b/Engine/Source/Runtime/CoreUObject/Private/UObject/SavePackage/SaveContext.h
@@ -374,7 +374,7 @@ public:
 		, SaveArgs(InSaveArgs)
 		, PackageWriter(InSaveArgs.SavePackageContext ? InSaveArgs.SavePackageContext->PackageWriter : nullptr)
 		, SerializeContext(InSerializeContext)
-		, ExcludedObjectMarks(SavePackageUtilities::GetExcludedObjectMarksForTargetPlatform(SaveArgs.TargetPlatform))
+		, ExcludedObjectMarks(SavePackageUtilities::GetExcludedObjectMarksForTargetPlatform(SaveArgs.GetTargetPlatform()))
 	{
 		// Assumptions & checks
 		check(InPackage);
@@ -401,7 +401,7 @@ public:
 			TargetPackagePath.SetHeaderExtension(EPackageExtension::EmptyString);
 		}
 
-		bCanUseUnversionedPropertySerialization = CanUseUnversionedPropertySerialization(SaveArgs.TargetPlatform);
+		bCanUseUnversionedPropertySerialization = CanUseUnversionedPropertySerialization(SaveArgs.GetTargetPlatform());
 		bTextFormat = FString(Filename).EndsWith(FPackageName::GetTextAssetPackageExtension()) || FString(Filename).EndsWith(FPackageName::GetTextMapPackageExtension());
 		static const IConsoleVariable* ProcessPrestreamingRequests = IConsoleManager::Get().FindConsoleVariable(TEXT("s.ProcessPrestreamingRequests"));
 		if (ProcessPrestreamingRequests)
@@ -433,9 +433,14 @@ public:
 		return SaveArgs;
 	}
 
+	FArchiveCookData* GetCookData()
+	{
+		return SaveArgs.ArchiveCookData;
+	}
+
 	const ITargetPlatform* GetTargetPlatform() const
 	{
-		return SaveArgs.TargetPlatform;
+		return SaveArgs.GetTargetPlatform();
 	}
 
 	UPackage* GetPackage() const
@@ -490,7 +495,7 @@ public:
 
 	bool IsCooking() const
 	{
-		return SaveArgs.TargetPlatform != nullptr;
+		return SaveArgs.IsCooking();
 	}
 
 	bool IsProceduralSave() const
diff --git a/Engine/Source/Runtime/CoreUObject/Private/UObject/SavePackage2.cpp b/Engine/Source/Runtime/CoreUObject/Private/UObject/SavePackage2.cpp
--- a/Engine/Source/Runtime/CoreUObject/Private/UObject/SavePackage2.cpp
+++ b/Engine/Source/Runtime/CoreUObject/Private/UObject/SavePackage2.cpp
@@ -917,7 +917,7 @@ ESavePackageResult BuildLinker(FSaveContext& SaveContext)
 
 		SaveContext.GetLinker()->SetPortFlags(SaveContext.GetPortFlags());
 		SaveContext.GetLinker()->SetFilterEditorOnly(SaveContext.IsFilterEditorOnly());
-		SaveContext.GetLinker()->SetCookingTarget(SaveContext.GetTargetPlatform());
+		SaveContext.GetLinker()->SetCookData(SaveContext.GetCookData());
 
 		bool bUseUnversionedProperties = SaveContext.IsSaveUnversionedProperties();
 		SaveContext.GetLinker()->SetUseUnversionedPropertySerialization(bUseUnversionedProperties);
@@ -2597,7 +2597,7 @@ ESavePackageResult UPackage::SaveConcurrent(TArrayView<FPackageSaveInfo> InPacka
 	{
 		SCOPED_SAVETIMER(UPackage_SaveConcurrent);
 		// Use concurrent new save only if new save is enabled, otherwise use old save
-		if (SavePackageUtilities::IsNewSaveEnabled(SaveArgs.TargetPlatform != nullptr))
+		if (SavePackageUtilities::IsNewSaveEnabled(SaveArgs.IsCooking()))
 		{
 			// @todo: Once GIsSavingPackage is reworked we should reinstore the saving flag here for the gc lock
 			// Passing in false here so that GIsSavingPackage is set to true on top of locking the GC
diff --git a/Engine/Source/Runtime/CoreUObject/Public/AssetRegistry/CookTagList.h b/Engine/Source/Runtime/CoreUObject/Public/AssetRegistry/CookTagList.h
new file mode 100644
--- /dev/null
+++ b/Engine/Source/Runtime/CoreUObject/Public/AssetRegistry/CookTagList.h
@@ -0,0 +1,63 @@
+// Copyright Epic Games, Inc. All Rights Reserved.
+
+#pragma once
+#include "UObject/Object.h"
+#include "UObject/SoftObjectPath.h"
+
+/**
+*	Used for passing key-value pairs to the asset registry during cook, and is accessed off of
+*	FArchive::CookContext()
+* 
+*	This exists so that the asset registry can have per-platform data, as the general purpose 
+*	GetAssetRegistryTags function can't have a TargetPlatform and more generally can't assume 
+*	that platform data has had a chance to build.
+*
+*	All tags added to this list have "Cook_" prepended to their names prior to being
+*	added to the asset registry.
+*
+*	Cook tags are *only* added to the development asset registry, and will not
+*	show up in the shipped runtime registry.
+*
+*	Cook tags are only generated during Cook By The Book.
+*
+*	Example:
+* 
+*	// Cook context is present for all cooks, but the tag list is only for CBTB.
+*	if (Ar.CookContext() && Ar.CookContext()->CookTagList())
+*	{
+*		FCookTagList* CookTags = Ar.CookContext()->CookTagList();
+*		CookTags->Add(Object, "TagName", FString(TEXT("TagValue")));
+*	}
+*/
+struct FCookTagList
+{
+	UPackage* Package;
+	using FTagNameValuePair = TPair<FName, FString>;
+	TMap<FSoftObjectPath, TArray<FTagNameValuePair>> ObjectToTags;
+
+	FCookTagList(UPackage* InPackage) : Package(InPackage) {}
+
+	/**
+	*	Adds a tag to the list to be added to the development asset registry
+	*	for this object. Note that we only carry over the name/value, not
+	*	any of the other metadata available when adding tags via GetAssetRegistryTags.
+	*/
+	void Add(UObject* InObject, FName InTagName, FString&& InTagValue)
+	{
+		if (!ensure(InObject->GetOutermost() == Package))
+		{
+			return;
+		}
+
+		if (InObject->IsAsset())
+		{
+			TArray<FTagNameValuePair>& ObjectTags = ObjectToTags.FindOrAdd(FSoftObjectPath(InObject));
+			ObjectTags.Add(FTagNameValuePair(InTagName, InTagValue));
+		}
+	}
+
+	void Reset()
+	{
+		ObjectToTags.Reset();
+	}
+};
diff --git a/Engine/Source/Runtime/CoreUObject/Public/UObject/ArchiveCookContext.h b/Engine/Source/Runtime/CoreUObject/Public/UObject/ArchiveCookContext.h
new file mode 100644
--- /dev/null
+++ b/Engine/Source/Runtime/CoreUObject/Public/UObject/ArchiveCookContext.h
@@ -0,0 +1,49 @@
+// Copyright Epic Games, Inc. All Rights Reserved.
+
+#pragma once
+#include "AssetRegistry/CookTagList.h"
+
+
+/**
+*	Accessor for data about the package being cooked during UObject::Serialize calls.
+*/
+struct FArchiveCookContext
+{
+public:
+	enum ECookType
+	{
+		ECookTypeUnknown,
+		ECookOnTheFly,
+		ECookByTheBook
+	};
+
+private:
+
+	// Only valid for cook by the book.
+	FCookTagList CookTagList;
+	bool bCookTagListEnabled = false;
+	ECookType CookType;
+
+public:
+
+
+	FArchiveCookContext(UPackage* InPackage, ECookType InCookType) :
+		CookTagList(InPackage),
+		bCookTagListEnabled(InPackage && InCookType == ECookByTheBook),
+		CookType(InCookType)
+	{
+		
+	}
+
+	void Reset()
+	{
+		CookTagList.Reset();
+	}
+
+	FCookTagList* GetCookTagList() { return bCookTagListEnabled ? &CookTagList : nullptr; }
+
+	bool IsCookByTheBook() const { return CookType == ECookType::ECookByTheBook; }
+	bool IsCookOnTheFly() const { return CookType == ECookType::ECookOnTheFly; }
+	bool IsCookTypeUnknown() const { return CookType == ECookType::ECookTypeUnknown; }
+	ECookType GetCookType() const { return CookType; }
+};
\ No newline at end of file
diff --git a/Engine/Source/Runtime/CoreUObject/Public/UObject/SavePackage.h b/Engine/Source/Runtime/CoreUObject/Public/UObject/SavePackage.h
--- a/Engine/Source/Runtime/CoreUObject/Public/UObject/SavePackage.h
+++ b/Engine/Source/Runtime/CoreUObject/Public/UObject/SavePackage.h
@@ -42,8 +42,15 @@ struct FPackageSaveInfo
  */
 struct FSavePackageArgs
 {
-	/* The platform being saved for when cooking, or nullptr if not cooking. */
+	UE_DEPRECATED(5.1, "TargetPlatform is now on the ArchiveCookData, use GetTargetPlatform() or IsCooking()")
 	const ITargetPlatform* TargetPlatform = nullptr;
+
+	/* nullptr if not cooking, passed to the FArchive */
+	FArchiveCookData* ArchiveCookData = nullptr;
+
+	bool IsCooking() const { return ArchiveCookData != nullptr; }
+	const ITargetPlatform* GetTargetPlatform() const { return ArchiveCookData ? &ArchiveCookData->TargetPlatform : nullptr; }
+	
 	/**
 	 * For all objects which are not referenced[either directly, or indirectly] through the InAsset provided
 	 * to the Save call (See UPackage::Save), only objects that contain any of these flags will be saved.
diff --git a/Engine/Source/Runtime/Engine/Private/Materials/MaterialShared.cpp b/Engine/Source/Runtime/Engine/Private/Materials/MaterialShared.cpp
--- a/Engine/Source/Runtime/Engine/Private/Materials/MaterialShared.cpp
+++ b/Engine/Source/Runtime/Engine/Private/Materials/MaterialShared.cpp
@@ -4933,7 +4933,7 @@ FMaterialResourceMemoryWriter::FMaterialResourceMemoryWriter(FArchive& Ar) :
 {
 	check(Ar.IsSaving());
 	this->SetByteSwapping(Ar.IsByteSwapping());
-	this->SetCookingTarget(Ar.CookingTarget());
+	this->SetCookData(Ar.GetCookData());
 }
 
 FMaterialResourceMemoryWriter::~FMaterialResourceMemoryWriter()
diff --git a/Engine/Source/Runtime/Engine/Private/ShaderCompiler/ShaderCompiler.cpp b/Engine/Source/Runtime/Engine/Private/ShaderCompiler/ShaderCompiler.cpp
--- a/Engine/Source/Runtime/Engine/Private/ShaderCompiler/ShaderCompiler.cpp
+++ b/Engine/Source/Runtime/Engine/Private/ShaderCompiler/ShaderCompiler.cpp
@@ -67,6 +67,7 @@
 
 #if WITH_EDITOR
 #include "Compression/OodleDataCompression.h"
+#include "UObject/ArchiveCookContext.h"
 #include "DerivedDataCache.h"
 #include "DerivedDataCacheRecord.h"
 #include "DerivedDataRequestOwner.h"
@@ -6654,10 +6655,18 @@ FString SaveGlobalShaderFile(EShaderPlatform Platform, FString SavePath, class I
 	TArray<uint8> GlobalShaderData;
 	{
 		FMemoryWriter MemoryWriter(GlobalShaderData, true);
-		if (TargetPlatform)
+
+#if WITH_EDITOR
+		TOptional<FArchiveCookData> CookData;
+		FArchiveCookContext CookContext(nullptr /*InPackage*/, FArchiveCookContext::ECookTypeUnknown);
+		if (TargetPlatform != nullptr)
 		{
-			MemoryWriter.SetCookingTarget(TargetPlatform);
+			CookData.Emplace(*TargetPlatform, CookContext);
 		}
+		
+		MemoryWriter.SetCookData(CookData.GetPtrOrNull());
+#endif
+
 		GlobalShaderMap->SaveToGlobalArchive(MemoryWriter);
 	}
 
@@ -7233,7 +7242,14 @@ void RecompileShadersForRemote(
 					// write the shader compilation info to memory, converting fnames to strings
 					FMemoryWriter MemWriter(*Args.GlobalShaderMap, true);
 					FNameAsStringProxyArchive Ar(MemWriter);
-					Ar.SetCookingTarget(TargetPlatform);
+
+					TOptional<FArchiveCookData> CookData;
+					FArchiveCookContext CookContext(nullptr /*InPackage*/, FArchiveCookContext::ECookTypeUnknown);
+					if (TargetPlatform != nullptr)
+					{
+						CookData.Emplace(*TargetPlatform, CookContext);
+					}
+					Ar.SetCookData(CookData.GetPtrOrNull());
 
 					// save out the global shader map to the byte array
 					SaveGlobalShadersForRemoteRecompile(Ar, ShaderPlatform);
@@ -7249,7 +7265,15 @@ void RecompileShadersForRemote(
 					// write the shader compilation info to memory, converting fnames to strings
 					FMemoryWriter MemWriter(*Args.MeshMaterialMaps, true);
 					FNameAsStringProxyArchive Ar(MemWriter);
-					Ar.SetCookingTarget(TargetPlatform);
+
+					TOptional<FArchiveCookData> CookData;
+					FArchiveCookContext CookContext(nullptr /*InPackage*/, FArchiveCookContext::ECookTypeUnknown);
+					if (TargetPlatform != nullptr)
+					{
+						CookData.Emplace(*TargetPlatform, CookContext);
+					}
+
+					Ar.SetCookData(CookData.GetPtrOrNull());
 
 					// save out the shader map to the byte array
 					FMaterialShaderMap::SaveForRemoteRecompile(Ar, CompiledShaderMaps);
diff --git a/Engine/Source/Runtime/Engine/Private/SkeletalMeshLODRenderData.cpp b/Engine/Source/Runtime/Engine/Private/SkeletalMeshLODRenderData.cpp
--- a/Engine/Source/Runtime/Engine/Private/SkeletalMeshLODRenderData.cpp
+++ b/Engine/Source/Runtime/Engine/Private/SkeletalMeshLODRenderData.cpp
@@ -825,7 +825,7 @@ void FSkeletalMeshLODRenderData::Serialize(FArchive& Ar, UObject* Owner, int32 I
 		if (Ar.IsSaving())
 		{
 			FSkeletalMeshLODSizeCounter LODSizeCounter;
-			LODSizeCounter.SetCookingTarget(Ar.CookingTarget());
+			LODSizeCounter.SetCookData(Ar.GetCookData());
 			LODSizeCounter.SetByteSwapping(Ar.IsByteSwapping());
 			SerializeStreamedData(LODSizeCounter, OwnerMesh, Idx, ClassDataStripFlags, bNeedsCPUAccess, bForceKeepCPUResources);
 			BuffersSize = LODSizeCounter.TotalSize();
@@ -853,7 +853,7 @@ void FSkeletalMeshLODRenderData::Serialize(FArchive& Ar, UObject* Owner, int32 I
 				if (!bDiscardBulkData)
 				{
 					FMemoryWriter MemWriter(TmpBuff, true);
-					MemWriter.SetCookingTarget(Ar.CookingTarget());
+					MemWriter.SetCookData(Ar.GetCookData());
 					MemWriter.SetByteSwapping(Ar.IsByteSwapping());
 					SerializeStreamedData(MemWriter, OwnerMesh, Idx, ClassDataStripFlags, bNeedsCPUAccess, bForceKeepCPUResources);
 				}
diff --git a/Engine/Source/Runtime/Engine/Private/StaticMesh.cpp b/Engine/Source/Runtime/Engine/Private/StaticMesh.cpp
--- a/Engine/Source/Runtime/Engine/Private/StaticMesh.cpp
+++ b/Engine/Source/Runtime/Engine/Private/StaticMesh.cpp
@@ -785,7 +785,7 @@ void FStaticMeshLODResources::Serialize(FArchive& Ar, UObject* Owner, int32 Inde
 				if (!bDiscardBulkData)
 				{
 					FMemoryWriter MemWriter(TmpBuff, true);
-					MemWriter.SetCookingTarget(Ar.CookingTarget());
+					MemWriter.SetCookData(Ar.GetCookData());
 					MemWriter.SetByteSwapping(Ar.IsByteSwapping());
 					SerializeBuffers(MemWriter, OwnerStaticMesh, ClassDataStripFlags, TmpBuffersSize);
 				}
