FScopeCycleCounterUObject trace: suppress per-object timer names (using per-class instead) if AssetLoadTime channel is not enabled. This prevents poor performance due to 2million timers in insights on large cooks. Also change per-class timers to display"Blueprint" for blueprint classes, to prevent many new timers for projects with a large number of blueprints.
#rnx

diff --git a/Engine/Source/Editor/UnrealEd/Private/Cooker/CookTypes.h b/Engine/Source/Editor/UnrealEd/Private/Cooker/CookTypes.h
--- a/Engine/Source/Editor/UnrealEd/Private/Cooker/CookTypes.h
+++ b/Engine/Source/Editor/UnrealEd/Private/Cooker/CookTypes.h
@@ -219,13 +219,13 @@ namespace UE::Cook
 
 inline void RouteBeginCacheForCookedPlatformData(UObject* Obj, const ITargetPlatform* TargetPlatform)
 {
-	UE_SCOPED_TEXT_COOKTIMER(*WriteToString<128>(Obj->GetClass()->GetFName(), TEXT("_BeginCacheForCookedPlatformData")));
+	UE_SCOPED_TEXT_COOKTIMER(*WriteToString<128>(GetClassTraceScope(Obj), TEXT("_BeginCacheForCookedPlatformData")));
 	Obj->BeginCacheForCookedPlatformData(TargetPlatform);
 }
 
 inline bool RouteIsCachedCookedPlatformDataLoaded(UObject* Obj, const ITargetPlatform* TargetPlatform)
 {
-	UE_SCOPED_TEXT_COOKTIMER(*WriteToString<128>(Obj->GetClass()->GetFName(), TEXT("_IsCachedCookedPlatformDataLoaded")));
+	UE_SCOPED_TEXT_COOKTIMER(*WriteToString<128>(GetClassTraceScope(Obj), TEXT("_IsCachedCookedPlatformDataLoaded")));
 	return Obj->IsCachedCookedPlatformDataLoaded(TargetPlatform);
 }
 
diff --git a/Engine/Source/Runtime/CoreUObject/Private/UObject/LinkerLoad.cpp b/Engine/Source/Runtime/CoreUObject/Private/UObject/LinkerLoad.cpp
--- a/Engine/Source/Runtime/CoreUObject/Private/UObject/LinkerLoad.cpp
+++ b/Engine/Source/Runtime/CoreUObject/Private/UObject/LinkerLoad.cpp
@@ -4044,7 +4044,8 @@ void FLinkerLoad::Preload( UObject* Object )
 
 		if (Object->GetLinker() == this)
 		{
-			TRACE_CPUPROFILER_EVENT_SCOPE_TEXT(*FString::Printf(TEXT("FLinkerLoad::Preload %s"), *Object->GetFName().ToString()));
+			TRACE_CPUPROFILER_EVENT_SCOPE(FLinkerLoad::Preload);
+			TRACE_CPUPROFILER_EVENT_SCOPE_TEXT_ON_CHANNEL(*WriteToString<256>(TEXT("FLinkerLoad::Preload "), Object->GetFName()), AssetLoadTimeChannel);
 
 			UClass* Cls = Cast<UClass>(Object);
 			check(!GEventDrivenLoaderEnabled || !bLockoutLegacyOperations || !EVENT_DRIVEN_ASYNC_LOAD_ACTIVE_AT_RUNTIME);
@@ -4217,7 +4218,7 @@ void FLinkerLoad::Preload( UObject* Object )
 #endif
 						check(CurrentLoadContext);
 #if WITH_EDITOR
-						SCOPED_LOADTIMER_TEXT(*WriteToString<128>(Object->GetClass()->GetFName(), TEXTVIEW("_LoadSerialize")));
+						SCOPED_LOADTIMER_TEXT(*WriteToString<128>(GetClassTraceScope(Object), TEXTVIEW("_LoadSerialize")));
 						TRACE_CPUPROFILER_EVENT_SCOPE_TEXT_ON_CHANNEL(*Object->GetFullName(), AssetLoadTimeChannel);
 #endif
 
diff --git a/Engine/Source/Runtime/CoreUObject/Private/UObject/Obj.cpp b/Engine/Source/Runtime/CoreUObject/Private/UObject/Obj.cpp
--- a/Engine/Source/Runtime/CoreUObject/Private/UObject/Obj.cpp
+++ b/Engine/Source/Runtime/CoreUObject/Private/UObject/Obj.cpp
@@ -1095,7 +1095,7 @@ void UObject::ConditionalPostLoad()
 			else
 			{
 #if WITH_EDITOR
-				SCOPED_LOADTIMER_TEXT(*WriteToString<128>(GetClass()->GetFName(), TEXTVIEW("_PostLoad")));
+				SCOPED_LOADTIMER_TEXT(*WriteToString<128>(GetClassTraceScope(this), TEXTVIEW("_PostLoad")));
 #endif
 				LLM_SCOPED_TAG_WITH_OBJECT_IN_SET(GetOutermost(), ELLMTagSet::Assets);
 				LLM_SCOPED_TAG_WITH_OBJECT_IN_SET(GetClass(), ELLMTagSet::AssetClasses);
diff --git a/Engine/Source/Runtime/CoreUObject/Private/UObject/ObjectBaseUtility.cpp b/Engine/Source/Runtime/CoreUObject/Private/UObject/ObjectBaseUtility.cpp
--- a/Engine/Source/Runtime/CoreUObject/Private/UObject/ObjectBaseUtility.cpp
+++ b/Engine/Source/Runtime/CoreUObject/Private/UObject/ObjectBaseUtility.cpp
@@ -2,6 +2,7 @@
 
 #include "CoreMinimal.h"
 #include "Misc/Guid.h"
+#include "ProfilingDebugging/LoadTimeTracker.h"
 #include "Serialization/CustomVersion.h"
 #include "UObject/UObjectBaseUtility.h"
 #include "UObject/Package.h"
@@ -124,6 +125,39 @@ int32 UObjectBaseUtility::GetLinkerLicenseeUEVersion() const
 	}
 }
 
+#if CPUPROFILERTRACE_ENABLED
+COREUOBJECT_API FName GetClassTraceScope(const UObjectBaseUtility* Object)
+{
+	UClass* Class = Object->GetClass();
+	if (Class->IsNative())
+	{
+		return Class->GetFName();
+	}
+	else
+	{
+		static const FName NAME_Blueprint(TEXT("Blueprint"));
+		return NAME_Blueprint;
+	}
+}
+#endif
+
+#if STATS && CPUPROFILERTRACE_ENABLED
+void FScopeCycleCounterUObject::StartObjectTrace(const UObjectBaseUtility* Object)
+{
+#if LOADTIMEPROFILERTRACE_ENABLED
+	if (UE_TRACE_CHANNELEXPR_IS_ENABLED(AssetLoadTimeChannel))
+	{
+		StartTrace(Object->GetFName());
+	}
+	else
+#endif
+	{
+		StartTrace(GetClassTraceScope(Object));
+	}
+}
+#endif
+
+
 // Console variable so that GarbageCollectorSettings work in the editor but we don't want to use it in runtime as we can't support changing its value from console
 int32 GPendingKillEnabled = 1;
 static FAutoConsoleVariableRef CVarPendingKillEnabled(
diff --git a/Engine/Source/Runtime/CoreUObject/Private/UObject/SavePackage/PackageHarvester.cpp b/Engine/Source/Runtime/CoreUObject/Private/UObject/SavePackage/PackageHarvester.cpp
--- a/Engine/Source/Runtime/CoreUObject/Private/UObject/SavePackage/PackageHarvester.cpp
+++ b/Engine/Source/Runtime/CoreUObject/Private/UObject/SavePackage/PackageHarvester.cpp
@@ -249,7 +249,7 @@ void FPackageHarvester::ProcessExport(const FExportWithContext& InProcessContext
 	// In the CDO case the above would serialize most of the references, including transient properties
 	// but we still want to serialize the object using the normal path to collect all custom versions it might be using.
 	{
-		SCOPED_SAVETIMER_TEXT(*WriteToString<128>(Export->GetClass()->GetFName(), TEXT("_SaveSerialize")));
+		SCOPED_SAVETIMER_TEXT(*WriteToString<128>(GetClassTraceScope(Export), TEXT("_SaveSerialize")));
 		Export->Serialize(*this);
 	}
 
diff --git a/Engine/Source/Runtime/CoreUObject/Private/UObject/SavePackage/SavePackageUtilities.cpp b/Engine/Source/Runtime/CoreUObject/Private/UObject/SavePackage/SavePackageUtilities.cpp
--- a/Engine/Source/Runtime/CoreUObject/Private/UObject/SavePackage/SavePackageUtilities.cpp
+++ b/Engine/Source/Runtime/CoreUObject/Private/UObject/SavePackage/SavePackageUtilities.cpp
@@ -852,7 +852,7 @@ bool IsProceduralSave(bool bIsCooking, const FPackagePath& TargetPackagePath, ui
 
 void CallPreSave(UObject* Object, FObjectSaveContextData& ObjectSaveContext)
 {
-	SCOPED_SAVETIMER_TEXT(*WriteToString<256>(Object->GetClass()->GetFName(), TEXTVIEW("_PreSave")));
+	SCOPED_SAVETIMER_TEXT(*WriteToString<256>(GetClassTraceScope(Object), TEXTVIEW("_PreSave")));
 	PRAGMA_DISABLE_DEPRECATION_WARNINGS;
 	Object->PreSave(ObjectSaveContext.TargetPlatform);
 	PRAGMA_ENABLE_DEPRECATION_WARNINGS;
@@ -875,7 +875,7 @@ void CallPreSave(UObject* Object, FObjectSaveContextData& ObjectSaveContext)
 
 void CallPreSaveRoot(UObject* Object, FObjectSaveContextData& ObjectSaveContext)
 {
-	SCOPED_SAVETIMER_TEXT(*WriteToString<256>(Object->GetClass()->GetFName(), TEXTVIEW("_PreSave")));
+	SCOPED_SAVETIMER_TEXT(*WriteToString<256>(GetClassTraceScope(Object), TEXTVIEW("_PreSave")));
 	PRAGMA_DISABLE_DEPRECATION_WARNINGS;
 	bool bLegacyNeedsCleanup = Object->PreSaveRoot(*ObjectSaveContext.TargetFilename);
 	PRAGMA_ENABLE_DEPRECATION_WARNINGS;
@@ -887,7 +887,7 @@ void CallPreSaveRoot(UObject* Object, FObjectSaveContextData& ObjectSaveContext)
 
 void CallPostSaveRoot(UObject* Object, FObjectSaveContextData& ObjectSaveContext, bool bNeedsCleanup)
 {
-	SCOPED_SAVETIMER_TEXT(*WriteToString<256>(Object->GetClass()->GetFName(), TEXTVIEW("_PreSave")));
+	SCOPED_SAVETIMER_TEXT(*WriteToString<256>(GetClassTraceScope(Object), TEXTVIEW("_PreSave")));
 	PRAGMA_DISABLE_DEPRECATION_WARNINGS;
 	Object->PostSaveRoot(bNeedsCleanup);
 	PRAGMA_ENABLE_DEPRECATION_WARNINGS;
diff --git a/Engine/Source/Runtime/CoreUObject/Private/UObject/SavePackage2.cpp b/Engine/Source/Runtime/CoreUObject/Private/UObject/SavePackage2.cpp
--- a/Engine/Source/Runtime/CoreUObject/Private/UObject/SavePackage2.cpp
+++ b/Engine/Source/Runtime/CoreUObject/Private/UObject/SavePackage2.cpp
@@ -1720,7 +1720,7 @@ ESavePackageResult WriteExports(FStructuredArchive::FRecord& StructuredArchiveRo
 			bool bSupportsText = false;
 #endif
 
-			SCOPED_SAVETIMER_TEXT(*WriteToString<128>(Export.Object->GetClass()->GetFName(), TEXT("_SaveSerialize")));
+			SCOPED_SAVETIMER_TEXT(*WriteToString<128>(GetClassTraceScope(Export.Object), TEXT("_SaveSerialize")));
 			if (Export.Object->HasAnyFlags(RF_ClassDefaultObject))
 			{
 				if (bSupportsText)
diff --git a/Engine/Source/Runtime/CoreUObject/Public/UObject/UObjectBaseUtility.h b/Engine/Source/Runtime/CoreUObject/Public/UObject/UObjectBaseUtility.h
--- a/Engine/Source/Runtime/CoreUObject/Public/UObject/UObjectBaseUtility.h
+++ b/Engine/Source/Runtime/CoreUObject/Public/UObject/UObjectBaseUtility.h
@@ -934,6 +934,10 @@ COREUOBJECT_API UClass* GetParentNativeClass(UClass* Class);
 #define USE_LIGHTWEIGHT_UOBJECT_STATS_FOR_HITCH_DETECTION (1)
 #endif
 
+#if CPUPROFILERTRACE_ENABLED
+COREUOBJECT_API FName GetClassTraceScope(const UObjectBaseUtility* Object);
+#endif
+
 #if STATS
 
 /** Structure used to track time spent by a UObject */
@@ -967,7 +971,7 @@ public:
 #if CPUPROFILERTRACE_ENABLED
 			if (!bStarted && UE_TRACE_CHANNELEXPR_IS_ENABLED(CpuChannel))
 			{
-				StartTrace(Object->GetFName());
+				StartObjectTrace(Object);
 			}
 #endif
 		}
@@ -1000,7 +1004,7 @@ public:
 #if CPUPROFILERTRACE_ENABLED
 			if (!bStarted && UE_TRACE_CHANNELEXPR_IS_ENABLED(CpuChannel))
 			{
-				StartTrace(Object->GetFName());
+				StartObjectTrace(Object);
 			}
 #endif
 		}
@@ -1013,6 +1017,10 @@ public:
 #endif
 	}
 
+#if CPUPROFILERTRACE_ENABLED
+	COREUOBJECT_API void StartObjectTrace(const UObjectBaseUtility* Object);
+#endif
+
 	/**
 	 * Updates the stat with the time spent
 	 */
