Corrected comments on UFUNCTION keywords Reliable and Unreliable to indicated they are valid in conjunction with NetMulticast, not just Client and Server.

#UE5 #RNX #NoReleaseNotes

diff --git a/Engine/Source/Runtime/CoreUObject/Public/UObject/ObjectMacros.h b/Engine/Source/Runtime/CoreUObject/Public/UObject/ObjectMacros.h
--- a/Engine/Source/Runtime/CoreUObject/Public/UObject/ObjectMacros.h
+++ b/Engine/Source/Runtime/CoreUObject/Public/UObject/ObjectMacros.h
@@ -887,11 +887,11 @@ namespace UF
 		NetMulticast,
 
 		/// Replication of calls to this function should be done on a reliable channel.
-		/// Only valid when used in conjunction with Client or Server
+		/// Only valid when used in conjunction with Client, Server, or NetMulticast
 		Reliable,
 
 		/// Replication of calls to this function can be done on an unreliable channel.
-		/// Only valid when used in conjunction with Client or Server
+		/// Only valid when used in conjunction with Client, Server, or NetMulticast
 		Unreliable,
 
 		/// This function fulfills a contract of producing no side effects, and additionally implies BlueprintCallable.
