Fix assignment between Object and ObjectPtr from property utilities methods used by blueprints.

diff --git a/Engine/Source/Runtime/CoreUObject/Private/UObject/PropertyAccessUtil.cpp b/Engine/Source/Runtime/CoreUObject/Private/UObject/PropertyAccessUtil.cpp
--- a/Engine/Source/Runtime/CoreUObject/Private/UObject/PropertyAccessUtil.cpp
+++ b/Engine/Source/Runtime/CoreUObject/Private/UObject/PropertyAccessUtil.cpp
@@ -71,8 +71,18 @@ bool ArePropertiesCompatible(const FProperty* InSrcProp, const FProperty* InDest
 	}
 
 	// Compare the classes as these must be an *exact* match as the access is low-level and without property coercion
+	// HOWEVER:
+	// Object properties can either be an ObjectProperty, or an ObjectPtrProperty
+	// We allow coercion between these two types
 	if (InSrcProp->GetClass() != InDestProp->GetClass())
 	{
+		const bool bSrcIsObject = CastField<FObjectProperty>(InSrcProp) || CastField<FObjectPtrProperty>(InSrcProp);
+		const bool bDestIsObject = CastField<FObjectProperty>(InDestProp) || CastField<FObjectPtrProperty>(InDestProp);
+		if (bSrcIsObject && bDestIsObject)
+		{
+			return true;
+		}
+
 		return false;
 	}
 
@@ -153,6 +163,18 @@ bool CopySinglePropertyValue(const FProperty* InSrcProp, const void* InSrcValue,
 		DestBoolProp->SetPropertyValue(InDestValue, bBoolValue);
 		return true;
 	}
+
+	// If the source is an ObjectPtr and the destination is an Object, handle the copy in a custom way to ensure we always resolve
+	// the source before copying it to the destination.
+	if (const FObjectPtrProperty* SrcObjectPtrProp = CastField<FObjectPtrProperty>(InSrcProp))
+	{
+		if (const FObjectProperty* DestObjectProp = CastField<FObjectProperty>(InDestProp))
+		{
+			UObject* SrcObject = SrcObjectPtrProp->GetObjectPropertyValue(InSrcValue);
+			DestObjectProp->SetObjectPropertyValue(InDestValue, SrcObject);
+			return true;
+		}
+	}
 	
 	InSrcProp->CopySingleValue(InDestValue, InSrcValue);
 	return true;
@@ -195,6 +217,24 @@ bool CopyCompletePropertyValue(const FProperty* InSrcProp, const void* InSrcValu
 		}
 		return true;
 	}
+
+	// If the source is an ObjectPtr and the destination is an Object, handle the copy in a custom way to ensure we always resolve
+	// the source before copying it to the destination.
+	if (const FObjectPtrProperty* SrcObjectPtrProp = CastField<FObjectPtrProperty>(InSrcProp))
+	{
+		if (const FObjectProperty* DestObjectProp = CastField<FObjectProperty>(InDestProp))
+		{
+			for (int32 Idx = 0; Idx < InSrcProp->ArrayDim; ++Idx)
+			{
+				const void* SrcElemValue = static_cast<const uint8*>(InSrcValue) + (InSrcProp->ElementSize * Idx);
+				void* DestElemValue = static_cast<uint8*>(InDestValue) + (InDestProp->ElementSize * Idx);
+
+				UObject* SrcObject = SrcObjectPtrProp->GetObjectPropertyValue(SrcElemValue);
+				DestObjectProp->SetObjectPropertyValue(DestElemValue, SrcObject);
+			}
+			return true;
+		}
+	}
 	
 	InSrcProp->CopyCompleteValue(InDestValue, InSrcValue);
 	return true;
