Added support for native FProperty setters and getters.

diff --git a/Engine/Source/Runtime/CoreUObject/Private/UObject/PropertyMap.cpp b/Engine/Source/Runtime/CoreUObject/Private/UObject/PropertyMap.cpp
--- a/Engine/Source/Runtime/CoreUObject/Private/UObject/PropertyMap.cpp
+++ b/Engine/Source/Runtime/CoreUObject/Private/UObject/PropertyMap.cpp
@@ -640,7 +640,7 @@ FString FMapProperty::GetCPPMacroType( FString& ExtendedTypeText ) const
 	return TEXT("TMAP");
 }
 
-void FMapProperty::ExportTextItem(FString& ValueStr, const void* PropertyValue, const void* DefaultValue, UObject* Parent, int32 PortFlags, UObject* ExportRootScope) const
+void FMapProperty::ExportText_Internal(FString& ValueStr, const void* ContainerOrPropertyPtr, EPropertyPointerType PropertyPointerType, const void* DefaultValue, UObject* Parent, int32 PortFlags, UObject* ExportRootScope) const
 {
 	if (0 != (PortFlags & PPF_ExportCpp))
 	{
@@ -651,7 +651,34 @@ void FMapProperty::ExportTextItem(FString& ValueStr, const void* PropertyValue,
 	checkSlow(KeyProp);
 	checkSlow(ValueProp);
 
-	FScriptMapHelper MapHelper(this, PropertyValue);
+	uint8* TempMapStorage = nullptr;
+	void* PropertyValuePtr = nullptr;
+	if (PropertyPointerType == EPropertyPointerType::Container && HasGetter())
+	{
+		// Allocate temporary map as we first need to initialize it with the value provided by the getter function and then export it
+		TempMapStorage = (uint8*)FMemory::MallocZeroed(GetSize());
+		if (!HasAnyPropertyFlags(CPF_ZeroConstructor)) // this stuff is already zero
+		{
+			InitializeValue(TempMapStorage);
+		}
+		PropertyValuePtr = TempMapStorage;
+		FProperty::GetValue_InContainer(ContainerOrPropertyPtr, PropertyValuePtr);
+	}
+	else
+	{
+		PropertyValuePtr = PointerToValuePtr(ContainerOrPropertyPtr, PropertyPointerType);
+	}
+
+	ON_SCOPE_EXIT
+	{
+		if (TempMapStorage)
+		{
+			DestroyValue(TempMapStorage);
+			FMemory::Free(TempMapStorage);
+		}
+	};
+
+	FScriptMapHelper MapHelper(this, PropertyValuePtr);
 
 	if (MapHelper.Num() == 0)
 	{
@@ -703,7 +730,7 @@ void FMapProperty::ExportTextItem(FString& ValueStr, const void* PropertyValue,
 				}
 
 				ValueStr += TEXT("[");
-				KeyProp->ExportTextItem(ValueStr, PropData, nullptr, Parent, PortFlags | PPF_Delimited, ExportRootScope);
+				KeyProp->ExportText_Internal(ValueStr, PropData, EPropertyPointerType::Direct, nullptr, Parent, PortFlags | PPF_Delimited, ExportRootScope);
 				ValueStr += TEXT("] ");
 
 				// Always use struct defaults if the inner is a struct, for symmetry with the import of array inner struct defaults
@@ -715,7 +742,7 @@ void FMapProperty::ExportTextItem(FString& ValueStr, const void* PropertyValue,
 					PropDefault = PropData;
 				}
 
-				ValueProp->ExportTextItem(ValueStr, PropData + MapLayout.ValueOffset, PropDefault + MapLayout.ValueOffset, Parent, PortFlags | PPF_Delimited, ExportRootScope);
+				ValueProp->ExportText_Internal(ValueStr, PropData + MapLayout.ValueOffset, EPropertyPointerType::Direct, PropDefault + MapLayout.ValueOffset, Parent, PortFlags | PPF_Delimited, ExportRootScope);
 
 				--Count;
 			}
@@ -741,7 +768,7 @@ void FMapProperty::ExportTextItem(FString& ValueStr, const void* PropertyValue,
 
 				ValueStr += TEXT("(");
 
-				KeyProp->ExportTextItem(ValueStr, PropData, nullptr, Parent, PortFlags | PPF_Delimited, ExportRootScope);
+				KeyProp->ExportText_Internal(ValueStr, PropData, EPropertyPointerType::Direct, nullptr, Parent, PortFlags | PPF_Delimited, ExportRootScope);
 
 				ValueStr += TEXT(", ");
 
@@ -754,7 +781,7 @@ void FMapProperty::ExportTextItem(FString& ValueStr, const void* PropertyValue,
 					PropDefault = PropData;
 				}
 
-				ValueProp->ExportTextItem(ValueStr, PropData + MapLayout.ValueOffset, PropDefault + MapLayout.ValueOffset, Parent, PortFlags | PPF_Delimited, ExportRootScope);
+				ValueProp->ExportText_Internal(ValueStr, PropData + MapLayout.ValueOffset, EPropertyPointerType::Direct, PropDefault + MapLayout.ValueOffset, Parent, PortFlags | PPF_Delimited, ExportRootScope);
 
 				ValueStr += TEXT(")");
 
@@ -766,12 +793,39 @@ void FMapProperty::ExportTextItem(FString& ValueStr, const void* PropertyValue,
 	}
 }
 
-const TCHAR* FMapProperty::ImportText_Internal(const TCHAR* Buffer, void* Data, int32 PortFlags, UObject* Parent, FOutputDevice* ErrorText) const
+const TCHAR* FMapProperty::ImportText_Internal(const TCHAR* Buffer, void* ContainerOrPropertyPtr, EPropertyPointerType PropertyPointerType, UObject* Parent, int32 PortFlags, FOutputDevice* ErrorText) const
 {
 	checkSlow(KeyProp);
 	checkSlow(ValueProp);
 
-	FScriptMapHelper MapHelper(this, Data);
+	FScriptMapHelper MapHelper(this, PointerToValuePtr(ContainerOrPropertyPtr, PropertyPointerType));
+	uint8* TempMapStorage = nullptr;
+
+	ON_SCOPE_EXIT
+	{
+		if (TempMapStorage)
+		{
+			// TempMap is used by property setter so if it was allocated call the setter now
+			FProperty::SetValue_InContainer(ContainerOrPropertyPtr, TempMapStorage);
+
+			// Destroy and free the temp map used by property setter
+			DestroyValue(TempMapStorage);
+			FMemory::Free(TempMapStorage);
+		}
+	};
+
+	if (PropertyPointerType == EPropertyPointerType::Container && HasSetter())
+	{
+		// Allocate temporary map as we first need to initialize it with the parsed items and then use the setter to update the property
+		TempMapStorage = (uint8*)FMemory::MallocZeroed(GetSize());
+		if (!HasAnyPropertyFlags(CPF_ZeroConstructor)) // this stuff is already zero
+		{
+			InitializeValue(TempMapStorage);
+		}
+		// Reinitialize the map helper with the temp value
+		MapHelper = FScriptMapHelper(this, TempMapStorage);
+	}
+
 	MapHelper.EmptyValues();
 
 	// If we export an empty array we export an empty string, so ensure that if we're passed an empty string
@@ -819,7 +873,7 @@ const TCHAR* FMapProperty::ImportText_Internal(const TCHAR* Buffer, void* Data,
 
 		// Parse the key
 		SkipWhitespace(Buffer);
-		Buffer = KeyProp->ImportText(Buffer, TempPairStorage, PortFlags | PPF_Delimited, Parent, ErrorText);
+		Buffer = KeyProp->ImportText_Internal(Buffer, TempPairStorage, EPropertyPointerType::Direct, Parent, PortFlags | PPF_Delimited, ErrorText);
 		if (!Buffer)
 		{
 			return nullptr;
@@ -836,7 +890,7 @@ const TCHAR* FMapProperty::ImportText_Internal(const TCHAR* Buffer, void* Data,
 
 		// Parse the value
 		SkipWhitespace(Buffer);
-		Buffer = ValueProp->ImportText(Buffer, TempPairStorage + MapLayout.ValueOffset, PortFlags | PPF_Delimited, Parent, ErrorText);
+		Buffer = ValueProp->ImportText_Internal(Buffer, TempPairStorage + MapLayout.ValueOffset, EPropertyPointerType::Direct, Parent, PortFlags | PPF_Delimited, ErrorText);
 		if (!Buffer)
 		{
 			return nullptr;
diff --git a/Engine/Source/Runtime/CoreUObject/Private/UObject/PropertyMulticastDelegate.cpp b/Engine/Source/Runtime/CoreUObject/Private/UObject/PropertyMulticastDelegate.cpp
--- a/Engine/Source/Runtime/CoreUObject/Private/UObject/PropertyMulticastDelegate.cpp
+++ b/Engine/Source/Runtime/CoreUObject/Private/UObject/PropertyMulticastDelegate.cpp
@@ -184,7 +184,7 @@ FString FMulticastDelegateProperty::GetCPPTypeForwardDeclaration() const
 }
 
 
-void FMulticastDelegateProperty::ExportTextItem( FString& ValueStr, const void* PropertyValue, const void* DefaultValue, UObject* Parent, int32 PortFlags, UObject* ExportRootScope ) const
+void FMulticastDelegateProperty::ExportText_Internal( FString& ValueStr, const void* PropertyValueOrContainer, EPropertyPointerType PropertyPointerType, const void* DefaultValue, UObject* Parent, int32 PortFlags, UObject* ExportRootScope ) const
 {
 	if (0 != (PortFlags & PPF_ExportCpp))
 	{
@@ -192,13 +192,24 @@ void FMulticastDelegateProperty::ExportTextItem( FString& ValueStr, const void*
 		return;
 	}
 
-	const FMulticastScriptDelegate::FInvocationList& InvocationList = GetInvocationList(PropertyValue);
+	const FMulticastScriptDelegate::FInvocationList* InvocationList = nullptr;
+	
+	if (PropertyPointerType == EPropertyPointerType::Container && HasGetter())
+	{
+		FMulticastScriptDelegate Delegate;
+		GetValue_InContainer(PropertyValueOrContainer, &Delegate);
+		InvocationList = &GetInvocationList(&Delegate);
+	}
+	else
+	{
+		InvocationList = &GetInvocationList(PointerToValuePtr(PropertyValueOrContainer, PropertyPointerType));
+	}
 
 	// Start delegate array with open paren
 	ValueStr += TEXT( "(" );
 
 	bool bIsFirstFunction = true;
-	for (FMulticastScriptDelegate::FInvocationList::TConstIterator CurInvocation(InvocationList); CurInvocation; ++CurInvocation)
+	for (FMulticastScriptDelegate::FInvocationList::TConstIterator CurInvocation(*InvocationList); CurInvocation; ++CurInvocation)
 	{
 		if (CurInvocation->IsBound())
 		{
@@ -406,10 +417,24 @@ void FMulticastInlineDelegateProperty::SerializeItem(FStructuredArchive::FSlot S
 	Ar << *GetPropertyValuePtr(Value);
 }
 
-const TCHAR* FMulticastInlineDelegateProperty::ImportText_Internal(const TCHAR* Buffer, void* PropertyValue, int32 PortFlags, UObject* Parent, FOutputDevice* ErrorText) const
+const TCHAR* FMulticastInlineDelegateProperty::ImportText_Internal(const TCHAR* Buffer, void* ContainerOrPropertyPtr, EPropertyPointerType PropertyPointerType, UObject* Parent, int32 PortFlags, FOutputDevice* ErrorText) const
 {
-	FMulticastScriptDelegate& MulticastDelegate = (*(FMulticastScriptDelegate*)PropertyValue);
-	return ImportDelegateFromText(MulticastDelegate, Buffer, Parent, ErrorText);
+	const TCHAR* Result = nullptr;
+	if (PropertyPointerType == EPropertyPointerType::Container && HasSetter())
+	{
+		FMulticastScriptDelegate MulticastDelegate;		
+		Result = ImportDelegateFromText(MulticastDelegate, Buffer, Parent, ErrorText);
+		if (Result)
+		{
+			SetValue_InContainer(ContainerOrPropertyPtr, MulticastDelegate);
+		}
+	}
+	else
+	{
+		FMulticastScriptDelegate& MulticastDelegate = *(FMulticastScriptDelegate*)PointerToValuePtr(ContainerOrPropertyPtr, PropertyPointerType);
+		Result = ImportDelegateFromText(MulticastDelegate, Buffer, Parent, ErrorText);
+	}
+	return Result;	
 }
 
 void ResolveDelegateReference(const FMulticastInlineDelegateProperty* InlineProperty, UObject*& Parent, void*& PropertyValue)
@@ -559,25 +584,32 @@ void FMulticastSparseDelegateProperty::SerializeItemInternal(FArchive& Ar, void*
 	}
 }
 
-const TCHAR* FMulticastSparseDelegateProperty::ImportText_Internal(const TCHAR* Buffer, void* PropertyValue, int32 PortFlags, UObject* Parent, FOutputDevice* ErrorText) const
+const TCHAR* FMulticastSparseDelegateProperty::ImportText_Internal(const TCHAR* Buffer, void* ContainerOrPropertyPtr, EPropertyPointerType PropertyPointerType, UObject* Parent, int32 PortFlags, FOutputDevice* ErrorText) const
 {
 	FMulticastScriptDelegate Delegate;
 	const TCHAR* Result = ImportDelegateFromText(Delegate, Buffer, Parent, ErrorText);
 
 	if (Result)
 	{
-		FSparseDelegate& SparseDelegate = *(FSparseDelegate*)PropertyValue;
-		USparseDelegateFunction* SparseDelegateFunc = CastChecked<USparseDelegateFunction>(SignatureFunction);
-
-		if (Delegate.IsBound())
+		if (PropertyPointerType == EPropertyPointerType::Container && HasSetter())
 		{
-			FSparseDelegateStorage::SetMulticastDelegate(Parent, SparseDelegateFunc->DelegateName, MoveTemp(Delegate));
-			SparseDelegate.bIsBound = true;
+			FProperty::SetValue_InContainer(ContainerOrPropertyPtr, &Delegate);
 		}
 		else
 		{
-			FSparseDelegateStorage::Clear(Parent, SparseDelegateFunc->DelegateName);
-			SparseDelegate.bIsBound = false;
+			FSparseDelegate& SparseDelegate = *(FSparseDelegate*)PointerToValuePtr(ContainerOrPropertyPtr, PropertyPointerType);
+			USparseDelegateFunction* SparseDelegateFunc = CastChecked<USparseDelegateFunction>(SignatureFunction);
+
+			if (Delegate.IsBound())
+			{
+				FSparseDelegateStorage::SetMulticastDelegate(Parent, SparseDelegateFunc->DelegateName, MoveTemp(Delegate));
+				SparseDelegate.bIsBound = true;
+			}
+			else
+			{
+				FSparseDelegateStorage::Clear(Parent, SparseDelegateFunc->DelegateName);
+				SparseDelegate.bIsBound = false;
+			}
 		}
 	}
 
diff --git a/Engine/Source/Runtime/CoreUObject/Private/UObject/PropertyName.cpp b/Engine/Source/Runtime/CoreUObject/Private/UObject/PropertyName.cpp
--- a/Engine/Source/Runtime/CoreUObject/Private/UObject/PropertyName.cpp
+++ b/Engine/Source/Runtime/CoreUObject/Private/UObject/PropertyName.cpp
@@ -16,9 +16,17 @@ FNameProperty::FNameProperty(FFieldVariant InOwner, const UECodeGen_Private::FNa
 {
 }
 
-void FNameProperty::ExportTextItem( FString& ValueStr, const void* PropertyValue, const void* DefaultValue, UObject* Parent, int32 PortFlags, UObject* ExportRootScope ) const
+void FNameProperty::ExportText_Internal( FString& ValueStr, const void* PropertyValueOrContainer, EPropertyPointerType PropertyPointerType, const void* DefaultValue, UObject* Parent, int32 PortFlags, UObject* ExportRootScope ) const
 {
-	FName Temp = *(FName*)PropertyValue;
+	FName Temp;
+	if (PropertyPointerType == EPropertyPointerType::Container && HasGetter())
+	{
+		GetValue_InContainer(PropertyValueOrContainer, &Temp);
+	}
+	else
+	{
+		Temp = *(FName*)PointerToValuePtr(PropertyValueOrContainer, PropertyPointerType);
+	}
 	if (0 != (PortFlags & PPF_ExportCpp))
 	{
 		ValueStr += (Temp == NAME_None) 
@@ -38,11 +46,12 @@ void FNameProperty::ExportTextItem( FString& ValueStr, const void* PropertyValue
 		ValueStr += TEXT("\"\"");
 	}
 }
-const TCHAR* FNameProperty::ImportText_Internal( const TCHAR* Buffer, void* Data, int32 PortFlags, UObject* Parent, FOutputDevice* ErrorText ) const
+const TCHAR* FNameProperty::ImportText_Internal( const TCHAR* Buffer, void* ContainerOrPropertyPtr, EPropertyPointerType PropertyPointerType, UObject* Parent, int32 PortFlags, FOutputDevice* ErrorText ) const
 {
+	FName ImportedName;
 	if (!(PortFlags & PPF_Delimited))
 	{
-		*(FName*)Data = FName(Buffer);
+		ImportedName = FName(Buffer);		
 
 		// in order to indicate that the value was successfully imported, advance the buffer past the last character that was imported
 		Buffer += FCString::Strlen(Buffer);
@@ -52,10 +61,21 @@ const TCHAR* FNameProperty::ImportText_Internal( const TCHAR* Buffer, void* Data
 		TStringBuilder<256> Token;
 		Buffer = FPropertyHelpers::ReadToken(Buffer, /* out */ Token, true);
 		if (!Buffer)
-			return NULL;
+		{
+			return nullptr;
+		}
 
-		*(FName*)Data = FName(Token);
+		ImportedName = FName(Token);
 	}
+	if (PropertyPointerType == EPropertyPointerType::Container && HasSetter())
+	{
+		SetValue_InContainer(ContainerOrPropertyPtr, ImportedName);
+	}
+	else
+	{
+		*(FName*)PointerToValuePtr(ContainerOrPropertyPtr, PropertyPointerType) = ImportedName;
+	}
+
 	return Buffer;
 }
 
diff --git a/Engine/Source/Runtime/CoreUObject/Private/UObject/PropertyNumeric.cpp b/Engine/Source/Runtime/CoreUObject/Private/UObject/PropertyNumeric.cpp
--- a/Engine/Source/Runtime/CoreUObject/Private/UObject/PropertyNumeric.cpp
+++ b/Engine/Source/Runtime/CoreUObject/Private/UObject/PropertyNumeric.cpp
@@ -64,7 +64,7 @@ int64 FNumericProperty::ReadEnumAsInt64(FStructuredArchive::FSlot Slot, UStruct*
 	return Result;
 };
 
-const TCHAR* FNumericProperty::ImportText_Internal( const TCHAR* Buffer, void* Data, int32 PortFlags, UObject* Parent, FOutputDevice* ErrorText ) const
+const TCHAR* FNumericProperty::ImportText_Internal( const TCHAR* Buffer, void* ContainerOrPropertyPtr, EPropertyPointerType PropertyPointerType, UObject* Parent, int32 PortFlags, FOutputDevice* ErrorText ) const
 {
 	if ( Buffer != NULL )
 	{
@@ -76,7 +76,14 @@ const TCHAR* FNumericProperty::ImportText_Internal( const TCHAR* Buffer, void* D
 				int64 EnumValue = UEnum::ParseEnum(Buffer);
 				if (EnumValue != INDEX_NONE)
 				{
-					SetIntPropertyValue(Data, EnumValue);
+					if (PropertyPointerType == EPropertyPointerType::Container && HasSetter())
+					{
+						SetValue_InContainer(ContainerOrPropertyPtr, &EnumValue);
+					}
+					else
+					{
+						SetIntPropertyValue(PointerToValuePtr(ContainerOrPropertyPtr, PropertyPointerType), EnumValue);
+					}
 					return Buffer;
 				}
 				else
@@ -127,14 +134,28 @@ const TCHAR* FNumericProperty::ImportText_Internal( const TCHAR* Buffer, void* D
 				Buffer++;
 			}
 		}
-		SetNumericPropertyValueFromString(Data, Start);
+		if (PropertyPointerType == EPropertyPointerType::Container && HasSetter())
+		{
+			SetNumericPropertyValueFromString_InContainer(ContainerOrPropertyPtr, Start);
+		}
+		else
+		{
+			SetNumericPropertyValueFromString(PointerToValuePtr(ContainerOrPropertyPtr, PropertyPointerType), Start);
+		}
 	}
 	return Buffer;
 }
 
-void FNumericProperty::ExportTextItem( FString& ValueStr, const void* PropertyValue, const void* DefaultValue, UObject* Parent, int32 PortFlags, UObject* ExportRootScope ) const
+void FNumericProperty::ExportText_Internal( FString& ValueStr, const void* PropertyValueOrContainer, EPropertyPointerType PropertyPointerType, const void* DefaultValue, UObject* Parent, int32 PortFlags, UObject* ExportRootScope ) const
 {
-	ValueStr += GetNumericPropertyValueToString(PropertyValue);
+	if (PropertyPointerType == EPropertyPointerType::Container && HasGetter())
+	{
+		ValueStr += GetNumericPropertyValueToString_InContainer(PropertyValueOrContainer);
+	}
+	else
+	{
+		ValueStr += GetNumericPropertyValueToString(PointerToValuePtr(PropertyValueOrContainer, PropertyPointerType));
+	}
 }
 
 bool FNumericProperty::IsFloatingPoint() const
@@ -192,7 +213,10 @@ void FNumericProperty::SetNumericPropertyValueFromString(void* Data, TCHAR const
 {
 	check(0);
 }
-
+void FNumericProperty::SetNumericPropertyValueFromString_InContainer(void* Container, TCHAR const* Value) const
+{
+	check(0);
+}
 /** 
 	* Gets the value of a signed integral property type
 	* @param Data - pointer to property data to get
@@ -203,6 +227,11 @@ int64 FNumericProperty::GetSignedIntPropertyValue(void const* Data) const
 	check(0);
 	return 0;
 }
+int64 FNumericProperty::GetSignedIntPropertyValue_InContainer(void const* Container) const
+{
+	check(0);
+	return 0;
+}
 
 /** 
 	* Gets the value of an unsigned integral property type
@@ -238,6 +267,13 @@ FString FNumericProperty::GetNumericPropertyValueToString(void const* Data) cons
 	return FString();
 }
 
+FString FNumericProperty::GetNumericPropertyValueToString_InContainer(void const* Container) const
+{
+	check(0);
+	return FString();
+}
+
+
 FInt8Property::FInt8Property(FFieldVariant InOwner, const UECodeGen_Private::FInt8PropertyParams& Prop)
 	: TProperty_Numeric(InOwner, (const UECodeGen_Private::FPropertyParamsBaseWithOffset&)Prop)
 {
diff --git a/Engine/Source/Runtime/CoreUObject/Private/UObject/PropertyObject.cpp b/Engine/Source/Runtime/CoreUObject/Private/UObject/PropertyObject.cpp
--- a/Engine/Source/Runtime/CoreUObject/Private/UObject/PropertyObject.cpp
+++ b/Engine/Source/Runtime/CoreUObject/Private/UObject/PropertyObject.cpp
@@ -142,11 +142,12 @@ void FObjectProperty::SerializeItem( FStructuredArchive::FSlot Slot, void* Value
 	}
 }
 
-const TCHAR* FObjectProperty::ImportText_Internal(const TCHAR* Buffer, void* Data, int32 PortFlags, UObject* OwnerObject, FOutputDevice* ErrorText) const
+const TCHAR* FObjectProperty::ImportText_Internal(const TCHAR* Buffer, void* ContainerOrPropertyPtr, EPropertyPointerType PropertyPointerType, UObject* OwnerObject, int32 PortFlags, FOutputDevice* ErrorText) const
 {
-	const TCHAR* Result = TFObjectPropertyBase<UObject*>::ImportText_Internal(Buffer, Data, PortFlags, OwnerObject, ErrorText);
+	const TCHAR* Result = TFObjectPropertyBase<UObject*>::ImportText_Internal(Buffer, ContainerOrPropertyPtr, PropertyPointerType, OwnerObject, PortFlags, ErrorText);
 	if (Result)
 	{
+		void* Data = PointerToValuePtr(ContainerOrPropertyPtr, PropertyPointerType);
 		CheckValidObject(Data);
 
 #if USE_CIRCULAR_DEPENDENCY_LOAD_DEFERRING
diff --git a/Engine/Source/Runtime/CoreUObject/Private/UObject/PropertySet.cpp b/Engine/Source/Runtime/CoreUObject/Private/UObject/PropertySet.cpp
--- a/Engine/Source/Runtime/CoreUObject/Private/UObject/PropertySet.cpp
+++ b/Engine/Source/Runtime/CoreUObject/Private/UObject/PropertySet.cpp
@@ -553,7 +553,7 @@ FString FSetProperty::GetCPPTypeForwardDeclaration() const
 	return ElementProp->GetCPPTypeForwardDeclaration();
 }
 
-void FSetProperty::ExportTextItem(FString& ValueStr, const void* PropertyValue, const void* DefaultValue, UObject* Parent, int32 PortFlags, UObject* ExportRootScope) const
+void FSetProperty::ExportText_Internal(FString& ValueStr, const void* ContainerOrPropertyPtr, EPropertyPointerType PropertyPointerType, const void* DefaultValue, UObject* Parent, int32 PortFlags, UObject* ExportRootScope) const
 {
 	if (0 != (PortFlags & PPF_ExportCpp))
 	{
@@ -563,7 +563,34 @@ void FSetProperty::ExportTextItem(FString& ValueStr, const void* PropertyValue,
 
 	checkSlow(ElementProp);
 
-	FScriptSetHelper SetHelper(this, PropertyValue);
+	uint8* TempSetStorage = nullptr;
+	void* PropertyValuePtr = nullptr;
+	if (PropertyPointerType == EPropertyPointerType::Container && HasGetter())
+	{
+		// Allocate temporary map as we first need to initialize it with the value provided by the getter function and then export it
+		TempSetStorage = (uint8*)FMemory::MallocZeroed(GetSize());
+		if (!HasAnyPropertyFlags(CPF_ZeroConstructor)) // this stuff is already zero
+		{
+			InitializeValue(TempSetStorage);
+		}
+		PropertyValuePtr = TempSetStorage;
+		FProperty::GetValue_InContainer(ContainerOrPropertyPtr, PropertyValuePtr);
+	}
+	else
+	{
+		PropertyValuePtr = PointerToValuePtr(ContainerOrPropertyPtr, PropertyPointerType);
+	}
+
+	ON_SCOPE_EXIT
+	{
+		if (TempSetStorage)
+		{
+			DestroyValue(TempSetStorage);
+			FMemory::Free(TempSetStorage);
+		}
+	};
+
+	FScriptSetHelper SetHelper(this, PropertyValuePtr);
 
 	if (SetHelper.Num() == 0)
 	{
@@ -624,7 +651,7 @@ void FSetProperty::ExportTextItem(FString& ValueStr, const void* PropertyValue,
 					PropDefault = PropData;
 				}
 
-				ElementProp->ExportTextItem(ValueStr, PropData, PropDefault, Parent, PortFlags | PPF_Delimited, ExportRootScope);
+				ElementProp->ExportTextItem_Direct(ValueStr, PropData, PropDefault, Parent, PortFlags | PPF_Delimited, ExportRootScope);
 
 				--Count;
 			}
@@ -656,7 +683,7 @@ void FSetProperty::ExportTextItem(FString& ValueStr, const void* PropertyValue,
 					PropDefault = PropData;
 				}
 
-				ElementProp->ExportTextItem(ValueStr, PropData, PropDefault, Parent, PortFlags | PPF_Delimited, ExportRootScope);
+				ElementProp->ExportTextItem_Direct(ValueStr, PropData, PropDefault, Parent, PortFlags | PPF_Delimited, ExportRootScope);
 
 				--Count;
 			}
@@ -666,11 +693,49 @@ void FSetProperty::ExportTextItem(FString& ValueStr, const void* PropertyValue,
 	}
 }
 
-const TCHAR* FSetProperty::ImportText_Internal(const TCHAR* Buffer, void* Data, int32 PortFlags, UObject* Parent, FOutputDevice* ErrorText) const
+const TCHAR* FSetProperty::ImportText_Internal(const TCHAR* Buffer, void* ContainerOrPropertyPtr, EPropertyPointerType PropertyPointerType, UObject* Parent, int32 PortFlags, FOutputDevice* ErrorText) const
 {
 	checkSlow(ElementProp);
 
-	FScriptSetHelper SetHelper(this, Data);
+	FScriptSetHelper SetHelper(this, PointerToValuePtr(ContainerOrPropertyPtr, PropertyPointerType));
+	uint8* TempSetStorage = nullptr;
+	uint8* TempElementStorage = nullptr;
+	bool bSuccess = true;
+
+	ON_SCOPE_EXIT
+	{
+		FMemory::Free(TempElementStorage);
+
+		// If we are returning because of an error, remove any already-added elements from the map before returning
+		// to ensure we're not left with a partial state.
+		if (!bSuccess)
+		{
+			SetHelper.EmptyElements();
+		}
+
+		if (TempSetStorage)
+		{
+			// TempSet is used by property setter so if it was allocated call the setter now
+			FProperty::SetValue_InContainer(ContainerOrPropertyPtr, TempSetStorage);
+
+			// Destroy and free the temp set used by property setter
+			DestroyValue(TempSetStorage);
+			FMemory::Free(TempSetStorage);
+		}
+	};
+
+	if (PropertyPointerType == EPropertyPointerType::Container && HasSetter())
+	{
+		// Allocate temporary set as we first need to initialize it with the parsed items and then use the setter to update the property
+		TempSetStorage = (uint8*)FMemory::MallocZeroed(GetSize());
+		if (!HasAnyPropertyFlags(CPF_ZeroConstructor)) // this stuff is already zero
+		{
+			InitializeValue(TempSetStorage);
+		}
+		// Reinitialize the set helper with the temp value
+		SetHelper = FScriptSetHelper(this, TempSetStorage);
+	}
+
 	SetHelper.EmptyElements();
 
 	// If we export an empty array we export an empty string, so ensure that if we're passed an empty string
@@ -686,21 +751,11 @@ const TCHAR* FSetProperty::ImportText_Internal(const TCHAR* Buffer, void* Data,
 		return Buffer + 1;
 	}
 
-	uint8* TempElementStorage = (uint8*)FMemory::Malloc(ElementProp->ElementSize);
-
-	bool bSuccess = false;
-	ON_SCOPE_EXIT
-	{
-		FMemory::Free(TempElementStorage);
-
-		// If we are returning because of an error, remove any already-added elements from the map before returning
-		// to ensure we're not left with a partial state.
-		if (!bSuccess)
-		{
-			SetHelper.EmptyElements();
-		}
-	};
+	TempElementStorage = (uint8*)FMemory::Malloc(ElementProp->ElementSize);
+	// From this point failure should empty the set
+	bSuccess = false;
 
+	check(ElementProp->GetOffset_ForInternal() == 0);
 	for (;;)
 	{
 		ElementProp->InitializeValue(TempElementStorage);
@@ -710,7 +765,7 @@ const TCHAR* FSetProperty::ImportText_Internal(const TCHAR* Buffer, void* Data,
 		};
 
 		// Read key into temporary storage
-		Buffer = ElementProp->ImportText(Buffer, TempElementStorage, PortFlags | PPF_Delimited, Parent, ErrorText);
+		Buffer = ElementProp->ImportText_Direct(Buffer, TempElementStorage, Parent, PortFlags | PPF_Delimited, ErrorText);
 		if (!Buffer)
 		{
 			return nullptr;
diff --git a/Engine/Source/Runtime/CoreUObject/Private/UObject/PropertySoftObjectPtr.cpp b/Engine/Source/Runtime/CoreUObject/Private/UObject/PropertySoftObjectPtr.cpp
--- a/Engine/Source/Runtime/CoreUObject/Private/UObject/PropertySoftObjectPtr.cpp
+++ b/Engine/Source/Runtime/CoreUObject/Private/UObject/PropertySoftObjectPtr.cpp
@@ -92,9 +92,18 @@ bool FSoftObjectProperty::NetSerializeItem(FArchive& Ar, UPackageMap* Map, void*
 	return true;
 }
 
-void FSoftObjectProperty::ExportTextItem( FString& ValueStr, const void* PropertyValue, const void* DefaultValue, UObject* Parent, int32 PortFlags, UObject* ExportRootScope ) const
+void FSoftObjectProperty::ExportText_Internal( FString& ValueStr, const void* PropertyValueOrContainer, EPropertyPointerType PropertyPointerType, const void* DefaultValue, UObject* Parent, int32 PortFlags, UObject* ExportRootScope ) const
 {
-	FSoftObjectPtr& SoftObjectPtr = *(FSoftObjectPtr*)PropertyValue;
+	FSoftObjectPtr SoftObjectPtr;
+	
+	if (PropertyPointerType == EPropertyPointerType::Container && HasGetter())
+	{
+		GetValue_InContainer(PropertyValueOrContainer, &SoftObjectPtr);
+	}
+	else
+	{
+		SoftObjectPtr = *(FSoftObjectPtr*)PointerToValuePtr(PropertyValueOrContainer, PropertyPointerType);
+	}
 
 	FSoftObjectPath SoftObjectPath;
 	UObject *Object = SoftObjectPtr.Get();
@@ -118,10 +127,8 @@ void FSoftObjectProperty::ExportTextItem( FString& ValueStr, const void* Propert
 	SoftObjectPath.ExportTextItem(ValueStr, SoftObjectPath, Parent, PortFlags, ExportRootScope);
 }
 
-const TCHAR* FSoftObjectProperty::ImportText_Internal( const TCHAR* InBuffer, void* Data, int32 PortFlags, UObject* Parent, FOutputDevice* ErrorText ) const
+const TCHAR* FSoftObjectProperty::ImportText_Internal( const TCHAR* InBuffer, void* ContainerOrPropertyPtr, EPropertyPointerType PropertyPointerType, UObject* Parent, int32 PortFlags, FOutputDevice* ErrorText ) const
 {
-	FSoftObjectPtr& SoftObjectPtr = *(FSoftObjectPtr*)Data;
-
 	FSoftObjectPath SoftObjectPath;
 
 	bool bImportTextSuccess = false;
@@ -139,13 +146,30 @@ const TCHAR* FSoftObjectProperty::ImportText_Internal( const TCHAR* InBuffer, vo
 
 	if (bImportTextSuccess)
 	{
-		SoftObjectPtr = SoftObjectPath;
+		if (PropertyPointerType == EPropertyPointerType::Container && HasSetter())
+		{
+			FSoftObjectPtr SoftObjectPtr(SoftObjectPath);
+			SetValue_InContainer(ContainerOrPropertyPtr, SoftObjectPtr);
+		}
+		else
+		{
+			FSoftObjectPtr& SoftObjectPtr = *(FSoftObjectPtr*)PointerToValuePtr(ContainerOrPropertyPtr, PropertyPointerType);
+			SoftObjectPtr = SoftObjectPath;
+		}
 		return InBuffer;
 	}
-
 	else
 	{
-		SoftObjectPtr = nullptr;
+		if (PropertyPointerType == EPropertyPointerType::Container && HasSetter())
+		{
+			FSoftObjectPtr NullPtr(nullptr);
+			SetValue_InContainer(ContainerOrPropertyPtr, NullPtr);
+		}
+		else
+		{
+			FSoftObjectPtr& SoftObjectPtr = *(FSoftObjectPtr*)PointerToValuePtr(ContainerOrPropertyPtr, PropertyPointerType);
+			SoftObjectPtr = nullptr;
+		}
 		return nullptr;
 	}
 }
diff --git a/Engine/Source/Runtime/CoreUObject/Private/UObject/PropertyStr.cpp b/Engine/Source/Runtime/CoreUObject/Private/UObject/PropertyStr.cpp
--- a/Engine/Source/Runtime/CoreUObject/Private/UObject/PropertyStr.cpp
+++ b/Engine/Source/Runtime/CoreUObject/Private/UObject/PropertyStr.cpp
@@ -158,9 +158,18 @@ FString FStrProperty::ExportCppHardcodedText(const FString& InSource, const FStr
 	return Result.ToString();
 }
 
-void FStrProperty::ExportTextItem( FString& ValueStr, const void* PropertyValue, const void* DefaultValue, UObject* Parent, int32 PortFlags, UObject* ExportRootScope ) const
+void FStrProperty::ExportText_Internal( FString& ValueStr, const void* PropertyValueOrContainer, EPropertyPointerType PropertyPointerType, const void* DefaultValue, UObject* Parent, int32 PortFlags, UObject* ExportRootScope ) const
 {
-	FString& StringValue = *(FString*)PropertyValue;
+	FString StringValue;
+	if (PropertyPointerType == EPropertyPointerType::Container && HasGetter())
+	{
+		GetValue_InContainer(PropertyValueOrContainer, &StringValue);
+	}
+	else
+	{
+		StringValue = *(FString*)PointerToValuePtr(PropertyValueOrContainer, PropertyPointerType);
+	}
+
 	if (0 != (PortFlags & PPF_ExportCpp))
 	{
 		ValueStr += FString::Printf(TEXT("FString(%s)"), *ExportCppHardcodedText(StringValue, FString()));
@@ -178,11 +187,12 @@ void FStrProperty::ExportTextItem( FString& ValueStr, const void* PropertyValue,
 		ValueStr += TEXT("\"\"");
 	}
 }
-const TCHAR* FStrProperty::ImportText_Internal( const TCHAR* Buffer, void* Data, int32 PortFlags, UObject* Parent, FOutputDevice* ErrorText ) const
+const TCHAR* FStrProperty::ImportText_Internal( const TCHAR* Buffer, void* ContainerOrPropertyPtr, EPropertyPointerType PropertyPointerType, UObject* Parent, int32 PortFlags, FOutputDevice* ErrorText ) const
 {
+	FString ImportedText;
 	if( !(PortFlags & PPF_Delimited) )
 	{
-		*(FString*)Data = Buffer;
+		ImportedText = Buffer;
 
 		// in order to indicate that the value was successfully imported, advance the buffer past the last character that was imported
 		Buffer += FCString::Strlen(Buffer);
@@ -207,7 +217,15 @@ const TCHAR* FStrProperty::ImportText_Internal( const TCHAR* Buffer, void* Data,
 			ErrorText->Logf(TEXT("Missing terminating '\"' in string property value: %s"), Start);
 			return NULL;
 		}
-		*(FString*)Data = MoveTemp(Temp);
+		ImportedText = MoveTemp(Temp);
+	}
+	if (PropertyPointerType == EPropertyPointerType::Container && HasSetter())
+	{
+		SetValue_InContainer(ContainerOrPropertyPtr, ImportedText);
+	}
+	else
+	{
+		*(FString*)PointerToValuePtr(ContainerOrPropertyPtr, PropertyPointerType) = MoveTemp(ImportedText);
 	}
 	return Buffer;
 }
diff --git a/Engine/Source/Runtime/CoreUObject/Private/UObject/PropertyStruct.cpp b/Engine/Source/Runtime/CoreUObject/Private/UObject/PropertyStruct.cpp
--- a/Engine/Source/Runtime/CoreUObject/Private/UObject/PropertyStruct.cpp
+++ b/Engine/Source/Runtime/CoreUObject/Private/UObject/PropertyStruct.cpp
@@ -286,12 +286,32 @@ FString FStructProperty::GetCPPMacroType( FString& ExtendedTypeText ) const
 	return TEXT("STRUCT");
 }
 
-void FStructProperty::ExportTextItem( FString& ValueStr, const void* PropertyValue, const void* DefaultValue, UObject* Parent, int32 PortFlags, UObject* ExportRootScope ) const
+void FStructProperty::ExportText_Internal( FString& ValueStr, const void* PropertyValueOrContainer, EPropertyPointerType PropertyPointerType, const void* DefaultValue, UObject* Parent, int32 PortFlags, UObject* ExportRootScope ) const
 {
-	Struct->ExportText(ValueStr, PropertyValue, DefaultValue, Parent, PortFlags, ExportRootScope, true);
+	void* StructData = nullptr;
+	if (PropertyPointerType == EPropertyPointerType::Container && HasGetter())
+	{
+		int32 RequiredAllocSize = Struct->GetStructureSize();
+		StructData = FMemory::Malloc(RequiredAllocSize);
+		Struct->InitializeStruct(StructData);
+		GetValue_InContainer(PropertyValueOrContainer, StructData);
+	}
+	else
+	{
+		StructData = PointerToValuePtr(PropertyValueOrContainer, PropertyPointerType);
+	}
+
+	Struct->ExportText(ValueStr, StructData, DefaultValue, Parent, PortFlags, ExportRootScope, true);
+
+	if (PropertyPointerType == EPropertyPointerType::Container && HasGetter())
+	{
+		Struct->DestroyStruct(StructData);
+		FMemory::Free(StructData);
+		StructData = nullptr;
+	}
 }
 
-const TCHAR* FStructProperty::ImportText_Internal(const TCHAR* InBuffer, void* Data, int32 PortFlags, UObject* Parent, FOutputDevice* ErrorText) const
+const TCHAR* FStructProperty::ImportText_Internal(const TCHAR* InBuffer, void* ContainerOrPropertyPtr, EPropertyPointerType PropertyPointerType, UObject* Parent, int32 PortFlags, FOutputDevice* ErrorText) const
 {
 #if USE_CIRCULAR_DEPENDENCY_LOAD_DEFERRING
 	FScopedPlaceholderPropertyTracker ImportPropertyTracker(this);
@@ -310,13 +330,34 @@ const TCHAR* FStructProperty::ImportText_Internal(const TCHAR* InBuffer, void* D
 		StructLinker->LoadFlags |= OldFlags | PropagatedLoadFlags;
 	}
 #endif 
-	const TCHAR* Result = Struct->ImportText(InBuffer, Data, Parent, PortFlags, ErrorText, [this]() { return GetName(); }, true);
+	void* StructData = nullptr;
+	if (PropertyPointerType == EPropertyPointerType::Container && HasSetter())
+	{
+		int32 RequiredAllocSize = Struct->GetStructureSize();
+		StructData = FMemory::Malloc(RequiredAllocSize);
+		Struct->InitializeStruct(StructData);
+		GetValue_InContainer(ContainerOrPropertyPtr, StructData);
+	}
+	else
+	{
+		StructData = PointerToValuePtr(ContainerOrPropertyPtr, PropertyPointerType);
+	}
+
+	const TCHAR* Result = Struct->ImportText(InBuffer, StructData, Parent, PortFlags, ErrorText, [this]() { return GetName(); }, true);
+
+	if (PropertyPointerType == EPropertyPointerType::Container && HasSetter())
+	{
+		SetValue_InContainer(ContainerOrPropertyPtr, StructData);
+		Struct->DestroyStruct(StructData);
+		FMemory::Free(StructData);
+		StructData = nullptr;
+	}
 
 #if USE_CIRCULAR_DEPENDENCY_LOAD_DEFERRING
 	if (StructLinker)
 	{
 		StructLinker->LoadFlags = OldFlags;
-}
+	}
 #endif
 
 	return Result;
diff --git a/Engine/Source/Runtime/CoreUObject/Private/UObject/PropertyTempVal.cpp b/Engine/Source/Runtime/CoreUObject/Private/UObject/PropertyTempVal.cpp
--- a/Engine/Source/Runtime/CoreUObject/Private/UObject/PropertyTempVal.cpp
+++ b/Engine/Source/Runtime/CoreUObject/Private/UObject/PropertyTempVal.cpp
@@ -29,5 +29,5 @@ void FPropertyTempVal::Serialize(FArchive& Ar, const void* Defaults)
 
 void FPropertyTempVal::ExportText(FString& ValueStr, const void* Defaults, UObject* Parent, int32 PortFlags, UObject* ExportRootScope)
 {
-	Prop->ExportTextItem(ValueStr, Value, Defaults, Parent, PortFlags, ExportRootScope);
+	Prop->ExportTextItem_Direct(ValueStr, Value, Defaults, Parent, PortFlags, ExportRootScope);
 }
diff --git a/Engine/Source/Runtime/CoreUObject/Private/UObject/SavePackage.cpp b/Engine/Source/Runtime/CoreUObject/Private/UObject/SavePackage.cpp
--- a/Engine/Source/Runtime/CoreUObject/Private/UObject/SavePackage.cpp
+++ b/Engine/Source/Runtime/CoreUObject/Private/UObject/SavePackage.cpp
@@ -2378,8 +2378,7 @@ FSavePackageResultStruct UPackage::Save(UPackage* InOuter, UObject* InAsset, con
 									for (FProperty* Prop : TFieldRange<FProperty>(Object->GetClass()))
 									{
 										FString PropState;
-										const void* PropAddr = Prop->ContainerPtrToValuePtr<void>(Object);
-										Prop->ExportTextItem(PropState, PropAddr, nullptr, Object, PPF_None);
+										Prop->ExportTextItem_InContainer(PropState, Object, nullptr, Object, PPF_None);
 
 										Result.Emplace(Prop, MoveTemp(PropState));
 									}
diff --git a/Engine/Source/Runtime/CoreUObject/Private/UObject/ScriptCore.cpp b/Engine/Source/Runtime/CoreUObject/Private/UObject/ScriptCore.cpp
--- a/Engine/Source/Runtime/CoreUObject/Private/UObject/ScriptCore.cpp
+++ b/Engine/Source/Runtime/CoreUObject/Private/UObject/ScriptCore.cpp
@@ -437,12 +437,21 @@ void FFrame::StepExplicitProperty(void*const Result, FProperty* Property)
 			checkSlow(Out);
 		}
 		MostRecentPropertyAddress = Out->PropAddr;
+		MostRecentPropertyContainer = nullptr;
 		// no need to copy property value, since the caller is just looking for MostRecentPropertyAddress
 	}
 	else
 	{
 		MostRecentPropertyAddress = Property->ContainerPtrToValuePtr<uint8>(Locals);
-		Property->CopyCompleteValueToScriptVM(Result, MostRecentPropertyAddress);
+		MostRecentPropertyContainer = Locals;
+		if (Property->HasGetter())
+		{
+			Property->GetValue_InContainer(MostRecentPropertyContainer, Result);
+		}
+		else
+		{
+			Property->CopyCompleteValueToScriptVM(Result, MostRecentPropertyAddress);
+		}
 	}
 }
 
@@ -756,7 +765,8 @@ void UObject::SkipFunction(FFrame& Stack, RESULT_DECL, UFunction* Function)
 	FMemory::Memzero(Frame, Function->PropertiesSize);
 	for (FProperty* Property = (FProperty*)(Function->ChildProperties); *Stack.Code != EX_EndFunctionParms; Property = (FProperty*)(Property->Next))
 	{
-		Stack.MostRecentPropertyAddress = NULL;
+		Stack.MostRecentPropertyAddress = nullptr;
+		Stack.MostRecentPropertyContainer = nullptr;
 		// evaluate the expression into our temporary memory space
 		// it'd be nice to be able to skip the copy, but most native functions assume a non-NULL Result pointer
 		// so we can only do that if we know the expression is an l-value (out parameter)
@@ -841,8 +851,9 @@ void ProcessScriptFunction(UObject* Context, UFunction* Function, FFrame& Stack,
 	{
 		checkfSlow(Property, TEXT("NULL Property in Function %s"), *Function->GetPathName()); 
 
-		Stack.MostRecentPropertyAddress = NULL;
-			
+		Stack.MostRecentPropertyAddress = nullptr;
+		Stack.MostRecentPropertyContainer = nullptr;
+
 		// Skip the return parameter case, as we've already handled it above
 		const bool bIsReturnParam = ((Property->PropertyFlags & CPF_ReturnParm) != 0);
 		if( bIsReturnParam )
@@ -1262,7 +1273,7 @@ bool UObject::CallFunctionByNameWithArguments(const TCHAR* Str, FOutputDevice& A
 			{
 				bFoundDefault = true;
 
-				const TCHAR* Result = It->ImportText( *PropertyDefaultValue, It->ContainerPtrToValuePtr<uint8>(Parms), ExportFlags, NULL );
+				const TCHAR* Result = It->ImportText_InContainer(*PropertyDefaultValue, Parms, nullptr, ExportFlags);
 				bFailedImport = (Result == nullptr);
 			}
 		}
@@ -1278,7 +1289,7 @@ bool UObject::CallFunctionByNameWithArguments(const TCHAR* Str, FOutputDevice& A
 				ArgStr = FString(RemainingStr).TrimStart();
 			}
 
-			const TCHAR* Result = It->ImportText(*ArgStr, It->ContainerPtrToValuePtr<uint8>(Parms), ExportFlags, NULL );
+			const TCHAR* Result = It->ImportText_InContainer(*ArgStr, Parms, nullptr, ExportFlags);
 			bFailedImport = (Result == nullptr);
 		}
 		
@@ -2044,14 +2055,23 @@ DEFINE_FUNCTION(UObject::execLocalVariable)
 		FBlueprintCoreDelegates::ThrowScriptException(P_THIS, Stack, ExceptionInfo);
 
 		Stack.MostRecentPropertyAddress = nullptr;
+		Stack.MostRecentPropertyContainer = nullptr;
 	}
 	else
 	{
 		Stack.MostRecentPropertyAddress = VarProperty->ContainerPtrToValuePtr<uint8>(Stack.Locals);
+		Stack.MostRecentPropertyContainer = Stack.Locals;
 
 		if (RESULT_PARAM)
 		{
-			VarProperty->CopyCompleteValueToScriptVM(RESULT_PARAM, Stack.MostRecentPropertyAddress);
+			if (VarProperty->HasGetter())
+			{
+				VarProperty->GetValue_InContainer(Stack.MostRecentPropertyContainer, RESULT_PARAM);
+			}
+			else
+			{
+				VarProperty->CopyCompleteValueToScriptVM(RESULT_PARAM, Stack.MostRecentPropertyAddress);
+			}
 		}
 	}
 }
@@ -2068,14 +2088,22 @@ DEFINE_FUNCTION(UObject::execInstanceVariable)
 		FBlueprintCoreDelegates::ThrowScriptException(P_THIS, Stack, ExceptionInfo);
 
 		Stack.MostRecentPropertyAddress = nullptr;
+		Stack.MostRecentPropertyContainer = nullptr;
 	}
 	else
 	{
 		Stack.MostRecentPropertyAddress = VarProperty->ContainerPtrToValuePtr<uint8>(P_THIS);
-
+		Stack.MostRecentPropertyContainer = (uint8*)P_THIS;
 		if (RESULT_PARAM)
 		{
-			VarProperty->CopyCompleteValueToScriptVM(RESULT_PARAM, Stack.MostRecentPropertyAddress);
+			if (VarProperty->HasGetter())
+			{
+				VarProperty->GetValue_InContainer(Stack.MostRecentPropertyContainer, RESULT_PARAM);
+			}
+			else
+			{
+				VarProperty->CopyCompleteValueToScriptVM(RESULT_PARAM, Stack.MostRecentPropertyAddress);
+			}
 		}
 	}
 
@@ -2094,15 +2122,24 @@ DEFINE_FUNCTION(UObject::execClassSparseDataVariable)
 		FBlueprintCoreDelegates::ThrowScriptException(P_THIS, Stack, ExceptionInfo);
 
 		Stack.MostRecentPropertyAddress = nullptr;
+		Stack.MostRecentPropertyContainer = nullptr;
 	}
 	else
 	{
 		void* SparseDataBaseAddress = const_cast<void*>(P_THIS->GetClass()->GetSparseClassData(EGetSparseClassDataMethod::ArchetypeIfNull));
 		Stack.MostRecentPropertyAddress = VarProperty->ContainerPtrToValuePtr<uint8>(SparseDataBaseAddress);
+		Stack.MostRecentPropertyContainer = (uint8*)SparseDataBaseAddress;
 
 		if (RESULT_PARAM)
 		{
-			VarProperty->CopyCompleteValueToScriptVM(RESULT_PARAM, Stack.MostRecentPropertyAddress);
+			if (VarProperty->HasGetter())
+			{
+				VarProperty->GetValue_InContainer(Stack.MostRecentPropertyContainer, RESULT_PARAM);
+			}
+			else
+			{
+				VarProperty->CopyCompleteValueToScriptVM(RESULT_PARAM, Stack.MostRecentPropertyAddress);
+			}
 		}
 	}
 }
@@ -2113,6 +2150,7 @@ DEFINE_FUNCTION(UObject::execDefaultVariable)
 	FProperty* VarProperty = (FProperty*)Stack.ReadObject();
 	Stack.MostRecentProperty = VarProperty;
 	Stack.MostRecentPropertyAddress = nullptr;
+	Stack.MostRecentPropertyContainer = nullptr;
 
 	UObject* DefaultObject = nullptr;
 	if (P_THIS->HasAnyFlags(RF_ClassDefaultObject))
@@ -2134,9 +2172,18 @@ DEFINE_FUNCTION(UObject::execDefaultVariable)
 		if(DefaultObject != nullptr)
 		{
 			Stack.MostRecentPropertyAddress = VarProperty->ContainerPtrToValuePtr<uint8>(DefaultObject);
+			Stack.MostRecentPropertyContainer = (uint8*)DefaultObject;
+
 			if(RESULT_PARAM)
 			{
-				VarProperty->CopyCompleteValueToScriptVM(RESULT_PARAM, Stack.MostRecentPropertyAddress);
+				if (VarProperty->HasGetter())
+				{
+					VarProperty->GetValue_InContainer(Stack.MostRecentPropertyContainer, RESULT_PARAM);
+				}
+				else
+				{
+					VarProperty->CopyCompleteValueToScriptVM(RESULT_PARAM, Stack.MostRecentPropertyAddress);
+				}
 			}
 		}
 		else
@@ -2226,8 +2273,9 @@ DEFINE_FUNCTION(UObject::execClassContext)
 		FProperty* RValueProperty = nullptr;
 		const VariableSizeType bSize = Stack.ReadVariableSize(&RValueProperty); // Code += sizeof(ScriptPointerType) + sizeof(uint8)
 		Stack.Code += wSkip;
-		Stack.MostRecentPropertyAddress = NULL;
-		Stack.MostRecentProperty = NULL;
+		Stack.MostRecentPropertyAddress = nullptr;
+		Stack.MostRecentPropertyContainer = nullptr;
+		Stack.MostRecentProperty = nullptr;
 
 		if (RESULT_PARAM && RValueProperty)
 		{
@@ -2473,8 +2521,9 @@ IMPLEMENT_VM_FUNCTION( EX_PopExecutionFlowIfNot, execPopExecutionFlowIfNot );
 DEFINE_FUNCTION(UObject::execLetValueOnPersistentFrame)
 {
 #if USE_UBER_GRAPH_PERSISTENT_FRAME
-	Stack.MostRecentProperty = NULL;
-	Stack.MostRecentPropertyAddress = NULL;
+	Stack.MostRecentProperty = nullptr;
+	Stack.MostRecentPropertyAddress = nullptr;
+	Stack.MostRecentPropertyContainer = nullptr;
 
 	FProperty* DestProperty = Stack.ReadProperty();
 	checkSlow(DestProperty);
@@ -2498,6 +2547,7 @@ DEFINE_FUNCTION(UObject::execSwitchValue)
 
 	Stack.MostRecentProperty = nullptr;
 	Stack.MostRecentPropertyAddress = nullptr;
+	Stack.MostRecentPropertyContainer = nullptr;
 	Stack.Step(Stack.Object, nullptr);
 
 	FProperty* IndexProperty = Stack.MostRecentProperty;
@@ -2562,7 +2612,8 @@ IMPLEMENT_VM_FUNCTION(EX_SwitchValue, execSwitchValue);
 DEFINE_FUNCTION(UObject::execArrayGetByRef)
 {
 	// Get variable address.
-	Stack.MostRecentPropertyAddress = NULL;
+	Stack.MostRecentPropertyAddress = nullptr;
+	Stack.MostRecentPropertyContainer = nullptr;
 	Stack.Step( Stack.Object, NULL ); // Evaluate variable.
 
 	if (Stack.MostRecentPropertyAddress == NULL)
@@ -2590,6 +2641,7 @@ DEFINE_FUNCTION(UObject::execArrayGetByRef)
 	if (ArrayHelper.IsValidIndex(ArrayIndex))
 	{
 		Stack.MostRecentPropertyAddress = ArrayHelper.GetRawPtr(ArrayIndex);
+		Stack.MostRecentPropertyContainer = nullptr;
 
 		if (RESULT_PARAM)
 		{
@@ -2600,6 +2652,7 @@ DEFINE_FUNCTION(UObject::execArrayGetByRef)
 	{
 		// clear so other methods don't try to use a stale value (depends on this method succeeding)
 		Stack.MostRecentPropertyAddress = nullptr;
+		Stack.MostRecentPropertyContainer = nullptr;
 		// sometimes other exec functions guard on MostRecentProperty, and expect 
 		// MostRecentPropertyAddress to be filled out; since this was a failure
 		// clear this too (so all reliant execs can properly detect)
@@ -2627,9 +2680,12 @@ DEFINE_FUNCTION(UObject::execLet)
 	// Get variable address.
 	Stack.MostRecentProperty = nullptr;
 	Stack.MostRecentPropertyAddress = nullptr;
+	Stack.MostRecentPropertyContainer = nullptr;
 	Stack.Step(Stack.Object, nullptr); // Evaluate variable.
 
 	uint8* LocalTempResult = nullptr;
+	uint8* PreviousPropertyAddress = nullptr;
+
 	if (Stack.MostRecentPropertyAddress == nullptr)
 	{
 		FBlueprintExceptionInfo ExceptionInfo(
@@ -2649,13 +2705,31 @@ DEFINE_FUNCTION(UObject::execLet)
 			FMemory::Memzero(Stack.MostRecentPropertyAddress, sizeof(FString));
 		}
 	}
+	else if (LocallyKnownProperty && LocallyKnownProperty->HasSetter())
+	{
+		// We can't assign a value directly to a property if it's got a setter or getter
+		LocalTempResult = (uint8*)FMemory_Alloca_Aligned(LocallyKnownProperty->GetSize(), LocallyKnownProperty->GetMinAlignment());
+		LocallyKnownProperty->InitializeValue(LocalTempResult);
+		PreviousPropertyAddress = Stack.MostRecentPropertyAddress;
+		Stack.MostRecentPropertyAddress = LocalTempResult;
+	}
 
 	// Evaluate expression into variable.
 	Stack.Step(Stack.Object, Stack.MostRecentPropertyAddress);
 
-	if (LocalTempResult && LocallyKnownProperty)
+	if (LocallyKnownProperty)
 	{
-		LocallyKnownProperty->DestroyValue(LocalTempResult);
+		if (LocallyKnownProperty->HasSetter())
+		{
+			check(Stack.MostRecentPropertyContainer != nullptr);
+			LocallyKnownProperty->SetValue_InContainer(Stack.MostRecentPropertyContainer, LocalTempResult);
+			Stack.MostRecentPropertyAddress = PreviousPropertyAddress;
+		}
+
+		if (LocalTempResult)
+		{
+			LocallyKnownProperty->DestroyValue(LocalTempResult);
+		}
 	}
 }
 IMPLEMENT_VM_FUNCTION( EX_Let, execLet );
@@ -2663,7 +2737,8 @@ IMPLEMENT_VM_FUNCTION( EX_Let, execLet );
 DEFINE_FUNCTION(UObject::execLetObj)
 {
 	// Get variable address.
-	Stack.MostRecentPropertyAddress = NULL;
+	Stack.MostRecentPropertyAddress = nullptr;
+	Stack.MostRecentPropertyContainer = nullptr;
 	Stack.Step( Stack.Object, NULL ); // Evaluate variable.
 
 	if (Stack.MostRecentPropertyAddress == NULL)
@@ -2692,7 +2767,15 @@ DEFINE_FUNCTION(UObject::execLetObj)
 	if (ObjAddr)
 	{
 		checkSlow(ObjectProperty);
-		ObjectProperty->SetObjectPropertyValue(ObjAddr, NewValue);
+		if (ObjectProperty->HasSetter())
+		{
+			check(Stack.MostRecentPropertyContainer != nullptr);
+			ObjectProperty->SetValue_InContainer(Stack.MostRecentPropertyContainer, &NewValue);
+		}
+		else
+		{
+			ObjectProperty->SetObjectPropertyValue(ObjAddr, NewValue);
+		}
 	}
 }
 IMPLEMENT_VM_FUNCTION( EX_LetObj, execLetObj );
@@ -2700,7 +2783,8 @@ IMPLEMENT_VM_FUNCTION( EX_LetObj, execLetObj );
 DEFINE_FUNCTION(UObject::execLetWeakObjPtr)
 {
 	// Get variable address.
-	Stack.MostRecentPropertyAddress = NULL;
+	Stack.MostRecentPropertyAddress = nullptr;
+	Stack.MostRecentPropertyContainer = nullptr;
 	Stack.Step( Stack.Object, NULL ); // Evaluate variable.
 
 	if (Stack.MostRecentPropertyAddress == NULL)
@@ -2729,15 +2813,25 @@ DEFINE_FUNCTION(UObject::execLetWeakObjPtr)
 	if (ObjAddr)
 	{
 		checkSlow(ObjectProperty);
-		ObjectProperty->SetObjectPropertyValue(ObjAddr, NewValue);
+		if (ObjectProperty->HasSetter())
+		{
+			check(Stack.MostRecentPropertyContainer != nullptr);
+			FWeakObjectPtr NewWeakPtrValue(NewValue);
+			ObjectProperty->SetValue_InContainer(Stack.MostRecentPropertyContainer, &NewWeakPtrValue);
+		}
+		else
+		{
+			ObjectProperty->SetObjectPropertyValue(ObjAddr, NewValue);
+		}
 	}
 }
 IMPLEMENT_VM_FUNCTION( EX_LetWeakObjPtr, execLetWeakObjPtr );
 
 DEFINE_FUNCTION(UObject::execLetBool)
 {
-	Stack.MostRecentPropertyAddress = NULL;
-	Stack.MostRecentProperty = NULL;
+	Stack.MostRecentPropertyAddress = nullptr;
+	Stack.MostRecentPropertyContainer = nullptr;
+	Stack.MostRecentProperty = nullptr;
 
 	// Get the variable and address to place the data.
 	Stack.Step( Stack.Object, NULL );
@@ -2773,7 +2867,15 @@ DEFINE_FUNCTION(UObject::execLetBool)
 	if( BoolAddr )
 	{
 		checkSlow(CastField<FBoolProperty>(BoolProperty));
-		BoolProperty->SetPropertyValue( BoolAddr, NewValue );
+		if (BoolProperty->HasSetter())
+		{
+			check(Stack.MostRecentPropertyContainer != nullptr);
+			BoolProperty->SetValue_InContainer(Stack.MostRecentPropertyContainer, &NewValue);
+		}
+		else
+		{
+			BoolProperty->SetPropertyValue(BoolAddr, NewValue);
+		}
 	}
 }
 IMPLEMENT_VM_FUNCTION( EX_LetBool, execLetBool );
@@ -2782,8 +2884,9 @@ IMPLEMENT_VM_FUNCTION( EX_LetBool, execLetBool );
 DEFINE_FUNCTION(UObject::execLetDelegate)
 {
 	// Get variable address.
-	Stack.MostRecentPropertyAddress = NULL;
-	Stack.MostRecentProperty = NULL;
+	Stack.MostRecentPropertyAddress = nullptr;
+	Stack.MostRecentPropertyContainer = nullptr;
+	Stack.MostRecentProperty = nullptr;
 	Stack.Step( Stack.Object, NULL ); // Variable.
 
 	FScriptDelegate* DelegateAddr = (FScriptDelegate*)Stack.MostRecentPropertyAddress;
@@ -2801,8 +2904,9 @@ IMPLEMENT_VM_FUNCTION( EX_LetDelegate, execLetDelegate );
 DEFINE_FUNCTION(UObject::execLetMulticastDelegate)
 {
 	// Get variable address.
-	Stack.MostRecentPropertyAddress = NULL;
-	Stack.MostRecentProperty = NULL;
+	Stack.MostRecentPropertyAddress = nullptr;
+	Stack.MostRecentPropertyContainer = nullptr;
+	Stack.MostRecentProperty = nullptr;
 	Stack.Step( Stack.Object, NULL ); // Variable.
 
 	FMulticastDelegateProperty* DelegateProp = CastFieldCheckedNullAllowed<FMulticastDelegateProperty>(Stack.MostRecentProperty);
@@ -2919,8 +3023,9 @@ void UObject::ProcessContextOpcode( FFrame& Stack, RESULT_DECL, bool bCanFailSil
 		FProperty* RValueProperty = nullptr;
 		const VariableSizeType bSize = Stack.ReadVariableSize(&RValueProperty); // Code += sizeof(ScriptPointerType) + sizeof(uint8)
 		Stack.Code += wSkip;
-		Stack.MostRecentPropertyAddress = NULL;
-		Stack.MostRecentProperty = NULL;
+		Stack.MostRecentPropertyAddress = nullptr;
+		Stack.MostRecentPropertyContainer = nullptr;
+		Stack.MostRecentProperty = nullptr;
 
 		if (RESULT_PARAM && RValueProperty)
 		{
@@ -2936,20 +3041,29 @@ DEFINE_FUNCTION(UObject::execStructMemberContext)
 	checkSlow(StructProperty);
 
 	// Evaluate an expression leading to the struct.
-	Stack.MostRecentProperty = NULL;
-	Stack.MostRecentPropertyAddress = NULL;
-	Stack.Step(Stack.Object, NULL);
+	Stack.MostRecentProperty = nullptr;
+	Stack.MostRecentPropertyAddress = nullptr;
+	Stack.MostRecentPropertyContainer = nullptr;
+	Stack.Step(Stack.Object, nullptr);
 
-	if (Stack.MostRecentProperty != NULL)
+	if (Stack.MostRecentProperty != nullptr)
 	{
 		// Offset into the specific member
-		Stack.MostRecentPropertyAddress = StructProperty->ContainerPtrToValuePtr<uint8>(Stack.MostRecentPropertyAddress);
+		Stack.MostRecentPropertyContainer = Stack.MostRecentPropertyAddress;
+		Stack.MostRecentPropertyAddress = StructProperty->ContainerPtrToValuePtr<uint8>(Stack.MostRecentPropertyAddress);		
 		Stack.MostRecentProperty = StructProperty;
 
 		// Handle variable reads
 		if (RESULT_PARAM)
 		{
-			StructProperty->CopyCompleteValueToScriptVM(RESULT_PARAM, Stack.MostRecentPropertyAddress);
+			if (StructProperty->HasGetter())
+			{
+				StructProperty->GetValue_InContainer(Stack.MostRecentPropertyContainer, RESULT_PARAM);
+			}
+			else
+			{
+				StructProperty->CopyCompleteValueToScriptVM(RESULT_PARAM, Stack.MostRecentPropertyAddress);
+			}
 		}
 	}
 	else
@@ -2964,8 +3078,9 @@ DEFINE_FUNCTION(UObject::execStructMemberContext)
 		);
 		FBlueprintCoreDelegates::ThrowScriptException(P_THIS, Stack, ExceptionInfo);
 
-		Stack.MostRecentPropertyAddress = NULL;
-		Stack.MostRecentProperty = NULL;
+		Stack.MostRecentPropertyAddress = nullptr;
+		Stack.MostRecentPropertyContainer = nullptr;
+		Stack.MostRecentProperty = nullptr;
 	}
 }
 IMPLEMENT_VM_FUNCTION( EX_StructMemberContext, execStructMemberContext );
@@ -3005,9 +3120,10 @@ public:
 	{
 		//Get delegate
 		UFunction* SignatureFunction = CastChecked<UFunction>(Stack.ReadObject());
-		Stack.MostRecentPropertyAddress = NULL;
-		Stack.MostRecentProperty = NULL;
-		Stack.Step( Stack.Object, NULL );
+		Stack.MostRecentPropertyAddress = nullptr;
+		Stack.MostRecentPropertyContainer = nullptr;
+		Stack.MostRecentProperty = nullptr;
+		Stack.Step( Stack.Object, nullptr );
 		FMulticastDelegateProperty* DelegateProp = CastFieldCheckedNullAllowed<FMulticastDelegateProperty>(Stack.MostRecentProperty);
 		const FMulticastScriptDelegate* DelegateAddr = (DelegateProp ? DelegateProp->GetMulticastDelegate(Stack.MostRecentPropertyAddress) : nullptr);
 
@@ -3016,7 +3132,8 @@ public:
 		FMemory::Memzero(Parameters, SignatureFunction->ParmsSize);
 		for (FProperty* Property = (FProperty*)SignatureFunction->ChildProperties; *Stack.Code != EX_EndFunctionParms; Property = (FProperty*)Property->Next)
 		{
-			Stack.MostRecentPropertyAddress = NULL;
+			Stack.MostRecentPropertyAddress = nullptr;
+			Stack.MostRecentPropertyContainer = nullptr;
 			if (Property->PropertyFlags & CPF_OutParm)
 			{
 				Stack.Step(Stack.Object, NULL);
@@ -3060,9 +3177,10 @@ IMPLEMENT_VM_FUNCTION( EX_CallMulticastDelegate, execCallMulticastDelegate );
 DEFINE_FUNCTION(UObject::execAddMulticastDelegate)
 {
 	// Get variable address.
-	Stack.MostRecentPropertyAddress = NULL;
-	Stack.MostRecentProperty = NULL;
-	Stack.Step( Stack.Object, NULL ); // Variable.
+	Stack.MostRecentPropertyAddress = nullptr;
+	Stack.MostRecentPropertyContainer = nullptr;
+	Stack.MostRecentProperty = nullptr;
+	Stack.Step( Stack.Object, nullptr ); // Variable.
 
 	FMulticastDelegateProperty* DelegateProp = CastFieldCheckedNullAllowed<FMulticastDelegateProperty>(Stack.MostRecentProperty);
 	void* DelegateAddr = Stack.MostRecentPropertyAddress;
@@ -3080,9 +3198,10 @@ IMPLEMENT_VM_FUNCTION( EX_AddMulticastDelegate, execAddMulticastDelegate );
 DEFINE_FUNCTION(UObject::execRemoveMulticastDelegate)
 {
 	// Get variable address.
-	Stack.MostRecentPropertyAddress = NULL;
-	Stack.MostRecentProperty = NULL;
-	Stack.Step( Stack.Object, NULL ); // Variable.
+	Stack.MostRecentPropertyAddress = nullptr;
+	Stack.MostRecentPropertyContainer = nullptr;
+	Stack.MostRecentProperty = nullptr;
+	Stack.Step( Stack.Object, nullptr ); // Variable.
 
 	FMulticastDelegateProperty* DelegateProp = CastFieldCheckedNullAllowed<FMulticastDelegateProperty>(Stack.MostRecentProperty);
 	void* DelegateAddr = Stack.MostRecentPropertyAddress;
@@ -3100,9 +3219,10 @@ IMPLEMENT_VM_FUNCTION( EX_RemoveMulticastDelegate, execRemoveMulticastDelegate )
 DEFINE_FUNCTION(UObject::execClearMulticastDelegate)
 {
 	// Get the delegate address
-	Stack.MostRecentPropertyAddress = NULL;
-	Stack.MostRecentProperty = NULL;
-	Stack.Step( Stack.Object, NULL );
+	Stack.MostRecentPropertyAddress = nullptr;
+	Stack.MostRecentPropertyContainer = nullptr;
+	Stack.MostRecentProperty = nullptr;
+	Stack.Step( Stack.Object, nullptr );
 
 	FMulticastDelegateProperty* DelegateProp = CastFieldCheckedNullAllowed<FMulticastDelegateProperty>(Stack.MostRecentProperty);
 	void* DelegateAddr = Stack.MostRecentPropertyAddress;
@@ -3285,9 +3405,10 @@ DEFINE_FUNCTION(UObject::execBindDelegate)
 	FName FunctionName = Stack.ReadName();
 
 	// Get delegate address.
-	Stack.MostRecentPropertyAddress = NULL;
-	Stack.MostRecentProperty = NULL;
-	Stack.Step( Stack.Object, NULL ); // Variable.
+	Stack.MostRecentPropertyAddress = nullptr;
+	Stack.MostRecentPropertyContainer = nullptr;
+	Stack.MostRecentProperty = nullptr;
+	Stack.Step( Stack.Object, nullptr ); // Variable.
 
 	FScriptDelegate* DelegateAddr = (FScriptDelegate*)Stack.MostRecentPropertyAddress;
 
@@ -3398,9 +3519,10 @@ IMPLEMENT_VM_FUNCTION( EX_StructConst, execStructConst );
 DEFINE_FUNCTION(UObject::execSetArray)
 {
 	// Get the array address
-	Stack.MostRecentPropertyAddress = NULL;
-	Stack.MostRecentProperty = NULL;
-	Stack.Step( Stack.Object, NULL ); // Array to set
+	Stack.MostRecentPropertyAddress = nullptr;
+	Stack.MostRecentPropertyContainer = nullptr;
+	Stack.MostRecentProperty = nullptr;
+	Stack.Step( Stack.Object, nullptr ); // Array to set
 	
 	FArrayProperty* ArrayProperty = CastFieldChecked<FArrayProperty>(Stack.MostRecentProperty);
  	FScriptArrayHelper ArrayHelper(ArrayProperty, Stack.MostRecentPropertyAddress);
@@ -3422,6 +3544,7 @@ DEFINE_FUNCTION(UObject::execSetSet)
 {
 	// Get the set address
 	Stack.MostRecentPropertyAddress = nullptr;
+	Stack.MostRecentPropertyContainer = nullptr;
 	Stack.MostRecentProperty = nullptr;
 	Stack.Step( Stack.Object, nullptr ); // Set to set
 	const int32 Num = Stack.ReadInt<int32>();
@@ -3455,6 +3578,7 @@ DEFINE_FUNCTION(UObject::execSetMap)
 {
 	// Get the map address
 	Stack.MostRecentPropertyAddress = nullptr;
+	Stack.MostRecentPropertyContainer = nullptr;
 	Stack.MostRecentProperty = nullptr;
 	Stack.Step( Stack.Object, nullptr ); // Map to set
 	const int32 Num = Stack.ReadInt<int32>();
diff --git a/Engine/Source/Runtime/CoreUObject/Private/UObject/TextProperty.cpp b/Engine/Source/Runtime/CoreUObject/Private/UObject/TextProperty.cpp
--- a/Engine/Source/Runtime/CoreUObject/Private/UObject/TextProperty.cpp
+++ b/Engine/Source/Runtime/CoreUObject/Private/UObject/TextProperty.cpp
@@ -114,9 +114,17 @@ void FTextProperty::SerializeItem(FStructuredArchive::FSlot Slot, void* Value, v
 	Slot << *TextPtr;
 }
 
-void FTextProperty::ExportTextItem( FString& ValueStr, const void* PropertyValue, const void* DefaultValue, UObject* Parent, int32 PortFlags, UObject* ExportRootScope ) const
+void FTextProperty::ExportText_Internal( FString& ValueStr, const void* PropertyValueOrContainer, EPropertyPointerType PropertyPointerType, const void* DefaultValue, UObject* Parent, int32 PortFlags, UObject* ExportRootScope ) const
 {
-	const FText& TextValue = GetPropertyValue(PropertyValue);
+	FText TextValue;
+	if (PropertyPointerType == EPropertyPointerType::Container && HasGetter())
+	{
+		GetValue_InContainer(PropertyValueOrContainer, &TextValue);
+	}
+	else
+	{
+		TextValue = GetPropertyValue(PointerToValuePtr(PropertyValueOrContainer, PropertyPointerType));
+	}
 
 	if (PortFlags & PPF_ExportCpp)
 	{
@@ -141,9 +149,9 @@ void FTextProperty::ExportTextItem( FString& ValueStr, const void* PropertyValue
 	}
 }
 
-const TCHAR* FTextProperty::ImportText_Internal( const TCHAR* Buffer, void* Data, int32 PortFlags, UObject* Parent, FOutputDevice* ErrorText ) const
+const TCHAR* FTextProperty::ImportText_Internal( const TCHAR* Buffer, void* ContainerOrPropertyPtr, EPropertyPointerType PropertyPointerType, UObject* Parent, int32 PortFlags, FOutputDevice* ErrorText ) const
 {
-	FText* TextPtr = GetPropertyValuePtr(Data);
+	FText ImportedText;
 
 	FString TextNamespace;
 	if (Parent && HasAnyPropertyFlags(CPF_Config))
@@ -181,7 +189,20 @@ const TCHAR* FTextProperty::ImportText_Internal( const TCHAR* Buffer, void* Data
 	}
 #endif // USE_STABLE_LOCALIZATION_KEYS
 
-	return FTextStringHelper::ReadFromBuffer(Buffer, *TextPtr, *TextNamespace, *PackageNamespace, !!(PortFlags & PPF_Delimited));
+	const TCHAR* Result = FTextStringHelper::ReadFromBuffer(Buffer, ImportedText, *TextNamespace, *PackageNamespace, !!(PortFlags & PPF_Delimited));
+	if (Result)
+	{
+		if (PropertyPointerType == EPropertyPointerType::Container && HasSetter())
+		{
+			SetValue_InContainer(ContainerOrPropertyPtr, ImportedText);
+		}
+		else
+		{
+			FText* TextPtr = GetPropertyValuePtr(PointerToValuePtr(ContainerOrPropertyPtr, PropertyPointerType));
+			*TextPtr = ImportedText;
+		}
+	}
+	return Result;
 }
 
 FString FTextProperty::GenerateCppCodeForTextValue(const FText& InValue, const FString& Indent)
diff --git a/Engine/Source/Runtime/CoreUObject/Private/UObject/UObjectGlobals.cpp b/Engine/Source/Runtime/CoreUObject/Private/UObject/UObjectGlobals.cpp
--- a/Engine/Source/Runtime/CoreUObject/Private/UObject/UObjectGlobals.cpp
+++ b/Engine/Source/Runtime/CoreUObject/Private/UObject/UObjectGlobals.cpp
@@ -64,6 +64,7 @@
 #include "ProfilingDebugging/LoadTimeTracker.h"
 #include "Misc/PackageAccessTracking.h"
 #include "Misc/PackageAccessTracking.h"
+#include "UObject/PropertyWithSetterAndGetter.h"
 
 DEFINE_LOG_CATEGORY(LogUObjectGlobals);
 
@@ -458,7 +459,7 @@ void GlobalSetProperty( const TCHAR* Value, UClass* Class, FProperty* Property,
 						Object->PreEditChange(Property);
 					}
 #endif // WITH_EDITOR
-					Property->ImportText( Value, Property->ContainerPtrToValuePtr<uint8>(Object), 0, Object );
+					Property->ImportText_InContainer(Value, Object, Object, 0);
 #if WITH_EDITOR
 					if( !Object->HasAnyFlags(RF_ClassDefaultObject) && bNotifyObjectOfChange )
 					{
@@ -4667,7 +4668,16 @@ namespace UECodeGen_Private
 	PropertyType* NewFProperty(FFieldVariant Outer, const FPropertyParamsBase& PropBase)
 	{
 		const PropertyParamsType& Prop = (const PropertyParamsType&)PropBase;
-		PropertyType* NewProp = new PropertyType(Outer, Prop);
+		PropertyType* NewProp = nullptr;
+
+		if (Prop.SetterFunc || Prop.GetterFunc)
+		{
+			NewProp = new TPropertyWithSetterAndGetter<PropertyType>(Outer, Prop);
+		}
+		else
+		{
+			NewProp = new PropertyType(Outer, Prop);
+		}
 
 #if WITH_METADATA
 		if (Prop.NumMetaData)
diff --git a/Engine/Source/Runtime/CoreUObject/Public/UObject/EnumProperty.h b/Engine/Source/Runtime/CoreUObject/Public/UObject/EnumProperty.h
--- a/Engine/Source/Runtime/CoreUObject/Public/UObject/EnumProperty.h
+++ b/Engine/Source/Runtime/CoreUObject/Public/UObject/EnumProperty.h
@@ -53,8 +53,8 @@ public:
 	virtual bool Identical( const void* A, const void* B, uint32 PortFlags ) const override;
 	virtual void SerializeItem(FStructuredArchive::FSlot Slot, void* Value, void const* Defaults) const override;
 	virtual bool NetSerializeItem( FArchive& Ar, UPackageMap* Map, void* Data, TArray<uint8> * MetaData = NULL ) const override;
-	virtual void ExportTextItem( FString& ValueStr, const void* PropertyValue, const void* DefaultValue, UObject* Parent, int32 PortFlags, UObject* ExportRootScope ) const override;
-	virtual const TCHAR* ImportText_Internal( const TCHAR* Buffer, void* Data, int32 PortFlags, UObject* OwnerObject, FOutputDevice* ErrorText ) const override;
+	virtual void ExportText_Internal( FString& ValueStr, const void* PropertyValueOrContainer, EPropertyPointerType PropertyPointerType, const void* DefaultValue, UObject* Parent, int32 PortFlags, UObject* ExportRootScope ) const override;
+	virtual const TCHAR* ImportText_Internal(const TCHAR* Buffer, void* ContainerOrPropertyPtr, EPropertyPointerType PropertyPointerType, UObject* OwnerObject, int32 PortFlags, FOutputDevice* ErrorText) const override;
 	virtual int32 GetMinAlignment() const override;
 	virtual bool SameType(const FProperty* Other) const override;
 	virtual EConvertFromTypeResult ConvertFromType(const FPropertyTag& Tag, FStructuredArchive::FSlot Slot, uint8* Data, UStruct* DefaultsStruct) override;
diff --git a/Engine/Source/Runtime/CoreUObject/Public/UObject/FieldPathProperty.h b/Engine/Source/Runtime/CoreUObject/Public/UObject/FieldPathProperty.h
--- a/Engine/Source/Runtime/CoreUObject/Public/UObject/FieldPathProperty.h
+++ b/Engine/Source/Runtime/CoreUObject/Public/UObject/FieldPathProperty.h
@@ -57,8 +57,8 @@ PRAGMA_ENABLE_DEPRECATION_WARNINGS
 	virtual EConvertFromTypeResult ConvertFromType(const FPropertyTag& Tag, FStructuredArchive::FSlot Slot, uint8* Data, UStruct* DefaultsStruct) override;
 	virtual bool Identical( const void* A, const void* B, uint32 PortFlags ) const override;
 	virtual void SerializeItem(FStructuredArchive::FSlot Slot, void* Value, void const* Defaults) const override;
-	virtual void ExportTextItem( FString& ValueStr, const void* PropertyValue, const void* DefaultValue, UObject* Parent, int32 PortFlags, UObject* ExportRootScope ) const override;
-	virtual const TCHAR* ImportText_Internal( const TCHAR* Buffer, void* Data, int32 PortFlags, UObject* OwnerObject, FOutputDevice* ErrorText ) const override;
+	virtual void ExportText_Internal( FString& ValueStr, const void* PropertyValueOrContainer, EPropertyPointerType PropertyPointerType, const void* DefaultValue, UObject* Parent, int32 PortFlags, UObject* ExportRootScope ) const override;
+	virtual const TCHAR* ImportText_Internal(const TCHAR* Buffer, void* ContainerOrPropertyPtr, EPropertyPointerType PropertyPointerType, UObject* OwnerObject, int32 PortFlags, FOutputDevice* ErrorText) const override;
 	virtual void Serialize(FArchive& Ar) override;
 	virtual bool ContainsObjectReference(TArray<const FStructProperty*>& EncounteredStructProps, EPropertyObjectReferenceType InReferenceType = EPropertyObjectReferenceType::Strong) const override;
 	virtual void EmitReferenceInfo(UClass& OwnerClass, int32 BaseOffset, TArray<const FStructProperty*>& EncounteredStructProps, FGCStackSizeHelper& StackSizeHelper) override;
diff --git a/Engine/Source/Runtime/CoreUObject/Public/UObject/Stack.h b/Engine/Source/Runtime/CoreUObject/Public/UObject/Stack.h
--- a/Engine/Source/Runtime/CoreUObject/Public/UObject/Stack.h
+++ b/Engine/Source/Runtime/CoreUObject/Public/UObject/Stack.h
@@ -94,6 +94,7 @@ public:
 
 	FProperty* MostRecentProperty;
 	uint8* MostRecentPropertyAddress;
+	uint8* MostRecentPropertyContainer;
 
 	/** The execution flow stack for compiled Kismet code */
 	FlowStackType FlowStack;
@@ -224,12 +225,13 @@ inline FFrame::FFrame( UObject* InObject, UFunction* InNode, void* InLocals, FFr
 	, Object(InObject)
 	, Code(InNode->Script.GetData())
 	, Locals((uint8*)InLocals)
-	, MostRecentProperty(NULL)
-	, MostRecentPropertyAddress(NULL)
+	, MostRecentProperty(nullptr)
+	, MostRecentPropertyAddress(nullptr)
+	, MostRecentPropertyContainer(nullptr)
 	, PreviousFrame(InPreviousFrame)
 	, OutParms(NULL)
 	, PropertyChainForCompiledIn(InPropertyChainForCompiledIn)
-	, CurrentNativeFunction(NULL)
+	, CurrentNativeFunction(nullptr)
 	, bArrayContextFailed(false)
 #if PER_FUNCTION_SCRIPT_STATS
 	, DepthCounter(0)
@@ -379,7 +381,8 @@ FORCEINLINE_DEBUGGABLE void FFrame::StepCompiledIn(void* Result, const FFieldCla
 template<class TProperty, typename TNativeType>
 FORCEINLINE_DEBUGGABLE TNativeType& FFrame::StepCompiledInRef(void*const TemporaryBuffer)
 {
-	MostRecentPropertyAddress = NULL;
+	MostRecentPropertyAddress = nullptr;
+	MostRecentPropertyContainer = nullptr;
 
 	if (Code)
 	{
diff --git a/Engine/Source/Runtime/CoreUObject/Public/UObject/TextProperty.h b/Engine/Source/Runtime/CoreUObject/Public/UObject/TextProperty.h
--- a/Engine/Source/Runtime/CoreUObject/Public/UObject/TextProperty.h
+++ b/Engine/Source/Runtime/CoreUObject/Public/UObject/TextProperty.h
@@ -49,8 +49,10 @@ PRAGMA_ENABLE_DEPRECATION_WARNINGS
 	virtual EConvertFromTypeResult ConvertFromType(const FPropertyTag& Tag, FStructuredArchive::FSlot Slot, uint8* Data, UStruct* DefaultsStruct) override;
 	virtual bool Identical( const void* A, const void* B, uint32 PortFlags ) const override;
 	virtual void SerializeItem(FStructuredArchive::FSlot Slot, void* Value, void const* Defaults) const override;
-	virtual void ExportTextItem( FString& ValueStr, const void* PropertyValue, const void* DefaultValue, UObject* Parent, int32 PortFlags, UObject* ExportRootScope ) const override;
-	virtual const TCHAR* ImportText_Internal( const TCHAR* Buffer, void* Data, int32 PortFlags, UObject* OwnerObject, FOutputDevice* ErrorText ) const override;
+protected:
+	virtual void ExportText_Internal( FString& ValueStr, const void* PropertyValueOrContainer, EPropertyPointerType PropertyPointerType, const void* DefaultValue, UObject* Parent, int32 PortFlags, UObject* ExportRootScope ) const override;
+	virtual const TCHAR* ImportText_Internal(const TCHAR* Buffer, void* ContainerOrPropertyPtr, EPropertyPointerType PropertyPointerType, UObject* OwnerObject, int32 PortFlags, FOutputDevice* ErrorText) const override;
+public:
 	virtual FString GetCPPTypeForwardDeclaration() const override;
 	// End of FProperty interface
 
