Added support for native FProperty setters and getters.

diff --git a/Engine/Source/Runtime/CoreUObject/Public/UObject/UnrealType.h b/Engine/Source/Runtime/CoreUObject/Public/UObject/UnrealType.h
--- a/Engine/Source/Runtime/CoreUObject/Public/UObject/UnrealType.h
+++ b/Engine/Source/Runtime/CoreUObject/Public/UObject/UnrealType.h
@@ -172,6 +172,13 @@ public:
 	}
 };
 
+/** Type of pointer provided for property API functions */
+enum class EPropertyPointerType
+{
+	Direct = 0, /** Raw property access */
+	Container = 1, /** Property access through its owner container */
+};
+
 //
 // An UnrealScript variable.
 //
@@ -281,6 +288,41 @@ public:
 	/** Gets the wrapper object for this property or creates one if it doesn't exist yet */
 	UPropertyWrapper* GetUPropertyWrapper();
 #endif
+
+	/** Checks if this property as a native setter function */
+	virtual bool HasSetter() const
+	{
+		return false;
+	}
+	/** Checks if this property as a native getter function */
+	virtual bool HasGetter() const
+	{
+		return false;
+	}
+	/** Checks if this property as a native setter or getter function */
+	virtual bool HasSetterOrGetter() const
+	{
+		return false;
+	}
+	/** 
+	 * Calls the native setter function for this property
+	 * @param Container Pointer to the owner of this property (either UObject or struct)
+	 * @param InValue Pointer to the new value
+	 */
+	virtual void CallSetter(void* Container, const void* InValue) const 
+	{
+		checkf(HasSetter(), TEXT("Calling a setter on %s but it doesn't have one"), *GetFullName());
+	}
+	/**
+	 * Calls the native getter function for this property
+	 * @param Container Pointer to the owner of this property (either UObject or struct)
+	 * @param OutValue Pointer to the value where the existing property value will be copied to
+	 */
+	virtual void CallGetter(const void* Container, void* OutValue) const 
+	{
+		checkf(HasGetter(), TEXT("Calling a getter on %s but it doesn't have one"), *GetFullName());
+	}
+
 private:
 	/** Set the alignment offset for this property 
 	 * @return the size of the structure including this newly added property
@@ -434,15 +476,88 @@ public:
 	virtual void SerializeItem(FStructuredArchive::FSlot Slot, void* Value, void const* Defaults = NULL) const PURE_VIRTUAL(FProperty::SerializeItem, );
 	virtual bool NetSerializeItem( FArchive& Ar, UPackageMap* Map, void* Data, TArray<uint8> * MetaData = NULL ) const;
 	virtual bool SupportsNetSharedSerialization() const;
-	virtual void ExportTextItem( FString& ValueStr, const void* PropertyValue, const void* DefaultValue, UObject* Parent, int32 PortFlags, UObject* ExportRootScope = NULL ) const PURE_VIRTUAL(FProperty::ExportTextItem,);
+
+	UE_DEPRECATED(5.1, "Please use ExportTextItem_InContainer or ExportTextItem_Direct instead.")
+	virtual void ExportTextItem(FString& ValueStr, const void* PropertyValue, const void* DefaultValue, UObject* Parent, int32 PortFlags, UObject* ExportRootScope = nullptr) const
+	{
+		ExportTextItem_Direct(ValueStr, PropertyValue, DefaultValue, Parent, PortFlags, ExportRootScope);
+	}	
+
+	void ExportTextItem_Direct(FString& ValueStr, const void* PropertyValue, const void* DefaultValue, UObject* Parent, int32 PortFlags, UObject* ExportRootScope = nullptr) const
+	{
+		ExportText_Internal(ValueStr, PropertyValue, EPropertyPointerType::Direct, DefaultValue, Parent, PortFlags, ExportRootScope);
+	}
+
+	void ExportTextItem_InContainer(FString& ValueStr, const void* Container, const void* DefaultValue, UObject* Parent, int32 PortFlags, UObject* ExportRootScope = nullptr) const
+	{
+		ExportText_Internal(ValueStr, Container, EPropertyPointerType::Container, DefaultValue, Parent, PortFlags, ExportRootScope);
+	}
+
+	UE_DEPRECATED(5.1, "ImportText that takes a direct property pointer is deprecated.  Please use ImportText_Direct or ImportText_InContainer instead.")
 	const TCHAR* ImportText( const TCHAR* Buffer, void* Data, int32 PortFlags, UObject* OwnerObject, FOutputDevice* ErrorText = (FOutputDevice*)GWarn ) const
 	{
-		if ( !ValidateImportFlags(PortFlags,ErrorText) || Buffer == NULL )
+		return ImportText_Direct(Buffer, Data, OwnerObject, PortFlags, ErrorText);
+	}
+
+	/**
+	 * Import a text value
+	 * @param Buffer		Text representing the property value
+	 * @param Container	Pointer to the container that owns this property (either UObject pointer or a struct pointer)
+	 * @param OwnerObject	Object that owns the property container (if the container is an UObject then Container is also OwnerObject)
+	 * @param PortFlags	Flags controlling the behavior when importing the value
+	 * @param ErrorText	Output device for throwing warnings or errors on import
+	 * @returns Buffer pointer advanced by the number of characters consumed when reading the text value
+	 */
+	const TCHAR* ImportText_InContainer(const TCHAR* Buffer, void* Container, UObject* OwnerObject, int32 PortFlags, FOutputDevice* ErrorText = (FOutputDevice*)GWarn) const
+	{
+		if (!ValidateImportFlags(PortFlags, ErrorText) || Buffer == nullptr)
+		{
+			return nullptr;
+		}
+		PortFlags |= EPropertyPortFlags::PPF_UseDeprecatedProperties; // Imports should always process deprecated properties
+		return ImportText_Internal(Buffer, Container, EPropertyPointerType::Container, OwnerObject, PortFlags, ErrorText);
+	}
+
+	/**
+	 * Import a text value
+	 * @param Buffer		Text representing the property value
+	 * @param PropertyPtr	Pointer to property value
+	 * @param OwnerObject	Object that owns the property
+	 * @param PortFlags	Flags controlling the behavior when importing the value
+	 * @param ErrorText	Output device for throwing warnings or errors on import
+	 * @returns Buffer pointer advanced by the number of characters consumed when reading the text value
+	 */
+	const TCHAR* ImportText_Direct(const TCHAR* Buffer, void* PropertyPtr, UObject* OwnerObject, int32 PortFlags, FOutputDevice* ErrorText = (FOutputDevice*)GWarn) const
+	{
+		if (!ValidateImportFlags(PortFlags, ErrorText) || Buffer == NULL)
 		{
 			return NULL;
 		}
 		PortFlags |= EPropertyPortFlags::PPF_UseDeprecatedProperties; // Imports should always process deprecated properties
-		return ImportText_Internal( Buffer, Data, PortFlags, OwnerObject, ErrorText );
+		return ImportText_Internal(Buffer, PropertyPtr, EPropertyPointerType::Direct, OwnerObject, PortFlags, ErrorText);
+	}
+
+	FORCENOINLINE void SetValue_InContainer(void* InContainer, const void* InValue) const
+	{
+		if (!HasSetter())
+		{
+			CopyCompleteValue(ContainerVoidPtrToValuePtrInternal(InContainer, 0), InValue);
+		}
+		else
+		{
+			CallSetter(InContainer, InValue);
+		}
+	}
+	FORCENOINLINE void GetValue_InContainer(void const* InContainer, void* OutValue) const
+	{
+		if (!HasGetter())
+		{
+			CopyCompleteValue(OutValue, ContainerVoidPtrToValuePtrInternal((void*)InContainer, 0));
+		}
+		else
+		{
+			CallGetter(InContainer, OutValue);
+		}
 	}
 
 #if WITH_EDITORONLY_DATA
@@ -460,11 +575,14 @@ public:
 	virtual void AppendSchemaHash(FBlake3& Builder, bool bSkipEditorOnly) const;
 #endif
 protected:
-	virtual const TCHAR* ImportText_Internal( const TCHAR* Buffer, void* Data, int32 PortFlags, UObject* OwnerObject, FOutputDevice* ErrorText ) const PURE_VIRTUAL(FProperty::ImportText,return NULL;);
+
+	virtual void ExportText_Internal(FString& ValueStr, const void* PropertyValueOrContainer, EPropertyPointerType PointerType, const void* DefaultValue, UObject* Parent, int32 PortFlags, UObject* ExportRootScope = nullptr) const PURE_VIRTUAL(FProperty::ExportText, );
+	virtual const TCHAR* ImportText_Internal(const TCHAR* Buffer, void* ContainerOrPropertyPtr, EPropertyPointerType PointerType, UObject* OwnerObject, int32 PortFlags, FOutputDevice* ErrorText) const PURE_VIRTUAL(FProperty::ImportText, return nullptr;);
+
 public:
 	
-	bool ExportText_Direct( FString& ValueStr, const void* Data, const void* Delta, UObject* Parent, int32 PortFlags, UObject* ExportRootScope = NULL ) const;
-	FORCEINLINE bool ExportText_InContainer( int32 Index, FString& ValueStr, const void* Data, const void* Delta, UObject* Parent, int32 PortFlags, UObject* ExportRootScope = NULL ) const
+	bool ExportText_Direct(FString& ValueStr, const void* Data, const void* Delta, UObject* Parent, int32 PortFlags, UObject* ExportRootScope = nullptr) const;
+	FORCEINLINE bool ExportText_InContainer(int32 Index, FString& ValueStr, const void* Data, const void* Delta, UObject* Parent, int32 PortFlags, UObject* ExportRootScope = nullptr) const
 	{
 		return ExportText_Direct(ValueStr, ContainerPtrToValuePtr<void>(Data, Index), ContainerPtrToValuePtrForDefaults<void>(NULL, Delta, Index), Parent, PortFlags, ExportRootScope);
 	}
@@ -512,6 +630,20 @@ private:
 		return (uint8*)ContainerPtr + Offset_Internal + ElementSize * ArrayIndex;
 	}
 
+protected:
+
+	FORCEINLINE void* PointerToValuePtr(void const* ContainerOrPropertyPtr, EPropertyPointerType PropertyPointerType, int32 ArrayIndex = 0) const
+	{
+		if (PropertyPointerType == EPropertyPointerType::Container)
+		{
+			return (uint8*)ContainerOrPropertyPtr + Offset_Internal + ElementSize * ArrayIndex;
+		}
+		else
+		{
+			return (void*)ContainerOrPropertyPtr;
+		}
+	}
+
 public:
 
 	/** 
@@ -1265,6 +1397,16 @@ public:
 		*GetPropertyValuePtr_InContainer(A, ArrayIndex) = Value;
 	}
 
+	FORCENOINLINE void SetValue_InContainer(void* InContainer, const TCppType& InValue) const
+	{
+		TInPropertyBaseClass::SetValue_InContainer(InContainer, &InValue);
+	}
+
+	FORCENOINLINE void GetValue_InContainer(void const* InContainer, TCppType* OutValue) const
+	{
+		TInPropertyBaseClass::GetValue_InContainer(InContainer, OutValue);
+	}
+
 protected:
 	FORCEINLINE void SetElementSize()
 	{
@@ -1366,9 +1508,8 @@ PRAGMA_ENABLE_DEPRECATION_WARNINGS
 #endif // WITH_EDITORONLY_DATA
 
 	// FProperty interface.
-	virtual const TCHAR* ImportText_Internal( const TCHAR* Buffer, void* Data, int32 PortFlags, UObject* Parent, FOutputDevice* ErrorText ) const override;
-
-	virtual void ExportTextItem( FString& ValueStr, const void* PropertyValue, const void* DefaultValue, UObject* Parent, int32 PortFlags, UObject* ExportRootScope ) const override;
+	virtual const TCHAR* ImportText_Internal(const TCHAR* Buffer, void* ContainerOrPropertyPtr, EPropertyPointerType PropertyPointerType, UObject* Parent, int32 PortFlags, FOutputDevice* ErrorText) const override;
+	virtual void ExportText_Internal( FString& ValueStr, const void* PropertyValueOrContainer, EPropertyPointerType PropertyPointerType, const void* DefaultValue, UObject* Parent, int32 PortFlags, UObject* ExportRootScope ) const override;
 	// End of FProperty interface
 
 	// FNumericProperty interface.
@@ -1434,6 +1575,7 @@ PRAGMA_ENABLE_DEPRECATION_WARNINGS
 	 * CAUTION: This routine does not do enum name conversion
 	**/
 	virtual void SetNumericPropertyValueFromString(void* Data, TCHAR const* Value) const;
+	virtual void SetNumericPropertyValueFromString_InContainer(void* Container, TCHAR const* Value) const;
 
 	/** 
 	 * Gets the value of a signed integral property type
@@ -1441,6 +1583,7 @@ PRAGMA_ENABLE_DEPRECATION_WARNINGS
 	 * @return Data as a signed int
 	**/
 	virtual int64 GetSignedIntPropertyValue(void const* Data) const;
+	virtual int64 GetSignedIntPropertyValue_InContainer(void const* Container) const;
 
 	/** 
 	 * Gets the value of an unsigned integral property type
@@ -1463,6 +1606,7 @@ PRAGMA_ENABLE_DEPRECATION_WARNINGS
 	 * CAUTION: This routine does not do enum name conversion
 	**/
 	virtual FString GetNumericPropertyValueToString(void const* Data) const;
+	virtual FString GetNumericPropertyValueToString_InContainer(void const* Container) const;
 	// End of FNumericProperty interface
 
 	static int64 ReadEnumAsInt64(FStructuredArchive::FSlot Slot, UStruct* DefaultsStruct, const FPropertyTag& Tag);
@@ -1670,15 +1814,34 @@ public:
 	{
 		LexFromString(*TTypeFundamentals::GetPropertyValuePtr(Data), Value);
 	}
+	virtual void SetNumericPropertyValueFromString_InContainer(void* Container, TCHAR const* Value) const override
+	{
+		TCppType LocalValue{};
+		LexFromString(LocalValue, Value);
+		FNumericProperty::SetValue_InContainer(Container, &LocalValue);
+	}
 	virtual FString GetNumericPropertyValueToString(void const* Data) const override
 	{
 		return LexToString(TTypeFundamentals::GetPropertyValue(Data));
 	}
+	virtual FString GetNumericPropertyValueToString_InContainer(void const* Container) const override
+	{
+		TCppType LocalValue{};
+		FNumericProperty::GetValue_InContainer(Container, &LocalValue);
+		return LexToString(LocalValue);
+	}
 	virtual int64 GetSignedIntPropertyValue(void const* Data) const override
 	{
 		check(TIsIntegral<TCppType>::Value);
 		return (int64)TTypeFundamentals::GetPropertyValue(Data);
 	}
+	virtual int64 GetSignedIntPropertyValue_InContainer(void const* Container) const override
+	{
+		check(TIsIntegral<TCppType>::Value);
+		TCppType LocalValue{};
+		FNumericProperty::GetValue_InContainer(Container, &LocalValue);
+		return (int64)LocalValue;
+	}
 	virtual uint64 GetUnsignedIntPropertyValue(void const* Data) const override
 	{
 		check(TIsIntegral<TCppType>::Value);
@@ -1764,8 +1927,10 @@ PRAGMA_ENABLE_DEPRECATION_WARNINGS
 	// FProperty interface.
 	virtual void SerializeItem(FStructuredArchive::FSlot Slot, void* Value, void const* Defaults) const override;
 	virtual bool NetSerializeItem( FArchive& Ar, UPackageMap* Map, void* Data, TArray<uint8> * MetaData = NULL ) const override;
-	virtual void ExportTextItem( FString& ValueStr, const void* PropertyValue, const void* DefaultValue, UObject* Parent, int32 PortFlags, UObject* ExportRootScope ) const override;
-	virtual const TCHAR* ImportText_Internal( const TCHAR* Buffer, void* Data, int32 PortFlags, UObject* Parent, FOutputDevice* ErrorText ) const override;
+protected:
+	virtual void ExportText_Internal( FString& ValueStr, const void* PropertyValueOrContainer, EPropertyPointerType PropertyPointerType, const void* DefaultValue, UObject* Parent, int32 PortFlags, UObject* ExportRootScope ) const override;
+	virtual const TCHAR* ImportText_Internal(const TCHAR* Buffer, void* ContainerOrPropertyPtr, EPropertyPointerType PropertyPointerType, UObject* Parent, int32 PortFlags, FOutputDevice* ErrorText) const override;
+public:
 	virtual EConvertFromTypeResult ConvertFromType(const FPropertyTag& Tag, FStructuredArchive::FSlot Slot, uint8* Data, UStruct* DefaultsStruct) override;
 #if WITH_EDITORONLY_DATA
 	virtual void AppendSchemaHash(FBlake3& Builder, bool bSkipEditorOnly) const override;
@@ -2116,7 +2281,8 @@ PRAGMA_ENABLE_DEPRECATION_WARNINGS
 #endif // WITH_EDITORONLY_DATA
 
 	// FProperty interface
-	virtual void ExportTextItem(FString& ValueStr, const void* PropertyValue, const void* DefaultValue, UObject* Parent, int32 PortFlags, UObject* ExportRootScope) const override;
+protected:
+	virtual void ExportText_Internal(FString& ValueStr, const void* PropertyValueOrContainer, EPropertyPointerType PropertyPointerType, const void* DefaultValue, UObject* Parent, int32 PortFlags, UObject* ExportRootScope) const override;
 	// End of FProperty interface
 };
 
@@ -2240,8 +2406,10 @@ public:
 	virtual bool Identical( const void* A, const void* B, uint32 PortFlags ) const override;
 	virtual void SerializeItem( FStructuredArchive::FSlot Slot, void* Value, void const* Defaults) const override;
 	virtual bool NetSerializeItem( FArchive& Ar, UPackageMap* Map, void* Data, TArray<uint8> * MetaData = NULL ) const override;
-	virtual void ExportTextItem( FString& ValueStr, const void* PropertyValue, const void* DefaultValue, UObject* Parent, int32 PortFlags, UObject* ExportRootScope ) const override;
-	virtual const TCHAR* ImportText_Internal( const TCHAR* Buffer, void* Data, int32 PortFlags, UObject* Parent, FOutputDevice* ErrorText ) const override;
+protected:
+	virtual void ExportText_Internal( FString& ValueStr, const void* PropertyValueOrContainer, EPropertyPointerType PropertyPointerType, const void* DefaultValue, UObject* Parent, int32 PortFlags, UObject* ExportRootScope ) const override;
+	virtual const TCHAR* ImportText_Internal(const TCHAR* Buffer, void* ContainerOrPropertyPtr, EPropertyPointerType PropertyPointerType, UObject* Parent, int32 PortFlags, FOutputDevice* ErrorText) const override;
+public:
 	virtual void CopyValuesInternal( void* Dest, void const* Src, int32 Count ) const override;
 	virtual void ClearValueInternal( void* Data ) const override;
 	virtual void InitializeValueInternal( void* Dest ) const override;
@@ -2368,8 +2536,10 @@ PRAGMA_ENABLE_DEPRECATION_WARNINGS
 	virtual bool Identical( const void* A, const void* B, uint32 PortFlags ) const override;
 	virtual bool NetSerializeItem( FArchive& Ar, UPackageMap* Map, void* Data, TArray<uint8> * MetaData = NULL ) const override;
 	virtual bool SupportsNetSharedSerialization() const override { return false; }
-	virtual void ExportTextItem( FString& ValueStr, const void* PropertyValue, const void* DefaultValue, UObject* Parent, int32 PortFlags, UObject* ExportRootScope ) const override;
-	virtual const TCHAR* ImportText_Internal( const TCHAR* Buffer, void* Data, int32 PortFlags, UObject* OwnerObject, FOutputDevice* ErrorText ) const override;
+protected:
+	virtual void ExportText_Internal( FString& ValueStr, const void* PropertyValueOrContainer, EPropertyPointerType PropertyPointerType, const void* DefaultValue, UObject* Parent, int32 PortFlags, UObject* ExportRootScope ) const override;
+	virtual const TCHAR* ImportText_Internal(const TCHAR* Buffer, void* ContainerOrPropertyPtr, EPropertyPointerType PropertyPointerType, UObject* OwnerObject, int32 PortFlags, FOutputDevice* ErrorText) const override;
+public:
 	virtual FName GetID() const override;
 	virtual void InstanceSubobjects( void* Data, void const* DefaultData, UObject* Owner, struct FObjectInstancingGraph* InstanceGraph ) override;
 	virtual bool SameType(const FProperty* Other) const override;
@@ -2607,7 +2777,7 @@ PRAGMA_ENABLE_DEPRECATION_WARNINGS
 	// FProperty interface
 	virtual void SerializeItem(FStructuredArchive::FSlot Slot, void* Value, void const* Defaults) const override;
 	virtual void EmitReferenceInfo(UClass& OwnerClass, int32 BaseOffset, TArray<const FStructProperty*>& EncounteredStructProps, FGCStackSizeHelper& StackSizeHelper) override;
-	virtual const TCHAR* ImportText_Internal(const TCHAR* Buffer, void* Data, int32 PortFlags, UObject* OwnerObject, FOutputDevice* ErrorText) const override;
+	virtual const TCHAR* ImportText_Internal(const TCHAR* Buffer, void* ContainerOrPropertyPtr, EPropertyPointerType PropertyPointerType, UObject* OwnerObject, int32 PortFlags, FOutputDevice* ErrorText) const override;
 	virtual EConvertFromTypeResult ConvertFromType(const FPropertyTag& Tag, FStructuredArchive::FSlot Slot, uint8* Data, UStruct* DefaultsStruct) override;
 
 private:
@@ -2829,8 +2999,10 @@ PRAGMA_ENABLE_DEPRECATION_WARNINGS
 	virtual bool Identical( const void* A, const void* B, uint32 PortFlags ) const override;
 	virtual void SerializeItem( FStructuredArchive::FSlot Slot, void* Value, void const* Defaults) const override;
 	virtual bool NetSerializeItem(FArchive& Ar, UPackageMap* Map, void* Data, TArray<uint8> * MetaData = NULL) const override;
-	virtual void ExportTextItem( FString& ValueStr, const void* PropertyValue, const void* DefaultValue, UObject* Parent, int32 PortFlags, UObject* ExportRootScope ) const override;
-	virtual const TCHAR* ImportText_Internal( const TCHAR* Buffer, void* Data, int32 PortFlags, UObject* OwnerObject, FOutputDevice* ErrorText ) const override;
+protected:
+	virtual void ExportText_Internal( FString& ValueStr, const void* PropertyValueOrContainer, EPropertyPointerType PropertyPointerType, const void* DefaultValue, UObject* Parent, int32 PortFlags, UObject* ExportRootScope ) const override;
+	virtual const TCHAR* ImportText_Internal(const TCHAR* Buffer, void* ContainerOrPropertyPtr, EPropertyPointerType PropertyPointerType, UObject* OwnerObject, int32 PortFlags, FOutputDevice* ErrorText) const override;
+public:
 	virtual EConvertFromTypeResult ConvertFromType(const FPropertyTag& Tag, FStructuredArchive::FSlot Slot, uint8* Data, UStruct* DefaultsStruct) override;
 	virtual void EmitReferenceInfo(UClass& OwnerClass, int32 BaseOffset, TArray<const FStructProperty*>& EncounteredStructProps, FGCStackSizeHelper& StackSizeHelper) override;
 	// End of FProperty interface
@@ -2924,7 +3096,7 @@ PRAGMA_ENABLE_DEPRECATION_WARNINGS
 	// End of UHT interface
 
 	// FProperty interface
-	virtual const TCHAR* ImportText_Internal( const TCHAR* Buffer, void* Data, int32 PortFlags, UObject* OwnerObject, FOutputDevice* ErrorText ) const override;
+	virtual const TCHAR* ImportText_Internal(const TCHAR* Buffer, void* ContainerOrPropertyPtr, EPropertyPointerType PropertyPointerType, UObject* OwnerObject, int32 PortFlags, FOutputDevice* ErrorText) const override;
 	virtual bool SameType(const FProperty* Other) const override;
 	virtual bool Identical( const void* A, const void* B, uint32 PortFlags ) const override;
 #if WITH_EDITORONLY_DATA
@@ -3125,8 +3297,10 @@ PRAGMA_ENABLE_DEPRECATION_WARNINGS
 	virtual void SerializeItem( FStructuredArchive::FSlot Slot, void* Value, void const* Defaults) const override;
 	virtual bool NetSerializeItem( FArchive& Ar, UPackageMap* Map, void* Data, TArray<uint8> * MetaData = NULL ) const override;
 	virtual bool SupportsNetSharedSerialization() const override { return false; }
-	virtual void ExportTextItem( FString& ValueStr, const void* PropertyValue, const void* DefaultValue, UObject* Parent, int32 PortFlags, UObject* ExportRootScope ) const override;
-	virtual const TCHAR* ImportText_Internal( const TCHAR* Buffer, void* Data, int32 PortFlags, UObject* OwnerObject, FOutputDevice* ErrorText ) const override;
+protected:
+	virtual void ExportText_Internal( FString& ValueStr, const void* PropertyValueOrContainer, EPropertyPointerType PropertyPointerType, const void* DefaultValue, UObject* Parent, int32 PortFlags, UObject* ExportRootScope ) const override;
+	virtual const TCHAR* ImportText_Internal(const TCHAR* Buffer, void* ContainerOrPropertyPtr, EPropertyPointerType PropertyPointerType, UObject* OwnerObject, int32 PortFlags, FOutputDevice* ErrorText) const override;
+public:
 	virtual bool ContainsObjectReference(TArray<const FStructProperty*>& EncounteredStructProps, EPropertyObjectReferenceType InReferenceType = EPropertyObjectReferenceType::Strong) const override;
 	virtual bool SameType(const FProperty* Other) const override;
 #if WITH_EDITORONLY_DATA
@@ -3203,8 +3377,10 @@ PRAGMA_ENABLE_DEPRECATION_WARNINGS
 #endif // WITH_EDITORONLY_DATA
 
 	// FProperty interface
-	virtual void ExportTextItem( FString& ValueStr, const void* PropertyValue, const void* DefaultValue, UObject* Parent, int32 PortFlags, UObject* ExportRootScope ) const override;
-	virtual const TCHAR* ImportText_Internal( const TCHAR* Buffer, void* Data, int32 PortFlags, UObject* OwnerObject, FOutputDevice* ErrorText ) const override;
+protected:
+	virtual void ExportText_Internal( FString& ValueStr, const void* PropertyValueOrContainer, EPropertyPointerType PropertyPointerType, const void* DefaultValue, UObject* Parent, int32 PortFlags, UObject* ExportRootScope ) const override;
+	virtual const TCHAR* ImportText_Internal(const TCHAR* Buffer, void* ContainerOrPropertyPtr, EPropertyPointerType PropertyPointerType, UObject* OwnerObject, int32 PortFlags, FOutputDevice* ErrorText) const override;
+public:
 	virtual EConvertFromTypeResult ConvertFromType(const FPropertyTag& Tag, FStructuredArchive::FSlot Slot, uint8* Data, UStruct* DefaultsStruct) override;
 	virtual FString GetCPPTypeForwardDeclaration() const override;
 	uint32 GetValueTypeHashInternal(const void* Src) const override;
@@ -3257,8 +3433,10 @@ PRAGMA_ENABLE_DEPRECATION_WARNINGS
 #endif // WITH_EDITORONLY_DATA
 
 	// FProperty interface
-	virtual void ExportTextItem( FString& ValueStr, const void* PropertyValue, const void* DefaultValue, UObject* Parent, int32 PortFlags, UObject* ExportRootScope ) const override;
-	virtual const TCHAR* ImportText_Internal( const TCHAR* Buffer, void* Data, int32 PortFlags, UObject* OwnerObject, FOutputDevice* ErrorText ) const override;
+protected:
+	virtual void ExportText_Internal( FString& ValueStr, const void* PropertyValueOrContainer, EPropertyPointerType PropertyPointerType, const void* DefaultValue, UObject* Parent, int32 PortFlags, UObject* ExportRootScope ) const override;
+	virtual const TCHAR* ImportText_Internal(const TCHAR* Buffer, void* ContainerOrPropertyPtr, EPropertyPointerType PropertyPointerType, UObject* OwnerObject, int32 PortFlags, FOutputDevice* ErrorText) const override;
+public:
 	virtual EConvertFromTypeResult ConvertFromType(const FPropertyTag& Tag, FStructuredArchive::FSlot Slot, uint8* Data, UStruct* DefaultsStruct) override;
 	virtual FString GetCPPTypeForwardDeclaration() const override;
 	uint32 GetValueTypeHashInternal(const void* Src) const override;
@@ -3363,8 +3541,10 @@ PRAGMA_ENABLE_DEPRECATION_WARNINGS
 	virtual bool Identical( const void* A, const void* B, uint32 PortFlags ) const override;
 	virtual void SerializeItem(FStructuredArchive::FSlot Slot, void* Value, void const* Defaults) const override;
 	virtual bool NetSerializeItem( FArchive& Ar, UPackageMap* Map, void* Data, TArray<uint8> * MetaData = NULL ) const override;
-	virtual void ExportTextItem( FString& ValueStr, const void* PropertyValue, const void* DefaultValue, UObject* Parent, int32 PortFlags, UObject* ExportRootScope ) const override;
-	virtual const TCHAR* ImportText_Internal( const TCHAR* Buffer, void* Data, int32 PortFlags, UObject* OwnerObject, FOutputDevice* ErrorText ) const override;
+protected:
+	virtual void ExportText_Internal( FString& ValueStr, const void* PropertyValueOrContainer, EPropertyPointerType PropertyPointerType, const void* DefaultValue, UObject* Parent, int32 PortFlags, UObject* ExportRootScope ) const override;
+	virtual const TCHAR* ImportText_Internal(const TCHAR* Buffer, void* ContainerOrPropertyPtr, EPropertyPointerType PropertyPointerType, UObject* OwnerObject, int32 PortFlags, FOutputDevice* ErrorText) const override;
+public:
 	virtual void InitializeValueInternal(void* Dest) const override
 	{
 		if (EnumHasAnyFlags(ArrayFlags, EArrayPropertyFlags::UsesMemoryImageAllocator))
@@ -3403,7 +3583,7 @@ PRAGMA_ENABLE_DEPRECATION_WARNINGS
 
 	FString GetCPPTypeCustom(FString* ExtendedTypeText, uint32 CPPExportFlags, const FString& InnerTypeText, const FString& InInnerExtendedTypeText) const;
 
-	/** Called by ExportTextItem, but can also be used by a non-ArrayProperty whose ArrayDim is > 1. */
+	/** Called by ExportText_Internal, but can also be used by a non-ArrayProperty whose ArrayDim is > 1. */
 	static void ExportTextInnerItem(FString& ValueStr, const FProperty* Inner, const void* PropertyValue, int32 PropertySize, const void* DefaultValue, int32 DefaultSize, UObject* Parent = nullptr, int32 PortFlags = 0, UObject* ExportRootScope = nullptr);
 
 	/** Called by ImportTextItem, but can also be used by a non-ArrayProperty whose ArrayDim is > 1. ArrayHelper should be supplied by ArrayProperties and nullptr for fixed-size arrays. */
@@ -3495,8 +3675,10 @@ public:
 	virtual bool Identical(const void* A, const void* B, uint32 PortFlags) const override;
 	virtual void SerializeItem(FStructuredArchive::FSlot Slot, void* Value, void const* Defaults) const override;
 	virtual bool NetSerializeItem(FArchive& Ar, UPackageMap* Map, void* Data, TArray<uint8> * MetaData = NULL) const override;
-	virtual void ExportTextItem(FString& ValueStr, const void* PropertyValue, const void* DefaultValue, UObject* Parent, int32 PortFlags, UObject* ExportRootScope) const override;
-	virtual const TCHAR* ImportText_Internal(const TCHAR* Buffer, void* Data, int32 PortFlags, UObject* OwnerObject, FOutputDevice* ErrorText) const override;
+protected:
+	virtual void ExportText_Internal(FString& ValueStr, const void* PropertyValueOrContainer, EPropertyPointerType PropertyPointerType, const void* DefaultValue, UObject* Parent, int32 PortFlags, UObject* ExportRootScope) const override;
+	virtual const TCHAR* ImportText_Internal(const TCHAR* Buffer, void* ContainerOrPropertyPtr, EPropertyPointerType PropertyPointerType, UObject* OwnerObject, int32 PortFlags, FOutputDevice* ErrorText) const override;
+public:
 	virtual void InitializeValueInternal(void* Dest) const override
 	{
 		if (EnumHasAnyFlags(MapFlags, EMapPropertyFlags::UsesMemoryImageAllocator))
@@ -3635,8 +3817,10 @@ public:
 	virtual bool Identical(const void* A, const void* B, uint32 PortFlags) const override;
 	virtual void SerializeItem(FStructuredArchive::FSlot Slot, void* Value, void const* Defaults) const override;
 	virtual bool NetSerializeItem(FArchive& Ar, UPackageMap* Map, void* Data, TArray<uint8> * MetaData = NULL) const override;
-	virtual void ExportTextItem(FString& ValueStr, const void* PropertyValue, const void* DefaultValue, UObject* Parent, int32 PortFlags, UObject* ExportRootScope) const override;
-	virtual const TCHAR* ImportText_Internal(const TCHAR* Buffer, void* Data, int32 PortFlags, UObject* OwnerObject, FOutputDevice* ErrorText) const override;
+protected:
+	virtual void ExportText_Internal(FString& ValueStr, const void* PropertyValueOrContainer, EPropertyPointerType PropertyPointerType, const void* DefaultValue, UObject* Parent, int32 PortFlags, UObject* ExportRootScope) const override;
+	virtual const TCHAR* ImportText_Internal(const TCHAR* Buffer, void* ContainerOrPropertyPtr, EPropertyPointerType PropertyPointerType, UObject* OwnerObject, int32 PortFlags, FOutputDevice* ErrorText) const override;
+public:
 	virtual void CopyValuesInternal(void* Dest, void const* Src, int32 Count) const override;
 	virtual void ClearValueInternal(void* Data) const override;
 	virtual void DestroyValueInternal(void* Dest) const override;
@@ -5466,8 +5650,10 @@ public:
 	virtual void SerializeItem(FStructuredArchive::FSlot Slot, void* Value, void const* Defaults) const override;
 	virtual bool NetSerializeItem( FArchive& Ar, UPackageMap* Map, void* Data, TArray<uint8> * MetaData = NULL ) const override;
 	virtual bool SupportsNetSharedSerialization() const override;
-	virtual void ExportTextItem( FString& ValueStr, const void* PropertyValue, const void* DefaultValue, UObject* Parent, int32 PortFlags, UObject* ExportRootScope ) const override;
-	virtual const TCHAR* ImportText_Internal( const TCHAR* Buffer, void* Data, int32 PortFlags, UObject* OwnerObject, FOutputDevice* ErrorText ) const override;
+protected:
+	virtual void ExportText_Internal( FString& ValueStr, const void* PropertyValueOrContainer, EPropertyPointerType PropertyPointerType, const void* DefaultValue, UObject* Parent, int32 PortFlags, UObject* ExportRootScope ) const override;
+	virtual const TCHAR* ImportText_Internal(const TCHAR* Buffer, void* ContainerOrPropertyPtr, EPropertyPointerType PropertyPointerType, UObject* OwnerObject, int32 PortFlags, FOutputDevice* ErrorText) const override;
+public:
 	virtual void CopyValuesInternal( void* Dest, void const* Src, int32 Count  ) const override;
 	virtual void ClearValueInternal( void* Data ) const override;
 	virtual void DestroyValueInternal( void* Dest ) const override;
@@ -5563,8 +5749,10 @@ PRAGMA_ENABLE_DEPRECATION_WARNINGS
 	virtual bool Identical( const void* A, const void* B, uint32 PortFlags ) const override;
 	virtual void SerializeItem(FStructuredArchive::FSlot Slot, void* Value, void const* Defaults) const override;
 	virtual bool NetSerializeItem( FArchive& Ar, UPackageMap* Map, void* Data, TArray<uint8> * MetaData = NULL ) const override;
-	virtual void ExportTextItem( FString& ValueStr, const void* PropertyValue, const void* DefaultValue, UObject* Parent, int32 PortFlags, UObject* ExportRootScope ) const override;
-	virtual const TCHAR* ImportText_Internal( const TCHAR* Buffer, void* Data, int32 PortFlags, UObject* OwnerObject, FOutputDevice* ErrorText ) const override;
+protected:
+	virtual void ExportText_Internal( FString& ValueStr, const void* PropertyValueOrContainer, EPropertyPointerType PropertyPointerType, const void* DefaultValue, UObject* Parent, int32 PortFlags, UObject* ExportRootScope ) const override;
+	virtual const TCHAR* ImportText_Internal(const TCHAR* Buffer, void* ContainerOrPropertyPtr, EPropertyPointerType PropertyPointerType, UObject* OwnerObject, int32 PortFlags, FOutputDevice* ErrorText) const override;
+public:
 	virtual bool ContainsObjectReference(TArray<const FStructProperty*>& EncounteredStructProps, EPropertyObjectReferenceType InReferenceType = EPropertyObjectReferenceType::Strong) const override;
 	virtual void EmitReferenceInfo(UClass& OwnerClass, int32 BaseOffset, TArray<const FStructProperty*>& EncounteredStructProps, FGCStackSizeHelper& StackSizeHelper) override;
 	virtual void InstanceSubobjects( void* Data, void const* DefaultData, UObject* Owner, struct FObjectInstancingGraph* InstanceGraph ) override;
@@ -5632,7 +5820,9 @@ PRAGMA_ENABLE_DEPRECATION_WARNINGS
 	virtual FString GetCPPTypeForwardDeclaration() const override;
 	virtual bool Identical( const void* A, const void* B, uint32 PortFlags ) const override;
 	virtual bool NetSerializeItem( FArchive& Ar, UPackageMap* Map, void* Data, TArray<uint8> * MetaData = NULL ) const override;
-	virtual void ExportTextItem( FString& ValueStr, const void* PropertyValue, const void* DefaultValue, UObject* Parent, int32 PortFlags, UObject* ExportRootScope ) const override;
+protected:
+	virtual void ExportText_Internal( FString& ValueStr, const void* PropertyValueOrContainer, EPropertyPointerType PropertyPointerType, const void* DefaultValue, UObject* Parent, int32 PortFlags, UObject* ExportRootScope ) const override;
+public:
 	virtual bool ContainsObjectReference(TArray<const FStructProperty*>& EncounteredStructProps, EPropertyObjectReferenceType InReferenceType = EPropertyObjectReferenceType::Strong) const override;
 	virtual void EmitReferenceInfo(UClass& OwnerClass, int32 BaseOffset, TArray<const FStructProperty*>& EncounteredStructProps, FGCStackSizeHelper& StackSizeHelper) override;
 	virtual void InstanceSubobjects( void* Data, void const* DefaultData, UObject* Owner, struct FObjectInstancingGraph* InstanceGraph ) override;
@@ -5754,7 +5944,7 @@ PRAGMA_ENABLE_DEPRECATION_WARNINGS
 
 	// FProperty interface
 	virtual void SerializeItem(FStructuredArchive::FSlot Slot, void* Value, void const* Defaults) const override;
-	virtual const TCHAR* ImportText_Internal(const TCHAR* Buffer, void* Data, int32 PortFlags, UObject* OwnerObject, FOutputDevice* ErrorText) const override;
+	virtual const TCHAR* ImportText_Internal(const TCHAR* Buffer, void* ContainerOrPropertyPtr, EPropertyPointerType PropertyPointerType, UObject* OwnerObject, int32 PortFlags, FOutputDevice* ErrorText) const override;
 	// End of FProperty interface
 
 	// FMulticastDelegateProperty interface
@@ -5805,7 +5995,7 @@ PRAGMA_ENABLE_DEPRECATION_WARNINGS
 
 	// FProperty interface
 	virtual void SerializeItem(FStructuredArchive::FSlot Slot, void* Value, void const* Defaults) const override;
-	virtual const TCHAR* ImportText_Internal(const TCHAR* Buffer, void* Data, int32 PortFlags, UObject* OwnerObject, FOutputDevice* ErrorText) const override;
+	virtual const TCHAR* ImportText_Internal(const TCHAR* Buffer, void* ContainerOrPropertyPtr, EPropertyPointerType PropertyPointerType, UObject* OwnerObject, int32 PortFlags, FOutputDevice* ErrorText) const override;
 	// End of FProperty interface
 
 	// FMulticastDelegateProperty interface
diff --git a/Engine/Source/Runtime/DeveloperSettings/Private/Engine/DeveloperSettings.cpp b/Engine/Source/Runtime/DeveloperSettings/Private/Engine/DeveloperSettings.cpp
--- a/Engine/Source/Runtime/DeveloperSettings/Private/Engine/DeveloperSettings.cpp
+++ b/Engine/Source/Runtime/DeveloperSettings/Private/Engine/DeveloperSettings.cpp
@@ -120,7 +120,7 @@ void UDeveloperSettings::ImportConsoleVariableValues()
 			IConsoleVariable* CVar = IConsoleManager::Get().FindConsoleVariable(*CVarName);
 			if (CVar)
 			{
-				if (Property->ImportText(*CVar->GetString(), Property->ContainerPtrToValuePtr<uint8>(this, 0), PPF_ConsoleVariable, this) == NULL)
+				if (Property->ImportText_InContainer(*CVar->GetString(), this, this, PPF_ConsoleVariable) == NULL)
 				{
 					UE_LOG(LogTemp, Error, TEXT("%s import failed for %s on console variable %s (=%s)"), *GetClass()->GetName(), *Property->GetName(), *CVarName, *CVar->GetString());
 				}
diff --git a/Engine/Source/Runtime/Engine/Classes/Engine/DataTable.h b/Engine/Source/Runtime/Engine/Classes/Engine/DataTable.h
--- a/Engine/Source/Runtime/Engine/Classes/Engine/DataTable.h
+++ b/Engine/Source/Runtime/Engine/Classes/Engine/DataTable.h
@@ -535,7 +535,7 @@ struct ENGINE_API FDataTableCategoryHandle
 		// check each row to see if the value in the Property element is the one we're looking for (RowContents). If it is, add the row to OutRows
 		uint8* RowContentsAsBinary = (uint8*)FMemory_Alloca(Property->GetSize());
 		Property->InitializeValue(RowContentsAsBinary);
-		if (Property->ImportText(*RowContents.ToString(), RowContentsAsBinary, PPF_None, nullptr) == nullptr)
+		if (Property->ImportText_Direct(*RowContents.ToString(), RowContentsAsBinary, nullptr, PPF_None) == nullptr)
 		{
 			Property->DestroyValue(RowContentsAsBinary);
 			return;
diff --git a/Engine/Source/Runtime/Engine/Classes/Engine/DemoNetDriver.h b/Engine/Source/Runtime/Engine/Classes/Engine/DemoNetDriver.h
--- a/Engine/Source/Runtime/Engine/Classes/Engine/DemoNetDriver.h
+++ b/Engine/Source/Runtime/Engine/Classes/Engine/DemoNetDriver.h
@@ -295,9 +295,6 @@ private:
 	/** Does the actual work of TickFlush, either on the main thread or in a task thread in parallel with Slate. */
 	void TickFlushInternal(float DeltaSeconds);
 
-	/** Returns either CheckpointSaveMaxMSPerFrame or the value of demo.CheckpointSaveMaxMSPerFrameOverride if it's >= 0. */
-	float GetCheckpointSaveMaxMSPerFrame() const;
-
 	/** Returns the last checkpoint time in integer milliseconds. */
 	uint32 GetLastCheckpointTimeInMS() const { return ReplayHelper.GetLastCheckpointTimeInMS(); }
 
diff --git a/Engine/Source/Runtime/Engine/Classes/Engine/StaticMesh.h b/Engine/Source/Runtime/Engine/Classes/Engine/StaticMesh.h
--- a/Engine/Source/Runtime/Engine/Classes/Engine/StaticMesh.h
+++ b/Engine/Source/Runtime/Engine/Classes/Engine/StaticMesh.h
@@ -625,12 +625,12 @@ public:
 private:
 	/** Imported raw mesh bulk data. */
 	UE_DEPRECATED(5.00, "This must be protected for async build, always use the accessors even internally.")
-	UPROPERTY()
+	UPROPERTY(Setter = None, Getter = None)
 	TArray<FStaticMeshSourceModel> SourceModels;
 
 	/** Optional hi-res source data */
 	UE_DEPRECATED(5.00, "This must be protected for async build, always use the accessors even internally.")
-	UPROPERTY()
+	UPROPERTY(Setter = None, Getter = None)
 	FStaticMeshSourceModel HiResSourceModel;
 
 	void SetLightmapUVVersion(int32 InLightmapUVVersion)
diff --git a/Engine/Source/Runtime/Engine/Classes/Engine/World.h b/Engine/Source/Runtime/Engine/Classes/Engine/World.h
--- a/Engine/Source/Runtime/Engine/Classes/Engine/World.h
+++ b/Engine/Source/Runtime/Engine/Classes/Engine/World.h
@@ -1380,7 +1380,7 @@ private:
 	TArray<TObjectPtr<class ULevel>>						Levels;
 
 	/** Array of level collections currently in this world. */
-	UPROPERTY(Transient, NonTransactional)
+	UPROPERTY(Transient, NonTransactional, Setter = None, Getter = None)
 	TArray<FLevelCollection>					LevelCollections;
 
 	/** Index of the level collection that's currently ticking. */
diff --git a/Engine/Source/Runtime/Engine/Classes/Kismet/BlueprintMapLibrary.h b/Engine/Source/Runtime/Engine/Classes/Kismet/BlueprintMapLibrary.h
--- a/Engine/Source/Runtime/Engine/Classes/Kismet/BlueprintMapLibrary.h
+++ b/Engine/Source/Runtime/Engine/Classes/Kismet/BlueprintMapLibrary.h
@@ -136,7 +136,8 @@ public:
 		void* KeyStorageSpace = FMemory_Alloca(KeyPropertySize);
 		CurrKeyProp->InitializeValue(KeyStorageSpace);
 
-		Stack.MostRecentPropertyAddress = NULL;
+		Stack.MostRecentPropertyAddress = nullptr;
+		Stack.MostRecentPropertyContainer = nullptr;
 		Stack.StepCompiledIn<FProperty>(KeyStorageSpace);
 		
 		const FProperty* CurrValueProp = MapProperty->ValueProp;
@@ -144,7 +145,8 @@ public:
 		void* ValueStorageSpace = FMemory_Alloca(ValuePropertySize);
 		CurrValueProp->InitializeValue(ValueStorageSpace);
 		
-		Stack.MostRecentPropertyAddress = NULL;
+		Stack.MostRecentPropertyAddress = nullptr;
+		Stack.MostRecentPropertyContainer = nullptr;
 		Stack.StepCompiledIn<FProperty>(ValueStorageSpace);
 
 		P_FINISH;
@@ -175,7 +177,8 @@ public:
 		void* KeyStorageSpace = FMemory_Alloca(KeyPropertySize);
 		CurrKeyProp->InitializeValue(KeyStorageSpace);
 
-		Stack.MostRecentPropertyAddress = NULL;
+		Stack.MostRecentPropertyAddress = nullptr;
+		Stack.MostRecentPropertyContainer = nullptr;
 		Stack.StepCompiledIn<FProperty>(KeyStorageSpace);
 		
 		P_FINISH;
@@ -204,7 +207,8 @@ public:
 		void* KeyStorageSpace = FMemory_Alloca(KeyPropertySize);
 		CurrKeyProp->InitializeValue(KeyStorageSpace);
 
-		Stack.MostRecentPropertyAddress = NULL;
+		Stack.MostRecentPropertyAddress = nullptr;
+		Stack.MostRecentPropertyContainer = nullptr;
 		Stack.StepCompiledIn<FProperty>(KeyStorageSpace);
 		
 		const FProperty* CurrValueProp = MapProperty->ValueProp;
@@ -212,7 +216,8 @@ public:
 		void* ValueStorageSpace = FMemory_Alloca(ValuePropertySize);
 		CurrValueProp->InitializeValue(ValueStorageSpace);
 		
-		Stack.MostRecentPropertyAddress = NULL;
+		Stack.MostRecentPropertyAddress = nullptr;
+		Stack.MostRecentPropertyContainer = nullptr;
 		Stack.StepCompiledIn<FProperty>(ValueStorageSpace);
 		const FFieldClass* CurrValuePropClass = CurrValueProp->GetClass();
 		const FFieldClass* MostRecentPropClass = Stack.MostRecentProperty->GetClass();
@@ -312,7 +317,8 @@ public:
 		void* KeyStorageSpace = FMemory_Alloca(KeyPropertySize);
 		CurrKeyProp->InitializeValue(KeyStorageSpace);
 
-		Stack.MostRecentPropertyAddress = NULL;
+		Stack.MostRecentPropertyAddress = nullptr;
+		Stack.MostRecentPropertyContainer = nullptr;
 		Stack.StepCompiledIn<FProperty>(KeyStorageSpace);
 
 		P_FINISH;
diff --git a/Engine/Source/Runtime/Engine/Classes/Kismet/BlueprintSetLibrary.h b/Engine/Source/Runtime/Engine/Classes/Kismet/BlueprintSetLibrary.h
--- a/Engine/Source/Runtime/Engine/Classes/Kismet/BlueprintSetLibrary.h
+++ b/Engine/Source/Runtime/Engine/Classes/Kismet/BlueprintSetLibrary.h
@@ -169,7 +169,8 @@ public:
 		void* StorageSpace = FMemory_Alloca(PropertySize);
 		ElementProp->InitializeValue(StorageSpace);
 
-		Stack.MostRecentPropertyAddress = NULL;
+		Stack.MostRecentPropertyAddress = nullptr;
+		Stack.MostRecentPropertyContainer = nullptr;
 		Stack.StepCompiledIn<FProperty>(StorageSpace);
 		void* ItemPtr = StorageSpace;
 
@@ -360,7 +361,8 @@ public:
 		void* StorageSpace = FMemory_Alloca(PropertySize);
 		ElementProp->InitializeValue(StorageSpace);
 
-		Stack.MostRecentPropertyAddress = NULL;
+		Stack.MostRecentPropertyAddress = nullptr;
+		Stack.MostRecentPropertyContainer = nullptr;
 		Stack.StepCompiledIn<FProperty>(StorageSpace);
 		void* ItemToFindPtr = StorageSpace;
 
diff --git a/Engine/Source/Runtime/Engine/Classes/Kismet/KismetArrayLibrary.h b/Engine/Source/Runtime/Engine/Classes/Kismet/KismetArrayLibrary.h
--- a/Engine/Source/Runtime/Engine/Classes/Kismet/KismetArrayLibrary.h
+++ b/Engine/Source/Runtime/Engine/Classes/Kismet/KismetArrayLibrary.h
@@ -311,7 +311,8 @@ public:
  		void* StorageSpace = FMemory_Alloca(PropertySize);
  		InnerProp->InitializeValue(StorageSpace);
  
- 		Stack.MostRecentPropertyAddress = NULL;
+ 		Stack.MostRecentPropertyAddress = nullptr;
+		Stack.MostRecentPropertyContainer = nullptr;
  		Stack.StepCompiledIn<FProperty>(StorageSpace);
  
  		P_FINISH;
@@ -340,7 +341,8 @@ public:
 		void* StorageSpace = FMemory_Alloca(PropertySize);
 		InnerProp->InitializeValue(StorageSpace);
 
-		Stack.MostRecentPropertyAddress = NULL;
+		Stack.MostRecentPropertyAddress = nullptr;
+		Stack.MostRecentPropertyContainer = nullptr;
 		Stack.StepCompiledIn<FProperty>(StorageSpace);
 
 		P_FINISH;
@@ -447,7 +449,8 @@ public:
 		void* StorageSpace = FMemory_Alloca(PropertySize);
 		InnerProp->InitializeValue(StorageSpace);
 
-		Stack.MostRecentPropertyAddress = NULL;
+		Stack.MostRecentPropertyAddress = nullptr;
+		Stack.MostRecentPropertyContainer = nullptr;
 		Stack.StepCompiledIn<FProperty>(StorageSpace);
 
 		P_GET_PROPERTY(FIntProperty, Index);
@@ -496,7 +499,8 @@ public:
 		void* StorageSpace = FMemory_Alloca(PropertySize);
 		InnerProp->InitializeValue(StorageSpace);
 
-		Stack.MostRecentPropertyAddress = NULL;
+		Stack.MostRecentPropertyAddress = nullptr;
+		Stack.MostRecentPropertyContainer = nullptr;
 		Stack.StepCompiledIn<FProperty>(StorageSpace);
 		void* ItemPtr = StorageSpace;
 
@@ -653,7 +657,8 @@ public:
 		void* StorageSpace = FMemory_Alloca(PropertySize);
 		InnerProp->InitializeValue(StorageSpace);
 
-		Stack.MostRecentPropertyAddress = NULL;
+		Stack.MostRecentPropertyAddress = nullptr;
+		Stack.MostRecentPropertyContainer = nullptr;
 		Stack.StepCompiledIn<FProperty>(StorageSpace);
 		const FFieldClass* InnerPropClass = InnerProp->GetClass();
 		const FFieldClass* MostRecentPropClass = Stack.MostRecentProperty->GetClass();
@@ -695,7 +700,8 @@ public:
 		void* StorageSpace = FMemory_Alloca(PropertySize);
 		InnerProp->InitializeValue(StorageSpace);
 
-		Stack.MostRecentPropertyAddress = NULL;
+		Stack.MostRecentPropertyAddress = nullptr;
+		Stack.MostRecentPropertyContainer = nullptr;
 		Stack.StepCompiledIn<FProperty>(StorageSpace);
 
 		P_GET_UBOOL(bSizeToFit);
@@ -749,7 +755,8 @@ public:
 		void* StorageSpace = FMemory_Alloca(PropertySize);
 		InnerProp->InitializeValue(StorageSpace);
 
-		Stack.MostRecentPropertyAddress = NULL;
+		Stack.MostRecentPropertyAddress = nullptr;
+		Stack.MostRecentPropertyContainer = nullptr;
 		Stack.StepCompiledIn<FProperty>(StorageSpace);
 		void* ItemToFindPtr = StorageSpace;
 
@@ -783,7 +790,8 @@ public:
 		void* StorageSpace = FMemory_Alloca(PropertySize);
 		InnerProp->InitializeValue(StorageSpace);
 
-		Stack.MostRecentPropertyAddress = NULL;
+		Stack.MostRecentPropertyAddress = nullptr;
+		Stack.MostRecentPropertyContainer = nullptr;
 		Stack.StepCompiledIn<FProperty>(StorageSpace);
 		void* ItemToFindPtr = StorageSpace;
 
diff --git a/Engine/Source/Runtime/Engine/Private/BlueprintGeneratedClass.cpp b/Engine/Source/Runtime/Engine/Private/BlueprintGeneratedClass.cpp
--- a/Engine/Source/Runtime/Engine/Private/BlueprintGeneratedClass.cpp
+++ b/Engine/Source/Runtime/Engine/Private/BlueprintGeneratedClass.cpp
@@ -2075,7 +2075,7 @@ void UBlueprintGeneratedClass::GetEditorTags(FEditorTags& Tags) const
 			const FProperty* Property = FindFieldChecked<FProperty>(UBlueprint::StaticClass(), PropertyName);
 			const uint8* PropertyAddr = PropertyValueOverride ? PropertyValueOverride : Property->ContainerPtrToValuePtr<uint8>(BP);
 			FString PropertyValueAsText;
-			Property->ExportTextItem(PropertyValueAsText, PropertyAddr, PropertyAddr, nullptr, PPF_None);
+			Property->ExportTextItem_Direct(PropertyValueAsText, PropertyAddr, PropertyAddr, nullptr, PPF_None);
 			if (!PropertyValueAsText.IsEmpty())
 			{
 				Tags.Add(TagName, MoveTemp(PropertyValueAsText));
diff --git a/Engine/Source/Runtime/Engine/Private/DataTableUtils.cpp b/Engine/Source/Runtime/Engine/Private/DataTableUtils.cpp
--- a/Engine/Source/Runtime/Engine/Private/DataTableUtils.cpp
+++ b/Engine/Source/Runtime/Engine/Private/DataTableUtils.cpp
@@ -30,7 +30,7 @@ void AssignStringToPropertyDirect(const FString& InString, const FProperty* InPr
 {
 	auto DoImportText = [&](const FString& InStringToImport)
 	{
-		InProp->ImportText(*InStringToImport, InData, InPortFlags, nullptr, &OutImportError);
+		InProp->ImportText_Direct(*InStringToImport, InData, nullptr, InPortFlags, &OutImportError);
 	};
 
 	bool bNeedsImport = true;
diff --git a/Engine/Source/Runtime/Engine/Private/EdGraph/EdGraphNode.cpp b/Engine/Source/Runtime/Engine/Private/EdGraph/EdGraphNode.cpp
--- a/Engine/Source/Runtime/Engine/Private/EdGraph/EdGraphNode.cpp
+++ b/Engine/Source/Runtime/Engine/Private/EdGraph/EdGraphNode.cpp
@@ -269,7 +269,7 @@ FString UEdGraphNode::GetPropertyNameAndValueForDiff(const FProperty* Prop, cons
 	}
 	else
 	{
-		Prop->ExportTextItem(ExportedStringValue, PropertyAddr, NULL, NULL, PPF_PropertyWindow, NULL);
+		Prop->ExportTextItem_Direct(ExportedStringValue, PropertyAddr, NULL, NULL, PPF_PropertyWindow, NULL);
 	}
 
 	const bool bIsBool = Prop->IsA(FBoolProperty::StaticClass());
diff --git a/Engine/Source/Runtime/Engine/Private/EdGraph/EdGraphPin.cpp b/Engine/Source/Runtime/Engine/Private/EdGraph/EdGraphPin.cpp
--- a/Engine/Source/Runtime/Engine/Private/EdGraph/EdGraphPin.cpp
+++ b/Engine/Source/Runtime/Engine/Private/EdGraph/EdGraphPin.cpp
@@ -916,7 +916,7 @@ bool UEdGraphPin::ExportTextItem(FString& ValueStr, int32 PortFlags) const
 	if (PinName != DefaultPin.PinName)
 	{
 		ValueStr += PinHelpers::PinNameName + TEXT("=");
-		NamePropCDO->ExportTextItem(ValueStr, &PinName, nullptr, nullptr, PortFlags, nullptr);
+		NamePropCDO->ExportTextItem_Direct(ValueStr, &PinName, nullptr, nullptr, PortFlags, nullptr);
 		ValueStr += PinHelpers::ExportTextPropDelimiter;
 	}
 
@@ -924,7 +924,7 @@ bool UEdGraphPin::ExportTextItem(FString& ValueStr, int32 PortFlags) const
 	if (!PinFriendlyName.EqualTo(DefaultPin.PinFriendlyName))
 	{
 		ValueStr += PinHelpers::PinFriendlyNameName + TEXT("=");
-		TextPropCDO->ExportTextItem(ValueStr, &PinFriendlyName, nullptr, nullptr, PortFlags, nullptr);
+		TextPropCDO->ExportTextItem_Direct(ValueStr, &PinFriendlyName, nullptr, nullptr, PortFlags, nullptr);
 		ValueStr += PinHelpers::ExportTextPropDelimiter;
 	}
 #endif // WITH_EDITORONLY_DATA
@@ -932,7 +932,7 @@ bool UEdGraphPin::ExportTextItem(FString& ValueStr, int32 PortFlags) const
 	if (PinToolTip != DefaultPin.PinToolTip)
 	{
 		ValueStr += PinHelpers::PinToolTipName + TEXT("=");
-		StrPropCDO->ExportTextItem(ValueStr, &PinToolTip, nullptr, nullptr, PortFlags, nullptr);
+		StrPropCDO->ExportTextItem_Direct(ValueStr, &PinToolTip, nullptr, nullptr, PortFlags, nullptr);
 		ValueStr += PinHelpers::ExportTextPropDelimiter;
 	}
 
@@ -940,7 +940,7 @@ bool UEdGraphPin::ExportTextItem(FString& ValueStr, int32 PortFlags) const
 	{
 		const FString DirectionString = UEnum::GetValueAsString(TEXT("/Script/Engine.EEdGraphPinDirection"), Direction);
 		ValueStr += PinHelpers::DirectionName + TEXT("=");
-		StrPropCDO->ExportTextItem(ValueStr, &DirectionString, nullptr, nullptr, PortFlags, nullptr);
+		StrPropCDO->ExportTextItem_Direct(ValueStr, &DirectionString, nullptr, nullptr, PortFlags, nullptr);
 		ValueStr += PinHelpers::ExportTextPropDelimiter;
 	}
 
@@ -952,7 +952,7 @@ bool UEdGraphPin::ExportTextItem(FString& ValueStr, int32 PortFlags) const
 			FString PropertyStr;
 			const uint8* PropertyAddr = Prop->ContainerPtrToValuePtr<uint8>(&PinType);
 			const uint8* DefaultAddr = Prop->ContainerPtrToValuePtr<uint8>(&DefaultPin.PinType);
-			Prop->ExportTextItem(PropertyStr, PropertyAddr, DefaultAddr, NULL, PortFlags, nullptr);
+			Prop->ExportTextItem_Direct(PropertyStr, PropertyAddr, DefaultAddr, NULL, PortFlags, nullptr);
 
 			if (!PropertyStr.IsEmpty())
 			{
@@ -965,14 +965,14 @@ bool UEdGraphPin::ExportTextItem(FString& ValueStr, int32 PortFlags) const
 	if (DefaultValue != DefaultPin.DefaultValue)
 	{
 		ValueStr += PinHelpers::DefaultValueName + TEXT("=");
-		StrPropCDO->ExportTextItem(ValueStr, &DefaultValue, nullptr, nullptr, PortFlags, nullptr);
+		StrPropCDO->ExportTextItem_Direct(ValueStr, &DefaultValue, nullptr, nullptr, PortFlags, nullptr);
 		ValueStr += PinHelpers::ExportTextPropDelimiter;
 	}
 
 	if (AutogeneratedDefaultValue != DefaultPin.AutogeneratedDefaultValue)
 	{
 		ValueStr += PinHelpers::AutogeneratedDefaultValueName + TEXT("=");
-		StrPropCDO->ExportTextItem(ValueStr, &AutogeneratedDefaultValue, nullptr, nullptr, PortFlags, nullptr);
+		StrPropCDO->ExportTextItem_Direct(ValueStr, &AutogeneratedDefaultValue, nullptr, nullptr, PortFlags, nullptr);
 		ValueStr += PinHelpers::ExportTextPropDelimiter;
 	}
 
@@ -980,14 +980,14 @@ bool UEdGraphPin::ExportTextItem(FString& ValueStr, int32 PortFlags) const
 	{
 		const FString DefaultObjectPath = GetPathNameSafe(DefaultObject);
 		ValueStr += PinHelpers::DefaultObjectName + TEXT("=");
-		StrPropCDO->ExportTextItem(ValueStr, &DefaultObjectPath, nullptr, nullptr, PortFlags, nullptr);
+		StrPropCDO->ExportTextItem_Direct(ValueStr, &DefaultObjectPath, nullptr, nullptr, PortFlags, nullptr);
 		ValueStr += PinHelpers::ExportTextPropDelimiter;
 	}
 
 	if (!DefaultTextValue.EqualTo(DefaultPin.DefaultTextValue))
 	{
 		ValueStr += PinHelpers::DefaultTextValueName + TEXT("=");
-		TextPropCDO->ExportTextItem(ValueStr, &DefaultTextValue, nullptr, nullptr, PortFlags, nullptr);
+		TextPropCDO->ExportTextItem_Direct(ValueStr, &DefaultTextValue, nullptr, nullptr, PortFlags, nullptr);
 		ValueStr += PinHelpers::ExportTextPropDelimiter;
 	}
 
@@ -1023,36 +1023,36 @@ bool UEdGraphPin::ExportTextItem(FString& ValueStr, int32 PortFlags) const
 
 		ValueStr += PinHelpers::bHiddenName + TEXT("=");
 		bool LocalHidden = bHidden;
-		BoolPropCDO->ExportTextItem(ValueStr, &LocalHidden, nullptr, nullptr, PortFlags, nullptr);
+		BoolPropCDO->ExportTextItem_Direct(ValueStr, &LocalHidden, nullptr, nullptr, PortFlags, nullptr);
 		ValueStr += PinHelpers::ExportTextPropDelimiter;
 
 		ValueStr += PinHelpers::bNotConnectableName + TEXT("=");
 		bool LocalNotConnectable = bNotConnectable;
-		BoolPropCDO->ExportTextItem(ValueStr, &LocalNotConnectable, nullptr, nullptr, PortFlags, nullptr);
+		BoolPropCDO->ExportTextItem_Direct(ValueStr, &LocalNotConnectable, nullptr, nullptr, PortFlags, nullptr);
 		ValueStr += PinHelpers::ExportTextPropDelimiter;
 
 		ValueStr += PinHelpers::bDefaultValueIsReadOnlyName + TEXT("=");
 		bool LocalDefaultValueIsReadOnly = bDefaultValueIsReadOnly;
-		BoolPropCDO->ExportTextItem(ValueStr, &LocalDefaultValueIsReadOnly, nullptr, nullptr, PortFlags, nullptr);
+		BoolPropCDO->ExportTextItem_Direct(ValueStr, &LocalDefaultValueIsReadOnly, nullptr, nullptr, PortFlags, nullptr);
 		ValueStr += PinHelpers::ExportTextPropDelimiter;
 
 		ValueStr += PinHelpers::bDefaultValueIsIgnoredName + TEXT("=");
 		bool LocalDefaultValueIsIgnored = bDefaultValueIsIgnored;
-		BoolPropCDO->ExportTextItem(ValueStr, &LocalDefaultValueIsIgnored, nullptr, nullptr, PortFlags, nullptr);
+		BoolPropCDO->ExportTextItem_Direct(ValueStr, &LocalDefaultValueIsIgnored, nullptr, nullptr, PortFlags, nullptr);
 		ValueStr += PinHelpers::ExportTextPropDelimiter;
 
 		// Intentionally not exporting bIsDiffing as it is transient
 
 		ValueStr += PinHelpers::bAdvancedViewName + TEXT("=");
 		bool LocalAdvancedView = bAdvancedView;
-		BoolPropCDO->ExportTextItem(ValueStr, &LocalAdvancedView, nullptr, nullptr, PortFlags, nullptr);
+		BoolPropCDO->ExportTextItem_Direct(ValueStr, &LocalAdvancedView, nullptr, nullptr, PortFlags, nullptr);
 		ValueStr += PinHelpers::ExportTextPropDelimiter;
 
 		// Intentionally not exporting bDisplayAsMutableRef as it is transient
 
 		ValueStr += PinHelpers::bOrphanedPinName + TEXT("=");
 		bool LocalOrphanedPin = bOrphanedPin;
-		BoolPropCDO->ExportTextItem(ValueStr, &LocalOrphanedPin, nullptr, nullptr, PortFlags, nullptr);
+		BoolPropCDO->ExportTextItem_Direct(ValueStr, &LocalOrphanedPin, nullptr, nullptr, PortFlags, nullptr);
 		ValueStr += PinHelpers::ExportTextPropDelimiter;
 
 		// Intentionally not exporting bSavePinIfOrphaned as it is transient
@@ -1139,25 +1139,25 @@ bool UEdGraphPin::ImportTextItem(const TCHAR*& Buffer, int32 PortFlags, class UO
 		}
 		else if (PropertyToken == PinHelpers::PinNameName)
 		{
-			Buffer = NamePropCDO->ImportText(Buffer, &PinName, PortFlags, Parent, ErrorText);
+			Buffer = NamePropCDO->ImportText_Direct(Buffer, &PinName, Parent, PortFlags, ErrorText);
 			bParseSuccess = (Buffer != nullptr);
 		}
 #if WITH_EDITORONLY_DATA
 		else if (PropertyToken == PinHelpers::PinFriendlyNameName)
 		{
-			Buffer = TextPropCDO->ImportText(Buffer, &PinFriendlyName, PortFlags, Parent, ErrorText);
+			Buffer = TextPropCDO->ImportText_Direct(Buffer, &PinFriendlyName, Parent, PortFlags, ErrorText);
 			bParseSuccess = (Buffer != nullptr);
 		}
 #endif //WITH_EDITORONLY_DATA
 		else if (PropertyToken == PinHelpers::PinToolTipName)
 		{
-			Buffer = StrPropCDO->ImportText(Buffer, &PinToolTip, PortFlags, Parent, ErrorText);
+			Buffer = StrPropCDO->ImportText_Direct(Buffer, &PinToolTip, Parent, PortFlags, ErrorText);
 			bParseSuccess = (Buffer != nullptr);
 		}
 		else if (PropertyToken == PinHelpers::DirectionName)
 		{
 			FString DirectionString;
-			Buffer = StrPropCDO->ImportText(Buffer, &DirectionString, PortFlags, Parent, ErrorText);
+			Buffer = StrPropCDO->ImportText_Direct(Buffer, &DirectionString, Parent, PortFlags, ErrorText);
 			bParseSuccess = (Buffer != nullptr);
 			if (bParseSuccess)
 			{
@@ -1185,14 +1185,14 @@ bool UEdGraphPin::ImportTextItem(const TCHAR*& Buffer, int32 PortFlags, class UO
 				if (FoundProp)
 				{
 					uint8* PropertyAddr = FoundProp->ContainerPtrToValuePtr<uint8>(&PinType);
-					Buffer = FoundProp->ImportText(Buffer, PropertyAddr, PortFlags, Parent, ErrorText);
+					Buffer = FoundProp->ImportText_Direct(Buffer, PropertyAddr, Parent, PortFlags, ErrorText);
 					bParseSuccess = (Buffer != nullptr);
 				}
 				// UE_DEPRECATED(4.17) - For some time bIsMap and bIsSet would have been in exported text and will cause issues if we don't handle them
 				else if (PropertyName == TEXT("bIsMap") || PropertyName == TEXT("bIsSet"))
 				{
 					bool bDummyBool = false;
-					Buffer = BoolPropCDO->ImportText(Buffer, &bDummyBool, PortFlags, Parent, ErrorText);
+					Buffer = BoolPropCDO->ImportText_Direct(Buffer, &bDummyBool, Parent, PortFlags, ErrorText);
 					bParseSuccess = (Buffer != nullptr);
 				}
 				else
@@ -1213,18 +1213,18 @@ bool UEdGraphPin::ImportTextItem(const TCHAR*& Buffer, int32 PortFlags, class UO
 		}
 		else if (PropertyToken == PinHelpers::DefaultValueName)
 		{
-			Buffer = StrPropCDO->ImportText(Buffer, &DefaultValue, PortFlags, Parent, ErrorText);
+			Buffer = StrPropCDO->ImportText_Direct(Buffer, &DefaultValue, Parent, PortFlags, ErrorText);
 			bParseSuccess = (Buffer != nullptr);
 		}
 		else if (PropertyToken == PinHelpers::AutogeneratedDefaultValueName)
 		{
-			Buffer = StrPropCDO->ImportText(Buffer, &AutogeneratedDefaultValue, PortFlags, Parent, ErrorText);
+			Buffer = StrPropCDO->ImportText_Direct(Buffer, &AutogeneratedDefaultValue, Parent, PortFlags, ErrorText);
 			bParseSuccess = (Buffer != nullptr);
 		}
 		else if (PropertyToken == PinHelpers::DefaultObjectName)
 		{
 			FString DefaultObjectString;
-			Buffer = StrPropCDO->ImportText(Buffer, &DefaultObjectString, PortFlags, Parent, ErrorText);
+			Buffer = StrPropCDO->ImportText_Direct(Buffer, &DefaultObjectString, Parent, PortFlags, ErrorText);
 			bParseSuccess = (Buffer != nullptr);
 			if (bParseSuccess)
 			{
@@ -1241,7 +1241,7 @@ bool UEdGraphPin::ImportTextItem(const TCHAR*& Buffer, int32 PortFlags, class UO
 		}
 		else if (PropertyToken == PinHelpers::DefaultTextValueName)
 		{
-			Buffer = TextPropCDO->ImportText(Buffer, &DefaultTextValue, PortFlags, Parent, ErrorText);
+			Buffer = TextPropCDO->ImportText_Direct(Buffer, &DefaultTextValue, Parent, PortFlags, ErrorText);
 			bParseSuccess = (Buffer != nullptr);
 		}
 		else if (PropertyToken == PinHelpers::LinkedToName)
@@ -1268,28 +1268,28 @@ bool UEdGraphPin::ImportTextItem(const TCHAR*& Buffer, int32 PortFlags, class UO
 		else if (PropertyToken == PinHelpers::bHiddenName)
 		{
 			bool LocalHidden = bHidden;
-			Buffer = BoolPropCDO->ImportText(Buffer, &LocalHidden, PortFlags, Parent, ErrorText);
+			Buffer = BoolPropCDO->ImportText_Direct(Buffer, &LocalHidden, Parent, PortFlags, ErrorText);
 			bParseSuccess = (Buffer != nullptr);
 			bHidden = LocalHidden;
 		}
 		else if (PropertyToken == PinHelpers::bNotConnectableName)
 		{
 			bool LocalNotConnectable = bNotConnectable;
-			Buffer = BoolPropCDO->ImportText(Buffer, &LocalNotConnectable, PortFlags, Parent, ErrorText);
+			Buffer = BoolPropCDO->ImportText_Direct(Buffer, &LocalNotConnectable, Parent, PortFlags, ErrorText);
 			bParseSuccess = (Buffer != nullptr);
 			bNotConnectable = LocalNotConnectable;
 		}
 		else if (PropertyToken == PinHelpers::bDefaultValueIsReadOnlyName)
 		{
 			bool LocalDefaultValueIsReadOnly = bDefaultValueIsReadOnly;
-			Buffer = BoolPropCDO->ImportText(Buffer, &LocalDefaultValueIsReadOnly, PortFlags, Parent, ErrorText);
+			Buffer = BoolPropCDO->ImportText_Direct(Buffer, &LocalDefaultValueIsReadOnly, Parent, PortFlags, ErrorText);
 			bParseSuccess = (Buffer != nullptr);
 			bDefaultValueIsReadOnly = LocalDefaultValueIsReadOnly;
 		}
 		else if (PropertyToken == PinHelpers::bDefaultValueIsIgnoredName)
 		{
 			bool LocalDefaultValueIsIgnored = bDefaultValueIsIgnored;
-			Buffer = BoolPropCDO->ImportText(Buffer, &LocalDefaultValueIsIgnored, PortFlags, Parent, ErrorText);
+			Buffer = BoolPropCDO->ImportText_Direct(Buffer, &LocalDefaultValueIsIgnored, Parent, PortFlags, ErrorText);
 			bParseSuccess = (Buffer != nullptr);
 			bDefaultValueIsIgnored = LocalDefaultValueIsIgnored;
 		}
@@ -1297,7 +1297,7 @@ bool UEdGraphPin::ImportTextItem(const TCHAR*& Buffer, int32 PortFlags, class UO
 		else if (PropertyToken == PinHelpers::bAdvancedViewName)
 		{
 			bool LocalAdvancedView = bAdvancedView;
-			Buffer = BoolPropCDO->ImportText(Buffer, &LocalAdvancedView, PortFlags, Parent, ErrorText);
+			Buffer = BoolPropCDO->ImportText_Direct(Buffer, &LocalAdvancedView, Parent, PortFlags, ErrorText);
 			bParseSuccess = (Buffer != nullptr);
 			bAdvancedView = LocalAdvancedView;
 		}
@@ -1305,7 +1305,7 @@ bool UEdGraphPin::ImportTextItem(const TCHAR*& Buffer, int32 PortFlags, class UO
 		else if (PropertyToken == PinHelpers::bOrphanedPinName)
 		{
 			bool LocalOrphanedPin = bOrphanedPin;
-			Buffer = BoolPropCDO->ImportText(Buffer, &LocalOrphanedPin, PortFlags, Parent, ErrorText);
+			Buffer = BoolPropCDO->ImportText_Direct(Buffer, &LocalOrphanedPin, Parent, PortFlags, ErrorText);
 			bParseSuccess = (Buffer != nullptr);
 			if (AreOrphanPinsEnabled())
 			{
diff --git a/Engine/Source/Runtime/Engine/Private/Net/RPCDoSDetection.cpp b/Engine/Source/Runtime/Engine/Private/Net/RPCDoSDetection.cpp
--- a/Engine/Source/Runtime/Engine/Private/Net/RPCDoSDetection.cpp
+++ b/Engine/Source/Runtime/Engine/Private/Net/RPCDoSDetection.cpp
@@ -235,7 +235,7 @@ bool FRPCDoSStateConfig::LoadStructConfig(const TCHAR* SectionName, const TCHAR*
 				{
 					if (CurProp->HasAllPropertyFlags(CPF_Config))
 					{
-						CurProp->ImportText(*Value, CurProp->ContainerPtrToValuePtr<void>(this), 0, nullptr);
+						CurProp->ImportText_InContainer(*Value, this, nullptr, 0);
 					}
 					else
 					{
diff --git a/Engine/Source/Runtime/Engine/Private/NetDriver.cpp b/Engine/Source/Runtime/Engine/Private/NetDriver.cpp
--- a/Engine/Source/Runtime/Engine/Private/NetDriver.cpp
+++ b/Engine/Source/Runtime/Engine/Private/NetDriver.cpp
@@ -3880,13 +3880,13 @@ bool FPacketSimulationSettings::LoadEmulationProfile(const TCHAR* ProfileName)
 			}
 			else if (FProperty* StructProperty = ThisStruct->FindPropertyByName(FName(*VarName, FNAME_Find)))
 			{
-				StructProperty->ImportText(*VarValue, StructProperty->ContainerPtrToValuePtr<void>(this), 0, nullptr);
-	}
+				StructProperty->ImportText_InContainer(*VarValue, this, nullptr, 0);
+			}
 			else
-	{
+			{
 				UE_LOG(LogNet, Warning, TEXT("FPacketSimulationSettings::LoadEmulationProfile could not find property named %s"), *VarName);
 			}
-	}
+		}
 	}
 	
 	ValidateSettings();
diff --git a/Engine/Source/Runtime/Engine/Private/ObjectEditorUtils.cpp b/Engine/Source/Runtime/Engine/Private/ObjectEditorUtils.cpp
--- a/Engine/Source/Runtime/Engine/Private/ObjectEditorUtils.cpp
+++ b/Engine/Source/Runtime/Engine/Private/ObjectEditorUtils.cpp
@@ -294,7 +294,7 @@ namespace FObjectEditorUtils
 			FString ExportedTextString;
 			if ( InSourceProperty->ExportText_InContainer(0, ExportedTextString, InSourcePtr, InSourcePtr, SourceObject, PPF_Copy, SourceObject) )
 			{
-				InDestinationProperty->ImportText(*ExportedTextString, InDestinationProperty->ContainerPtrToValuePtr<void>(InTargetPtr), 0, InDestinationObject);
+				InDestinationProperty->ImportText_InContainer(*ExportedTextString, InTargetPtr, InDestinationObject, 0);
 			}
 		}
 	}
diff --git a/Engine/Source/Runtime/Engine/Private/RuntimeOptionsBase.cpp b/Engine/Source/Runtime/Engine/Private/RuntimeOptionsBase.cpp
--- a/Engine/Source/Runtime/Engine/Private/RuntimeOptionsBase.cpp
+++ b/Engine/Source/Runtime/Engine/Private/RuntimeOptionsBase.cpp
@@ -55,12 +55,10 @@ void URuntimeOptionsBase::ApplyCommandlineOverrides()
 	{
 		const FString FullyQualifiedName = FString::Printf(TEXT("%s.%s"), *OptionCommandPrefix, *Property->GetName());
 
-		uint8* DataPtr = Property->ContainerPtrToValuePtr<uint8>(this, 0);
-
 		FString CommandLineOverride;
 		if (FParse::Value(FCommandLine::Get(), *FString::Printf(TEXT("%s="), *FullyQualifiedName), /*out*/ CommandLineOverride))
 		{
-			Property->ImportText(*CommandLineOverride, DataPtr, PPF_None, this);
+			Property->ImportText_InContainer(*CommandLineOverride, this, this, PPF_None);
 		}
 	}
 #endif
diff --git a/Engine/Source/Runtime/Engine/Private/Tests/Internationalization/TextPropertyTest.cpp b/Engine/Source/Runtime/Engine/Private/Tests/Internationalization/TextPropertyTest.cpp
--- a/Engine/Source/Runtime/Engine/Private/Tests/Internationalization/TextPropertyTest.cpp
+++ b/Engine/Source/Runtime/Engine/Private/Tests/Internationalization/TextPropertyTest.cpp
@@ -40,7 +40,7 @@ bool FTextPropertyTest::RunTest (const FString& Parameters)
 		// Test ExportText - Export text should provide the localized form of the text.
 		{
 			FString ExportedStringValue;
-			DefaultedTextProperty->ExportTextItem(ExportedStringValue, &(NewUObject->DefaultedText), NULL, NULL, 0, NULL);
+			DefaultedTextProperty->ExportTextItem_Direct(ExportedStringValue, &(NewUObject->DefaultedText), nullptr, nullptr, 0, nullptr);
 			FText ExportedTextValue;
 			FTextStringHelper::ReadFromBuffer(*ExportedStringValue, ExportedTextValue);
 			if (ExportedTextValue.ToString() != NewUObject->DefaultedText.ToString())
@@ -52,7 +52,7 @@ bool FTextPropertyTest::RunTest (const FString& Parameters)
 		// Test ImportText - Import text should set the source string to the input string.
 		{
 			FString ImportedStringValue = TEXT("ImportValue");
-			DefaultedTextProperty->ImportText(*ImportedStringValue, &(NewUObject->DefaultedText), 0, NULL);
+			DefaultedTextProperty->ImportText_Direct(*ImportedStringValue, &(NewUObject->DefaultedText), nullptr, 0);
 			const FString* const SourceString = FTextInspector::GetSourceString(NewUObject->DefaultedText);
 			if( !SourceString || ImportedStringValue != *SourceString )
 			{
diff --git a/Engine/Source/Runtime/Engine/Private/UnrealEngine.cpp b/Engine/Source/Runtime/Engine/Private/UnrealEngine.cpp
--- a/Engine/Source/Runtime/Engine/Private/UnrealEngine.cpp
+++ b/Engine/Source/Runtime/Engine/Private/UnrealEngine.cpp
@@ -9159,7 +9159,7 @@ bool UEngine::HandleObjCommand( const TCHAR* Cmd, FOutputDevice& Ar )
 							for( int32 i=0; i<FMath::Min(ArrayHelper.Num(),100); i++ )
 							{
 								Value.Empty();
-								ArrayProp->Inner->ExportTextItem( Value, ArrayHelper.GetRawPtr(i), ArrayHelper.GetRawPtr(i), Obj, PPF_IncludeTransient );
+								ArrayProp->Inner->ExportTextItem_Direct( Value, ArrayHelper.GetRawPtr(i), ArrayHelper.GetRawPtr(i), Obj, PPF_IncludeTransient );
 								Ar.Logf(TEXT("  %s(%i)=%s"), *ArrayProp->GetName(), i, *Value);
 							}
 
@@ -13197,7 +13197,7 @@ void UEngine::SpawnServerActors(UWorld *World)
 						if(	FCString::Stricmp(*It->GetName(),Str)==0
 							&&	(It->PropertyFlags & CPF_Config) )
 						{
-							It->ImportText( Value, It->ContainerPtrToValuePtr<uint8>(Actor), 0, Actor );
+							It->ImportText_InContainer( Value, Actor, Actor, 0 );
 						}
 					}
 				}
diff --git a/Engine/Source/Runtime/Engine/Private/UnrealExporter.cpp b/Engine/Source/Runtime/Engine/Private/UnrealExporter.cpp
--- a/Engine/Source/Runtime/Engine/Private/UnrealExporter.cpp
+++ b/Engine/Source/Runtime/Engine/Private/UnrealExporter.cpp
@@ -756,7 +756,7 @@ void ExportProperties
 						bool bExportItem = DiffData == NULL || (!bHasDiffData && (DynamicArrayIndex == ArrayHelper.Num()-1)) || (DiffData != SourceData && !InnerProp->Identical(SourceData, DiffData, ExportFlags));
 						if (bExportItem)
 						{
-							InnerProp->ExportTextItem(Value, SourceData, DiffData, Parent, ExportFlags, ExportRootScope);
+							InnerProp->ExportTextItem_Direct(Value, SourceData, DiffData, Parent, ExportFlags, ExportRootScope);
 							if (ExportObjectProp)
 							{
 								UObject* Obj = ExportObjectProp->GetObjectPropertyValue(ArrayHelper.GetRawPtr(DynamicArrayIndex));
diff --git a/Engine/Source/Runtime/Experimental/Interchange/Core/Private/InterchangePipelineBase.cpp b/Engine/Source/Runtime/Experimental/Interchange/Core/Private/InterchangePipelineBase.cpp
--- a/Engine/Source/Runtime/Experimental/Interchange/Core/Private/InterchangePipelineBase.cpp
+++ b/Engine/Source/Runtime/Experimental/Interchange/Core/Private/InterchangePipelineBase.cpp
@@ -62,7 +62,7 @@ void UInterchangePipelineBase::LoadSettings(const FName PipelineStackName)
 					ArrayHelper.EmptyAndAddValues(List.Num());
 					for (int32 i = List.Num() - 1, c = 0; i >= 0; i--, c++)
 					{
-						Array->Inner->ImportText(*List[i].GetValue(), ArrayHelper.GetRawPtr(c), PortFlags, this);
+						Array->Inner->ImportText_Direct(*List[i].GetValue(), ArrayHelper.GetRawPtr(c), this, PortFlags);
 					}
 				}
 				else
@@ -87,7 +87,7 @@ void UInterchangePipelineBase::LoadSettings(const FName PipelineStackName)
 						{
 							// expand the array if necessary so that Index is a valid element
 							ArrayHelper.ExpandForIndex(Index);
-							Array->Inner->ImportText(*ElementValue->GetValue(), ArrayHelper.GetRawPtr(Index), PortFlags, this);
+							Array->Inner->ImportText_Direct(*ElementValue->GetValue(), ArrayHelper.GetRawPtr(Index), this, PortFlags);
 						}
 
 						Index++;
@@ -120,7 +120,7 @@ void UInterchangePipelineBase::LoadSettings(const FName PipelineStackName)
 
 				if (bFoundValue)
 				{
-					if (Property->ImportText(*Value, Property->ContainerPtrToValuePtr<uint8>(this, i), PortFlags, this) == NULL)
+					if (Property->ImportText_Direct(*Value, Property->ContainerPtrToValuePtr<uint8>(this, i), this, PortFlags) == NULL)
 					{
 						// this should be an error as the properties from the .ini / .int file are not correctly being read in and probably are affecting things in subtle ways
 						UE_LOG(LogInterchangeCore, Error, TEXT("UInterchangePipeline (class:%s) failed to load settings. Property: %s Value: %s"), *this->GetClass()->GetName(), *Property->GetName(), *Value);
@@ -162,7 +162,7 @@ void UInterchangePipelineBase::SaveSettings(const FName PipelineStackName)
 			for (int32 i = 0; i < ArrayHelper.Num(); i++)
 			{
 				FString	Buffer;
-				Array->Inner->ExportTextItem(Buffer, ArrayHelper.GetRawPtr(i), ArrayHelper.GetRawPtr(i), this, PortFlags);
+				Array->Inner->ExportTextItem_Direct(Buffer, ArrayHelper.GetRawPtr(i), ArrayHelper.GetRawPtr(i), this, PortFlags);
 				Sec->Add(*Key, *Buffer);
 			}
 		}
diff --git a/Engine/Source/Runtime/Experimental/Interchange/Core/Private/InterchangeResult.cpp b/Engine/Source/Runtime/Experimental/Interchange/Core/Private/InterchangeResult.cpp
--- a/Engine/Source/Runtime/Experimental/Interchange/Core/Private/InterchangeResult.cpp
+++ b/Engine/Source/Runtime/Experimental/Interchange/Core/Private/InterchangeResult.cpp
@@ -23,7 +23,7 @@ FString UInterchangeResult::ToJson()
 		FProperty* Field = *FieldIterator;
 		FString PropertyName = Field->GetNameCPP();
 		FString PropertyValue;
-		Field->ExportTextItem(PropertyValue, Field->ContainerPtrToValuePtr<uint8>(this), nullptr, this, 0);
+		Field->ExportTextItem_InContainer(PropertyValue, this, nullptr, this, 0);
 
 		TSharedPtr<FJsonObject> Property = MakeShared<FJsonObject>();
 		Property->SetStringField(TEXT("Name"), PropertyName);
@@ -90,7 +90,7 @@ UInterchangeResult* UInterchangeResult::FromJson(const FString& JsonString)
 		}
 
 		FProperty* Field = ClassToImport->FindPropertyByName(FName(*PropertyName));
-		Field->ImportText(*PropertyValue, Field->ContainerPtrToValuePtr<uint8>(Result), 0, Result);
+		Field->ImportText_InContainer(*PropertyValue, Result, Result, 0);
 	}
 
 	return Result;
diff --git a/Engine/Source/Runtime/GameplayTags/Private/GameplayTagContainer.cpp b/Engine/Source/Runtime/GameplayTags/Private/GameplayTagContainer.cpp
--- a/Engine/Source/Runtime/GameplayTags/Private/GameplayTagContainer.cpp
+++ b/Engine/Source/Runtime/GameplayTags/Private/GameplayTagContainer.cpp
@@ -1642,7 +1642,7 @@ FString UEditableGameplayTagQuery::GetTagQueryExportText(FGameplayTagQuery const
 	FProperty* const TQProperty = FindFProperty<FProperty>(GetClass(), TEXT("TagQueryExportText_Helper"));
 
 	FString OutString;
-	TQProperty->ExportTextItem(OutString, (void*)&TagQueryExportText_Helper, (void*)&TagQueryExportText_Helper, this, 0);
+	TQProperty->ExportTextItem_Direct(OutString, (void*)&TagQueryExportText_Helper, (void*)&TagQueryExportText_Helper, this, 0);
 	return OutString;
 }
 
diff --git a/Engine/Source/Runtime/JsonUtilities/Private/JsonObjectConverter.cpp b/Engine/Source/Runtime/JsonUtilities/Private/JsonObjectConverter.cpp
--- a/Engine/Source/Runtime/JsonUtilities/Private/JsonObjectConverter.cpp
+++ b/Engine/Source/Runtime/JsonUtilities/Private/JsonObjectConverter.cpp
@@ -134,7 +134,7 @@ TSharedPtr<FJsonValue> ConvertScalarFPropertyToJsonValue(FProperty* Property, co
 					FString KeyString;
 					if (!KeyElement->TryGetString(KeyString))
 					{
-						MapProperty->KeyProp->ExportTextItem(KeyString, Helper.GetKeyPtr(i), nullptr, nullptr, 0);
+						MapProperty->KeyProp->ExportTextItem_Direct(KeyString, Helper.GetKeyPtr(i), nullptr, nullptr, 0);
 						if (KeyString.IsEmpty())
 						{
 							UE_LOG(LogJson, Error, TEXT("Unable to convert key to string for property %s."), *MapProperty->GetAuthoredName())
@@ -193,7 +193,7 @@ TSharedPtr<FJsonValue> ConvertScalarFPropertyToJsonValue(FProperty* Property, co
 		else
 		{
 			FString StringValue;
-			Property->ExportTextItem(StringValue, Value, nullptr, nullptr, PPF_None);
+			Property->ExportTextItem_Direct(StringValue, Value, nullptr, nullptr, PPF_None);
 			return MakeShared<FJsonValueString>(StringValue);
 		}
 	}
@@ -201,7 +201,7 @@ TSharedPtr<FJsonValue> ConvertScalarFPropertyToJsonValue(FProperty* Property, co
 	{
 		// Default to export as string for everything else
 		FString StringValue;
-		Property->ExportTextItem(StringValue, Value, NULL, NULL, PPF_None);
+		Property->ExportTextItem_Direct(StringValue, Value, NULL, NULL, PPF_None);
 		return MakeShared<FJsonValueString>(StringValue);
 	}
 
@@ -676,7 +676,7 @@ namespace
 			if (!TheCppStructOps->ImportTextItem(ImportTextPtr, OutValue, PPF_None, nullptr, (FOutputDevice*)GWarn))
 			{
 				// Fall back to trying the tagged property approach if custom ImportTextItem couldn't get it done
-				if (Property->ImportText(ImportTextPtr, OutValue, PPF_None, nullptr) == nullptr)
+				if (Property->ImportText_Direct(ImportTextPtr, OutValue, nullptr, PPF_None) == nullptr)
 				{
 					UE_LOG(LogJson, Error, TEXT("JsonValueToUProperty - FJsonObjectConverter::JsonObjectToUStruct failed for property %s"), *Property->GetNameCPP());
 					return false;
@@ -687,7 +687,7 @@ namespace
 		{
 			FString ImportTextString = JsonValue->AsString();
 			const TCHAR* ImportTextPtr = *ImportTextString;
-			if (Property->ImportText(ImportTextPtr, OutValue, PPF_None, nullptr) == nullptr)
+			if (Property->ImportText_Direct(ImportTextPtr, OutValue, nullptr, PPF_None) == nullptr)
 			{
 				UE_LOG(LogJson, Error, TEXT("JsonValueToUProperty - FJsonObjectConverter::JsonObjectToUStruct failed for property %s"), *Property->GetNameCPP());
 				return false;
@@ -739,7 +739,7 @@ namespace
 			else if (JsonValue->Type == EJson::String)
 			{
 				// Default to expect a string for everything else
-				if (Property->ImportText(*JsonValue->AsString(), OutValue, 0, nullptr) == nullptr)
+				if (Property->ImportText_Direct(*JsonValue->AsString(), OutValue, nullptr, 0) == nullptr)
 				{
 					UE_LOG(LogJson, Error, TEXT("JsonValueToUProperty - Unable import property type %s from string value for property %s"), *Property->GetClass()->GetName(), *Property->GetNameCPP());
 					return false;
@@ -749,7 +749,7 @@ namespace
 	else
 	{
 		// Default to expect a string for everything else
-		if (Property->ImportText(*JsonValue->AsString(), OutValue, 0, nullptr) == nullptr)
+		if (Property->ImportText_Direct(*JsonValue->AsString(), OutValue, nullptr, 0) == nullptr)
 		{
 			UE_LOG(LogJson, Error, TEXT("JsonValueToUProperty - Unable import property type %s from string value for property %s"), *Property->GetClass()->GetName(), *Property->GetNameCPP());
 			return false;
diff --git a/Engine/Source/Runtime/Net/Core/Private/Net/Core/Connection/StateStruct.cpp b/Engine/Source/Runtime/Net/Core/Private/Net/Core/Connection/StateStruct.cpp
--- a/Engine/Source/Runtime/Net/Core/Private/Net/Core/Connection/StateStruct.cpp
+++ b/Engine/Source/Runtime/Net/Core/Private/Net/Core/Connection/StateStruct.cpp
@@ -190,7 +190,7 @@ bool UStatePerObjectConfig::LoadStructConfig(FStructOnScope& OutStruct, const TC
 					{
 						if (CurProp->HasAllPropertyFlags(CPF_Config))
 						{
-							CurProp->ImportText(*Value, CurProp->ContainerPtrToValuePtr<void>(OutStruct.GetStructMemory()), 0, nullptr);
+							CurProp->ImportText_InContainer(*Value, OutStruct.GetStructMemory(), nullptr, 0);
 						}
 						else
 						{
@@ -230,9 +230,7 @@ void UStatePerObjectConfig::DebugDump() const
 		for (TFieldIterator<FProperty> It(GetClass()); It; ++It)
 		{
 			FString TextValue;
-			const uint8* PropAddr = It->ContainerPtrToValuePtr<uint8>(this);
-
-			It->ExportTextItem(TextValue, PropAddr, nullptr, nullptr, PPF_DebugDump, nullptr);
+			It->ExportTextItem_InContainer(TextValue, this, nullptr, nullptr, PPF_DebugDump, nullptr);
 
 			UE_LOG(LogNetCore, Log, TEXT(" - %s: %s"), ToCStr(It->GetName()), ToCStr(TextValue));
 		}
@@ -268,9 +266,7 @@ void UStatePerObjectConfig::DebugDump() const
 					if (It->GetName() != TEXT("StateName"))
 					{
 						FString TextValue;
-						const uint8* PropAddr = It->ContainerPtrToValuePtr<uint8>(CurStateConfig.Get());
-
-						It->ExportTextItem(TextValue, PropAddr, nullptr, nullptr, PPF_DebugDump, nullptr);
+						It->ExportTextItem_InContainer(TextValue, CurStateConfig.Get(), nullptr, nullptr, PPF_DebugDump, nullptr);
 
 						UE_LOG(LogNetCore, Log, TEXT("   - %s: %s"), ToCStr(It->GetName()), ToCStr(TextValue));
 					}
diff --git a/Engine/Source/Runtime/PropertyPath/Private/PropertyPathHelpers.cpp b/Engine/Source/Runtime/PropertyPath/Private/PropertyPathHelpers.cpp
--- a/Engine/Source/Runtime/PropertyPath/Private/PropertyPathHelpers.cpp
+++ b/Engine/Source/Runtime/PropertyPath/Private/PropertyPathHelpers.cpp
@@ -180,7 +180,7 @@ namespace PropertyPathHelpersInternal
 						ReturnProperty->InitializeValue(TempBuffer.GetData());
 
 						InContainer->ProcessEvent(InFunction, TempBuffer.GetData());
-						ReturnProperty->ExportTextItem(OutValue, TempBuffer.GetData(), nullptr, nullptr, 0);
+						ReturnProperty->ExportTextItem_Direct(OutValue, TempBuffer.GetData(), nullptr, nullptr, 0);
 						return true;
 					}
 				}
@@ -209,17 +209,17 @@ namespace PropertyPathHelpersInternal
 				if ( ArrayHelper.IsValidIndex(ArrayIndex) )
 				{
 					OutProperty = ArrayProp->Inner;
-					OutProperty->ExportTextItem(OutValue, static_cast<void*>(ArrayHelper.GetRawPtr(ArrayIndex)), nullptr, nullptr, 0);
+					OutProperty->ExportTextItem_Direct(OutValue, static_cast<void*>(ArrayHelper.GetRawPtr(ArrayIndex)), nullptr, nullptr, 0);
 					return true;
 				}
 			}
 			else
 			{
 				// No index, so assume we want the array property itself
-				if ( void* ValuePtr = ArrayProp->ContainerPtrToValuePtr<void>(InContainer) )
+				if ( !!ArrayProp->ContainerPtrToValuePtr<void>(InContainer) )
 				{
 					OutProperty = ArrayProp;
-					OutProperty->ExportTextItem(OutValue, ValuePtr, nullptr, nullptr, 0);
+					OutProperty->ExportTextItem_InContainer(OutValue, InContainer, nullptr, nullptr, 0);
 					return true;
 				}
 			}
@@ -236,7 +236,7 @@ namespace PropertyPathHelpersInternal
 				if ( void* ValuePtr = Property->ContainerPtrToValuePtr<void>(InContainer, ArrayIndex) )
 				{
 					OutProperty = Property;
-					OutProperty->ExportTextItem(OutValue, ValuePtr, nullptr, nullptr, 0);
+					OutProperty->ExportTextItem_InContainer(OutValue, InContainer, nullptr, nullptr, 0);
 					return true;
 				}
 			}
@@ -275,7 +275,7 @@ namespace PropertyPathHelpersInternal
 						TempBuffer.AddUninitialized(ParamProperty->ElementSize);
 						ParamProperty->InitializeValue(TempBuffer.GetData());
 
-						ParamProperty->ImportText(*InValue, TempBuffer.GetData(), 0, nullptr);
+						ParamProperty->ImportText_Direct(*InValue, TempBuffer.GetData(), nullptr, 0);
 						InContainer->ProcessEvent(InFunction, TempBuffer.GetData());
 						return true;
 					}
@@ -304,16 +304,16 @@ namespace PropertyPathHelpersInternal
 				FScriptArrayHelper_InContainer ArrayHelper(ArrayProp, InContainer);
 				if ( ArrayHelper.IsValidIndex(ArrayIndex) )
 				{
-					ArrayProp->Inner->ImportText(*InValue, static_cast<void*>(ArrayHelper.GetRawPtr(ArrayIndex)), 0, nullptr);
+					ArrayProp->Inner->ImportText_Direct(*InValue, static_cast<void*>(ArrayHelper.GetRawPtr(ArrayIndex)), nullptr, 0);
 					return true;
 				}
 			}
 			else
 			{
 				// No index, so assume we want the array property itself
-				if ( void* ValuePtr = ArrayProp->ContainerPtrToValuePtr<void>(InContainer) )
+				if ( !!ArrayProp->ContainerPtrToValuePtr<void>(InContainer) )
 				{
-					ArrayProp->ImportText(*InValue, ValuePtr, 0, nullptr);
+					ArrayProp->ImportText_InContainer(*InValue, InContainer, nullptr, 0);
 					return true;
 				}
 			}
@@ -329,7 +329,7 @@ namespace PropertyPathHelpersInternal
 			{
 				if ( void* ValuePtr = Property->ContainerPtrToValuePtr<void>(InContainer, ArrayIndex) )
 				{
-					Property->ImportText(*InValue, ValuePtr, 0, nullptr);
+					Property->ImportText_Direct(*InValue, ValuePtr, nullptr, 0);
 					return true;
 				}
 			}
diff --git a/Engine/Source/Runtime/RigVM/Private/RigVMCore/RigVM.cpp b/Engine/Source/Runtime/RigVM/Private/RigVMCore/RigVM.cpp
--- a/Engine/Source/Runtime/RigVM/Private/RigVMCore/RigVM.cpp
+++ b/Engine/Source/Runtime/RigVM/Private/RigVMCore/RigVM.cpp
@@ -2220,7 +2220,7 @@ bool URigVM::Execute(TArrayView<URigVMMemoryStorage*> Memory, TArrayView<void*>
 					else
 					{
 						FString Value;
-						ElementPropertyA->ExportTextItem(Value, ArrayHelperA.GetRawPtr(Index), nullptr, nullptr, PPF_None);
+						ElementPropertyA->ExportTextItem_Direct(Value, ArrayHelperA.GetRawPtr(Index), nullptr, nullptr, PPF_None);
 						HashValue = TextKeyUtil::HashString(Value);
 					}
 					
@@ -2239,7 +2239,7 @@ bool URigVM::Execute(TArrayView<URigVMMemoryStorage*> Memory, TArrayView<void*>
 					else
 					{
 						FString Value;
-						ElementPropertyB->ExportTextItem(Value, ArrayHelperB.GetRawPtr(Index), nullptr, nullptr, PPF_None);
+						ElementPropertyB->ExportTextItem_Direct(Value, ArrayHelperB.GetRawPtr(Index), nullptr, nullptr, PPF_None);
 						HashValue = TextKeyUtil::HashString(Value);
 					}
 					if(!HashB.Contains(HashValue))
diff --git a/Engine/Source/Runtime/RigVM/Private/RigVMCore/RigVMMemoryStorage.cpp b/Engine/Source/Runtime/RigVM/Private/RigVMCore/RigVMMemoryStorage.cpp
--- a/Engine/Source/Runtime/RigVM/Private/RigVMCore/RigVMMemoryStorage.cpp
+++ b/Engine/Source/Runtime/RigVM/Private/RigVMCore/RigVMMemoryStorage.cpp
@@ -377,7 +377,7 @@ URigVMMemoryStorageGeneratorClass* URigVMMemoryStorageGeneratorClass::CreateStor
 		const FProperty* Property = LinkedProperties[PropertyIndex];
 		uint8* ValuePtr = Property->ContainerPtrToValuePtr<uint8>(CDO);
 
-		Property->ImportText(*DefaultValue, ValuePtr, EPropertyPortFlags::PPF_None, nullptr);
+		Property->ImportText_Direct(*DefaultValue, ValuePtr, nullptr, EPropertyPortFlags::PPF_None);
 	}
 
 	return Class;
@@ -641,7 +641,7 @@ FString URigVMMemoryStorage::GetDataAsString(int32 InPropertyIndex, int32 PortFl
 	const uint8* Data = GetData<uint8>(InPropertyIndex);
 
 	FString Value;
-	GetProperties()[InPropertyIndex]->ExportTextItem(Value, Data, nullptr, nullptr, PortFlags);
+	GetProperties()[InPropertyIndex]->ExportTextItem_Direct(Value, Data, nullptr, nullptr, PortFlags);
 	return Value;
 }
 
@@ -696,7 +696,7 @@ bool URigVMMemoryStorage::SetDataFromString(int32 InPropertyIndex, const FString
 	uint8* Data = GetData<uint8>(InPropertyIndex);
 	
 	FRigVMMemoryStorageImportErrorContext ErrorPipe;
-	GetProperties()[InPropertyIndex]->ImportText(*InValue, Data, EPropertyPortFlags::PPF_None, nullptr, &ErrorPipe);
+	GetProperties()[InPropertyIndex]->ImportText_Direct(*InValue, Data, nullptr, EPropertyPortFlags::PPF_None, &ErrorPipe);
 	return ErrorPipe.NumErrors == 0;
 }
 
diff --git a/Engine/Source/Runtime/RigVM/Private/RigVMCore/RigVMStruct.cpp b/Engine/Source/Runtime/RigVM/Private/RigVMCore/RigVMStruct.cpp
--- a/Engine/Source/Runtime/RigVM/Private/RigVMCore/RigVMStruct.cpp
+++ b/Engine/Source/Runtime/RigVM/Private/RigVMCore/RigVMStruct.cpp
@@ -407,7 +407,7 @@ FString FRigVMStruct::ExportToFullyQualifiedText(FProperty* InMemberProperty, co
 	}
 	else
 	{
-		InMemberProperty->ExportTextItem(DefaultValue, InMemberMemoryPtr, nullptr, nullptr, PPF_None);
+		InMemberProperty->ExportTextItem_Direct(DefaultValue, InMemberMemoryPtr, nullptr, nullptr, PPF_None);
 
 		if (CastField<FNameProperty>(InMemberProperty) != nullptr ||
 			CastField<FStrProperty>(InMemberProperty) != nullptr)
diff --git a/Engine/Source/Runtime/Serialization/Private/Backends/CborStructSerializerBackend.cpp b/Engine/Source/Runtime/Serialization/Private/Backends/CborStructSerializerBackend.cpp
--- a/Engine/Source/Runtime/Serialization/Private/Backends/CborStructSerializerBackend.cpp
+++ b/Engine/Source/Runtime/Serialization/Private/Backends/CborStructSerializerBackend.cpp
@@ -49,7 +49,7 @@ void FCborStructSerializerBackend::BeginArray(const FStructSerializerState& Stat
 	else if (State.KeyProperty != nullptr)
 	{
 		FString KeyString;
-		State.KeyProperty->ExportTextItem(KeyString, State.KeyData, nullptr, nullptr, PPF_None);
+		State.KeyProperty->ExportTextItem_Direct(KeyString, State.KeyData, nullptr, nullptr, PPF_None);
 		CborWriter.WriteValue(KeyString);
 	}
 	// Array nested in Object
@@ -80,7 +80,7 @@ void FCborStructSerializerBackend::BeginStructure(const FStructSerializerState&
 		else if (State.KeyProperty != nullptr)
 		{
 			FString KeyString;
-			State.KeyProperty->ExportTextItem(KeyString, State.KeyData, nullptr, nullptr, PPF_None);
+			State.KeyProperty->ExportTextItem_Direct(KeyString, State.KeyData, nullptr, nullptr, PPF_None);
 			CborWriter.WriteValue(KeyString);
 			CborWriter.WriteContainerStart(ECborCode::Map, -1/*Indefinite*/);
 		}
@@ -141,7 +141,7 @@ namespace CborStructSerializerBackend
 		else if (State.KeyProperty != nullptr)
 		{
 			FString KeyString;
-			State.KeyProperty->ExportTextItem(KeyString, State.KeyData, nullptr, nullptr, PPF_None);
+			State.KeyProperty->ExportTextItem_Direct(KeyString, State.KeyData, nullptr, nullptr, PPF_None);
 			CborWriter.WriteValue(KeyString);
 			CborWriter.WriteValue(Value);
 		}
@@ -167,7 +167,7 @@ namespace CborStructSerializerBackend
 		else if (State.KeyProperty != nullptr)
 		{
 			FString KeyString;
-			State.KeyProperty->ExportTextItem(KeyString, State.KeyData, nullptr, nullptr, PPF_None);
+			State.KeyProperty->ExportTextItem_Direct(KeyString, State.KeyData, nullptr, nullptr, PPF_None);
 			CborWriter.WriteValue(KeyString);
 			CborWriter.WriteNull();
 		}
diff --git a/Engine/Source/Runtime/Serialization/Private/Backends/JsonStructSerializerBackend.cpp b/Engine/Source/Runtime/Serialization/Private/Backends/JsonStructSerializerBackend.cpp
--- a/Engine/Source/Runtime/Serialization/Private/Backends/JsonStructSerializerBackend.cpp
+++ b/Engine/Source/Runtime/Serialization/Private/Backends/JsonStructSerializerBackend.cpp
@@ -19,7 +19,7 @@ void FJsonStructSerializerBackend::BeginArray(const FStructSerializerState& Stat
 	else if (State.KeyProperty != nullptr)
 	{
 		FString KeyString;
-		State.KeyProperty->ExportTextItem(KeyString, State.KeyData, nullptr, nullptr, PPF_None);
+		State.KeyProperty->ExportTextItem_Direct(KeyString, State.KeyData, nullptr, nullptr, PPF_None);
 		JsonWriter->WriteArrayStart(KeyString);
 	}
 	else
@@ -44,7 +44,7 @@ void FJsonStructSerializerBackend::BeginStructure(const FStructSerializerState&
 		else if (State.KeyProperty != nullptr)
 		{
 			FString KeyString;
-			State.KeyProperty->ExportTextItem(KeyString, State.KeyData, nullptr, nullptr, PPF_None);
+			State.KeyProperty->ExportTextItem_Direct(KeyString, State.KeyData, nullptr, nullptr, PPF_None);
 			JsonWriter->WriteObjectStart(KeyString);
 		}
 		else
diff --git a/Engine/Source/Runtime/Serialization/Private/StructDeserializer.cpp b/Engine/Source/Runtime/Serialization/Private/StructDeserializer.cpp
--- a/Engine/Source/Runtime/Serialization/Private/StructDeserializer.cpp
+++ b/Engine/Source/Runtime/Serialization/Private/StructDeserializer.cpp
@@ -252,7 +252,7 @@ bool FStructDeserializer::Deserialize( void* OutStruct, UStruct& TypeInfo, IStru
 					int32 PairIndex = MapHelper.AddDefaultValue_Invalid_NeedsRehash();
 					uint8* PairPtr = MapHelper.GetPairPtr(PairIndex);
 
-					MapProperty->KeyProp->ImportText(*PropertyName, PairPtr, PPF_None, nullptr);
+					MapProperty->KeyProp->ImportText_Direct(*PropertyName, PairPtr, nullptr, PPF_None);
 
 					if (!Backend.ReadProperty(Property, CurrentState.Property, PairPtr, CurrentState.ArrayIndex))
 					{
@@ -362,7 +362,7 @@ bool FStructDeserializer::Deserialize( void* OutStruct, UStruct& TypeInfo, IStru
 					NewState.Data = PairPtr + MapHelper.MapLayout.ValueOffset;
 					NewState.Property = MapProperty->ValueProp;
 
-					MapProperty->KeyProp->ImportText(*PropertyName, PairPtr, PPF_None, nullptr);
+					MapProperty->KeyProp->ImportText_Direct(*PropertyName, PairPtr, nullptr, PPF_None);
 				}
 				else
 				{
diff --git a/Engine/Source/Runtime/Serialization/Public/Backends/JsonStructSerializerBackend.h b/Engine/Source/Runtime/Serialization/Public/Backends/JsonStructSerializerBackend.h
--- a/Engine/Source/Runtime/Serialization/Public/Backends/JsonStructSerializerBackend.h
+++ b/Engine/Source/Runtime/Serialization/Public/Backends/JsonStructSerializerBackend.h
@@ -78,7 +78,7 @@ protected:
 		else if (State.KeyProperty != nullptr)
 		{
 			FString KeyString;
-			State.KeyProperty->ExportTextItem(KeyString, State.KeyData, nullptr, nullptr, PPF_None);
+			State.KeyProperty->ExportTextItem_Direct(KeyString, State.KeyData, nullptr, nullptr, PPF_None);
 			JsonWriter->WriteValue(KeyString, Value);
 		}
 		//Write PropertyName:Value for any other cases (single array element, single property, etc...)
@@ -102,7 +102,7 @@ protected:
 		else if (State.KeyProperty != nullptr)
 		{
 			FString KeyString;
-			State.KeyProperty->ExportTextItem(KeyString, State.KeyData, nullptr, nullptr, PPF_None);
+			State.KeyProperty->ExportTextItem_Direct(KeyString, State.KeyData, nullptr, nullptr, PPF_None);
 			JsonWriter->WriteNull(KeyString);
 		}
 		else
diff --git a/Engine/Source/Runtime/WebBrowser/Private/CEF/CEFJSStructSerializerBackend.cpp b/Engine/Source/Runtime/WebBrowser/Private/CEF/CEFJSStructSerializerBackend.cpp
--- a/Engine/Source/Runtime/WebBrowser/Private/CEF/CEFJSStructSerializerBackend.cpp
+++ b/Engine/Source/Runtime/WebBrowser/Private/CEF/CEFJSStructSerializerBackend.cpp
@@ -125,7 +125,7 @@ void FCEFJSStructSerializerBackend::BeginStructure(const FStructSerializerState&
 	if (State.KeyProperty != nullptr)
 	{
 		FString KeyString;
-		State.KeyProperty->ExportTextItem(KeyString, State.KeyData, nullptr, nullptr, PPF_None);
+		State.KeyProperty->ExportTextItem_Direct(KeyString, State.KeyData, nullptr, nullptr, PPF_None);
 
 		CefRefPtr<CefDictionaryValue> DictionaryValue = CefDictionaryValue::Create();
 		Stack.Push(StackItem(KeyString, DictionaryValue));
diff --git a/Engine/Source/Runtime/WebBrowser/Private/MobileJS/MobileJSStructSerializerBackend.cpp b/Engine/Source/Runtime/WebBrowser/Private/MobileJS/MobileJSStructSerializerBackend.cpp
--- a/Engine/Source/Runtime/WebBrowser/Private/MobileJS/MobileJSStructSerializerBackend.cpp
+++ b/Engine/Source/Runtime/WebBrowser/Private/MobileJS/MobileJSStructSerializerBackend.cpp
@@ -34,7 +34,7 @@ void FMobileJSStructSerializerBackend::WriteUObject(const FStructSerializerState
 	else if (State.KeyProperty != nullptr)
 	{
 		FString KeyString;
-		State.KeyProperty->ExportTextItem(KeyString, State.KeyData, nullptr, nullptr, PPF_None);
+		State.KeyProperty->ExportTextItem_Direct(KeyString, State.KeyData, nullptr, nullptr, PPF_None);
 		GetWriter()->WriteRawJSONValue(KeyString, RawValue);
 	}
 	else
