Changes that somehow got lost when shelving and unshelving property setter and getter changes across streams + a few fixes

diff --git a/Engine/Source/Programs/UnrealHeaderTool/Private/PropertyTypes.cpp b/Engine/Source/Programs/UnrealHeaderTool/Private/PropertyTypes.cpp
--- a/Engine/Source/Programs/UnrealHeaderTool/Private/PropertyTypes.cpp
+++ b/Engine/Source/Programs/UnrealHeaderTool/Private/PropertyTypes.cpp
@@ -1595,6 +1595,10 @@ struct FPropertyTypeTraitsObjectPtrReference : public FPropertyTypeTraitsObjectB
 		{
 			return FString::Printf(TEXT("TObjectPtr<%s%s>"), VarProperty.ClassDef->GetPrefixCPP(), *VarProperty.ClassDef->GetName());
 		}
+		else if ((CPPExportFlags & CPPF_ArgumentOrReturnValue) == CPPF_ArgumentOrReturnValue)
+		{
+			return FString::Printf(TEXT("%s%s*"), VarProperty.ClassDef->GetPrefixCPP(), *VarProperty.ClassDef->GetName());
+		}
 		else
 		{
 			return FString::Printf(TEXT("TObjectPtr<%s%s>"), VarProperty.ClassDef->GetPrefixCPP(), *VarProperty.ClassDef->GetName());
diff --git a/Engine/Source/Runtime/CoreUObject/Private/UObject/PropertyArray.cpp b/Engine/Source/Runtime/CoreUObject/Private/UObject/PropertyArray.cpp
--- a/Engine/Source/Runtime/CoreUObject/Private/UObject/PropertyArray.cpp
+++ b/Engine/Source/Runtime/CoreUObject/Private/UObject/PropertyArray.cpp
@@ -617,6 +617,7 @@ const TCHAR* FArrayProperty::ImportText_Internal(const TCHAR* Buffer, void* Cont
 		{
 			InitializeValue(TempArrayStorage);
 		}
+		ArrayPtr = TempArrayStorage;
 	}
 	else
 	{
diff --git a/Engine/Source/Runtime/CoreUObject/Private/UObject/PropertyBool.cpp b/Engine/Source/Runtime/CoreUObject/Private/UObject/PropertyBool.cpp
--- a/Engine/Source/Runtime/CoreUObject/Private/UObject/PropertyBool.cpp
+++ b/Engine/Source/Runtime/CoreUObject/Private/UObject/PropertyBool.cpp
@@ -348,15 +348,16 @@ void FBoolProperty::ExportText_Internal( FString& ValueStr, const void* Containe
 {
 	check(FieldSize != 0);
 	uint8 LocalByteValue = 0;
+	bool bValue = false;
 	if (PropertyPointerType == EPropertyPointerType::Container && HasGetter())
 	{
-		GetValue_InContainer(ContainerOrPropertyPtr, &LocalByteValue);
+		GetValue_InContainer(ContainerOrPropertyPtr, &bValue);
 	}
 	else
 	{
 		LocalByteValue = *((uint8*)PointerToValuePtr(ContainerOrPropertyPtr, PropertyPointerType) + ByteOffset);
+		bValue = 0 != (LocalByteValue & FieldMask);
 	}
-	const bool bValue = 0 != (LocalByteValue & FieldMask);
 	const TCHAR* Temp = nullptr;
 	if (0 != (PortFlags & PPF_ExportCpp))
 	{
@@ -381,7 +382,9 @@ const TCHAR* FBoolProperty::ImportText_Internal( const TCHAR* Buffer, void* Cont
 	uint8 LocalByteValue = 0;
 	if (PropertyPointerType == EPropertyPointerType::Container && HasGetter())
 	{
-		GetValue_InContainer(ContainerOrPropertyPtr, &LocalByteValue);
+		bool bValue = false;
+		GetValue_InContainer(ContainerOrPropertyPtr, &bValue);
+		LocalByteValue = bValue ? FieldMask : 0;
 	}
 	else
 	{
@@ -406,7 +409,8 @@ const TCHAR* FBoolProperty::ImportText_Internal( const TCHAR* Buffer, void* Cont
 
 	if (PropertyPointerType == EPropertyPointerType::Container && HasSetter())
 	{
-		SetValue_InContainer(ContainerOrPropertyPtr, &LocalByteValue);
+		bool bValue = 0 != (LocalByteValue & FieldMask);
+		SetValue_InContainer(ContainerOrPropertyPtr, &bValue);
 	}
 	else
 	{
