API exposed some PropertyAccessUtil functions

diff --git a/Engine/Source/Runtime/CoreUObject/Public/UObject/PropertyAccessUtil.h b/Engine/Source/Runtime/CoreUObject/Public/UObject/PropertyAccessUtil.h
--- a/Engine/Source/Runtime/CoreUObject/Public/UObject/PropertyAccessUtil.h
+++ b/Engine/Source/Runtime/CoreUObject/Public/UObject/PropertyAccessUtil.h
@@ -291,4 +291,66 @@ namespace PropertyAccessUtil
 	 * @return The found property, or null if the property cannot be found.
 	 */
 	COREUOBJECT_API FProperty* FindPropertyByName(const FName InPropName, const UStruct* InStruct);
+
+	/**
+	 * Low-level function to determine if two properties are compatible with each other (for use by CopySinglePropertyValue/CopyCompletePropertyValue).
+	 * 
+	 * @param InSrcProp Source property being compared.
+	 * @param InDestProp Destination property being compared.
+	 * 
+	 * @return True if the properties are compatible with each other, false otherwise.
+	 */
+	COREUOBJECT_API bool ArePropertiesCompatible(const FProperty* InSrcProp, const FProperty* InDestProp);
+
+	/**
+	 * Low-level function to determine if two properties have identical values (single property value, ignoring ArrayDim).
+	 * @note This function calls ArePropertiesCompatible internally.
+	 * 
+	 * @param InSrcProp Source property being compared.
+	 * @param InSrcValue Source value being compared.
+	 * @param InDestProp Destination property being compared.
+	 * @param InDestValue Destination value being compared.
+	 * 
+	 * @return True if the properties have identical values, false otherwise.
+	 */
+	COREUOBJECT_API bool IsSinglePropertyIdentical(const FProperty* InSrcProp, const void* InSrcValue, const FProperty* InDestProp, const void* InDestValue);
+
+	/**
+	 * Low-level function to determine if two properties have identical values (complete property value, for each ArrayDim).
+	 * @note This function calls ArePropertiesCompatible internally.
+	 *
+	 * @param InSrcProp Source property being compared.
+	 * @param InSrcValue Source value being compared.
+	 * @param InDestProp Destination property being compared.
+	 * @param InDestValue Destination value being compared.
+	 *
+	 * @return True if the properties have identical values, false otherwise.
+	 */
+	COREUOBJECT_API bool IsCompletePropertyIdentical(const FProperty* InSrcProp, const void* InSrcValue, const FProperty* InDestProp, const void* InDestValue);
+
+	/**
+	 * Low-level function to copy the value from one property to another (single property value, ignoring ArrayDim).
+	 * @note This function calls ArePropertiesCompatible internally.
+	 *
+	 * @param InSrcProp Source property being copied from.
+	 * @param InSrcValue Source value being copied from.
+	 * @param InDestProp Destination property being copied to.
+	 * @param InDestValue Destination value being copied to.
+	 *
+	 * @return True if the properties value was copied, false otherwise.
+	 */
+	COREUOBJECT_API bool CopySinglePropertyValue(const FProperty* InSrcProp, const void* InSrcValue, const FProperty* InDestProp, void* InDestValue);
+	
+	/**
+	 * Low-level function to copy the value from one property to another (complete property value, for each ArrayDim).
+	 * @note This function calls ArePropertiesCompatible internally.
+	 *
+	 * @param InSrcProp Source property being copied from.
+	 * @param InSrcValue Source value being copied from.
+	 * @param InDestProp Destination property being copied to.
+	 * @param InDestValue Destination value being copied to.
+	 *
+	 * @return True if the properties value was copied, false otherwise.
+	 */
+	COREUOBJECT_API bool CopyCompletePropertyValue(const FProperty* InSrcProp, const void* InSrcValue, const FProperty* InDestProp, void* InDestValue);
 }
