Fix test for licensee flag in Engine content.

diff --git a/Engine/Source/Runtime/CoreUObject/Private/UObject/LinkerLoad.cpp b/Engine/Source/Runtime/CoreUObject/Private/UObject/LinkerLoad.cpp
--- a/Engine/Source/Runtime/CoreUObject/Private/UObject/LinkerLoad.cpp
+++ b/Engine/Source/Runtime/CoreUObject/Private/UObject/LinkerLoad.cpp
@@ -1237,7 +1237,7 @@ FLinkerLoad::ELinkerStatus FLinkerLoad::SerializePackageFileSummaryInternal()
 		if (FEngineVersion::Current().IsLicenseeVersion() == false && Summary.CompatibleWithEngineVersion.IsLicenseeVersion())
 		{
 			// Only warn about things under Engine and Engine/Plugins so licensee projects can be opened
-			FString LocalFilename = GetPackagePath().GetLocalFullPath();
+			FString LocalFilename = FPaths::CreateStandardFilename(GetPackagePath().GetLocalFullPath());
 			bool IsEngineContent = LocalFilename.StartsWith(FPaths::EngineContentDir()) || LocalFilename.StartsWith(FPaths::EnginePluginsDir());
 
 			if (IsEngineContent)
