Make Subobject Class Overrides (more) const correct

diff --git a/Engine/Source/Runtime/CoreUObject/Private/UObject/UObjectGlobals.cpp b/Engine/Source/Runtime/CoreUObject/Private/UObject/UObjectGlobals.cpp
--- a/Engine/Source/Runtime/CoreUObject/Private/UObject/UObjectGlobals.cpp
+++ b/Engine/Source/Runtime/CoreUObject/Private/UObject/UObjectGlobals.cpp
@@ -3586,7 +3586,7 @@ void FObjectInitializer::InitProperties(UObject* Obj, UClass* DefaultsClass, UOb
 	}
 }
 
-void FObjectInitializer::FOverrides::Add(FName InComponentName, UClass* InComponentClass, const TArrayView<const FName>* InFullComponentPath)
+void FObjectInitializer::FOverrides::Add(FName InComponentName, const UClass* InComponentClass, const TArrayView<const FName>* InFullComponentPath)
 {
 	auto GetSubobjectPath = [InComponentName, InFullComponentPath]()
 	{
@@ -3650,7 +3650,7 @@ void FObjectInitializer::FOverrides::Add(FName InComponentName, UClass* InCompon
 	}
 }
 
-void FObjectInitializer::FOverrides::Add(FStringView InComponentPath, UClass* InComponentClass)
+void FObjectInitializer::FOverrides::Add(FStringView InComponentPath, const UClass* InComponentClass)
 {
 	TArray<FName> ComponentPath;
 
@@ -3666,7 +3666,7 @@ void FObjectInitializer::FOverrides::Add(FStringView InComponentPath, UClass* In
 	Add(PathArrayView, InComponentClass, &PathArrayView);
 }
 
-void FObjectInitializer::FOverrides::Add(TArrayView<const FName> InComponentPath, UClass* InComponentClass, const TArrayView<const FName>* InFullComponentPath)
+void FObjectInitializer::FOverrides::Add(TArrayView<const FName> InComponentPath, const UClass* InComponentClass, const TArrayView<const FName>* InFullComponentPath)
 {
 	if (InComponentPath.Num() > 1)
 	{
@@ -3690,7 +3690,7 @@ void FObjectInitializer::FOverrides::Add(TArrayView<const FName> InComponentPath
 }
 
 /**  Retrieve an override, or TClassToConstructByDefault::StaticClass or nullptr if this was removed by a derived class **/
-FObjectInitializer::FOverrides::FOverrideDetails FObjectInitializer::FOverrides::Get(FName InComponentName, UClass* ReturnType, UClass* ClassToConstructByDefault, bool bOptional) const
+FObjectInitializer::FOverrides::FOverrideDetails FObjectInitializer::FOverrides::Get(FName InComponentName, const UClass* ReturnType, const UClass* ClassToConstructByDefault, bool bOptional) const
 {
 	FOverrideDetails Result;
 
@@ -4466,7 +4466,7 @@ UScriptStruct* GetFallbackStruct()
 	return TBaseStructure<FFallbackStruct>::Get();
 }
 
-UObject* FObjectInitializer::CreateDefaultSubobject(UObject* Outer, FName SubobjectFName, UClass* ReturnType, UClass* ClassToCreateByDefault, bool bIsRequired, bool bIsTransient) const
+UObject* FObjectInitializer::CreateDefaultSubobject(UObject* Outer, FName SubobjectFName, const UClass* ReturnType, const UClass* ClassToCreateByDefault, bool bIsRequired, bool bIsTransient) const
 {
 	UE_CLOG(!FUObjectThreadContext::Get().IsInConstructor, LogUObjectGlobals, Fatal, TEXT("Subobjects cannot be created outside of UObject constructors. UObject constructing subobjects cannot be created using new or placement new operator."));
 	if (SubobjectFName == NAME_None)
@@ -4476,7 +4476,7 @@ UObject* FObjectInitializer::CreateDefaultSubobject(UObject* Outer, FName Subobj
 
 	UObject* Result = nullptr;
 	FOverrides::FOverrideDetails ComponentOverride = SubobjectOverrides.Get(SubobjectFName, ReturnType, ClassToCreateByDefault, !bIsRequired);
-	UClass* OverrideClass = ComponentOverride.Class;
+	const UClass* OverrideClass = ComponentOverride.Class;
 	if (OverrideClass)
 	{
 		check(OverrideClass->IsChildOf(ReturnType));
@@ -4555,7 +4555,7 @@ UObject* FObjectInitializer::CreateDefaultSubobject(UObject* Outer, FName Subobj
 	}
 	return Result;
 }
-UObject* FObjectInitializer::CreateEditorOnlyDefaultSubobject(UObject* Outer, FName SubobjectName, UClass* ReturnType, bool bTransient /*= false*/) const
+UObject* FObjectInitializer::CreateEditorOnlyDefaultSubobject(UObject* Outer, FName SubobjectName, const UClass* ReturnType, bool bTransient /*= false*/) const
 {
 #if WITH_EDITOR
 	if (GIsEditor)
diff --git a/Engine/Source/Runtime/CoreUObject/Public/UObject/Class.h b/Engine/Source/Runtime/CoreUObject/Public/UObject/Class.h
--- a/Engine/Source/Runtime/CoreUObject/Public/UObject/Class.h
+++ b/Engine/Source/Runtime/CoreUObject/Public/UObject/Class.h
@@ -3136,9 +3136,9 @@ public:
 	UObject* GetDefaultSubobjectByName(FName ToFind);
 
 	/** Adds a new default instance map item **/
-	void AddDefaultSubobject(UObject* NewSubobject, UClass* BaseClass)
+	void AddDefaultSubobject(UObject* NewSubobject, const UClass* BaseClass)
 	{
-		// this compoonent must be a derived class of the base class
+		// this component must be a derived class of the base class
 		check(NewSubobject->IsA(BaseClass));
 		// the outer of the component must be of my class or some superclass of me
 		check(IsChildOf(NewSubobject->GetOuter()->GetClass()));
@@ -3486,7 +3486,7 @@ public:
 	TArray<UObject*> Timelines;
 
 	/** Array of blueprint overrides of component classes in parent classes */
-	TArray<TPair<FName, UClass*>> ComponentClassOverrides;
+	TArray<TPair<FName, const UClass*>> ComponentClassOverrides;
 
 	/** IAnimClassInterface (UAnimClassData) or null */
 	UObject* AnimClassImplementation;
diff --git a/Engine/Source/Runtime/CoreUObject/Public/UObject/UObjectGlobals.h b/Engine/Source/Runtime/CoreUObject/Public/UObject/UObjectGlobals.h
--- a/Engine/Source/Runtime/CoreUObject/Public/UObject/UObjectGlobals.h
+++ b/Engine/Source/Runtime/CoreUObject/Public/UObject/UObjectGlobals.h
@@ -920,7 +920,7 @@ public:
 	template<class TReturnType>
 	TReturnType* CreateEditorOnlyDefaultSubobject(UObject* Outer, FName SubobjectName, bool bTransient = false) const
 	{
-		UClass* ReturnType = TReturnType::StaticClass();
+		const UClass* ReturnType = TReturnType::StaticClass();
 		return static_cast<TReturnType*>(CreateEditorOnlyDefaultSubobject(Outer, SubobjectName, ReturnType, bTransient));
 	}
 
@@ -931,7 +931,7 @@ public:
 	* @param	ReturnType					type of the new component
 	* @param	bTransient					true if the component is being assigned to a transient property
 	*/
-	UObject* CreateEditorOnlyDefaultSubobject(UObject* Outer, FName SubobjectName, UClass* ReturnType, bool bTransient = false) const;
+	UObject* CreateEditorOnlyDefaultSubobject(UObject* Outer, FName SubobjectName, const UClass* ReturnType, bool bTransient = false) const;
 
 	/**
 	 * Create a component or subobject
@@ -942,14 +942,14 @@ public:
 	 * @param	bIsRequired                 true if the component is required and will always be created even if DoNotCreateDefaultSubobject was specified.
 	 * @param	bIsTransient                true if the component is being assigned to a transient property
 	 */
-	UObject* CreateDefaultSubobject(UObject* Outer, FName SubobjectFName, UClass* ReturnType, UClass* ClassToCreateByDefault, bool bIsRequired = true, bool bIsTransient = false) const;
+	UObject* CreateDefaultSubobject(UObject* Outer, FName SubobjectFName, const UClass* ReturnType, const UClass* ClassToCreateByDefault, bool bIsRequired = true, bool bIsTransient = false) const;
 
 	/**
 	 * Sets the class to use for a subobject defined in a base class, the class must be a subclass of the class used by the base class.
 	 * @param	SubobjectName	name of the new component or subobject
 	 * @param	Class			The class to use for the specified subobject or component.
 	 */
-	const FObjectInitializer& SetDefaultSubobjectClass(FName SubobjectName, UClass* Class) const
+	const FObjectInitializer& SetDefaultSubobjectClass(FName SubobjectName, const UClass* Class) const
 	{
 		AssertIfSubobjectSetupIsNotAllowed(SubobjectName);
 		SubobjectOverrides.Add(SubobjectName, Class);
@@ -982,7 +982,7 @@ public:
 	 * @param	SubobjectName	path to the new component or subobject
 	 * @param	Class			The class to use for the specified subobject or component.
 	 */
-	const FObjectInitializer& SetNestedDefaultSubobjectClass(FStringView SubobjectName, UClass* Class) const
+	const FObjectInitializer& SetNestedDefaultSubobjectClass(FStringView SubobjectName, const UClass* Class) const
 	{
 		AssertIfSubobjectSetupIsNotAllowed(SubobjectName);
 		SubobjectOverrides.Add(SubobjectName, Class);
@@ -994,7 +994,7 @@ public:
 	 * @param	SubobjectName	path to the new component or subobject
 	 * @param	Class			The class to use for the specified subobject or component.
 	 */
-	const FObjectInitializer& SetNestedDefaultSubobjectClass(TArrayView<const FName> SubobjectNames, UClass* Class) const
+	const FObjectInitializer& SetNestedDefaultSubobjectClass(TArrayView<const FName> SubobjectNames, const UClass* Class) const
 	{
 		AssertIfSubobjectSetupIsNotAllowed(SubobjectNames);
 		SubobjectOverrides.Add(SubobjectNames, Class);
@@ -1112,22 +1112,22 @@ private:
 	struct FOverrides
 	{
 		/**  Add an override, make sure it is legal **/
-		void Add(FName InComponentName, UClass* InComponentClass, const TArrayView<const FName>* FullPath = nullptr);
+		void Add(FName InComponentName, const UClass* InComponentClass, const TArrayView<const FName>* FullPath = nullptr);
 
 		/**  Add a potentially nested override, make sure it is legal **/
-		void Add(FStringView InComponentPath, UClass* InComponentClass);
+		void Add(FStringView InComponentPath, const UClass* InComponentClass);
 
 		/**  Add a potentially nested override, make sure it is legal **/
-		void Add(TArrayView<const FName> InComponentPath, UClass* InComponentClass, const TArrayView<const FName>* FullPath = nullptr);
+		void Add(TArrayView<const FName> InComponentPath, const UClass* InComponentClass, const TArrayView<const FName>* FullPath = nullptr);
 
 		struct FOverrideDetails
 		{
-			UClass* Class = nullptr;
+			const UClass* Class = nullptr;
 			FOverrides* SubOverrides = nullptr;
 		};
 
 		/**  Retrieve an override, or TClassToConstructByDefault::StaticClass or nullptr if this was removed by a derived class **/
-		FOverrideDetails Get(FName InComponentName, UClass* ReturnType, UClass* ClassToConstructByDefault, bool bOptional) const;
+		FOverrideDetails Get(FName InComponentName, const UClass* ReturnType, const UClass* ClassToConstructByDefault, bool bOptional) const;
 
 	private:
 		static bool IsLegalOverride(const UClass* DerivedComponentClass, const UClass* BaseComponentClass);
@@ -1148,7 +1148,7 @@ private:
 		struct FOverride
 		{
 			FName ComponentName;
-			UClass* ComponentClass = nullptr;
+			const UClass* ComponentClass = nullptr;
 			TUniquePtr<FOverrides> SubOverrides;
 			bool bDoNotCreate = false;
 			FOverride(FName InComponentName)
diff --git a/Engine/Source/Runtime/Engine/Classes/Engine/BlueprintGeneratedClass.h b/Engine/Source/Runtime/Engine/Classes/Engine/BlueprintGeneratedClass.h
--- a/Engine/Source/Runtime/Engine/Classes/Engine/BlueprintGeneratedClass.h
+++ b/Engine/Source/Runtime/Engine/Classes/Engine/BlueprintGeneratedClass.h
@@ -608,14 +608,14 @@ struct FBPComponentClassOverride
 
 	/** The class to use when constructing the component. */
 	UPROPERTY()
-	TObjectPtr<UClass> ComponentClass;
+	TObjectPtr<const UClass> ComponentClass;
 
 	FBPComponentClassOverride()
 		: ComponentClass(nullptr)
 	{
 	}
 
-	FBPComponentClassOverride(FName InComponentName, UClass* InComponentClass)
+	FBPComponentClassOverride(FName InComponentName, const UClass* InComponentClass)
 		: ComponentName(InComponentName)
 		, ComponentClass(InComponentClass)
 	{
diff --git a/Engine/Source/Runtime/Engine/Private/BlueprintGeneratedClass.cpp b/Engine/Source/Runtime/Engine/Private/BlueprintGeneratedClass.cpp
--- a/Engine/Source/Runtime/Engine/Private/BlueprintGeneratedClass.cpp
+++ b/Engine/Source/Runtime/Engine/Private/BlueprintGeneratedClass.cpp
@@ -1534,7 +1534,7 @@ void UBlueprintGeneratedClass::GetDefaultObjectPreloadDependencies(TArray<UObjec
 		{
 			if (Override.ComponentClass)
 			{
-				OutDeps.Add(Override.ComponentClass);
+				OutDeps.Add(const_cast<UClass*>(Override.ComponentClass.Get()));
 			}
 		}
 	}
