Optimize FCoreRedirects::GetMatchingRedirects since it is showing up as a moderately high cost in cooks (38s out of 870s in CitySample).
#rnx

diff --git a/Engine/Source/Runtime/CoreUObject/Private/UObject/CoreRedirects.cpp b/Engine/Source/Runtime/CoreUObject/Private/UObject/CoreRedirects.cpp
--- a/Engine/Source/Runtime/CoreUObject/Private/UObject/CoreRedirects.cpp
+++ b/Engine/Source/Runtime/CoreUObject/Private/UObject/CoreRedirects.cpp
@@ -81,15 +81,19 @@ bool FCoreRedirectObjectName::Matches(const FCoreRedirectObjectName& Other, bool
 	if (bCheckSubstring)
 	{
 		// Much slower
-		if (ObjectName != NAME_None && !Other.ObjectName.ToString().Contains(ObjectName.ToString()))
+		auto NameContains = [](FName Text, FName LookFor)
+		{
+			return WriteToString<FName::StringBufferSize>(Text).ToView().Contains(WriteToString<FName::StringBufferSize>(LookFor));
+		};
+		if (ObjectName != NAME_None && !NameContains(Other.ObjectName, ObjectName))
 		{
 			return false;
 		}
-		if (OuterName != NAME_None && !Other.OuterName.ToString().Contains(OuterName.ToString()))
+		if (OuterName != NAME_None && !NameContains(Other.OuterName, OuterName))
 		{
 			return false;
 		}
-		if (PackageName != NAME_None && !Other.PackageName.ToString().Contains(PackageName.ToString()))
+		if (PackageName != NAME_None && !NameContains(Other.PackageName, PackageName))
 		{
 			return false;
 		}
@@ -444,9 +448,46 @@ bool FCoreRedirects::bIsInMultithreadedPhase = false;
 #endif
 
 TMap<FName, ECoreRedirectFlags> FCoreRedirects::ConfigKeyMap;
-TMap<ECoreRedirectFlags, FCoreRedirects::FRedirectNameMap> FCoreRedirects::RedirectTypeMap;
+FCoreRedirects::FRedirectTypeMap FCoreRedirects::RedirectTypeMap;
 FRWLock FCoreRedirects::KnownMissingLock;
 
+FCoreRedirects::FRedirectNameMap& FCoreRedirects::FRedirectTypeMap::FindOrAdd(ECoreRedirectFlags Key)
+{
+	FRedirectNameMap*& NameMap = Map.FindOrAdd(Key);
+	if (NameMap)
+	{
+		return *NameMap;
+	}
+
+	TPair<ECoreRedirectFlags, FRedirectNameMap>* OldData = FastIterable.GetData();
+	TPair<ECoreRedirectFlags, FRedirectNameMap>& NewPair =
+		FastIterable.Emplace_GetRef(Key, FRedirectNameMap());
+	if (FastIterable.GetData() == OldData)
+	{
+		NameMap = &NewPair.Value;
+	}
+	else
+	{
+		Map.Reset();
+		for (TPair<ECoreRedirectFlags, FRedirectNameMap>& Pair : FastIterable)
+		{
+			Map.Add(Pair.Key, &Pair.Value);
+		}
+	}
+	return NewPair.Value;
+}
+
+FCoreRedirects::FRedirectNameMap* FCoreRedirects::FRedirectTypeMap::Find(ECoreRedirectFlags Key)
+{
+	return Map.FindRef(Key);
+}
+
+void FCoreRedirects::FRedirectTypeMap::Empty()
+{
+	Map.Empty();
+	FastIterable.Empty();
+}
+
 void FCoreRedirects::Initialize()
 {
 	if (bInitialized)
@@ -481,7 +522,7 @@ void FCoreRedirects::Initialize()
 #if WITH_COREREDIRECTS_MULTITHREAD_WARNING
 	check(!bIsInMultithreadedPhase);
 #endif
-	RedirectTypeMap.Add(ECoreRedirectFlags::Type_Package | ECoreRedirectFlags::Category_Removed | ECoreRedirectFlags::Option_MissingLoad);
+	RedirectTypeMap.FindOrAdd(ECoreRedirectFlags::Type_Package | ECoreRedirectFlags::Category_Removed | ECoreRedirectFlags::Option_MissingLoad);
 
 	// Enable to run startup tests
 	//RunTests();
@@ -616,7 +657,7 @@ bool FCoreRedirects::FindPreviousNames(ECoreRedirectFlags SearchFlags, const FCo
 		// We need to check all maps that match the search or package flags
 		if (CheckRedirectFlagsMatch(PairFlags, SearchFlags) || (bSearchPackageRedirects && CheckRedirectFlagsMatch(PairFlags, ECoreRedirectFlags::Type_Package)))
 		{
-			for (const TPair<FName, TArray<FCoreRedirect> >& RedirectPair : Pair.Value.RedirectMap)
+			for (const TPair<FName, TArray<FCoreRedirect>>& RedirectPair : Pair.Value.RedirectMap)
 			{
 				for (const FCoreRedirect& Redirect : RedirectPair.Value)
 				{
@@ -687,7 +728,7 @@ void FCoreRedirects::ClearKnownMissing(ECoreRedirectFlags Type, ECoreRedirectFla
 bool FCoreRedirects::RunTests()
 {
 	bool bSuccess = true;
-	TMap<ECoreRedirectFlags, FRedirectNameMap > BackupMap = RedirectTypeMap;
+	FRedirectTypeMap BackupMap = MoveTemp(RedirectTypeMap);
 #if WITH_COREREDIRECTS_MULTITHREAD_WARNING
 	bool BackupIsInMultithreadedPhase = bIsInMultithreadedPhase;
 	bIsInMultithreadedPhase = false;
@@ -833,7 +874,7 @@ bool FCoreRedirects::RunTests()
 	}
 
 	// Restore old state
-	RedirectTypeMap = BackupMap;
+	RedirectTypeMap = MoveTemp(BackupMap);
 #if WITH_COREREDIRECTS_MULTITHREAD_WARNING
 	bIsInMultithreadedPhase = BackupIsInMultithreadedPhase;
 #endif
diff --git a/Engine/Source/Runtime/CoreUObject/Public/UObject/CoreRedirects.h b/Engine/Source/Runtime/CoreUObject/Public/UObject/CoreRedirects.h
--- a/Engine/Source/Runtime/CoreUObject/Public/UObject/CoreRedirects.h
+++ b/Engine/Source/Runtime/CoreUObject/Public/UObject/CoreRedirects.h
@@ -323,7 +323,20 @@ private:
 	static TMap<FName, ECoreRedirectFlags> ConfigKeyMap;
 
 	/** Map from name of thing being mapped to full list. List must be filtered further */
-	static TMap<ECoreRedirectFlags, FRedirectNameMap> RedirectTypeMap;
+	struct FRedirectTypeMap
+	{
+	public:
+		FRedirectNameMap& FindOrAdd(ECoreRedirectFlags Key);
+		FRedirectNameMap* Find(ECoreRedirectFlags Key);
+		void Empty();
+
+		TArray<TPair<ECoreRedirectFlags, FRedirectNameMap>>::RangedForIteratorType begin() { return FastIterable.begin(); }
+		TArray<TPair<ECoreRedirectFlags, FRedirectNameMap>>::RangedForIteratorType end() { return FastIterable.end(); }
+	private:
+		TMap<ECoreRedirectFlags, FRedirectNameMap*> Map;
+		TArray<TPair<ECoreRedirectFlags, FRedirectNameMap>> FastIterable;
+	};
+	static FRedirectTypeMap RedirectTypeMap;
 
 	/**
 	 * Lock to protect multithreaded access to *KnownMissing functions, which can be called from the async loading threads. 
