Downgrade SavePackage messages that the package did not save because it was editor-only from display to verbose. They cause a lot of spam on projects using WorldPartition (or OneFilePerActor packages) because OneFilePerActor packages are editor-only.
#rnx

diff --git a/Engine/Source/Runtime/CoreUObject/Private/UObject/SavePackage2.cpp b/Engine/Source/Runtime/CoreUObject/Private/UObject/SavePackage2.cpp
--- a/Engine/Source/Runtime/CoreUObject/Private/UObject/SavePackage2.cpp
+++ b/Engine/Source/Runtime/CoreUObject/Private/UObject/SavePackage2.cpp
@@ -102,12 +102,12 @@ ESavePackageResult ValidatePackage(FSaveContext& SaveContext)
 			// Don't save packages marked as editor-only.
 			if (SaveContext.CanSkipEditorReferencedPackagesWhenCooking() && SaveContext.GetPackage()->IsLoadedByEditorPropertiesOnly())
 			{
-				UE_CLOG(SaveContext.IsGenerateSaveError(), LogSavePackage, Display, TEXT("Package loaded by editor-only properties: %s. Package will not be saved."), *SaveContext.GetPackage()->GetName());
+				UE_CLOG(SaveContext.IsGenerateSaveError(), LogSavePackage, Verbose, TEXT("Package loaded by editor-only properties: %s. Package will not be saved."), *SaveContext.GetPackage()->GetName());
 				return ESavePackageResult::ReferencedOnlyByEditorOnlyData;
 			}
 			else if (SaveContext.GetPackage()->HasAnyPackageFlags(PKG_EditorOnly))
 			{
-				UE_CLOG(SaveContext.IsGenerateSaveError(), LogSavePackage, Display, TEXT("Package marked as editor-only: %s. Package will not be saved."), *SaveContext.GetPackage()->GetName());
+				UE_CLOG(SaveContext.IsGenerateSaveError(), LogSavePackage, Verbose, TEXT("Package marked as editor-only: %s. Package will not be saved."), *SaveContext.GetPackage()->GetName());
 				return ESavePackageResult::ReferencedOnlyByEditorOnlyData;
 			}
 		}
