Make a check for corrupted tags in an asset platform generic (instead of just on Windows)

diff --git a/Engine/Source/Runtime/CoreUObject/Private/UObject/LinkerLoad.cpp b/Engine/Source/Runtime/CoreUObject/Private/UObject/LinkerLoad.cpp
--- a/Engine/Source/Runtime/CoreUObject/Private/UObject/LinkerLoad.cpp
+++ b/Engine/Source/Runtime/CoreUObject/Private/UObject/LinkerLoad.cpp
@@ -1323,7 +1323,6 @@ FLinkerLoad::ELinkerStatus FLinkerLoad::SerializePackageFileSummaryInternal()
 		}
 	}
 
-#if PLATFORM_WINDOWS
 	if (!FPlatformProperties::RequiresCookedData() &&
 		// We can't check the post tag if the file is an EDL cooked package
 		!((Summary.GetPackageFlags() & PKG_FilterEditorOnly) && Summary.PreloadDependencyCount > 0 && Summary.PreloadDependencyOffset > 0)
@@ -1352,7 +1351,6 @@ FLinkerLoad::ELinkerStatus FLinkerLoad::SerializePackageFileSummaryInternal()
 		// seek back to the position after the package summary
 		Seek(OriginalOffset);
 	}
-#endif
 
 	return LINKER_Loaded;
 }
