Fixed FSoftObjectPath::GetLongPackageName returning an empty string when set to a package name

If a FSoftObjectPath was set to a package name (eg, "/Game/MyPackage") then GetLongPackageName would return an empty string which is inconsistent with GetLongPackageFName.

This change makes the behavior of GetLongPackageName and GetLongPackageFName consistent (returning the package name if there's no object in the path), and also removes a temporary FString allocation from GetLongPackageName and GetAssetName.

diff --git a/Engine/Source/Runtime/CoreUObject/Private/UObject/SoftObjectPath.cpp b/Engine/Source/Runtime/CoreUObject/Private/UObject/SoftObjectPath.cpp
--- a/Engine/Source/Runtime/CoreUObject/Private/UObject/SoftObjectPath.cpp
+++ b/Engine/Source/Runtime/CoreUObject/Private/UObject/SoftObjectPath.cpp
@@ -58,6 +58,25 @@ void FSoftObjectPath::ToString(FStringBuilderBase& Builder) const
 	}
 }
 
+FString FSoftObjectPath::GetLongPackageName() const
+{
+	if (!AssetPathName.IsNone())
+	{
+		FNameBuilder AssetPathNameBuilder(AssetPathName);
+		FStringView AssetPathNameView = AssetPathNameBuilder.ToView();
+
+		int32 DotPos = INDEX_NONE;
+		if (AssetPathNameView.FindChar(TEXT('.'), DotPos))
+		{
+			return FString(AssetPathNameView.Left(DotPos));
+		}
+
+		return FString(AssetPathNameView);
+	}
+
+	return FString();
+}
+
 FName FSoftObjectPath::GetLongPackageFName() const
 {
 	TCHAR Buffer[NAME_SIZE];
@@ -66,6 +85,23 @@ FName FSoftObjectPath::GetLongPackageFName() const
 	return DotPos == INDEX_NONE ? AssetPathName : FName(PlainAssetPath.Left(DotPos));
 }
 
+FString FSoftObjectPath::GetAssetName() const
+{
+	if (!AssetPathName.IsNone())
+	{
+		FNameBuilder AssetPathNameBuilder(AssetPathName);
+		FStringView AssetPathNameView = AssetPathNameBuilder.ToView();
+
+		int32 DotPos = INDEX_NONE;
+		if (AssetPathNameView.FindChar(TEXT('.'), DotPos))
+		{
+			return FString(AssetPathNameView.Mid(DotPos + 1));
+		}
+	}
+
+	return FString();
+}
+
 /** Helper function that adds info about the object currently being serialized when triggering an ensure about invalid soft object path */
 static FString GetObjectBeingSerializedForSoftObjectPath()
 {
diff --git a/Engine/Source/Runtime/CoreUObject/Public/UObject/SoftObjectPath.h b/Engine/Source/Runtime/CoreUObject/Public/UObject/SoftObjectPath.h
--- a/Engine/Source/Runtime/CoreUObject/Public/UObject/SoftObjectPath.h
+++ b/Engine/Source/Runtime/CoreUObject/Public/UObject/SoftObjectPath.h
@@ -115,23 +115,13 @@ struct COREUOBJECT_API FSoftObjectPath
 	}
 
 	/** Returns /package/path, leaving off the asset name and sub object */
-	FString GetLongPackageName() const
-	{
-		FString PackageName;
-		GetAssetPathString().Split(TEXT("."), &PackageName, nullptr, ESearchCase::CaseSensitive, ESearchDir::FromStart);
-		return PackageName;
-	}
+	FString GetLongPackageName() const;
 
 	/** Returns /package/path, leaving off the asset name and sub object */
 	FName GetLongPackageFName() const;
 
 	/** Returns assetname string, leaving off the /package/path part and sub object */
-	FString GetAssetName() const
-	{
-		FString AssetName;
-		GetAssetPathString().Split(TEXT("."), nullptr, &AssetName, ESearchCase::CaseSensitive, ESearchDir::FromStart);
-		return AssetName;
-	}
+	FString GetAssetName() const;
 
 	/** Sets asset path of this reference based on a string path */
 	void SetPath(FWideStringView Path);
