UE-141802 - Set Editor Property fails with float and object argument types

Editor properties go through the PropertyAccessUtil interface to change property values. However, they don't support potential float<->double conversions that might occur. This has become more prevalent since Blueprint floats default to double precision.

We can fix this by simply adding explicit float<->double conversions to PropertyAccessUtil. Unfortunately, lossiness will potentially occur since there's a narrowing conversion, but we're already expecting some of that when interfacing with native float properties.

diff --git a/Engine/Source/Runtime/CoreUObject/Private/UObject/PropertyAccessUtil.cpp b/Engine/Source/Runtime/CoreUObject/Private/UObject/PropertyAccessUtil.cpp
--- a/Engine/Source/Runtime/CoreUObject/Private/UObject/PropertyAccessUtil.cpp
+++ b/Engine/Source/Runtime/CoreUObject/Private/UObject/PropertyAccessUtil.cpp
@@ -5,6 +5,98 @@
 #include "UObject/Object.h"
 #include "UObject/Class.h"
 
+namespace UE::PropertyAccessUtil::Private
+{
+
+bool IsRealNumberConversion(const FProperty* InSrcProp, const FProperty* InDestProp)
+{
+	check(InSrcProp);
+	check(InDestProp);
+
+	const bool bIsRealNumberConversion =
+		(InSrcProp->IsA<FDoubleProperty>() && InDestProp->IsA<FFloatProperty>()) ||
+		(InSrcProp->IsA<FFloatProperty>() && InDestProp->IsA<FDoubleProperty>());
+
+	return bIsRealNumberConversion;
+}
+
+void ConvertRealNumber(const FProperty* InSrcProp, const void* InSrcValue, const FProperty* InDestProp, void* InDestValue, int InCount)
+{
+	check(InSrcProp);
+	check(InSrcValue);
+	check(InDestProp);
+	check(InDestValue);
+	check(InCount > 0);
+	check(InCount <= InSrcProp->ArrayDim);
+
+	if (const FDoubleProperty* SrcDoubleProp = CastField<FDoubleProperty>(InSrcProp))
+	{
+		const FFloatProperty* DestFloatProp = CastFieldChecked<FFloatProperty>(InDestProp);
+		for (int32 Idx = 0; Idx < InCount; ++Idx)
+		{
+			const void* SrcElemValue = static_cast<const uint8*>(InSrcValue) + (InSrcProp->ElementSize * Idx);
+			void* DestElemValue = static_cast<uint8*>(InDestValue) + (InDestProp->ElementSize * Idx);
+
+			const double Value = SrcDoubleProp->GetFloatingPointPropertyValue(SrcElemValue);
+			DestFloatProp->SetFloatingPointPropertyValue(DestElemValue, Value);
+		}
+	}
+	else if (const FFloatProperty* SrcFloatProp = CastField<FFloatProperty>(InSrcProp))
+	{
+		const FDoubleProperty* DestDoubleProp = CastFieldChecked<FDoubleProperty>(InDestProp);
+		for (int32 Idx = 0; Idx < InCount; ++Idx)
+		{
+			const void* SrcElemValue = static_cast<const uint8*>(InSrcValue) + (InSrcProp->ElementSize * Idx);
+			void* DestElemValue = static_cast<uint8*>(InDestValue) + (InDestProp->ElementSize * Idx);
+
+			const double Value = SrcFloatProp->GetFloatingPointPropertyValue(SrcElemValue);
+			DestDoubleProp->SetFloatingPointPropertyValue(DestElemValue, Value);
+		}
+	}
+	else
+	{
+		checkf(false, TEXT("Invalid property type used with ConvertRealNumber!"));
+	}
+}
+
+bool AreRealNumbersIdentical(const FProperty* InSrcProp, const void* InSrcValue, const FProperty* InDestProp, const void* InDestValue)
+{
+	check(InSrcProp);
+	check(InSrcValue);
+	check(InDestProp);
+	check(InDestValue);
+
+	bool bIdentical = false;
+
+	if (const FDoubleProperty* SrcDoubleProp = CastField<FDoubleProperty>(InSrcProp))
+	{
+		const FFloatProperty* DestFloatProp = CastFieldChecked<FFloatProperty>(InDestProp);
+
+		const double SrcValue = SrcDoubleProp->GetFloatingPointPropertyValue(InSrcValue);
+		const double DestValue = DestFloatProp->GetFloatingPointPropertyValue(InDestValue);
+
+		bIdentical = (SrcValue == DestValue);
+		
+	}
+	else if (const FFloatProperty* SrcFloatProp = CastField<FFloatProperty>(InSrcProp))
+	{
+		const FDoubleProperty* DestDoubleProp = CastFieldChecked<FDoubleProperty>(InDestProp);
+		
+		const double SrcValue = SrcFloatProp->GetFloatingPointPropertyValue(InSrcValue);
+		const double DestValue = DestDoubleProp->GetFloatingPointPropertyValue(InDestValue);
+
+		bIdentical = (SrcValue == DestValue);
+	}
+	else
+	{
+		checkf(false, TEXT("Invalid property type used with AreRealNumbersIdentical!"));
+	}
+
+	return bIdentical;
+}
+
+}
+
 namespace PropertyAccessUtil
 {
 
@@ -51,6 +143,8 @@ bool SetPropertyEnumValue(const FProperty* InProp, void* InPropValue, const int6
 
 bool ArePropertiesCompatible(const FProperty* InSrcProp, const FProperty* InDestProp)
 {
+	using namespace UE::PropertyAccessUtil::Private;
+
 	// Enum properties can either be a ByteProperty with an enum set, or an EnumProperty
 	// We allow coercion between these two types if they're using the same enum type
 	if (const UEnum* DestEnumType = GetPropertyEnumType(InDestProp))
@@ -70,6 +164,11 @@ bool ArePropertiesCompatible(const FProperty* InSrcProp, const FProperty* InDest
 		}
 	}
 
+	if (IsRealNumberConversion(InSrcProp, InDestProp))
+	{
+		return true;
+	}
+
 	// Compare the classes as these must be an *exact* match as the access is low-level and without property coercion
 	// HOWEVER:
 	// Object properties can either be an ObjectProperty, or an ObjectPtrProperty
@@ -109,6 +208,8 @@ bool ArePropertiesCompatible(const FProperty* InSrcProp, const FProperty* InDest
 
 bool IsSinglePropertyIdentical(const FProperty* InSrcProp, const void* InSrcValue, const FProperty* InDestProp, const void* InDestValue)
 {
+	using namespace UE::PropertyAccessUtil::Private;
+
 	if (!ArePropertiesCompatible(InSrcProp, InDestProp))
 	{
 		return false;
@@ -123,7 +224,12 @@ bool IsSinglePropertyIdentical(const FProperty* InSrcProp, const void* InSrcValu
 		const bool bDestBoolValue = DestBoolProp->GetPropertyValue(InDestValue);
 		return bSrcBoolValue == bDestBoolValue;
 	}
-	
+
+	if (IsRealNumberConversion(InSrcProp, InDestProp))
+	{
+		return AreRealNumbersIdentical(InSrcProp, InSrcValue, InDestProp, InDestValue);
+	}
+
 	return InSrcProp->Identical(InSrcValue, InDestValue);
 }
 
@@ -141,6 +247,8 @@ bool IsCompletePropertyIdentical(const FProperty* InSrcProp, const void* InSrcVa
 
 bool CopySinglePropertyValue(const FProperty* InSrcProp, const void* InSrcValue, const FProperty* InDestProp, void* InDestValue)
 {
+	using namespace UE::PropertyAccessUtil::Private;
+
 	if (!ArePropertiesCompatible(InSrcProp, InDestProp))
 	{
 		return false;
@@ -176,12 +284,20 @@ bool CopySinglePropertyValue(const FProperty* InSrcProp, const void* InSrcValue,
 		}
 	}
 	
+	if (IsRealNumberConversion(InSrcProp, InDestProp))
+	{
+		ConvertRealNumber(InSrcProp, InSrcValue, InDestProp, InDestValue, 1);
+		return true;
+	}
+
 	InSrcProp->CopySingleValue(InDestValue, InSrcValue);
 	return true;
 }
 
 bool CopyCompletePropertyValue(const FProperty* InSrcProp, const void* InSrcValue, const FProperty* InDestProp, void* InDestValue)
 {
+	using namespace UE::PropertyAccessUtil::Private;
+
 	if (!ArePropertiesCompatible(InSrcProp, InDestProp) || InSrcProp->ArrayDim != InDestProp->ArrayDim)
 	{
 		return false;
@@ -236,6 +352,12 @@ bool CopyCompletePropertyValue(const FProperty* InSrcProp, const void* InSrcValu
 		}
 	}
 	
+	if (IsRealNumberConversion(InSrcProp, InDestProp))
+	{
+		ConvertRealNumber(InSrcProp, InSrcValue, InDestProp, InDestValue, InSrcProp->ArrayDim);
+		return true;
+	}
+
 	InSrcProp->CopyCompleteValue(InDestValue, InSrcValue);
 	return true;
 }
