Fix for mismatched weakptr serialization between FJsonArchiveOutputFormatter and FJsonArchiveInputFormatter.

diff --git a/Engine/Source/Runtime/CoreUObject/Private/Serialization/Formatters/JsonArchiveOutputFormatter.cpp b/Engine/Source/Runtime/CoreUObject/Private/Serialization/Formatters/JsonArchiveOutputFormatter.cpp
--- a/Engine/Source/Runtime/CoreUObject/Private/Serialization/Formatters/JsonArchiveOutputFormatter.cpp
+++ b/Engine/Source/Runtime/CoreUObject/Private/Serialization/Formatters/JsonArchiveOutputFormatter.cpp
@@ -363,14 +363,8 @@ void FJsonArchiveOutputFormatter::Serialize(FText& Value)
 
 void FJsonArchiveOutputFormatter::Serialize(FWeakObjectPtr& Value)
 {
-	if (Value.IsValid() && IsObjectAllowed(Value.Get()))
-	{
-		SerializeStringInternal(FString::Printf(TEXT("Object:%s"), *Value.Get()->GetFullName(nullptr, EObjectFullNameFlags::IncludeClassPackage)));
-	}
-	else
-	{
-		WriteValue(TEXT("null"));
-	}
+	UObject* Ptr = Value.IsValid() ? Value.Get() : nullptr;
+	Serialize(Ptr);
 }
 
 void FJsonArchiveOutputFormatter::Serialize(FSoftObjectPtr& Value)
