Typo fix.

diff --git a/Engine/Source/Runtime/CoreUObject/Public/AssetRegistry/AssetData.h b/Engine/Source/Runtime/CoreUObject/Public/AssetRegistry/AssetData.h
--- a/Engine/Source/Runtime/CoreUObject/Public/AssetRegistry/AssetData.h
+++ b/Engine/Source/Runtime/CoreUObject/Public/AssetRegistry/AssetData.h
@@ -228,7 +228,7 @@ public:
 	{
 		OutFullNameBuilder.Reset();
 		AssetClass.AppendString(OutFullNameBuilder);
-		OutFullNameBuilder.AppendChar(TEXT*' '));
+		OutFullNameBuilder.AppendChar(TEXT(' '));
 		ObjectPath.AppendString(OutFullNameBuilder);
 	}
 
