Fixed incorrect FStructuredArchive operator<< overload for TObjectPtr
#rnx

diff --git a/Engine/Source/Runtime/CoreUObject/Public/UObject/ObjectPtr.h b/Engine/Source/Runtime/CoreUObject/Public/UObject/ObjectPtr.h
--- a/Engine/Source/Runtime/CoreUObject/Public/UObject/ObjectPtr.h
+++ b/Engine/Source/Runtime/CoreUObject/Public/UObject/ObjectPtr.h
@@ -329,7 +329,7 @@ public:
 		return Ar;
 	}
 
-	friend FORCEINLINE void operator<<(FStructuredArchiveSlot&& Slot, TObjectPtr<T>& ObjectPtr)
+	friend FORCEINLINE void operator<<(FStructuredArchiveSlot Slot, TObjectPtr<T>& ObjectPtr)
 	{
 		Slot << static_cast<FObjectPtr&>(ObjectPtr);
 	}
