Static Analysis Fixes

diff --git a/Engine/Source/Runtime/CoreUObject/Private/UObject/CoreNative.cpp b/Engine/Source/Runtime/CoreUObject/Private/UObject/CoreNative.cpp
--- a/Engine/Source/Runtime/CoreUObject/Private/UObject/CoreNative.cpp
+++ b/Engine/Source/Runtime/CoreUObject/Private/UObject/CoreNative.cpp
@@ -268,7 +268,7 @@ UObject* FObjectInstancingGraph::InstancePropertyValue(UObject* SubObjectTemplat
 
 	check(CurrentValue);
 
-	if (CurrentValue->GetClass()->HasAnyClassFlags(CLASS_DefaultToInstanced))
+	if (CurrentValue != nullptr && CurrentValue->GetClass()->HasAnyClassFlags(CLASS_DefaultToInstanced))
 	{
 		bCausesInstancing = true; // these are always instanced no matter what
 	}
@@ -283,7 +283,7 @@ UObject* FObjectInstancingGraph::InstancePropertyValue(UObject* SubObjectTemplat
 	// if the object we're instancing the subobjects for (Owner) has the current subobject's outer in its archetype chain, and its archetype has a nullptr value
 	// for this subobject property it means that the archetype didn't instance its subobject, so we shouldn't either.
 
-	if (SubObjectTemplate == nullptr && CurrentValue != nullptr && (Owner && Owner->IsBasedOnArchetype(CurrentValue->GetOuter())))
+	if (SubObjectTemplate == nullptr && (Owner != nullptr && CurrentValue != nullptr && Owner->IsBasedOnArchetype(CurrentValue->GetOuter())))
 	{
 		NewValue = nullptr;
 	}
diff --git a/Engine/Source/Runtime/CoreUObject/Private/UObject/SavePackage/PackageHarvester.cpp b/Engine/Source/Runtime/CoreUObject/Private/UObject/SavePackage/PackageHarvester.cpp
--- a/Engine/Source/Runtime/CoreUObject/Private/UObject/SavePackage/PackageHarvester.cpp
+++ b/Engine/Source/Runtime/CoreUObject/Private/UObject/SavePackage/PackageHarvester.cpp
@@ -325,7 +325,12 @@ void FPackageHarvester::TryHarvestImport(UObject* InObject)
 {
 	// Those should have been already validated
 	check(InObject);
-	check(!InObject->IsInPackage(SaveContext.GetPackage()));
+	check(InObject && !InObject->IsInPackage(SaveContext.GetPackage()));
+
+	if ( InObject==nullptr )
+	{
+		return;
+	}
 
 	auto IsObjNative = [](UObject* InObj)
 	{
diff --git a/Engine/Source/Runtime/JsonUtilities/Private/JsonObjectWrapper.cpp b/Engine/Source/Runtime/JsonUtilities/Private/JsonObjectWrapper.cpp
--- a/Engine/Source/Runtime/JsonUtilities/Private/JsonObjectWrapper.cpp
+++ b/Engine/Source/Runtime/JsonUtilities/Private/JsonObjectWrapper.cpp
@@ -19,7 +19,11 @@ bool FJsonObjectWrapper::ImportTextItem(const TCHAR*& Buffer, int32 PortFlags, U
 		int32 NumCharsRead = 0;
 		if (!FParse::QuotedString(Buffer, Json, &NumCharsRead))
 		{
-			ErrorText->Logf(ELogVerbosity::Warning, TEXT("FJsonObjectWrapper::ImportTextItem: Bad quoted string: %s\n"), Buffer);
+			if (ErrorText)
+			{
+				ErrorText->Logf(ELogVerbosity::Warning, TEXT("FJsonObjectWrapper::ImportTextItem: Bad quoted string: %s\n"), Buffer);
+			}
+			
 			return false;
 		}
 		Buffer += NumCharsRead;
