Nested subobjects that have overriden properties in a blueprint now be initialized correctly from their archetype, not the native default values

Nested subobjects that are part of Objects created from a template will now correctly initialize their properties from the supplied template

diff --git a/Engine/Source/Runtime/CoreUObject/Private/UObject/UObjectGlobals.cpp b/Engine/Source/Runtime/CoreUObject/Private/UObject/UObjectGlobals.cpp
--- a/Engine/Source/Runtime/CoreUObject/Private/UObject/UObjectGlobals.cpp
+++ b/Engine/Source/Runtime/CoreUObject/Private/UObject/UObjectGlobals.cpp
@@ -4508,13 +4508,6 @@ UObject* FObjectInitializer::CreateDefaultSubobject(UObject* Outer, FName Subobj
 		{
 			UObject* Template = OverrideClass->GetDefaultObject(); // force the CDO to be created if it hasn't already
 			EObjectFlags SubobjectFlags = Outer->GetMaskedFlags(RF_PropagateToSubObjects) | RF_DefaultSubObject;
-			bool bOwnerArchetypeIsNotNative;
-			UClass* OuterArchetypeClass;
-
-			// It is not safe to mark this component as properly transient, that results in it being nulled incorrectly
-
-			OuterArchetypeClass = Outer->GetArchetype()->GetClass();
-			bOwnerArchetypeIsNotNative = !OuterArchetypeClass->HasAnyClassFlags(CLASS_Native | CLASS_Intrinsic);
 
 			const bool bOwnerTemplateIsNotCDO = ObjectArchetype != nullptr && ObjectArchetype != Outer->GetClass()->GetDefaultObject(false) && !Outer->HasAnyFlags(RF_ClassDefaultObject);
 #if !UE_BUILD_SHIPPING
@@ -4535,21 +4528,27 @@ UObject* FObjectInitializer::CreateDefaultSubobject(UObject* Outer, FName Subobj
 			Params.SetFlags = SubobjectFlags;
 			Params.SubobjectOverrides = ComponentOverride.SubOverrides;
 
-			Result = StaticConstructObject_Internal(Params);
-			if (!bIsTransient && (bOwnerArchetypeIsNotNative || bOwnerTemplateIsNotCDO))
+			// If the object creating a subobject is being created from a template, not a CDO 
+			// then we need to use the subobject from that template as the new subobject's template
+			if (!bIsTransient && bOwnerTemplateIsNotCDO)
 			{
-				UObject* MaybeTemplate = nullptr;
-				if (bOwnerTemplateIsNotCDO)
+				UObject* MaybeTemplate = ObjectArchetype->GetDefaultSubobjectByName(SubobjectFName);
+				if (MaybeTemplate && Template != MaybeTemplate && MaybeTemplate->IsA(ReturnType))
 				{
-					// Try to get the subobject template from the specified object template
-					MaybeTemplate = ObjectArchetype->GetDefaultSubobjectByName(SubobjectFName);
+					Params.Template = MaybeTemplate;
 				}
-				if (!MaybeTemplate)
-				{
-					// The archetype of the outer is not native, so we need to copy properties to the subobjects after the C++ constructor chain for the outer has run (because those sets properties on the subobjects)
-					MaybeTemplate = OuterArchetypeClass->GetDefaultSubobjectByName(SubobjectFName);
-				}
-				if (MaybeTemplate && MaybeTemplate->IsA(ReturnType) && Template != MaybeTemplate)
+			}
+
+			Result = StaticConstructObject_Internal(Params);
+
+			if (Params.Template)
+			{
+				ComponentInits.Add(Result, Params.Template);
+			}
+			else if (!bIsTransient && Outer->GetArchetype()->IsInBlueprint())
+			{
+				UObject* MaybeTemplate = Result->GetArchetype();
+				if (MaybeTemplate && Template != MaybeTemplate && MaybeTemplate->IsA(ReturnType))
 				{
 					ComponentInits.Add(Result, MaybeTemplate);
 				}
