Added a few utility functions to help with using property setters and getters with properties representing containers (Arrays / Sets / Maps) or structs

diff --git a/Engine/Source/Runtime/CoreUObject/Private/UObject/Property.cpp b/Engine/Source/Runtime/CoreUObject/Private/UObject/Property.cpp
--- a/Engine/Source/Runtime/CoreUObject/Private/UObject/Property.cpp
+++ b/Engine/Source/Runtime/CoreUObject/Private/UObject/Property.cpp
@@ -1259,6 +1259,89 @@ bool FProperty::SameType(const FProperty* Other) const
 	return Other && (GetClass() == Other->GetClass());
 }
 
+void* FProperty::AllocateAndInitializeValue() const
+{
+	void* Memory = (uint8*)FMemory::MallocZeroed(GetSize(), GetMinAlignment());
+	if (!HasAnyPropertyFlags(CPF_ZeroConstructor)) // this stuff is already zero
+	{
+		InitializeValue(Memory);
+	}
+	return Memory;
+}
+
+void FProperty::DestroyAndFreeValue(void* InMemory) const
+{
+	if (InMemory)
+	{
+		DestroyValue(InMemory);
+		FMemory::Free(InMemory);
+	}
+}
+
+void* FProperty::GetValueAddressAtIndex_Direct(const FProperty* Inner, void* InValueAddress, int32 Index) const
+{
+	checkf(Inner == nullptr, TEXT("%s should not have an inner property or it's missing specialized GetValueAddressAtIndex_Direct override"), *GetFullName());
+	checkf(Index < ArrayDim && Index >= 0, TEXT("Array index (%d) out of range"), Index);
+	return (uint8*)InValueAddress + ElementSize * Index;
+}
+
+void FProperty::PerformOperationWithSetter(void* InContainer, void* DirectPropertyAddress, TFunctionRef<void(void*)> DirectValueAccessFunc) const
+{
+	if (InContainer && HasSetterOrGetter()) // If there's a getter we need to allocate a temp value even if there's no setter
+	{
+		// When modifying container or struct properties that have a setter or getter function we first allocate a temp value
+		// that we can operate on directly (add new elements or modify existing ones)
+		void* LocalValuePtr = AllocateAndInitializeValue();
+		// Copy the value to the allocated local (using a getter if present)
+		GetValue_InContainer(InContainer, LocalValuePtr);
+
+		// Perform operation on the temp value
+		DirectValueAccessFunc(LocalValuePtr);
+
+		// Assign the temp value back to the property using a setter function
+		SetValue_InContainer(InContainer, LocalValuePtr);
+		// Destroy and free the temp value
+		DestroyAndFreeValue(LocalValuePtr);
+	}
+	else
+	{
+		// When there's no setter or getter present it's ok to perform the operation directly on the container / struct memory
+		if (!DirectPropertyAddress)
+		{
+			checkf(InContainer, TEXT("Container pointr must be valid if DirectPropertyAddress is not valid"));
+			DirectPropertyAddress = PointerToValuePtr(InContainer, EPropertyPointerType::Container);
+		}
+		DirectValueAccessFunc(DirectPropertyAddress);
+	}
+}
+
+void FProperty::PerformOperationWithGetter(void* InContainer, const void* DirectPropertyAddress, TFunctionRef<void(const void*)> DirectValueAccessFunc) const
+{
+	if (InContainer && HasGetter())
+	{
+		// When modifying container or struct properties that have a getter function we first allocate a temp value
+		// that we can operate on directly (add new elements or modify existing ones)
+		void* LocalValuePtr = AllocateAndInitializeValue();
+		// Copy the value to the allocated local using a getter
+		GetValue_InContainer(InContainer, LocalValuePtr);
+
+		// Perform read-only operation on the temp value
+		DirectValueAccessFunc(LocalValuePtr);
+
+		// Destroy and free the temp value
+		DestroyAndFreeValue(LocalValuePtr);
+	}
+	else
+	{
+		if (!DirectPropertyAddress)
+		{
+			checkf(InContainer, TEXT("Container pointr must be valid if DirectPropertyAddress is not valid"));
+			DirectPropertyAddress = PointerToValuePtr(InContainer, EPropertyPointerType::Container);
+		}
+		DirectValueAccessFunc(DirectPropertyAddress);
+	}
+}
+
 /**
  * Attempts to read an array index (xxx) sequence.  Handles const/enum replacements, etc.
  * @param	ObjectStruct	the scope of the object/struct containing the property we're currently importing
diff --git a/Engine/Source/Runtime/CoreUObject/Private/UObject/PropertyArray.cpp b/Engine/Source/Runtime/CoreUObject/Private/UObject/PropertyArray.cpp
--- a/Engine/Source/Runtime/CoreUObject/Private/UObject/PropertyArray.cpp
+++ b/Engine/Source/Runtime/CoreUObject/Private/UObject/PropertyArray.cpp
@@ -478,11 +478,7 @@ void FArrayProperty::ExportText_Internal( FString& ValueStr, const void* Contain
 	if (PropertyPointerType == EPropertyPointerType::Container && HasGetter())
 	{
 		// Allocate temporary map as we first need to initialize it with the value provided by the getter function and then export it
-		TempArrayStorage = (uint8*)FMemory::MallocZeroed(GetSize());
-		if (!HasAnyPropertyFlags(CPF_ZeroConstructor)) // this stuff is already zero
-		{
-			InitializeValue(TempArrayStorage);
-		}
+		TempArrayStorage = (uint8*)AllocateAndInitializeValue();
 		PropertyValuePtr = TempArrayStorage;
 		FProperty::GetValue_InContainer(ContainerOrPropertyPtr, PropertyValuePtr);
 	}
@@ -493,11 +489,7 @@ void FArrayProperty::ExportText_Internal( FString& ValueStr, const void* Contain
 
 	ON_SCOPE_EXIT
 	{
-		if (TempArrayStorage)
-		{
-			DestroyValue(TempArrayStorage);
-			FMemory::Free(TempArrayStorage);
-		}
+		DestroyAndFreeValue(TempArrayStorage);
 	};
 
 	FScriptArrayHelper ArrayHelper(this, PropertyValuePtr);
@@ -603,8 +595,7 @@ const TCHAR* FArrayProperty::ImportText_Internal(const TCHAR* Buffer, void* Cont
 			FProperty::SetValue_InContainer(ContainerOrPropertyPtr, TempArrayStorage);
 
 			// Destroy and free the temp array used by property setter
-			DestroyValue(TempArrayStorage);
-			FMemory::Free(TempArrayStorage);
+			DestroyAndFreeValue(TempArrayStorage);
 		}
 	};
 
@@ -612,11 +603,7 @@ const TCHAR* FArrayProperty::ImportText_Internal(const TCHAR* Buffer, void* Cont
 	if (PropertyPointerType == EPropertyPointerType::Container && HasSetter())
 	{
 		// Allocate temporary map as we first need to initialize it with the parsed items and then use the setter to update the property
-		TempArrayStorage = (uint8*)FMemory::MallocZeroed(GetSize());
-		if (!HasAnyPropertyFlags(CPF_ZeroConstructor)) // this stuff is already zero
-		{
-			InitializeValue(TempArrayStorage);
-		}
+		TempArrayStorage = (uint8*)AllocateAndInitializeValue();
 		ArrayPtr = TempArrayStorage;
 	}
 	else
@@ -896,4 +883,13 @@ void FArrayProperty::GetInnerFields(TArray<FField*>& OutFields)
 	}
 }
 
+void* FArrayProperty::GetValueAddressAtIndex_Direct(const FProperty* InInner, void* InValueAddress, int32 Index) const
+{
+	FScriptArrayHelper ArrayHelper(this, InValueAddress);
+	checkf(Inner == InInner, TEXT("Passed in inner must be identical to the array property inner property"));
+	checkf(Index < ArrayHelper.Num() && Index >= 0, TEXT("Array index (%d) out of range"), Index);
+
+	return ArrayHelper.GetRawPtr(Index);
+}
+
 #include "UObject/DefineUPropertyMacros.h"
\ No newline at end of file
diff --git a/Engine/Source/Runtime/CoreUObject/Private/UObject/PropertyMap.cpp b/Engine/Source/Runtime/CoreUObject/Private/UObject/PropertyMap.cpp
--- a/Engine/Source/Runtime/CoreUObject/Private/UObject/PropertyMap.cpp
+++ b/Engine/Source/Runtime/CoreUObject/Private/UObject/PropertyMap.cpp
@@ -656,11 +656,7 @@ void FMapProperty::ExportText_Internal(FString& ValueStr, const void* ContainerO
 	if (PropertyPointerType == EPropertyPointerType::Container && HasGetter())
 	{
 		// Allocate temporary map as we first need to initialize it with the value provided by the getter function and then export it
-		TempMapStorage = (uint8*)FMemory::MallocZeroed(GetSize());
-		if (!HasAnyPropertyFlags(CPF_ZeroConstructor)) // this stuff is already zero
-		{
-			InitializeValue(TempMapStorage);
-		}
+		TempMapStorage = (uint8*)AllocateAndInitializeValue();
 		PropertyValuePtr = TempMapStorage;
 		FProperty::GetValue_InContainer(ContainerOrPropertyPtr, PropertyValuePtr);
 	}
@@ -671,11 +667,7 @@ void FMapProperty::ExportText_Internal(FString& ValueStr, const void* ContainerO
 
 	ON_SCOPE_EXIT
 	{
-		if (TempMapStorage)
-		{
-			DestroyValue(TempMapStorage);
-			FMemory::Free(TempMapStorage);
-		}
+		DestroyAndFreeValue(TempMapStorage);
 	};
 
 	FScriptMapHelper MapHelper(this, PropertyValuePtr);
@@ -809,19 +801,14 @@ const TCHAR* FMapProperty::ImportText_Internal(const TCHAR* Buffer, void* Contai
 			FProperty::SetValue_InContainer(ContainerOrPropertyPtr, TempMapStorage);
 
 			// Destroy and free the temp map used by property setter
-			DestroyValue(TempMapStorage);
-			FMemory::Free(TempMapStorage);
+			DestroyAndFreeValue(TempMapStorage);
 		}
 	};
 
 	if (PropertyPointerType == EPropertyPointerType::Container && HasSetter())
 	{
 		// Allocate temporary map as we first need to initialize it with the parsed items and then use the setter to update the property
-		TempMapStorage = (uint8*)FMemory::MallocZeroed(GetSize());
-		if (!HasAnyPropertyFlags(CPF_ZeroConstructor)) // this stuff is already zero
-		{
-			InitializeValue(TempMapStorage);
-		}
+		TempMapStorage = (uint8*)AllocateAndInitializeValue();
 		// Reinitialize the map helper with the temp value
 		MapHelper = FScriptMapHelper(this, TempMapStorage);
 	}
@@ -1335,4 +1322,32 @@ void FMapProperty::GetInnerFields(TArray<FField*>& OutFields)
 	}
 }
 
+void* FMapProperty::GetValueAddressAtIndex_Direct(const FProperty* Inner, void* InValueAddress, int32 Index) const
+{
+	FScriptMapHelper MapHelper(this, InValueAddress);
+	checkf(Inner == KeyProp || Inner == ValueProp, TEXT("Inner property must be either KeyProp or ValueProp"));
+
+	for (int32 MapIndex = 0, Num = MapHelper.Num(), LocalIndex = Index; LocalIndex >= 0 && Num > 0; ++MapIndex)
+	{
+		if (MapHelper.IsValidIndex(MapIndex))
+		{
+			if (LocalIndex == 0)
+			{
+				if (Inner == KeyProp)
+				{
+					return MapHelper.GetKeyPtr(MapIndex);
+				}
+				else
+				{
+					return MapHelper.GetValuePtr(MapIndex);
+				}
+			}
+			LocalIndex--;
+			Num--;
+		}
+	}
+	checkf(false, TEXT("Map element index (%d) out of range"), Index);
+	return nullptr;
+}
+
 #include "UObject/DefineUPropertyMacros.h"
diff --git a/Engine/Source/Runtime/CoreUObject/Private/UObject/PropertySet.cpp b/Engine/Source/Runtime/CoreUObject/Private/UObject/PropertySet.cpp
--- a/Engine/Source/Runtime/CoreUObject/Private/UObject/PropertySet.cpp
+++ b/Engine/Source/Runtime/CoreUObject/Private/UObject/PropertySet.cpp
@@ -568,11 +568,7 @@ void FSetProperty::ExportText_Internal(FString& ValueStr, const void* ContainerO
 	if (PropertyPointerType == EPropertyPointerType::Container && HasGetter())
 	{
 		// Allocate temporary map as we first need to initialize it with the value provided by the getter function and then export it
-		TempSetStorage = (uint8*)FMemory::MallocZeroed(GetSize());
-		if (!HasAnyPropertyFlags(CPF_ZeroConstructor)) // this stuff is already zero
-		{
-			InitializeValue(TempSetStorage);
-		}
+		TempSetStorage = (uint8*)AllocateAndInitializeValue();
 		PropertyValuePtr = TempSetStorage;
 		FProperty::GetValue_InContainer(ContainerOrPropertyPtr, PropertyValuePtr);
 	}
@@ -583,11 +579,7 @@ void FSetProperty::ExportText_Internal(FString& ValueStr, const void* ContainerO
 
 	ON_SCOPE_EXIT
 	{
-		if (TempSetStorage)
-		{
-			DestroyValue(TempSetStorage);
-			FMemory::Free(TempSetStorage);
-		}
+		DestroyAndFreeValue(TempSetStorage);
 	};
 
 	FScriptSetHelper SetHelper(this, PropertyValuePtr);
@@ -719,19 +711,14 @@ const TCHAR* FSetProperty::ImportText_Internal(const TCHAR* Buffer, void* Contai
 			FProperty::SetValue_InContainer(ContainerOrPropertyPtr, TempSetStorage);
 
 			// Destroy and free the temp set used by property setter
-			DestroyValue(TempSetStorage);
-			FMemory::Free(TempSetStorage);
+			DestroyAndFreeValue(TempSetStorage);
 		}
 	};
 
 	if (PropertyPointerType == EPropertyPointerType::Container && HasSetter())
 	{
 		// Allocate temporary set as we first need to initialize it with the parsed items and then use the setter to update the property
-		TempSetStorage = (uint8*)FMemory::MallocZeroed(GetSize());
-		if (!HasAnyPropertyFlags(CPF_ZeroConstructor)) // this stuff is already zero
-		{
-			InitializeValue(TempSetStorage);
-		}
+		TempSetStorage = (uint8*)AllocateAndInitializeValue();
 		// Reinitialize the set helper with the temp value
 		SetHelper = FScriptSetHelper(this, TempSetStorage);
 	}
@@ -1113,4 +1100,25 @@ void FSetProperty::GetInnerFields(TArray<FField*>& OutFields)
 	}
 }
 
+void* FSetProperty::GetValueAddressAtIndex_Direct(const FProperty* Inner, void* InValueAddress, int32 Index) const
+{
+	FScriptSetHelper SetHelper(this, InValueAddress);
+	checkf(Inner == ElementProp, TEXT("Inner property must be identical to ElementProp"));
+
+	for (int32 SetIndex = 0, Num = SetHelper.Num(), LocalIndex = Index; LocalIndex >= 0 && Num > 0; ++SetIndex)
+	{
+		if (SetHelper.IsValidIndex(SetIndex))
+		{
+			if (LocalIndex == 0)
+			{
+				return SetHelper.GetElementPtr(SetIndex);
+			}
+			LocalIndex--;
+			Num--;
+		}
+	}
+	checkf(false, TEXT("Set element index (%d) out of range"), Index);
+	return nullptr;
+}
+
 #include "UObject/DefineUPropertyMacros.h"
\ No newline at end of file
diff --git a/Engine/Source/Runtime/CoreUObject/Public/UObject/UnrealType.h b/Engine/Source/Runtime/CoreUObject/Public/UObject/UnrealType.h
--- a/Engine/Source/Runtime/CoreUObject/Public/UObject/UnrealType.h
+++ b/Engine/Source/Runtime/CoreUObject/Public/UObject/UnrealType.h
@@ -560,6 +560,36 @@ public:
 		}
 	}
 
+	/** Allocates and initializes memory to hold a value this property represents */
+	void* AllocateAndInitializeValue() const;
+
+	/** Destroys and frees memory with a value this property represents */
+	void DestroyAndFreeValue(void* InMemory) const;
+
+	/**
+	 * Helper function for setting container / struct property value and performing operation directly on the value memory
+	 * @param InContainer Pointer to the container that owns the property. Can be null but then setters and getters will not be used.
+	 * @param DirectPropertyAddress Direct property value address. Can be null only if InContainer is a valid pointer.
+	 * @param DirectValueAccessFunc Function that manipulates directly on property value address. The value address can be different than the passed in DirectPropertyAddress if setters and getters are present and InContainer pointer is valid.
+	 */
+	void PerformOperationWithSetter(void* InContainer, void* DirectPropertyAddress, TFunctionRef<void(void*)> DirectValueAccessFunc) const;
+
+	/**
+	 * Helper function for getting container / struct property value and performing operation directly on the value memory
+	 * @param InContainer Pointer to the container that owns the property. Can be null but then setters and getters will not be used.
+	 * @param DirectPropertyAddress Direct property value address. Can be null only if InContainer is a valid pointer.
+	 * @param DirectValueAccessFunc Function that manipulates directly on property value address. The value address can be different than the passed in DirectPropertyAddress if setters and getters are present and InContainer pointer is valid.
+	 */
+	void PerformOperationWithGetter(void* InContainer, const void* DirectPropertyAddress, TFunctionRef<void(const void*)> DirectValueAccessFunc) const;
+
+	/** 
+	 * Gets value address at given index inside of a static array or container
+	 * @param InValueAddress address of the value represented by this property
+	 * @param Index into the static array or container
+	 * @returns address of the value at given index
+	 */
+	virtual void* GetValueAddressAtIndex_Direct(const FProperty* Inner, void* InValueAddress, int32 Index) const;
+
 #if WITH_EDITORONLY_DATA
 	/**
 	 * Updates the given HashBuilder with name and type information of this Property.
@@ -3579,6 +3609,8 @@ public:
 		// This is the same as alignof(FFreezableScriptArray)
 		return alignof(FScriptArray);
 	}
+
+	virtual void* GetValueAddressAtIndex_Direct(const FProperty* Inner, void* InValueAddress, int32 Index) const override;
 	// End of FProperty interface
 
 	FString GetCPPTypeCustom(FString* ExtendedTypeText, uint32 CPPExportFlags, const FString& InnerTypeText, const FString& InInnerExtendedTypeText) const;
@@ -3707,6 +3739,7 @@ public:
 	virtual void EmitReferenceInfo(UClass& OwnerClass, int32 BaseOffset, TArray<const FStructProperty*>& EncounteredStructProps, FGCStackSizeHelper& StackSizeHelper) override;
 	virtual bool SameType(const FProperty* Other) const override;
 	virtual EConvertFromTypeResult ConvertFromType(const FPropertyTag& Tag, FStructuredArchive::FSlot Slot, uint8* Data, UStruct* DefaultsStruct) override;
+	virtual void* GetValueAddressAtIndex_Direct(const FProperty* Inner, void* InValueAddress, int32 Index) const override;
 	// End of FProperty interface
 
 	FString GetCPPTypeCustom(FString* ExtendedTypeText, uint32 CPPExportFlags, const FString& KeyTypeText, const FString& InKeyExtendedTypeText, const FString& ValueTypeText, const FString& InValueExtendedTypeText) const;
@@ -3830,6 +3863,7 @@ public:
 	virtual void EmitReferenceInfo(UClass& OwnerClass, int32 BaseOffset, TArray<const FStructProperty*>& EncounteredStructProps, FGCStackSizeHelper& StackSizeHelper) override;
 	virtual bool SameType(const FProperty* Other) const override;
 	virtual EConvertFromTypeResult ConvertFromType(const FPropertyTag& Tag, FStructuredArchive::FSlot Slot, uint8* Data, UStruct* DefaultsStruct) override;
+	virtual void* GetValueAddressAtIndex_Direct(const FProperty* Inner, void* InValueAddress, int32 Index) const override;
 	// End of FProperty interface
 
 	FString GetCPPTypeCustom(FString* ExtendedTypeText, uint32 CPPExportFlags, const FString& ElementTypeText, const FString& InElementExtendedTypeText) const;
