AsyncLoading2: Fix crash in seamless server travel.
When reloading an existing UPackage make sure to always set all serialized package flags again since they may have been cleared since the last load.
One example is UWorld::FinishDestroy that may clear PKG_ContainsMap on its world UPackage.

#rnx

diff --git a/Engine/Source/Runtime/CoreUObject/Private/Serialization/AsyncLoading2.cpp b/Engine/Source/Runtime/CoreUObject/Private/Serialization/AsyncLoading2.cpp
--- a/Engine/Source/Runtime/CoreUObject/Private/Serialization/AsyncLoading2.cpp
+++ b/Engine/Source/Runtime/CoreUObject/Private/Serialization/AsyncLoading2.cpp
@@ -5795,9 +5795,9 @@ void FAsyncPackage2::CreateUPackage(const FZenPackageSummary* PackageSummary, co
 	}
 	else
 	{
+		LinkerRoot->SetPackageFlags(PackageSummary->PackageFlags | PKG_Cooked);
 		check(LinkerRoot->CanBeImported() == Desc.bCanBeImported);
 		check(LinkerRoot->GetPackageId() == Desc.UPackageId);
-		check(LinkerRoot->GetPackageFlags() == (PackageSummary->PackageFlags | PKG_Cooked));
 		check(LinkerRoot->GetLinkerPackageVersion() == GPackageFileUEVersion);
 		check(LinkerRoot->GetLinkerLicenseeVersion() == GPackageFileLicenseeUEVersion);
 		check(LinkerRoot->HasAnyFlags(RF_WasLoaded));
