Fix for mixing SUBOBJECT_DELIMITER_CHAR with non-TCHARs - when used in a ternary in particular, it caused the result to become promoted to an integer.

diff --git a/Engine/Source/Runtime/AssetRegistry/Private/PackageReader.cpp b/Engine/Source/Runtime/AssetRegistry/Private/PackageReader.cpp
--- a/Engine/Source/Runtime/AssetRegistry/Private/PackageReader.cpp
+++ b/Engine/Source/Runtime/AssetRegistry/Private/PackageReader.cpp
@@ -621,7 +621,7 @@ bool FPackageReader::SerializeImportedClasses(const TArray<FObjectImport>& Impor
 		{
 			int32 NumTokens = ParentChain.Num();
 			check(NumTokens >= 1);
-			const TCHAR Delimiters[] = {'.', SUBOBJECT_DELIMITER_CHAR, '.' };
+			const TCHAR Delimiters[] = { TEXT('.'), SUBOBJECT_DELIMITER_CHAR, TEXT('.') };
 			int32 DelimiterIndex = 0;
 			ParentChain[NumTokens - 1].AppendString(ClassObjectPath);
 			for (int32 TokenIndex = NumTokens - 2; TokenIndex >= 0; --TokenIndex)
diff --git a/Engine/Source/Runtime/CoreUObject/Private/UObject/ObjectPtr.cpp b/Engine/Source/Runtime/CoreUObject/Private/UObject/ObjectPtr.cpp
--- a/Engine/Source/Runtime/CoreUObject/Private/UObject/ObjectPtr.cpp
+++ b/Engine/Source/Runtime/CoreUObject/Private/UObject/ObjectPtr.cpp
@@ -16,7 +16,7 @@ FString FObjectPtr::GetPath() const
 	CompletePath << Ref.PackageName;
 	for (int32 ResolvedNameIndex = 0; ResolvedNameIndex < ResolvedNames.Num(); ++ResolvedNameIndex)
 	{
-		CompletePath << (ResolvedNameIndex == 1 ? SUBOBJECT_DELIMITER_CHAR : '.') << ResolvedNames[ResolvedNameIndex];
+		CompletePath << (ResolvedNameIndex == 1 ? SUBOBJECT_DELIMITER_CHAR : TEXT('.')) << ResolvedNames[ResolvedNameIndex];
 	}
 	return FString(CompletePath);
 }
diff --git a/Engine/Source/Runtime/CoreUObject/Private/UObject/Property.cpp b/Engine/Source/Runtime/CoreUObject/Private/UObject/Property.cpp
--- a/Engine/Source/Runtime/CoreUObject/Private/UObject/Property.cpp
+++ b/Engine/Source/Runtime/CoreUObject/Private/UObject/Property.cpp
@@ -522,7 +522,7 @@ FORCEINLINE constexpr FStringView ParsePropertyToken(const TCHAR* Str, bool Dott
 {
 	constexpr FAsciiSet RegularTokenChars = AlphaNumericChars  + '_' + '-' + '+';
 	constexpr FAsciiSet RegularNonTokenChars = ~RegularTokenChars;
-	constexpr FAsciiSet DottedNonTokenChars = ~(RegularTokenChars + '.' + '/' + SUBOBJECT_DELIMITER_CHAR);
+	constexpr FAsciiSet DottedNonTokenChars = ~(RegularTokenChars + '.' + '/' + (char)SUBOBJECT_DELIMITER_CHAR);
 	FAsciiSet CurrentNonTokenChars = DottedNames ? DottedNonTokenChars : RegularNonTokenChars;
 
 	const TCHAR* TokenEnd = FAsciiSet::FindFirstOrEnd(Str, CurrentNonTokenChars);
diff --git a/Engine/Source/Runtime/CoreUObject/Private/UObject/UObjectGlobals.cpp b/Engine/Source/Runtime/CoreUObject/Private/UObject/UObjectGlobals.cpp
--- a/Engine/Source/Runtime/CoreUObject/Private/UObject/UObjectGlobals.cpp
+++ b/Engine/Source/Runtime/CoreUObject/Private/UObject/UObjectGlobals.cpp
@@ -739,7 +739,7 @@ bool ResolveName(UObject*& InPackage, FString& InOutName, bool Create, bool Thro
 	bool bSubobjectPath = false;
 
 	// Handle specified packages.
-	constexpr FAsciiSet Delimiters = FAsciiSet(".") + SUBOBJECT_DELIMITER_CHAR;
+	constexpr FAsciiSet Delimiters = FAsciiSet(".") + (char)SUBOBJECT_DELIMITER_CHAR;
 	while (true)
 	{
 		const TCHAR* DelimiterOrEnd = FAsciiSet::FindFirstOrEnd(*InOutName, Delimiters);
