Fix missing closing paren in log message

diff --git a/Engine/Source/Runtime/CoreUObject/Private/UObject/SavePackage2.cpp b/Engine/Source/Runtime/CoreUObject/Private/UObject/SavePackage2.cpp
--- a/Engine/Source/Runtime/CoreUObject/Private/UObject/SavePackage2.cpp
+++ b/Engine/Source/Runtime/CoreUObject/Private/UObject/SavePackage2.cpp
@@ -176,7 +176,7 @@ ESavePackageResult ValidatePackage(FSaveContext& SaveContext)
 				Arguments.Add(TEXT("Name"), FText::FromString(FilenameStr));
 				FText FormatText = SaveContext.GetPackage()->ContainsMap() 
 					? NSLOCTEXT("SavePackage2", "MapSaveNotAllowed", "Map '{Name}' is not allowed to save (see log for reason)") 
-					: NSLOCTEXT("SavePackage2", "AssetSaveNotAllowed", "Asset '{Name}' is not allowed to save (see log for reason");
+					: NSLOCTEXT("SavePackage2", "AssetSaveNotAllowed", "Asset '{Name}' is not allowed to save (see log for reason)");
 				FText ErrorText = FText::Format(FormatText, Arguments);
 				SaveContext.GetError()->Logf(ELogVerbosity::Warning, TEXT("%s"), *ErrorText.ToString());
 			}
