Add class and object name on check that reports abstract class instanciation failures so we can know more about where this is coming from

#rnx

diff --git a/Engine/Source/Runtime/CoreUObject/Private/UObject/UObjectGlobals.cpp b/Engine/Source/Runtime/CoreUObject/Private/UObject/UObjectGlobals.cpp
--- a/Engine/Source/Runtime/CoreUObject/Private/UObject/UObjectGlobals.cpp
+++ b/Engine/Source/Runtime/CoreUObject/Private/UObject/UObjectGlobals.cpp
@@ -2775,12 +2775,14 @@ UObject* StaticAllocateObject
 		}
 	}
 #endif // WITH_EDITOR
-	bool bCreatingCDO = (InFlags & RF_ClassDefaultObject) != 0;
+	const bool bCreatingCDO = (InFlags & RF_ClassDefaultObject) != 0;
 	const bool bCreatingArchetype = (InFlags & RF_ArchetypeObject) != 0;
 
 	check(InClass);
-	check(GIsEditor || !FScopedAllowAbstractClassAllocation::IsDisallowedAbstractClass(InClass, InFlags)); // this is a warning in the editor, otherwise it is illegal to create an abstract class, except the CDO
 	check(InOuter || (InClass == UPackage::StaticClass() && InName != NAME_None)); // only packages can not have an outer, and they must be named explicitly
+	// this is a warning in the editor, otherwise it is illegal to create an abstract class, except the CDO
+	checkf(GIsEditor || !FScopedAllowAbstractClassAllocation::IsDisallowedAbstractClass(InClass, InFlags), TEXT("Unable to create new object: %s %s.%s. Creating an instance of an abstract class is not allowed!"),
+		*GetNameSafe(InClass), *GetPathNameSafe(InOuter), *InName.ToString());
 	//checkf(InClass != UPackage::StaticClass() || !InOuter || bCreatingCDO, TEXT("Creating nested packages is not allowed: Outer=%s, Package=%s"), *GetNameSafe(InOuter), *InName.ToString());
 	check(bCreatingCDO || bCreatingArchetype || !InOuter || InOuter->IsA(InClass->ClassWithin));
 	checkf(!IsGarbageCollecting(), TEXT("Unable to create new object: %s %s.%s. Creating UObjects while Collecting Garbage is not allowed!"),
