Added FCoreUObjectDelegates::OnObjectPostCDOCompiled

This is similar to UObject::OnPostCDOCompiled, but doesn't require that you implement the logic on a particular class.

diff --git a/Engine/Source/Editor/KismetCompiler/Private/KismetCompiler.cpp b/Engine/Source/Editor/KismetCompiler/Private/KismetCompiler.cpp
--- a/Engine/Source/Editor/KismetCompiler/Private/KismetCompiler.cpp
+++ b/Engine/Source/Editor/KismetCompiler/Private/KismetCompiler.cpp
@@ -5047,6 +5047,7 @@ void FKismetCompilerContext::PostCDOCompiled(const UObject::FPostCDOCompiledCont
 		NewClass->ClassDefaultObject->PostCDOCompiled();
 		PRAGMA_ENABLE_DEPRECATION_WARNINGS;
 	}
+	FCoreUObjectDelegates::OnObjectPostCDOCompiled.Broadcast(NewClass->ClassDefaultObject, Context);
 
 	// Allow children to customize PostCDOCompile:
 	OnPostCDOCompiled(Context);
diff --git a/Engine/Source/Runtime/CoreUObject/Private/UObject/UObjectGlobals.cpp b/Engine/Source/Runtime/CoreUObject/Private/UObject/UObjectGlobals.cpp
--- a/Engine/Source/Runtime/CoreUObject/Private/UObject/UObjectGlobals.cpp
+++ b/Engine/Source/Runtime/CoreUObject/Private/UObject/UObjectGlobals.cpp
@@ -150,6 +150,7 @@ TSet<UObject*> FCoreUObjectDelegates::ObjectsModifiedThisFrame;
 FCoreUObjectDelegates::FOnObjectModified FCoreUObjectDelegates::OnObjectModified;
 FCoreUObjectDelegates::FOnObjectTransacted FCoreUObjectDelegates::OnObjectTransacted;
 FCoreUObjectDelegates::FOnObjectsReplaced FCoreUObjectDelegates::OnObjectsReplaced;
+FCoreUObjectDelegates::FOnObjectPostCDOCompiled FCoreUObjectDelegates::OnObjectPostCDOCompiled;
 FCoreUObjectDelegates::FOnAssetLoaded FCoreUObjectDelegates::OnAssetLoaded;
 FCoreUObjectDelegates::FOnEndLoadPackage FCoreUObjectDelegates::OnEndLoadPackage;
 PRAGMA_DISABLE_DEPRECATION_WARNINGS
diff --git a/Engine/Source/Runtime/CoreUObject/Public/UObject/Object.h b/Engine/Source/Runtime/CoreUObject/Public/UObject/Object.h
--- a/Engine/Source/Runtime/CoreUObject/Public/UObject/Object.h
+++ b/Engine/Source/Runtime/CoreUObject/Public/UObject/Object.h
@@ -10,6 +10,7 @@
 #include "UObject/Script.h"
 #include "UObject/ObjectMacros.h"
 #include "UObject/UObjectBaseUtility.h"
+#include "UObject/ObjectCompileContext.h"
 #include "ProfilingDebugging/ResourceSize.h"
 #include "UObject/PrimaryAssetId.h"
 
@@ -216,11 +217,7 @@ public:
 	}
 
 #if WITH_EDITOR
-	struct FPostCDOCompiledContext
-	{
-		/** True if this notification was from a 'skeleton-only' compile */
-		bool bIsSkeletonOnly = false;
-	};
+	using FPostCDOCompiledContext = FObjectPostCDOCompiledContext;
 
 	/**
 	 * Called after the Blueprint compiler has finished generating the Class Default Object (CDO) for a class. This can only happen in the editor.
diff --git a/Engine/Source/Runtime/CoreUObject/Public/UObject/ObjectCompileContext.h b/Engine/Source/Runtime/CoreUObject/Public/UObject/ObjectCompileContext.h
new file mode 100644
--- /dev/null
+++ b/Engine/Source/Runtime/CoreUObject/Public/UObject/ObjectCompileContext.h
@@ -0,0 +1,11 @@
+// Copyright Epic Games, Inc. All Rights Reserved.
+
+#pragma once
+
+#if WITH_EDITOR
+struct FObjectPostCDOCompiledContext
+{
+	/** True if this notification was from a 'skeleton-only' compile */
+	bool bIsSkeletonOnly = false;
+};
+#endif
diff --git a/Engine/Source/Runtime/CoreUObject/Public/UObject/UObjectGlobals.h b/Engine/Source/Runtime/CoreUObject/Public/UObject/UObjectGlobals.h
--- a/Engine/Source/Runtime/CoreUObject/Public/UObject/UObjectGlobals.h
+++ b/Engine/Source/Runtime/CoreUObject/Public/UObject/UObjectGlobals.h
@@ -21,6 +21,7 @@
 
 struct FCustomPropertyListNode;
 struct FObjectInstancingGraph;
+struct FObjectPostCDOCompiledContext;
 class FObjectPreSaveContext;
 struct FObjectPtr;
 struct FStaticConstructObjectParameters;
@@ -2430,6 +2431,14 @@ struct COREUOBJECT_API FCoreUObjectDelegates
 	DECLARE_MULTICAST_DELEGATE_OneParam(FOnObjectsReplaced, const FReplacementObjectMap&);
 	static FOnObjectsReplaced OnObjectsReplaced;
 
+	/**
+	 * Called after the Blueprint compiler has finished generating the Class Default Object (CDO) for a class. This can only happen in the editor.
+	 * This is called when the CDO and its associated class structure have been fully generated and populated, and allows the assignment of cached/derived data,
+	 * eg) caching the name/count of properties of a certain type, or inspecting the properties on the class and using their meta-data and CDO default values to derive game data.
+	 */
+	DECLARE_MULTICAST_DELEGATE_TwoParams(FOnObjectPostCDOCompiled, UObject*, const FObjectPostCDOCompiledContext&);
+	static FOnObjectPostCDOCompiled OnObjectPostCDOCompiled;
+
 	DECLARE_MULTICAST_DELEGATE_OneParam(FOnObjectSaved, UObject*);
 	UE_DEPRECATED(5.0, "Use OnObjectPreSave instead.")
 	static FOnObjectSaved OnObjectSaved;
