Moving IsRoutingPostLoad to a tighter scope to prevent crashes when loading packages with non-native objects

diff --git a/Engine/Source/Runtime/CoreUObject/Private/Serialization/AsyncLoading.cpp b/Engine/Source/Runtime/CoreUObject/Private/Serialization/AsyncLoading.cpp
--- a/Engine/Source/Runtime/CoreUObject/Private/Serialization/AsyncLoading.cpp
+++ b/Engine/Source/Runtime/CoreUObject/Private/Serialization/AsyncLoading.cpp
@@ -6365,7 +6365,6 @@ EAsyncPackageState::Type FAsyncPackage::PostLoadDeferredObjects(double InTickSta
 	FAsyncPackageScope PackageScope(this);
 
 	EAsyncPackageState::Type Result = EAsyncPackageState::Complete;
-	TGuardValue<bool> GuardIsRoutingPostLoad(PackageScope.ThreadContext.IsRoutingPostLoad, true);
 	FAsyncLoadingTickScope InAsyncLoadingTick(AsyncLoadingThread);
 
 	FUObjectSerializeContext* LoadContext = GetSerializeContext();
@@ -6374,67 +6373,70 @@ EAsyncPackageState::Type FAsyncPackage::PostLoadDeferredObjects(double InTickSta
 
 	STAT(double PostLoadStartTime = FPlatformTime::Seconds());
 
-	while (DeferredPostLoadIndex < DeferredPostLoadObjects.Num() && 
-		!AsyncLoadingThread.IsAsyncLoadingSuspendedInternal() &&
-		!::IsTimeLimitExceeded(InTickStartTime, bInUseTimeLimit, InOutTimeLimit, LastTypeOfWorkPerformed, LastObjectWorkWasPerformedOn))
 	{
-		UObject* Object = DeferredPostLoadObjects[DeferredPostLoadIndex++];
-		check(Object);
+		TGuardValue<bool> GuardIsRoutingPostLoad(PackageScope.ThreadContext.IsRoutingPostLoad, true);
 
-		if (!Object->IsReadyForAsyncPostLoad())
+		while (DeferredPostLoadIndex < DeferredPostLoadObjects.Num() &&
+			!AsyncLoadingThread.IsAsyncLoadingSuspendedInternal() &&
+			!::IsTimeLimitExceeded(InTickStartTime, bInUseTimeLimit, InOutTimeLimit, LastTypeOfWorkPerformed, LastObjectWorkWasPerformedOn))
 		{
-			--DeferredPostLoadIndex;
-			break;
-		}
+			UObject* Object = DeferredPostLoadObjects[DeferredPostLoadIndex++];
+			check(Object);
 
-		LastObjectWorkWasPerformedOn = Object;
-		LastTypeOfWorkPerformed = TEXT("postloading_gamethread");
+			if (!Object->IsReadyForAsyncPostLoad())
+			{
+				--DeferredPostLoadIndex;
+				break;
+			}
 
-		FScopeCycleCounterUObject ConstructorScope(Object, GET_STATID(STAT_FAsyncPackage_PostLoadObjectsGameThread));
+			LastObjectWorkWasPerformedOn = Object;
+			LastTypeOfWorkPerformed = TEXT("postloading_gamethread");
 
-		PackageScope.ThreadContext.CurrentlyPostLoadedObjectByALT = Object;
-		{
-			TRACE_LOADTIME_POSTLOAD_EXPORT_SCOPE(Object);			
-			Object->ConditionalPostLoad();
-		}
-		PackageScope.ThreadContext.CurrentlyPostLoadedObjectByALT = nullptr;
+			FScopeCycleCounterUObject ConstructorScope(Object, GET_STATID(STAT_FAsyncPackage_PostLoadObjectsGameThread));
 
-		if (ObjLoadedInPostLoad.Num())
-		{
-			// If there were any LoadObject calls inside of PostLoad, we need to pre-load those objects here. 
-			// There's no going back to the async tick loop from here.
-			UE_LOG(LogStreaming, Warning, TEXT("Detected %d objects loaded in PostLoad while streaming, this may cause hitches as we're blocking async loading to pre-load them."), ObjLoadedInPostLoad.Num());
-			
-			// Copy to local array because ObjLoadedInPostLoad can change while we're iterating over it
-			ObjLoadedInPostLoadLocal.Append(ObjLoadedInPostLoad);
-			ObjLoadedInPostLoad.Reset();
+			PackageScope.ThreadContext.CurrentlyPostLoadedObjectByALT = Object;
+			{
+				TRACE_LOADTIME_POSTLOAD_EXPORT_SCOPE(Object);
+				Object->ConditionalPostLoad();
+			}
+			PackageScope.ThreadContext.CurrentlyPostLoadedObjectByALT = nullptr;
 
-			while (ObjLoadedInPostLoadLocal.Num())
+			if (ObjLoadedInPostLoad.Num())
 			{
-				// Make sure all objects loaded in PostLoad get post-loaded too
-				DeferredPostLoadObjects.Append(ObjLoadedInPostLoadLocal);
+				// If there were any LoadObject calls inside of PostLoad, we need to pre-load those objects here. 
+				// There's no going back to the async tick loop from here.
+				UE_LOG(LogStreaming, Warning, TEXT("Detected %d objects loaded in PostLoad while streaming, this may cause hitches as we're blocking async loading to pre-load them."), ObjLoadedInPostLoad.Num());
 
-				// Preload (aka serialize) the objects loaded in PostLoad.
-				for (UObject* PreLoadObject : ObjLoadedInPostLoadLocal)
+				// Copy to local array because ObjLoadedInPostLoad can change while we're iterating over it
+				ObjLoadedInPostLoadLocal.Append(ObjLoadedInPostLoad);
+				ObjLoadedInPostLoad.Reset();
+
+				while (ObjLoadedInPostLoadLocal.Num())
 				{
-					if (PreLoadObject && PreLoadObject->GetLinker())
+					// Make sure all objects loaded in PostLoad get post-loaded too
+					DeferredPostLoadObjects.Append(ObjLoadedInPostLoadLocal);
+
+					// Preload (aka serialize) the objects loaded in PostLoad.
+					for (UObject* PreLoadObject : ObjLoadedInPostLoadLocal)
 					{
-						PreLoadObject->GetLinker()->Preload(PreLoadObject);
+						if (PreLoadObject && PreLoadObject->GetLinker())
+						{
+							PreLoadObject->GetLinker()->Preload(PreLoadObject);
+						}
 					}
-				}
 
-				// Other objects could've been loaded while we were preloading, continue until we've processed all of them.
-				ObjLoadedInPostLoadLocal.Reset();
-				ObjLoadedInPostLoadLocal.Append(ObjLoadedInPostLoad);
-				ObjLoadedInPostLoad.Reset();
-			}			
-		}
+					// Other objects could've been loaded while we were preloading, continue until we've processed all of them.
+					ObjLoadedInPostLoadLocal.Reset();
+					ObjLoadedInPostLoadLocal.Append(ObjLoadedInPostLoad);
+					ObjLoadedInPostLoad.Reset();
+				}
+			}
 
-		LastObjectWorkWasPerformedOn = Object;		
+			LastObjectWorkWasPerformedOn = Object;
 
-		UpdateLoadPercentage();
+			UpdateLoadPercentage();
+		}
 	}
-
 	// New objects might have been loaded during PostLoad.
 	Result = (DeferredPostLoadIndex == DeferredPostLoadObjects.Num()) ? EAsyncPackageState::Complete : EAsyncPackageState::TimeOut;
 
