Fixing obj refs -history=1 not looking in GC history

diff --git a/Engine/Source/Runtime/CoreUObject/Private/UObject/Obj.cpp b/Engine/Source/Runtime/CoreUObject/Private/UObject/Obj.cpp
--- a/Engine/Source/Runtime/CoreUObject/Private/UObject/Obj.cpp
+++ b/Engine/Source/Runtime/CoreUObject/Private/UObject/Obj.cpp
@@ -4200,7 +4200,7 @@ bool StaticExec( UWorld* InWorld, const TCHAR* Cmd, FOutputDevice& Ar )
 					{
 						if (FGCHistory::Get().IsActive())
 						{
-							int32 MaxHistoryLevel = FGCHistory::Get().GetHistorySize() - 1;
+							int32 MaxHistoryLevel = FGCHistory::Get().GetHistorySize();
 							if (FMath::Abs(HistoryLevel) >= MaxHistoryLevel)
 							{
 								UE_LOG(LogObj, Log, TEXT("GC History level %d will be clamped to the current max %d"), HistoryLevel, MaxHistoryLevel);
@@ -4224,7 +4224,7 @@ bool StaticExec( UWorld* InWorld, const TCHAR* Cmd, FOutputDevice& Ar )
 				else
 				{
 					FReferenceChainSearch HistorySearch(SearchModeFlags);
-					FGCSnapshot* GCSnapshot = FGCHistory::Get().GetSnapshot(HistoryLevel);
+					FGCSnapshot* GCSnapshot = FGCHistory::Get().GetSnapshot(FMath::Abs(HistoryLevel) - 1);
 					if (GCSnapshot)
 					{
 						HistorySearch.PerformSearchFromGCSnapshot(Object, *GCSnapshot);
