Fixed a condition that checks if 'obj refs -history=' value is correct

diff --git a/Engine/Source/Runtime/CoreUObject/Private/UObject/Obj.cpp b/Engine/Source/Runtime/CoreUObject/Private/UObject/Obj.cpp
--- a/Engine/Source/Runtime/CoreUObject/Private/UObject/Obj.cpp
+++ b/Engine/Source/Runtime/CoreUObject/Private/UObject/Obj.cpp
@@ -4201,7 +4201,7 @@ bool StaticExec( UWorld* InWorld, const TCHAR* Cmd, FOutputDevice& Ar )
 						if (FGCHistory::Get().IsActive())
 						{
 							int32 MaxHistoryLevel = FGCHistory::Get().GetHistorySize();
-							if (FMath::Abs(HistoryLevel) >= MaxHistoryLevel)
+							if (FMath::Abs(HistoryLevel) > MaxHistoryLevel)
 							{
 								UE_LOG(LogObj, Log, TEXT("GC History level %d will be clamped to the current max %d"), HistoryLevel, MaxHistoryLevel);
 								HistoryLevel = MaxHistoryLevel;
