Add a `FDynamicallyTypedValue` class that stores a value as some uninterpreted memory and a pointer to some runtime type information that contains methods to interpret the memory.
The garbage collector gains opcodes to mark references inside these values, though a FProperty to emit those opcodes is not part of this changelist.

diff --git a/Engine/Source/Runtime/Core/Public/GenericPlatform/GenericPlatformMath.h b/Engine/Source/Runtime/Core/Public/GenericPlatform/GenericPlatformMath.h
--- a/Engine/Source/Runtime/Core/Public/GenericPlatform/GenericPlatformMath.h
+++ b/Engine/Source/Runtime/Core/Public/GenericPlatform/GenericPlatformMath.h
@@ -821,6 +821,15 @@ struct FGenericPlatformMath
 		return (64 - CountLeadingZeros64(Arg - 1)) & (~Bitmask);
 	}
 
+	/**
+	 * Returns the smallest N such that (1<<N)>=Arg. This is a less efficient version of CeilLogTwo, but written in a
+	 * way that can be evaluated at compile-time.
+	 */
+	static FORCEINLINE constexpr uint8 ConstExprCeilLogTwo(SIZE_T Arg)
+	{
+		return Arg <= 1 ? 0 : (1 + ConstExprCeilLogTwo(Arg / 2));
+	}
+
 	/** @return Rounds the given number up to the next highest power of two. */
 	static FORCEINLINE uint32 RoundUpToPowerOfTwo(uint32 Arg)
 	{
diff --git a/Engine/Source/Runtime/CoreUObject/Private/UObject/DynamicallyTypedValue.cpp b/Engine/Source/Runtime/CoreUObject/Private/UObject/DynamicallyTypedValue.cpp
new file mode 100644
--- /dev/null
+++ b/Engine/Source/Runtime/CoreUObject/Private/UObject/DynamicallyTypedValue.cpp
@@ -0,0 +1,27 @@
+// Copyright Epic Games, Inc. All Rights Reserved.
+
+#include "UObject/DynamicallyTypedValue.h"
+
+UE::FDynamicallyTypedValueType& UE::FDynamicallyTypedValue::NullType()
+{
+	struct FNullType : FDynamicallyTypedValueType
+	{
+		FNullType(): FDynamicallyTypedValueType(0, 0, EContainsReferences::DoesNot) {}
+
+		virtual void MarkReachable() {}
+		virtual void MarkValueReachable(void* Data, FReferenceCollector& Collector) const {}
+
+		virtual void InitializeValue(void* Data) const {}
+		virtual void InitializeValueFromCopy(void* DestData, const void* SourceData) const {}
+		virtual void DestroyValue(void* Data) const {}
+
+		virtual void ClearValue(void* Data) const {}
+		virtual void SerializeValue(FStructuredArchive::FSlot Slot, void* Data) const {}
+
+		virtual uint32 GetValueHash(const void* Data) const override { return 0; }
+		virtual bool AreEqual(const void* DataA, const void* DataB) const override { return true; }
+	};
+
+	static FNullType Singleton;
+	return Singleton;
+}
\ No newline at end of file
diff --git a/Engine/Source/Runtime/CoreUObject/Private/UObject/GarbageCollection.cpp b/Engine/Source/Runtime/CoreUObject/Private/UObject/GarbageCollection.cpp
--- a/Engine/Source/Runtime/CoreUObject/Private/UObject/GarbageCollection.cpp
+++ b/Engine/Source/Runtime/CoreUObject/Private/UObject/GarbageCollection.cpp
@@ -2827,20 +2827,7 @@ void FArrayProperty::EmitReferenceInfo(UClass& OwnerClass, int32 BaseOffset, TAr
 		// Structs and nested arrays share the same implementation on the Garbage Collector side
 		// as arrays of structs already push the array memory into the GC stack and process its tokens
 		// which is exactly what is required for nested arrays to work
-		if( Inner->IsA(FStructProperty::StaticClass()) || Inner->IsA(FArrayProperty::StaticClass()) )
-		{
-			OwnerClass.EmitObjectReference(BaseOffset + GetOffset_ForGC(), *StackSizeHelper.GetPropertyPath(), bUsesFreezableAllocator ? GCRT_ArrayStructFreezable : GCRT_ArrayStruct);
-
-			// GCRT_ArrayStruct and GCRT_ArrayStructFreezable push a new stack frame in TFastReferenceCollector
-			FGCStackSizeHelperScope StackSizeScope(StackSizeHelper, Inner);
-
-			OwnerClass.ReferenceTokenStream.EmitStride(Inner->ElementSize);
-			const uint32 SkipIndexIndex = OwnerClass.ReferenceTokenStream.EmitSkipIndexPlaceholder();
-			Inner->EmitReferenceInfo(OwnerClass, 0, EncounteredStructProps, StackSizeHelper);
-			const uint32 SkipIndex = OwnerClass.ReferenceTokenStream.EmitReturn();
-			OwnerClass.ReferenceTokenStream.UpdateSkipIndexPlaceholder(SkipIndexIndex, SkipIndex);
-		}
-		else if( Inner->IsA(FObjectProperty::StaticClass()) || Inner->IsA(FObjectPtrProperty::StaticClass()) )
+		if( Inner->IsA(FObjectProperty::StaticClass()) || Inner->IsA(FObjectPtrProperty::StaticClass()) )
 		{
 			OwnerClass.EmitObjectReference(BaseOffset + GetOffset_ForGC(), *StackSizeHelper.GetPropertyPath(), bUsesFreezableAllocator ? GCRT_ArrayObjectFreezable : GCRT_ArrayObject);
 		}
@@ -2885,7 +2872,17 @@ void FArrayProperty::EmitReferenceInfo(UClass& OwnerClass, int32 BaseOffset, TAr
 		}
 		else
 		{
-			UE_LOG(LogGarbage, Fatal, TEXT("Encountered unknown property containing object or name reference: %s in %s"), *Inner->GetFullName(), *GetFullName() );
+			// In the general case, emit a GCRT_ArrayStruct that can handle any inner property type.
+			OwnerClass.EmitObjectReference(BaseOffset + GetOffset_ForGC(), *StackSizeHelper.GetPropertyPath(), bUsesFreezableAllocator ? GCRT_ArrayStructFreezable : GCRT_ArrayStruct);
+
+			// GCRT_ArrayStruct and GCRT_ArrayStructFreezable push a new stack frame in TFastReferenceCollector
+			FGCStackSizeHelperScope StackSizeScope(StackSizeHelper, Inner);
+
+			OwnerClass.ReferenceTokenStream.EmitStride(Inner->ElementSize);
+			const uint32 SkipIndexIndex = OwnerClass.ReferenceTokenStream.EmitSkipIndexPlaceholder();
+			Inner->EmitReferenceInfo(OwnerClass, 0, EncounteredStructProps, StackSizeHelper);
+			const uint32 SkipIndex = OwnerClass.ReferenceTokenStream.EmitReturn();
+			OwnerClass.ReferenceTokenStream.UpdateSkipIndexPlaceholder(SkipIndexIndex, SkipIndex);
 		}
 	}
 }
@@ -3345,6 +3342,7 @@ void FGCReferenceTokenStream::Fixup(void (*AddReferencedObjectsPtr)(UObject*, cl
 		case GCRT_ArrayDelegate:
 		case GCRT_MulticastDelegate:
 		case GCRT_ArrayMulticastDelegate:
+		case GCRT_DynamicallyTypedValue:
 			break;
 		default:
 			UE_LOG(LogGarbage, Fatal, TEXT("Unknown token type (%u) when trying to add ARO token."), (uint32)Token.Type);
diff --git a/Engine/Source/Runtime/CoreUObject/Public/UObject/DynamicallyTypedValue.h b/Engine/Source/Runtime/CoreUObject/Public/UObject/DynamicallyTypedValue.h
new file mode 100644
--- /dev/null
+++ b/Engine/Source/Runtime/CoreUObject/Public/UObject/DynamicallyTypedValue.h
@@ -0,0 +1,178 @@
+// Copyright Epic Games, Inc. All Rights Reserved.
+
+#pragma once
+
+#include "CoreMinimal.h"
+#include "Serialization/StructuredArchive.h"
+
+class FReferenceCollector;
+
+namespace UE
+{
+	// Provides methods to interact with values of a specific type.
+	struct FDynamicallyTypedValueType
+	{
+		enum class EContainsReferences : bool
+		{
+			DoesNot,
+			Maybe,
+		};
+
+		constexpr FDynamicallyTypedValueType(SIZE_T InNumBytes, SIZE_T InMinAlignmentLogTwo, EContainsReferences InContainsReferences)
+			: NumBytes(InNumBytes)
+			, MinAlignmentLogTwo(InMinAlignmentLogTwo)
+			, ContainsReferences(InContainsReferences)
+		{
+		}
+
+		// Marks the type itself as reachable.
+		virtual void MarkReachable() = 0;
+
+		// Marks a value of the type as reachable.
+		virtual void MarkValueReachable(void* Data, FReferenceCollector& Collector) const = 0;
+
+		virtual void InitializeValue(void* Data) const = 0;
+		virtual void InitializeValueFromCopy(void* DestData, const void* SourceData) const = 0;
+		virtual void DestroyValue(void* Data) const = 0;
+
+		virtual void SerializeValue(FStructuredArchive::FSlot Slot, void* Data) const = 0;
+
+		virtual uint32 GetValueHash(const void* Data) const = 0;
+		virtual bool AreEqual(const void* DataA, const void* DataB) const = 0;
+
+		SIZE_T GetNumBytes() const { return NumBytes; }
+		uint8 GetMinAlignmentLogTwo() const { return MinAlignmentLogTwo; }
+		uint32 GetMinAlignment() const { return 1 << MinAlignmentLogTwo; }
+		EContainsReferences GetContainsReferences() const { return ContainsReferences; }
+
+	private:
+		const SIZE_T NumBytes;
+		const uint8 MinAlignmentLogTwo;
+		const EContainsReferences ContainsReferences;
+	};
+
+	// An value stored in some uninterpreted memory and a pointer to a type that contains methods to interpret it.
+	struct COREUOBJECT_API FDynamicallyTypedValue
+	{
+		static FDynamicallyTypedValueType& NullType();
+
+		FDynamicallyTypedValue() { InitializeToNull(); }
+		FDynamicallyTypedValue(const FDynamicallyTypedValue& Copyee) { InitializeFromCopy(Copyee); }
+		FDynamicallyTypedValue(FDynamicallyTypedValue&& Movee) { InitializeFromMove(MoveTemp(Movee)); }
+
+		~FDynamicallyTypedValue() { Deinit(); }
+
+		FDynamicallyTypedValue& operator=(const FDynamicallyTypedValue& Copyee)
+		{
+			if (this != &Copyee)
+			{
+				Deinit();
+				InitializeFromCopy(Copyee);
+			}
+			return *this;
+		}
+		FDynamicallyTypedValue& operator=(FDynamicallyTypedValue&& Movee)
+		{
+			if (this != &Movee)
+			{
+				Deinit();
+				InitializeFromMove(MoveTemp(Movee));
+			}
+			return *this;
+		}
+
+		// Returns a pointer to the value's data.
+		const void* GetDataPointer() const { return IsInline() ? &InlineData : HeapData; }
+		void* GetDataPointer() { return IsInline() ? &InlineData : HeapData; }
+
+		// Returns the value's type.
+		FDynamicallyTypedValueType& GetType() const { return *Type; }
+
+		// Sets the value to the null state.
+		void SetToNull()
+		{
+			Deinit();
+			InitializeToNull();
+		}
+
+		// Sets the value to the initial value of a type.
+		void InitializeAsType(FDynamicallyTypedValueType& NewType)
+		{
+			check(&NewType != nullptr);
+			Deinit();
+			Type = &NewType;
+			AllocateData();
+			Type->InitializeValue(GetDataPointer());
+		}
+
+	private:
+
+		FDynamicallyTypedValueType* Type;
+
+		// Store pointer-sized or smaller values inline, heap allocate all others.
+		union
+		{
+			UPTRINT InlineData;
+			void* HeapData;
+		};
+
+		// Initialize this value from the primordial state to the null state.
+		void InitializeToNull()
+		{
+			Type = &NullType();
+			HeapData = nullptr;
+		}
+		// Deinitializes this value back to the primordial state.
+		void Deinit()
+		{
+			Type->DestroyValue(GetDataPointer());
+			FreeData();
+			Type = nullptr;
+		}
+		// Copies the data from another value to this one, which is assumed to be in the primordial state.
+		void InitializeFromCopy(const FDynamicallyTypedValue& Copyee)
+		{
+			Type = Copyee.Type;
+			AllocateData();
+			Type->InitializeValueFromCopy(GetDataPointer(), Copyee.GetDataPointer());
+		}
+		// Moves the data from another value to this one, which is assumed to be in the primordial state.
+		// The source value is set to the null state.
+		void InitializeFromMove(FDynamicallyTypedValue&& Movee)
+		{
+			// Simply copy the type and data from the source value.
+			// This assumes that the data is trivially relocatable.
+			Type = Movee.Type;
+			InlineData = Movee.InlineData;
+
+			// Reset the source value to null.
+			Movee.InitializeToNull();
+		}
+
+		// Whether the value's data is stored in InlineData or in the memory pointed to by HeapData.
+		bool IsInline() const
+		{
+			return Type->GetNumBytes() <= sizeof(UPTRINT)
+				&& Type->GetMinAlignmentLogTwo() <= FMath::ConstExprCeilLogTwo(alignof(UPTRINT));
+		}
+
+		// Allocates heap memory for the value if it uses it.
+		void AllocateData()
+		{
+			if (!IsInline())
+			{
+				HeapData = FMemory::Malloc(Type->GetNumBytes(), Type->GetMinAlignment());
+			}
+		}
+
+		// Frees heap memory for the value if it uses it.
+		void FreeData()
+		{
+			if (!IsInline())
+			{
+				FMemory::Free(HeapData);
+				HeapData = nullptr;
+			}
+		}
+	};
+}
\ No newline at end of file
diff --git a/Engine/Source/Runtime/CoreUObject/Public/UObject/FastReferenceCollector.h b/Engine/Source/Runtime/CoreUObject/Public/UObject/FastReferenceCollector.h
--- a/Engine/Source/Runtime/CoreUObject/Public/UObject/FastReferenceCollector.h
+++ b/Engine/Source/Runtime/CoreUObject/Public/UObject/FastReferenceCollector.h
@@ -15,6 +15,7 @@
 #include "Async/ParallelFor.h"
 #include "UObject/UObjectArray.h"
 #include "UObject/FastReferenceCollectorOptions.h"
+#include "UObject/DynamicallyTypedValue.h"
 
 struct FStackEntry;
 
@@ -1283,6 +1284,17 @@ private:
 						}
 					}
 					break;
+					case GCRT_DynamicallyTypedValue:
+					{
+						TokenReturnCount = ReferenceInfo.ReturnCount;
+						UE::FDynamicallyTypedValue* DynamicallyTypedValue = (UE::FDynamicallyTypedValue*)(StackEntryData + ReferenceInfo.Offset);
+						DynamicallyTypedValue->GetType().MarkReachable();
+						if (DynamicallyTypedValue->GetType().GetContainsReferences() != UE::FDynamicallyTypedValueType::EContainsReferences::DoesNot)
+						{
+							DynamicallyTypedValue->GetType().MarkValueReachable(DynamicallyTypedValue->GetDataPointer(), ReferenceCollector);
+						}
+					}
+					break;
 					case GCRT_EndOfStream:
 					{
 						// Break out of loop.
diff --git a/Engine/Source/Runtime/CoreUObject/Public/UObject/GarbageCollection.h b/Engine/Source/Runtime/CoreUObject/Public/UObject/GarbageCollection.h
--- a/Engine/Source/Runtime/CoreUObject/Public/UObject/GarbageCollection.h
+++ b/Engine/Source/Runtime/CoreUObject/Public/UObject/GarbageCollection.h
@@ -80,6 +80,7 @@ enum EGCReferenceType
 	GCRT_ArrayDelegate,
 	GCRT_MulticastDelegate,
 	GCRT_ArrayMulticastDelegate,
+	GCRT_DynamicallyTypedValue,
 };
 
 enum class EGCTokenType : uint32
