Allow stand alone programs to opt into full FPackagePath functionality even if the program is not compiling with the editor.

#rnx

- This is useful for stand alone programs that want to compile with core and coreuobject but don't want to compile against the editor/engine.
- To make use of this the program should add 'GlobalDefinitions.Add("UE_SUPPORT_FULL_PACKAGEPATH=1")' to it's Target.cs .
- If the program does not define UE_SUPPORT_FULL_PACKAGEPATH then it will default to WITH_EDITOR which was the pre-existing behavior.

diff --git a/Engine/Source/Runtime/CoreUObject/Private/Misc/PackagePath.cpp b/Engine/Source/Runtime/CoreUObject/Private/Misc/PackagePath.cpp
--- a/Engine/Source/Runtime/CoreUObject/Private/Misc/PackagePath.cpp
+++ b/Engine/Source/Runtime/CoreUObject/Private/Misc/PackagePath.cpp
@@ -368,7 +368,7 @@ FString FPackagePath::GetDebugNameWithExtension(EPackageSegment PackageSegment)
 	return Result;
 }
 
-#if WITH_EDITOR
+#if UE_SUPPORT_FULL_PACKAGEPATH
 
 FPackagePath& FPackagePath::operator=(const FPackagePath& Other)
 {
@@ -1207,4 +1207,4 @@ FStringView FPackagePath::GetCustomExtension() const
 	return FStringView();
 }
 
-#endif
\ No newline at end of file
+#endif //UE_SUPPORT_FULL_PACKAGEPATH
diff --git a/Engine/Source/Runtime/CoreUObject/Public/Misc/PackagePath.h b/Engine/Source/Runtime/CoreUObject/Public/Misc/PackagePath.h
--- a/Engine/Source/Runtime/CoreUObject/Public/Misc/PackagePath.h
+++ b/Engine/Source/Runtime/CoreUObject/Public/Misc/PackagePath.h
@@ -10,6 +10,10 @@
 #include "Internationalization/Text.h"
 #include "Templates/UniquePtr.h"
 
+#ifndef UE_SUPPORT_FULL_PACKAGEPATH
+	#define UE_SUPPORT_FULL_PACKAGEPATH WITH_EDITOR
+#endif
+
 class FArchive;
 
 /**
@@ -475,7 +479,7 @@ private:
 		LocalOnlyPath,
 	};
 
-#if WITH_EDITOR
+#if UE_SUPPORT_FULL_PACKAGEPATH
 	/**
 	 * Combined storage for (EPackageIDType-specific path+filename storage) followed by (possibly empty)
 	 * PackageNameRoot followed by (possibly empty) FilePathRoot followed by (possibly empty) Extension
@@ -491,5 +495,5 @@ private:
 #else
 	FName PackageName;
 	mutable EPackageExtension HeaderExtension = EPackageExtension::Unspecified;
-#endif
+#endif //UE_SUPPORT_FULL_PACKAGEPATH
 };
