Avoid dereferencing a null DefaultValue pointer when calling ExportTextItem

diff --git a/Engine/Source/Runtime/CoreUObject/Public/UObject/Class.h b/Engine/Source/Runtime/CoreUObject/Public/UObject/Class.h
--- a/Engine/Source/Runtime/CoreUObject/Public/UObject/Class.h
+++ b/Engine/Source/Runtime/CoreUObject/Public/UObject/Class.h
@@ -931,7 +931,7 @@ public:
 
 	/** Template to manage dynamic access to C++ struct construction and destruction **/
 	template<class CPPSTRUCT>
-	struct TCppStructOps : public ICppStructOps
+	struct TCppStructOps final : public ICppStructOps
 	{
 		typedef TStructOpsTypeTraits<CPPSTRUCT> TTraits;
 		TCppStructOps()
@@ -1158,7 +1158,28 @@ public:
 			check(TTraits::WithExportTextItem); // don't call this if we have indicated it is not necessary
 			if constexpr (TStructOpsTypeTraits<CPPSTRUCT>::WithExportTextItem)
 			{
-				return ((const CPPSTRUCT*)PropertyValue)->ExportTextItem(ValueStr, *(const CPPSTRUCT*)DefaultValue, Parent, PortFlags, ExportRootScope);
+				if (DefaultValue)
+				{
+					return ((const CPPSTRUCT*)PropertyValue)->ExportTextItem(ValueStr, *(const CPPSTRUCT*)DefaultValue, Parent, PortFlags, ExportRootScope);
+				}
+				else
+				{
+					TTypeCompatibleBytes<CPPSTRUCT> TmpDefaultValue;
+					FMemory::Memzero(TmpDefaultValue.GetTypedPtr(), sizeof(CPPSTRUCT));
+					if (!HasZeroConstructor())
+					{
+						Construct(TmpDefaultValue.GetTypedPtr());
+					}
+
+					const bool bResult = ((const CPPSTRUCT*)PropertyValue)->ExportTextItem(ValueStr, *(const CPPSTRUCT*)TmpDefaultValue.GetTypedPtr(), Parent, PortFlags, ExportRootScope);
+
+					if (HasDestructor())
+					{
+						Destruct(TmpDefaultValue.GetTypedPtr());
+					}
+
+					return bResult;
+				}
 			}
 			else
 			{
