Moved the code to import default function argument values to a common location so it can be shared

diff --git a/Engine/Plugins/Experimental/PythonScriptPlugin/Source/PythonScriptPlugin/Private/PyUtil.cpp b/Engine/Plugins/Experimental/PythonScriptPlugin/Source/PythonScriptPlugin/Private/PyUtil.cpp
--- a/Engine/Plugins/Experimental/PythonScriptPlugin/Source/PythonScriptPlugin/Private/PyUtil.cpp
+++ b/Engine/Plugins/Experimental/PythonScriptPlugin/Source/PythonScriptPlugin/Private/PyUtil.cpp
@@ -483,45 +483,7 @@ bool IsOutputParameter(const FProperty* InParam)
 
 void ImportDefaultValue(const FProperty* InProp, void* InPropValue, const FString& InDefaultValue)
 {
-	if (!InDefaultValue.IsEmpty())
-	{
-		// Certain struct types export using a non-standard default value, so we have to import them manually rather than use ImportText
-		if (const FStructProperty* StructProp = CastField<FStructProperty>(InProp))
-		{
-			if (StructProp->Struct == TBaseStructure<FVector>::Get())
-			{
-				FVector* Vector = (FVector*)InPropValue;
-				FDefaultValueHelper::ParseVector(InDefaultValue, *Vector);
-				return;
-			}
-			else if (StructProp->Struct == TBaseStructure<FVector2D>::Get())
-			{
-				FVector2D* Vector2D = (FVector2D*)InPropValue;
-				FDefaultValueHelper::ParseVector2D(InDefaultValue, *Vector2D);
-				return;
-			}
-			else if (StructProp->Struct == TBaseStructure<FRotator>::Get())
-			{
-				FRotator* Rotator = (FRotator*)InPropValue;
-				FDefaultValueHelper::ParseRotator(InDefaultValue, *Rotator);
-				return;
-			}
-			else if (StructProp->Struct == TBaseStructure<FColor>::Get())
-			{
-				FColor* Color = (FColor*)InPropValue;
-				FDefaultValueHelper::ParseColor(InDefaultValue, *Color);
-				return;
-			}
-			else if (StructProp->Struct == TBaseStructure<FLinearColor>::Get())
-			{
-				FLinearColor* LinearColor = (FLinearColor*)InPropValue;
-				FDefaultValueHelper::ParseLinearColor(InDefaultValue, *LinearColor);
-				return;
-			}
-		}
-
-		InProp->ImportText_Direct(*InDefaultValue, InPropValue, nullptr, PPF_None);
-	}
+	PropertyAccessUtil::ImportDefaultPropertyValue(InProp, InPropValue, InDefaultValue);
 }
 
 bool InvokeFunctionCall(UObject* InObj, const UFunction* InFunc, void* InBaseParamsAddr, const TCHAR* InErrorCtxt)
diff --git a/Engine/Source/Runtime/CoreUObject/Private/UObject/PropertyAccessUtil.cpp b/Engine/Source/Runtime/CoreUObject/Private/UObject/PropertyAccessUtil.cpp
--- a/Engine/Source/Runtime/CoreUObject/Private/UObject/PropertyAccessUtil.cpp
+++ b/Engine/Source/Runtime/CoreUObject/Private/UObject/PropertyAccessUtil.cpp
@@ -4,6 +4,7 @@
 #include "UObject/EnumProperty.h"
 #include "UObject/Object.h"
 #include "UObject/Class.h"
+#include "Misc/DefaultValueHelper.h"
 
 namespace UE::PropertyAccessUtil::Private
 {
@@ -641,4 +642,49 @@ FProperty* FindPropertyByName(const FName InPropName, const UStruct* InStruct)
 	return Prop;
 }
 
+void ImportDefaultPropertyValue(const FProperty* InProp, void* InPropValue, const FString& InDefaultValue)
+{
+	if (InDefaultValue.IsEmpty())
+	{
+		return;
+	}
+
+	bool bDoImportText = true;
+
+	// Certain struct types export using a non-standard default value, so we have to import them manually rather than use ImportText
+	if (const FStructProperty* StructProp = CastField<FStructProperty>(InProp))
+	{
+		if (StructProp->Struct == TBaseStructure<FVector>::Get())
+		{
+			FVector* Vector = (FVector*)InPropValue;
+			bDoImportText = !FDefaultValueHelper::ParseVector(InDefaultValue, *Vector);
+		}
+		else if (StructProp->Struct == TBaseStructure<FVector2D>::Get())
+		{
+			FVector2D* Vector2D = (FVector2D*)InPropValue;
+			bDoImportText = !FDefaultValueHelper::ParseVector2D(InDefaultValue, *Vector2D);
+		}
+		else if (StructProp->Struct == TBaseStructure<FRotator>::Get())
+		{
+			FRotator* Rotator = (FRotator*)InPropValue;
+			bDoImportText = !FDefaultValueHelper::ParseRotator(InDefaultValue, *Rotator);
+		}
+		else if (StructProp->Struct == TBaseStructure<FColor>::Get())
+		{
+			FColor* Color = (FColor*)InPropValue;
+			bDoImportText = !FDefaultValueHelper::ParseColor(InDefaultValue, *Color);
+		}
+		else if (StructProp->Struct == TBaseStructure<FLinearColor>::Get())
+		{
+			FLinearColor* LinearColor = (FLinearColor*)InPropValue;
+			bDoImportText = !FDefaultValueHelper::ParseLinearColor(InDefaultValue, *LinearColor);
+		}
+	}
+
+	if (bDoImportText)
+	{
+		InProp->ImportText_Direct(*InDefaultValue, InPropValue, nullptr, PPF_None);
+	}
+}
+
 }
diff --git a/Engine/Source/Runtime/CoreUObject/Public/UObject/PropertyAccessUtil.h b/Engine/Source/Runtime/CoreUObject/Public/UObject/PropertyAccessUtil.h
--- a/Engine/Source/Runtime/CoreUObject/Public/UObject/PropertyAccessUtil.h
+++ b/Engine/Source/Runtime/CoreUObject/Public/UObject/PropertyAccessUtil.h
@@ -292,6 +292,17 @@ namespace PropertyAccessUtil
 	 */
 	COREUOBJECT_API FProperty* FindPropertyByName(const FName InPropName, const UStruct* InStruct);
 
+	/**
+	 * Low-level function to import a default value to the given property.
+	 * @note This is useful for applying defaults to function properties (eg, parameter data from UEdGraphSchema_K2::FindFunctionParameterDefaultValue, or local data from 
+	 *       FBlueprintEditorUtils::FindLocalVariable), as these may have come from UHT which can sometimes have a slightly different form than ImportText can handle natively.
+	 *
+	 * @param InProp The property to import the default value for.
+	 * @param InPropValue The destination location to receive the imported value.
+	 * @param InDefaultValue The default value to import.
+	 */
+	COREUOBJECT_API void ImportDefaultPropertyValue(const FProperty* InProp, void* InPropValue, const FString& InDefaultValue);
+
 	/**
 	 * Low-level function to determine if two properties are compatible with each other (for use by CopySinglePropertyValue/CopyCompletePropertyValue).
 	 * 
