Look for uncooked packages first in FPackageName::DoesPackageExist

diff --git a/Engine/Source/Runtime/CoreUObject/Private/Misc/PackageName.cpp b/Engine/Source/Runtime/CoreUObject/Private/Misc/PackageName.cpp
--- a/Engine/Source/Runtime/CoreUObject/Private/Misc/PackageName.cpp
+++ b/Engine/Source/Runtime/CoreUObject/Private/Misc/PackageName.cpp
@@ -1624,6 +1624,31 @@ FPackageName::EPackageLocationFilter FPackageName::DoesPackageExistEx(const FPac
 	}
 
 	EPackageLocationFilter Result = EPackageLocationFilter::None;
+
+	if ((uint8)Filter & (uint8)EPackageLocationFilter::FileSystem)
+	{
+		bool bFoundInFileSystem = false;
+		if (bMatchCaseOnDisk)
+		{
+			bFoundInFileSystem = IPackageResourceManager::Get().TryMatchCaseOnDisk(PackagePath, OutPackagePath);
+		}
+		else
+		{
+			bFoundInFileSystem = IPackageResourceManager::Get().DoesPackageExist(PackagePath, OutPackagePath);
+		}
+
+		if (bFoundInFileSystem)
+		{
+			Result = EPackageLocationFilter((uint8)Result | (uint8)EPackageLocationFilter::FileSystem);
+
+			// if we just want to find any existence, then we are done
+			if (Filter == EPackageLocationFilter::Any)
+			{
+				return Result;
+			}
+		}
+	}
+
 	if (((uint8)Filter & (uint8)EPackageLocationFilter::IoDispatcher))
 	{
 		if (DoesPackageExistOverrideDelegate.IsBound())
@@ -1641,7 +1666,7 @@ FPackageName::EPackageLocationFilter FPackageName::DoesPackageExistEx(const FPac
 
 				Result = EPackageLocationFilter((uint8)Result | (uint8)EPackageLocationFilter::IoDispatcher);
 
-				// if we just want to find any existence, then we are done and we can skip the on disk check lower
+				// if we just want to find any existence, then we are done
 				if (Filter == EPackageLocationFilter::Any)
 				{
 					return Result;
@@ -1650,24 +1675,6 @@ FPackageName::EPackageLocationFilter FPackageName::DoesPackageExistEx(const FPac
 		}
 	}
 
-	if ((uint8)Filter & (uint8)EPackageLocationFilter::FileSystem)
-	{
-		bool bFoundInFileSystem = false;
-		if (bMatchCaseOnDisk)
-		{
-			bFoundInFileSystem = IPackageResourceManager::Get().TryMatchCaseOnDisk(PackagePath, OutPackagePath);
-		}
-		else
-		{
-			bFoundInFileSystem = IPackageResourceManager::Get().DoesPackageExist(PackagePath, OutPackagePath);
-		}
-
-		if (bFoundInFileSystem)
-		{
-			Result = EPackageLocationFilter((uint8)Result | (uint8)EPackageLocationFilter::FileSystem);
-		}
-	}
-	
 	return Result;
 }
 
