If an export's class is not included in the harvesting, emit an error explaining why and exit SavePackage with an error rather than emitting an assertion in BuildLinker's Map Export Indices step.
#rnx

diff --git a/Engine/Source/Runtime/CoreUObject/Private/UObject/SavePackage/PackageHarvester.cpp b/Engine/Source/Runtime/CoreUObject/Private/UObject/SavePackage/PackageHarvester.cpp
--- a/Engine/Source/Runtime/CoreUObject/Private/UObject/SavePackage/PackageHarvester.cpp
+++ b/Engine/Source/Runtime/CoreUObject/Private/UObject/SavePackage/PackageHarvester.cpp
@@ -200,11 +200,16 @@ void FPackageHarvester::ProcessExport(const FExportWithContext& InProcessContext
 	FExportScope HarvesterScope(*this, InProcessContext, bReferencerIsEditorOnly);
 	
 	// The export scope set the current harvesting context
-	check(SaveContext.GetHarvestedRealm(CurrentExportHarvestingRealm).IsExport(Export));
+	FHarvestedRealm& HarvestedRealm = SaveContext.GetHarvestedRealm(CurrentExportHarvestingRealm);
+	check(HarvestedRealm.IsExport(Export));
 
 	// Harvest its class 
 	UClass* Class = Export->GetClass();
 	*this << Class;
+	if (!HarvestedRealm.IsIncluded(Class))
+	{
+		SaveContext.RecordIllegalReference(Export, Class, EIllegalRefReason::UnsaveableClass, GetUnsaveableReason(Class));
+	}
 
 	// Harvest the export outer
 	if (UObject* Outer = Export->GetOuter())
@@ -227,6 +232,14 @@ void FPackageHarvester::ProcessExport(const FExportWithContext& InProcessContext
 			FIgnoreDependenciesScope IgnoreDependencies(*this);
 			*this << Outer;
 		}
+		if (!HarvestedRealm.IsIncluded(Outer))
+		{
+			// Only packages or object having the currently saved package as outer are allowed to have no outer
+			if (!Export->IsA<UPackage>() && Outer != SaveContext.GetPackage())
+			{
+				SaveContext.RecordIllegalReference(Export, Outer, EIllegalRefReason::UnsaveableOuter, GetUnsaveableReason(Outer));
+			}
+		}
 	}
 
 	// Harvest its template, if any
@@ -657,4 +670,55 @@ void FPackageHarvester::AppendCurrentExportDependencies()
 	SaveContext.GetHarvestedRealm(CurrentExportHarvestingRealm).GetObjectDependencies().Add(CurrentExportDependencies.CurrentExport, MoveTemp(CurrentExportDependencies.ObjectReferences));
 	SaveContext.GetHarvestedRealm(CurrentExportHarvestingRealm).GetNativeObjectDependencies().Add(CurrentExportDependencies.CurrentExport, MoveTemp(CurrentExportDependencies.NativeObjectReferences));
 	CurrentExportDependencies.CurrentExport = nullptr;
+}
+
+FString FPackageHarvester::GetUnsaveableReason(UObject* Required)
+{
+	// Copy some of the code from operator<<(UObject*), TryHarvestExport, TryHarvestImport to find out why the Required object was not included
+	FString ReasonText(TEXTVIEW("It should be included but was excluded for an unknown reason."));
+	EIllegalRefReason UnusedRealmReason = EIllegalRefReason::None;
+	ESaveRealm HarvestContext = GetObjectHarvestingRealm(Required, UnusedRealmReason);
+	bool bShouldBeExport = Required->IsInPackage(SaveContext.GetPackage());
+
+	FSaveContext::ESaveableStatus CulpritStatus;
+	UObject* Culprit;
+	FSaveContext::ESaveableStatus Status = SaveContext.GetSaveableStatus(Required, &Culprit, &CulpritStatus);
+	if (Status != FSaveContext::ESaveableStatus::Success)
+	{
+		if (Status == FSaveContext::ESaveableStatus::OuterUnsaveable)
+		{
+			check(Culprit);
+			ReasonText = FString::Printf(TEXT("It has outer %s which %s."), *Culprit->GetPathName(), LexToString(CulpritStatus));
+		}
+		else
+		{
+			ReasonText =  FString::Printf(TEXT("It %s"), *LexToString(Status));
+		}
+	}
+	else if (ConditionallyExcludeObjectForTarget(SaveContext, Required, HarvestContext))
+	{
+		ReasonText = TEXTVIEW("It is excluded for the current cooking target.");
+	}
+	else if (bShouldBeExport)
+	{
+		// The class is in the package and so should be an export; we don't know of any other reasons why it would be excluded
+		ReasonText = TEXTVIEW("It should be an export but was excluded for an unknown reason.");
+	}
+	else
+	{
+		// The class is not in the package and so should be an import
+		bool bExcludePackageFromCook = FCoreUObjectDelegates::ShouldCookPackageForPlatform.IsBound() ?
+			!FCoreUObjectDelegates::ShouldCookPackageForPlatform.Execute(Required->GetOutermost(), CookingTarget()) : false;
+		if (bExcludePackageFromCook)
+		{
+			ReasonText = FString::Printf(TEXT("It is in package %s which is excluded from the cook by FCoreUObjectDelegates::ShouldCookPackageForPlatform."),
+				*Required->GetOutermost()->GetName());
+		}
+		else
+		{
+			// We don't know of any other reasons why it would be excluded
+			ReasonText = TEXTVIEW("It should be an import but was excluded for an unknown reason.");
+		}
+	}
+	return ReasonText;
 }
\ No newline at end of file
diff --git a/Engine/Source/Runtime/CoreUObject/Private/UObject/SavePackage/PackageHarvester.h b/Engine/Source/Runtime/CoreUObject/Private/UObject/SavePackage/PackageHarvester.h
--- a/Engine/Source/Runtime/CoreUObject/Private/UObject/SavePackage/PackageHarvester.h
+++ b/Engine/Source/Runtime/CoreUObject/Private/UObject/SavePackage/PackageHarvester.h
@@ -96,6 +96,8 @@ private:
 	void HarvestImport(UObject* InObject);
 
 	void AppendCurrentExportDependencies();
+	FString GetUnsaveableReason(UObject* Required);
+
 
 	struct FExportDependencies
 	{
diff --git a/Engine/Source/Runtime/CoreUObject/Private/UObject/SavePackage/SaveContext.cpp b/Engine/Source/Runtime/CoreUObject/Private/UObject/SavePackage/SaveContext.cpp
--- a/Engine/Source/Runtime/CoreUObject/Private/UObject/SavePackage/SaveContext.cpp
+++ b/Engine/Source/Runtime/CoreUObject/Private/UObject/SavePackage/SaveContext.cpp
@@ -28,6 +28,8 @@ void FSaveContext::MarkUnsaveable(UObject* InObject)
 {
 	if (IsUnsaveable(InObject))
 	{
+		// TODO: We should not be modifying objects during the save. Besides interfering with the objects outside of the save,
+		// it also prevents us from gathering reasons about why an object was marked unsaveable.
 		InObject->SetFlags(RF_Transient);
 	}
 
@@ -41,41 +43,100 @@ void FSaveContext::MarkUnsaveable(UObject* InObject)
 
 bool FSaveContext::IsUnsaveable(UObject* InObject, bool bEmitWarning) const
 {
-	UObject* Obj = InObject;
-	while (Obj)
+	UObject* Culprit;
+	ESaveableStatus CulpritStatus;
+	ESaveableStatus Status = GetSaveableStatus(InObject, &Culprit, &CulpritStatus);
+	if (Status == ESaveableStatus::Success)
 	{
-		// pending kill object are unsaveable
-		if (!IsValidChecked(Obj))
-		{
-			return true;
-		}
-
-		// transient object are considered unsaveable if non native
-		if (Obj->HasAnyFlags(RF_Transient) && !Obj->IsNative())
+		return false;
+	}
+	if (Status == ESaveableStatus::OuterUnsaveable && bEmitWarning &&
+		(CulpritStatus == ESaveableStatus::AbstractClass || CulpritStatus == ESaveableStatus::DeprecatedClass || CulpritStatus == ESaveableStatus::NewerVersionExistsClass))
+	{
+		check(Culprit);
+		// Only warn if the base object is fine but the outer is invalid. If an object is itself unsaveable, the old behavior is to ignore it
+		if (InObject->GetOutermost() == GetPackage())
 		{
-			return true;
+			UE_LOG(LogSavePackage, Warning, TEXT("%s has a deprecated or abstract class outer %s, so it will not be saved"),
+				*InObject->GetFullName(), *Culprit->GetFullName());
 		}
+	}
+	return true;
+}
 
-		// if the object class is abstract, has been marked as deprecated, there is a newer version that exist, or the class is marked transient, then the object is unsaveable
-		// @note: Although object instances of a transient class should definitely be unsaveable, it results in discrepancies with the old save algorithm and currently load problems
-		if (Obj->GetClass()->HasAnyClassFlags(CLASS_Abstract | CLASS_Deprecated | CLASS_NewerVersionExists /*| CLASS_Transient*/) 
-			&& !Obj->HasAnyFlags(RF_ClassDefaultObject))
+FSaveContext::ESaveableStatus FSaveContext::GetSaveableStatus(UObject* InObject, UObject** OutCulprit, ESaveableStatus* OutCulpritStatus) const
+{
+	UObject* Obj = InObject;
+	while (Obj)
+	{
+		ESaveableStatus Status = GetSaveableStatusNoOuter(Obj);
+		if (Status != ESaveableStatus::Success)
 		{
-			// only warn if the base object is fine but the outer is invalid. If an object is itself unsaveable, the old behavior is to ignore it
-			if (bEmitWarning
-				&& IsValidChecked(Obj)
-				&& InObject->GetOutermost() == GetPackage()
-				&& Obj != InObject)
+			if (OutCulprit)
 			{
-				UE_LOG(LogSavePackage, Warning, TEXT("%s has a deprecated or abstract class outer %s, so it will not be saved"), *InObject->GetFullName(), *Obj->GetFullName());
+				*OutCulprit = Obj;
 			}
-
-			// there used to be a check for reference if the class had the CLASS_HasInstancedReference,
-			// those reference were outer-ed to the object being flagged as unsaveable, making them unsaveable as well without having to look for them
-			return true;
+			if (OutCulpritStatus)
+			{
+				*OutCulpritStatus = Status;
+			}
+			return Obj == InObject ? Status : ESaveableStatus::OuterUnsaveable;
 		}
-
 		Obj = Obj->GetOuter();
 	}
-	return false;	
+	if (OutCulprit)
+	{
+		*OutCulprit = InObject;
+	}
+	if (OutCulpritStatus)
+	{
+		*OutCulpritStatus = ESaveableStatus::Success;
+	}
+	return ESaveableStatus::Success;
+}
+
+FSaveContext::ESaveableStatus FSaveContext::GetSaveableStatusNoOuter(UObject* Obj) const
+{
+	// pending kill object are unsaveable
+	if (!IsValidChecked(Obj))
+	{
+		return ESaveableStatus::PendingKill;
+	}
+
+	// transient object are considered unsaveable if non native
+	if (Obj->HasAnyFlags(RF_Transient) && !Obj->IsNative())
+	{
+		return ESaveableStatus::Transient;
+	}
+
+	UClass* Class = Obj->GetClass();
+	// if the object class is abstract, has been marked as deprecated, there is a newer version that exist, or the class is marked transient, then the object is unsaveable
+	// @note: Although object instances of a transient class should definitely be unsaveable, it results in discrepancies with the old save algorithm and currently load problems
+	if (Class->HasAnyClassFlags(CLASS_Abstract | CLASS_Deprecated | CLASS_NewerVersionExists /*| CLASS_Transient*/)
+		&& !Obj->HasAnyFlags(RF_ClassDefaultObject))
+	{
+		// There used to be a check for reference if the class had the CLASS_HasInstancedReference,
+		// but we don't need it because those references are outer-ed to the object being flagged as unsaveable, making them unsaveable as well without having to look for them
+		return Class->HasAnyClassFlags(CLASS_Abstract) ? ESaveableStatus::AbstractClass :
+			Class->HasAnyClassFlags(CLASS_Deprecated) ? ESaveableStatus::DeprecatedClass :
+			ESaveableStatus::NewerVersionExistsClass;
+	}
+
+	return ESaveableStatus::Success;
 }
+
+const TCHAR* LexToString(FSaveContext::ESaveableStatus Status)
+{
+	static_assert(static_cast<int32>(FSaveContext::ESaveableStatus::__Count) == 7);
+	switch (Status)
+	{
+	case FSaveContext::ESaveableStatus::Success: return TEXT("is saveable");
+	case FSaveContext::ESaveableStatus::PendingKill: return TEXT("is pendingkill");
+	case FSaveContext::ESaveableStatus::Transient: return TEXT("is transient");
+	case FSaveContext::ESaveableStatus::AbstractClass: return TEXT("has a Class with CLASS_Abstract");
+	case FSaveContext::ESaveableStatus::DeprecatedClass: return TEXT("has a Class with CLASS_Deprecated");
+	case FSaveContext::ESaveableStatus::NewerVersionExistsClass: return TEXT("has a Class with CLASS_NewerVersionExists");
+	case FSaveContext::ESaveableStatus::OuterUnsaveable: return TEXT("has an unsaveable Outer");
+	default: return TEXT("Unknown");
+	}
+}
\ No newline at end of file
diff --git a/Engine/Source/Runtime/CoreUObject/Private/UObject/SavePackage/SaveContext.h b/Engine/Source/Runtime/CoreUObject/Private/UObject/SavePackage/SaveContext.h
--- a/Engine/Source/Runtime/CoreUObject/Private/UObject/SavePackage/SaveContext.h
+++ b/Engine/Source/Runtime/CoreUObject/Private/UObject/SavePackage/SaveContext.h
@@ -80,6 +80,8 @@ enum class EIllegalRefReason : uint8
 	None = 0,
 	ReferenceToOptional,
 	ReferenceFromOptionalToMissingGameExport,
+	UnsaveableClass,
+	UnsaveableOuter,
 };
 
 /** Small struct to store illegal references harvested during save */
@@ -88,6 +90,7 @@ struct FIllegalReference
 	UObject* From = nullptr;
 	UObject* To = nullptr;
 	EIllegalRefReason Reason;
+	FString FormatStringArg;
 };
 
 /** Hold the harvested exports and imports for a realm */
@@ -663,10 +666,23 @@ public:
 	void MarkUnsaveable(UObject* InObject);
 
 	bool IsUnsaveable(UObject* InObject, bool bEmitWarning = true) const;
+	enum class ESaveableStatus
+	{
+		Success,
+		PendingKill,
+		Transient,
+		AbstractClass,
+		DeprecatedClass,
+		NewerVersionExistsClass,
+		OuterUnsaveable,
+		__Count,
+	};
+	ESaveableStatus GetSaveableStatus(UObject* InObject, UObject** OutCulprit = nullptr, ESaveableStatus* OutCulpritStatus = nullptr) const;
+	ESaveableStatus GetSaveableStatusNoOuter(UObject* InObject) const;
 
-	void RecordIllegalReference(UObject* InFrom, UObject* InTo, EIllegalRefReason InReason)
+	void RecordIllegalReference(UObject* InFrom, UObject* InTo, EIllegalRefReason InReason, FString&& InOptionalReasonText = FString())
 	{
-		HarvestedIllegalReferences.Add({ InFrom, InTo, InReason });
+		HarvestedIllegalReferences.Add({ InFrom, InTo, InReason, MoveTemp(InOptionalReasonText) });
 	}
 
 	const TArray<FIllegalReference>& GetIllegalReferences() const
@@ -1017,3 +1033,5 @@ private:
 	// Set of harvested prestream packages, should be deprecated
 	TSet<UPackage*> PrestreamPackages;
 };
+
+const TCHAR* LexToString(FSaveContext::ESaveableStatus Status);
diff --git a/Engine/Source/Runtime/CoreUObject/Private/UObject/SavePackage2.cpp b/Engine/Source/Runtime/CoreUObject/Private/UObject/SavePackage2.cpp
--- a/Engine/Source/Runtime/CoreUObject/Private/UObject/SavePackage2.cpp
+++ b/Engine/Source/Runtime/CoreUObject/Private/UObject/SavePackage2.cpp
@@ -437,6 +437,20 @@ ESavePackageResult ValidateRealms(FSaveContext& SaveContext)
 					Reference.From ? *Reference.From->GetPathName() : TEXT("Unknown"),
 					Reference.To ? *Reference.To->GetPathName() : TEXT("Unknown"));
 				break;
+			case EIllegalRefReason::UnsaveableClass:
+				ErrorMessage = FString::Printf(TEXT("Can't save %s: Object (%s) is an export but is an instance of class (%s) which is unsaveable: %s"),
+					SaveContext.GetFilename(),
+					Reference.From ? *Reference.From->GetPathName() : TEXT("Unknown"),
+					Reference.To ? *Reference.To->GetPathName() : TEXT("Unknown"),
+					*Reference.FormatStringArg);
+				break;
+			case EIllegalRefReason::UnsaveableOuter:
+				ErrorMessage = FString::Printf(TEXT("Can't save %s: Object (%s) is an export but is a subobject of (%s) which is unsaveable: %s"),
+					SaveContext.GetFilename(),
+					Reference.From ? *Reference.From->GetPathName() : TEXT("Unknown"),
+					Reference.To ? *Reference.To->GetPathName() : TEXT("Unknown"),
+					*Reference.FormatStringArg);
+				break;
 			default:
 				ErrorMessage = FString::Printf(TEXT("Can't save %s: Unknown Illegal reference from object (%s) to object (%s)"),
 					SaveContext.GetFilename(),
@@ -1135,6 +1149,7 @@ ESavePackageResult BuildLinker(FSaveContext& SaveContext)
 			if (ObjClass != UClass::StaticClass())
 			{
 				Export.ClassIndex = Linker->MapObject(ObjClass);
+				// The class should be mappable because it was checked in FPackageHarvester::ProcessExport and the save early-exited if not
 				checkf(!Export.ClassIndex.IsNull(), TEXT("Export %s class is not mapped when saving %s"), *Export.Object->GetFullName(), *Linker->LinkerRoot->GetName());
 			}
 			else
