Fix wrong test with TPropertyWithSetterAndGetter
#prefligh none

diff --git a/Engine/Source/Runtime/CoreUObject/Public/UObject/PropertyWithSetterAndGetter.h b/Engine/Source/Runtime/CoreUObject/Public/UObject/PropertyWithSetterAndGetter.h
--- a/Engine/Source/Runtime/CoreUObject/Public/UObject/PropertyWithSetterAndGetter.h
+++ b/Engine/Source/Runtime/CoreUObject/Public/UObject/PropertyWithSetterAndGetter.h
@@ -26,7 +26,7 @@ public:
 
 	virtual bool HasGetter() const override
 	{
-		return !!SetterFunc;
+		return !!GetterFunc;
 	}
 
 	virtual bool HasSetterOrGetter() const override
@@ -42,7 +42,7 @@ public:
 
 	virtual void CallGetter(const void* Container, void* OutValue) const override
 	{
-		checkf(SetterFunc, TEXT("Calling a getter on %s but the property has no getter defined."), *PropertyBaseClass::GetFullName());
+		checkf(GetterFunc, TEXT("Calling a getter on %s but the property has no getter defined."), *PropertyBaseClass::GetFullName());
 		GetterFunc(Container, OutValue);
 	}
 
