The editor is now able to load packages with package trailers even if the feature is disabled.

#rnx

- Disabling the package trailer was only added during development so we could turn it on as needed. It should prevent the creation of packages with the trailer, not prevent them from being loaded.

diff --git a/Engine/Source/Runtime/CoreUObject/Private/UObject/LinkerLoad.cpp b/Engine/Source/Runtime/CoreUObject/Private/UObject/LinkerLoad.cpp
--- a/Engine/Source/Runtime/CoreUObject/Private/UObject/LinkerLoad.cpp
+++ b/Engine/Source/Runtime/CoreUObject/Private/UObject/LinkerLoad.cpp
@@ -1531,7 +1531,7 @@ FLinkerLoad::ELinkerStatus FLinkerLoad::SerializePackageTrailer()
 
 	check(PackageTrailer == nullptr);
 
-	if (UE::FPackageTrailer::IsEnabled() && Summary.PayloadTocOffset > 0)
+	if (Summary.PayloadTocOffset > 0)
 	{
 		int64 CurPos = Tell();
 		Seek(Summary.PayloadTocOffset);
