DerivedData: Added the framework for serializing derived data references to cooked packages

#rnx

diff --git a/Engine/Source/Runtime/CoreUObject/Private/Serialization/DerivedData.cpp b/Engine/Source/Runtime/CoreUObject/Private/Serialization/DerivedData.cpp
--- a/Engine/Source/Runtime/CoreUObject/Private/Serialization/DerivedData.cpp
+++ b/Engine/Source/Runtime/CoreUObject/Private/Serialization/DerivedData.cpp
@@ -7,10 +7,12 @@
 #if WITH_EDITORONLY_DATA
 
 #include "DerivedDataBuildDefinition.h"
+#include "DerivedDataCacheKey.h"
 #include "DerivedDataValueId.h"
 #include "Serialization/EditorDerivedData.h"
 #include "Serialization/EditorDerivedDataIoStore.h"
 #include "Templates/SharedPointer.h"
+#include "UObject/LinkerSave.h"
 
 namespace UE::DerivedData::IoStore
 {
@@ -28,7 +30,24 @@ namespace UE
 
 void FDerivedData::Serialize(FArchive& Ar, UObject* Owner)
 {
-	unimplemented();
+	if (Ar.IsPersistent() && !Ar.IsObjectReferenceCollector() && !Ar.ShouldSkipBulkData())
+	{
+#if WITH_EDITORONLY_DATA
+		if (Ar.IsSaving())
+		{
+			checkf(Ar.IsCooking(), TEXT("Serializing FDerivedData is only supported for cooked packages."));
+			FLinkerSave* const Linker = Cast<FLinkerSave>(Ar.GetLinker());
+			checkf(Linker, TEXT("Serializing FDerivedData requires a linker."));
+			FIoChunkId SaveId = EditorData->Save(*Linker);
+			Ar << SaveId;
+		}
+		else
+#endif
+		{
+			Ar << ChunkId;
+		}
+		Ar << Flags;
+	}
 }
 
 FUniqueBuffer FDerivedData::LoadData(FArchive& Ar, FDerivedDataBufferAllocator Allocator)
@@ -101,8 +120,14 @@ FDerivedData::FDerivedData(const FCompressedBuffer& Data)
 {
 }
 
-FDerivedData::FDerivedData(FStringView CacheKey, FStringView CacheContext)
-	: EditorData(DerivedData::Private::MakeEditorDerivedData(CacheKey, CacheContext))
+FDerivedData::FDerivedData(const DerivedData::FSharedString& Name, const DerivedData::FCacheKey& Key)
+	: EditorData(DerivedData::Private::MakeEditorDerivedData(Name, Key, DerivedData::FValueId::Null))
+	, ChunkId(DerivedData::IoStore::GEditorDerivedDataIoStore->AddData(EditorData.Get()))
+{
+}
+
+FDerivedData::FDerivedData(const DerivedData::FSharedString& Name, const DerivedData::FCacheKey& Key, const DerivedData::FValueId& ValueId)
+	: EditorData(DerivedData::Private::MakeEditorDerivedData(Name, Key, ValueId))
 	, ChunkId(DerivedData::IoStore::GEditorDerivedDataIoStore->AddData(EditorData.Get()))
 {
 }
@@ -113,6 +138,12 @@ FDerivedData::FDerivedData(const DerivedData::FBuildDefinition& BuildDefinition,
 {
 }
 
+FDerivedData::FDerivedData(FStringView CacheKey, FStringView CacheContext)
+	: EditorData(DerivedData::Private::MakeEditorDerivedData(CacheKey, CacheContext))
+	, ChunkId(DerivedData::IoStore::GEditorDerivedDataIoStore->AddData(EditorData.Get()))
+{
+}
+
 FDerivedData::~FDerivedData()
 {
 	if (EditorData)
diff --git a/Engine/Source/Runtime/CoreUObject/Private/Serialization/EditorDerivedData.cpp b/Engine/Source/Runtime/CoreUObject/Private/Serialization/EditorDerivedData.cpp
--- a/Engine/Source/Runtime/CoreUObject/Private/Serialization/EditorDerivedData.cpp
+++ b/Engine/Source/Runtime/CoreUObject/Private/Serialization/EditorDerivedData.cpp
@@ -7,12 +7,18 @@
 #include "Compression/CompressedBuffer.h"
 #include "DerivedDataBuildDefinition.h"
 #include "DerivedDataBuildSession.h"
+#include "DerivedDataCache.h"
 #include "DerivedDataCacheInterface.h"
+#include "DerivedDataCacheKey.h"
+#include "DerivedDataRequestOwner.h"
+#include "DerivedDataSharedString.h"
+#include "DerivedDataValue.h"
 #include "DerivedDataValueId.h"
+#include "IO/IoDispatcher.h"
 #include "Memory/CompositeBuffer.h"
 #include "Memory/SharedBuffer.h"
 #include "Misc/CoreMisc.h"
-#include "Serialization/BulkDataRegistry.h"
+#include "UObject/LinkerSave.h"
 
 namespace UE::DerivedData::Private
 {
@@ -34,6 +40,7 @@ public:
 
 	void Read(IoStore::FDerivedDataIoRequest Request) const final;
 	bool TryGetSize(uint64& OutSize) const final;
+	FIoChunkId Save(FLinkerSave& Linker) const final;
 
 private:
 	FCompositeBuffer Data;
@@ -61,6 +68,11 @@ bool FEditorDerivedDataBuffer::TryGetSize(uint64& OutSize) const
 	return true;
 }
 
+FIoChunkId FEditorDerivedDataBuffer::Save(FLinkerSave& Linker) const
+{
+	return Linker.AddDerivedData(FValue::Compress(Data).GetData());
+}
+
 ///////////////////////////////////////////////////////////////////////////////////////////////////////////////////////
 
 class FEditorDerivedDataCompressedBuffer final : public FEditorDerivedData
@@ -75,6 +87,7 @@ public:
 
 	void Read(IoStore::FDerivedDataIoRequest Request) const final;
 	bool TryGetSize(uint64& OutSize) const final;
+	FIoChunkId Save(FLinkerSave& Linker) const final;
 
 private:
 	FCompressedBuffer Data;
@@ -109,14 +122,20 @@ bool FEditorDerivedDataCompressedBuffer::TryGetSize(uint64& OutSize) const
 	return true;
 }
 
+FIoChunkId FEditorDerivedDataCompressedBuffer::Save(FLinkerSave& Linker) const
+{
+	return Linker.AddDerivedData(Data);
+}
+
 ///////////////////////////////////////////////////////////////////////////////////////////////////////////////////////
 
 class FEditorDerivedDataCache final : public FEditorDerivedData
 {
 public:
-	FEditorDerivedDataCache(FStringView InCacheKey, FStringView InCacheContext)
-		: CacheKey(InCacheKey)
-		, CacheContext(InCacheContext)
+	FEditorDerivedDataCache(const FSharedString& InName, const FCacheKey& InKey, const FValueId& InValueId)
+		: Name(InName)
+		, Key(InKey)
+		, ValueId(InValueId)
 	{
 	}
 
@@ -124,22 +143,28 @@ public:
 
 	void Read(IoStore::FDerivedDataIoRequest Request) const final;
 	bool TryGetSize(uint64& OutSize) const final;
+	FIoChunkId Save(FLinkerSave& Linker) const final;
 
 private:
-	FString CacheKey;
-	FString CacheContext;
+	FSharedString Name;
+	FCacheKey Key;
+	FValueId ValueId;
 };
 
 void FEditorDerivedDataCache::Read(IoStore::FDerivedDataIoRequest Request) const
 {
-	TArray<uint8> Data;
-	if (GetDerivedDataCacheRef().GetSynchronous(*CacheKey, Data, CacheContext))
+	FSharedBuffer Data;
+
+	FRequestOwner Owner(EPriority::Blocking);
+	GetCache().GetChunks({{Name, Key, ValueId, Request.GetOffset(), Request.GetSize()}}, Owner, [&Data](FCacheGetChunkResponse&& Response)
 	{
-		const FMemoryView DataView = MakeMemoryView(Data);
-		const uint64 DataSize = DataView.GetSize();
-		const uint64 RequestOffset = Request.GetOffset();
-		const uint64 RequestSize = FMath::Min(Request.GetSize(), RequestOffset <= DataSize ? DataSize - RequestOffset : 0);
-		Request.CreateBuffer(RequestSize).CopyFrom(DataView.Mid(RequestOffset, RequestSize));
+		Data = MoveTemp(Response.RawData);
+	});
+	Owner.Wait();
+
+	if (Data)
+	{
+		Request.CreateBuffer(Data.GetSize()).CopyFrom(Data);
 		Request.SetComplete();
 	}
 	else
@@ -150,16 +175,25 @@ void FEditorDerivedDataCache::Read(IoStore::FDerivedDataIoRequest Request) const
 
 bool FEditorDerivedDataCache::TryGetSize(uint64& OutSize) const
 {
-	TArray<uint8> Data;
-	if (GetDerivedDataCacheRef().GetSynchronous(*CacheKey, Data, CacheContext))
+	bool bOk = false;
+	FRequestOwner Owner(EPriority::Blocking);
+	FCacheGetChunkRequest Request{Name, Key, ValueId};
+	Request.Policy |= ECachePolicy::SkipData;
+	GetCache().GetChunks({Request}, Owner, [&bOk, &OutSize](FCacheGetChunkResponse&& Response)
 	{
-		OutSize = MakeMemoryView(Data).GetSize();
-		return true;
-	}
-	else
-	{
-		return false;
-	}
+		if (Response.Status == EStatus::Ok)
+		{
+			bOk = true;
+			OutSize = Response.RawSize;
+		}
+	});
+	Owner.Wait();
+	return bOk;
+}
+
+FIoChunkId FEditorDerivedDataCache::Save(FLinkerSave& Linker) const
+{
+	return Linker.AddDerivedData(Key, ValueId);
 }
 
 ///////////////////////////////////////////////////////////////////////////////////////////////////////////////////////
@@ -178,6 +212,7 @@ public:
 
 	void Read(IoStore::FDerivedDataIoRequest Request) const final;
 	bool TryGetSize(uint64& OutSize) const final;
+	FIoChunkId Save(FLinkerSave& Linker) const final;
 
 private:
 	FBuildDefinition BuildDefinition;
@@ -194,6 +229,74 @@ bool FEditorDerivedDataBuild::TryGetSize(uint64& OutSize) const
 	return false;
 }
 
+FIoChunkId FEditorDerivedDataBuild::Save(FLinkerSave& Linker) const
+{
+	unimplemented();
+	return FIoChunkId::InvalidChunkId;
+}
+
+///////////////////////////////////////////////////////////////////////////////////////////////////////////////////////
+
+class FEditorDerivedDataLegacyCache final : public FEditorDerivedData
+{
+public:
+	FEditorDerivedDataLegacyCache(FStringView InCacheKey, FStringView InCacheContext)
+		: CacheKey(InCacheKey)
+		, CacheContext(InCacheContext)
+	{
+	}
+
+	TUniquePtr<FEditorDerivedData> Clone() const final { return MakeUnique<FEditorDerivedDataLegacyCache>(*this); }
+
+	void Read(IoStore::FDerivedDataIoRequest Request) const final;
+	bool TryGetSize(uint64& OutSize) const final;
+	FIoChunkId Save(FLinkerSave& Linker) const final;
+
+private:
+	FString CacheKey;
+	FString CacheContext;
+};
+
+void FEditorDerivedDataLegacyCache::Read(IoStore::FDerivedDataIoRequest Request) const
+{
+	TArray64<uint8> Data;
+	if (GetDerivedDataCacheRef().GetSynchronous(*CacheKey, Data, CacheContext))
+	{
+		const FMemoryView DataView = MakeMemoryView(Data);
+		const uint64 DataSize = DataView.GetSize();
+		const uint64 RequestOffset = Request.GetOffset();
+		const uint64 RequestSize = FMath::Min(Request.GetSize(), RequestOffset <= DataSize ? DataSize - RequestOffset : 0);
+		Request.CreateBuffer(RequestSize).CopyFrom(DataView.Mid(RequestOffset, RequestSize));
+		Request.SetComplete();
+	}
+	else
+	{
+		Request.SetFailed();
+	}
+}
+
+bool FEditorDerivedDataLegacyCache::TryGetSize(uint64& OutSize) const
+{
+	TArray<uint8> Data;
+	if (GetDerivedDataCacheRef().GetSynchronous(*CacheKey, Data, CacheContext))
+	{
+		OutSize = MakeMemoryView(Data).GetSize();
+		return true;
+	}
+	else
+	{
+		return false;
+	}
+}
+
+FIoChunkId FEditorDerivedDataLegacyCache::Save(FLinkerSave& Linker) const
+{
+	TArray64<uint8> Data;
+	const bool bOk = GetDerivedDataCacheRef().GetSynchronous(*CacheKey, Data, CacheContext);
+	checkf(bOk, TEXT("Failed to fetch %s from the cache for '%s'"), *CacheKey, *CacheContext);
+	return Linker.AddDerivedData(FValue::Compress(MakeSharedBufferFromArray(MoveTemp(Data))).GetData());
+}
+
 ///////////////////////////////////////////////////////////////////////////////////////////////////////////////////////
 
 TUniquePtr<FEditorDerivedData> MakeEditorDerivedData(const FSharedBuffer& Data)
@@ -211,18 +314,21 @@ TUniquePtr<FEditorDerivedData> MakeEditorDerivedData(const FCompressedBuffer& Da
 	return MakeUnique<FEditorDerivedDataCompressedBuffer>(Data);
 }
 
-TUniquePtr<FEditorDerivedData> MakeEditorDerivedData(FStringView CacheKey, FStringView CacheContext)
+TUniquePtr<FEditorDerivedData> MakeEditorDerivedData(const FSharedString& Name, const FCacheKey& Key, const FValueId& ValueId)
 {
-	return MakeUnique<FEditorDerivedDataCache>(CacheKey, CacheContext);
+	return MakeUnique<FEditorDerivedDataCache>(Name, Key, ValueId);
 }
 
-TUniquePtr<FEditorDerivedData> MakeEditorDerivedData(
-	const FBuildDefinition& BuildDefinition,
-	const FValueId& ValueId)
+TUniquePtr<FEditorDerivedData> MakeEditorDerivedData(const FBuildDefinition& BuildDefinition, const FValueId& ValueId)
 {
 	return MakeUnique<FEditorDerivedDataBuild>(BuildDefinition, ValueId);
 }
 
+TUniquePtr<FEditorDerivedData> MakeEditorDerivedData(FStringView CacheKey, FStringView CacheContext)
+{
+	return MakeUnique<FEditorDerivedDataLegacyCache>(CacheKey, CacheContext);
+}
+
 } // UE::DerivedData::Private
 
 #endif // WITH_EDITORONLY_DATA
diff --git a/Engine/Source/Runtime/CoreUObject/Private/Serialization/EditorDerivedData.h b/Engine/Source/Runtime/CoreUObject/Private/Serialization/EditorDerivedData.h
--- a/Engine/Source/Runtime/CoreUObject/Private/Serialization/EditorDerivedData.h
+++ b/Engine/Source/Runtime/CoreUObject/Private/Serialization/EditorDerivedData.h
@@ -7,14 +7,17 @@
 #if WITH_EDITORONLY_DATA
 
 #include "Containers/StringFwd.h"
+#include "DerivedDataSharedStringFwd.h"
 #include "Memory/MemoryFwd.h"
 #include "Templates/UniquePtr.h"
 
 class FCompressedBuffer;
+class FIoChunkId;
 class FIoRequestImpl;
 
 namespace UE::DerivedData { class FBuildDefinition; }
 namespace UE::DerivedData { class FRequestOwner; }
+namespace UE::DerivedData { struct FCacheKey; }
 namespace UE::DerivedData { struct FValueId; }
 namespace UE::DerivedData::IoStore { class FDerivedDataIoRequestQueue; }
 
@@ -55,17 +58,17 @@ public:
 	virtual TUniquePtr<FEditorDerivedData> Clone() const = 0;
 	virtual void Read(IoStore::FDerivedDataIoRequest Request) const = 0;
 	virtual bool TryGetSize(uint64& OutSize) const = 0;
+	virtual FIoChunkId Save(FLinkerSave& Linker) const = 0;
 };
 
 TUniquePtr<FEditorDerivedData> MakeEditorDerivedData(const FSharedBuffer& Data);
 TUniquePtr<FEditorDerivedData> MakeEditorDerivedData(const FCompositeBuffer& Data);
 TUniquePtr<FEditorDerivedData> MakeEditorDerivedData(const FCompressedBuffer& Data);
 
-TUniquePtr<FEditorDerivedData> MakeEditorDerivedData(FStringView CacheKey, FStringView CacheContext);
+TUniquePtr<FEditorDerivedData> MakeEditorDerivedData(const FSharedString& Name, const FCacheKey& Key, const FValueId& ValueId);
+TUniquePtr<FEditorDerivedData> MakeEditorDerivedData(const FBuildDefinition& BuildDefinition, const FValueId& ValueId);
 
-TUniquePtr<FEditorDerivedData> MakeEditorDerivedData(
-	const DerivedData::FBuildDefinition& BuildDefinition,
-	const DerivedData::FValueId& ValueId);
+TUniquePtr<FEditorDerivedData> MakeEditorDerivedData(FStringView CacheKey, FStringView CacheContext);
 
 } // UE::DerivedData::Private
 
diff --git a/Engine/Source/Runtime/CoreUObject/Private/UObject/LinkerSave.cpp b/Engine/Source/Runtime/CoreUObject/Private/UObject/LinkerSave.cpp
--- a/Engine/Source/Runtime/CoreUObject/Private/UObject/LinkerSave.cpp
+++ b/Engine/Source/Runtime/CoreUObject/Private/UObject/LinkerSave.cpp
@@ -15,6 +15,11 @@
 #include "UObject/UnrealType.h"
 #include "HAL/PlatformStackWalk.h"
 
+#if WITH_EDITORONLY_DATA
+#include "DerivedDataValueId.h"
+#include "IO/IoDispatcher.h"
+#endif
+
 /*----------------------------------------------------------------------------
 	FLinkerSave.
 ----------------------------------------------------------------------------*/
@@ -459,3 +464,35 @@ void FLinkerSave::SetUseUnversionedPropertySerialization(bool bInUseUnversioned)
 		}
 	}
 }
+
+#if WITH_EDITORONLY_DATA
+
+static FIoChunkId CreateDerivedDataChunkId(const FPackageId PackageId, const int32 ChunkIndex)
+{
+	checkf(ChunkIndex >= 0 && ChunkIndex < (1 << 24), TEXT("ChunkIndex %d is out of range."), ChunkIndex);
+
+	// PackageId                 ChunkIndex Type
+	// [00 01 02 03 04 05 06 07] [08 09 10] [11]
+	uint8 Data[12]{};
+	*reinterpret_cast<uint8*>(&Data[11]) = static_cast<uint8>(EIoChunkType::DerivedData);
+	*reinterpret_cast<uint32*>(&Data[7]) = NETWORK_ORDER32(ChunkIndex);
+	*reinterpret_cast<uint64*>(&Data[0]) = PackageId.Value();
+
+	FIoChunkId ChunkId;
+	ChunkId.Set(Data, 12);
+	return ChunkId;
+}
+
+FIoChunkId FLinkerSave::AddDerivedData(const FCompressedBuffer& Data)
+{
+	UE_LOG(LogLinker, Warning, TEXT("Data will not be able to load because derived data is not saved yet."));
+	return CreateDerivedDataChunkId(LinkerRoot->GetPackageId(), ++LastDerivedDataIndex);
+}
+
+FIoChunkId FLinkerSave::AddDerivedData(const UE::DerivedData::FCacheKey& Key, const UE::DerivedData::FValueId& ValueId)
+{
+	UE_LOG(LogLinker, Warning, TEXT("Data will not be able to load because derived data is not saved yet."));
+	return CreateDerivedDataChunkId(LinkerRoot->GetPackageId(), ++LastDerivedDataIndex);
+}
+
+#endif // WITH_EDITORONLY_DATA
diff --git a/Engine/Source/Runtime/CoreUObject/Public/Serialization/DerivedData.h b/Engine/Source/Runtime/CoreUObject/Public/Serialization/DerivedData.h
--- a/Engine/Source/Runtime/CoreUObject/Public/Serialization/DerivedData.h
+++ b/Engine/Source/Runtime/CoreUObject/Public/Serialization/DerivedData.h
@@ -17,7 +17,10 @@ class UObject;
 template <typename FuncType> class TFunctionRef;
 
 namespace UE::DerivedData { class FBuildDefinition; }
+namespace UE::DerivedData { struct FCacheKey; }
 namespace UE::DerivedData { struct FValueId; }
+namespace UE::DerivedData { template <typename CharType> class TSharedString; }
+namespace UE::DerivedData { using FSharedString = TSharedString<TCHAR>; }
 namespace UE::DerivedData::Private { class FEditorDerivedData; }
 
 namespace UE
@@ -59,14 +62,22 @@ public:
 	UE_API FDerivedData& operator=(FDerivedData&& Other);
 	UE_API FDerivedData& operator=(const FDerivedData& Other);
 
+	/** Reference a value that is stored in a buffer. */
 	UE_API explicit FDerivedData(const FSharedBuffer& Data);
 	UE_API explicit FDerivedData(const FCompositeBuffer& Data);
 	UE_API explicit FDerivedData(const FCompressedBuffer& Data);
 
-	UE_API FDerivedData(FStringView CacheKey, FStringView CacheContext);
-
+	/** Reference a value that was saved using ICache::PutValue. */
+	UE_API FDerivedData(const DerivedData::FSharedString& Name, const DerivedData::FCacheKey& Key);
+	/** Reference a value in a record that was saved using ICache::Put. */
+	UE_API FDerivedData(const DerivedData::FSharedString& Name, const DerivedData::FCacheKey& Key, const DerivedData::FValueId& ValueId);
+	/** Reference a value in the output from the build function for the build definition. */
 	UE_API FDerivedData(const DerivedData::FBuildDefinition& BuildDefinition, const DerivedData::FValueId& ValueId);
 
+	/** Reference a value that was saved using FDerivedDataCacheInterface::Put. */
+	UE_API FDerivedData(FStringView CacheKey, FStringView CacheContext);
+
+	/** Overwrite the existing flags. */
 	UE_API void SetFlags(EDerivedDataFlags Flags);
 
 private:
diff --git a/Engine/Source/Runtime/CoreUObject/Public/UObject/LinkerSave.h b/Engine/Source/Runtime/CoreUObject/Public/UObject/LinkerSave.h
--- a/Engine/Source/Runtime/CoreUObject/Public/UObject/LinkerSave.h
+++ b/Engine/Source/Runtime/CoreUObject/Public/UObject/LinkerSave.h
@@ -12,9 +12,13 @@
 #include "UObject/PackageTrailer.h"
 #include "UObject/UObjectThreadContext.h"
 
+class FIoChunkId;
 class FObjectPostSaveContext;
 struct FUntypedBulkData;
 
+namespace UE::DerivedData { struct FCacheKey; }
+namespace UE::DerivedData { struct FValueId; }
+
 /*----------------------------------------------------------------------------
 	FLinkerSave.
 ----------------------------------------------------------------------------*/
@@ -207,6 +211,33 @@ public:
 	{
 		return LogOutput;
 	}
+
+#if WITH_EDITORONLY_DATA
+
+	/**
+	 * Adds the derived data to the package.
+	 *
+	 * This is only supported when saving a cooked package.
+	 *
+	 * @return An ID that can be used to load the derived data from the cooked package.
+	 */
+	FIoChunkId AddDerivedData(const FCompressedBuffer& Data);
+
+	/**
+	 * Adds the derived data value referenced by the cache key and ID.
+	 *
+	 * This is only supported when saving a cooked package.
+	 * The key and a null ID must reference a value that was saved using ICache::PutValue.
+	 * The key and a non-null ID must reference a value in a record that was saved using ICache::Put.
+	 *
+	 * @param Key       Reference to a record or value in the cache.
+	 * @param ValueId   Reference to a value within a record, otherwise FValueId::Null.
+	 * @return An ID that can be used to load the derived data from the cooked package.
+	 */
+	FIoChunkId AddDerivedData(const UE::DerivedData::FCacheKey& Key, const UE::DerivedData::FValueId& ValueId);
+
+#endif // WITH_EDITORONLY_DATA
+
 protected:
 	/** Set the filename being saved to */
 	void SetFilename(FStringView InFilename);
@@ -214,4 +245,9 @@ protected:
 private:
 	/** Optional log output to bubble errors back up. */
 	FOutputDevice* LogOutput = nullptr;
+
+#if WITH_EDITORONLY_DATA
+	/** The index of the last derived data chunk added to the package. */
+	int32 LastDerivedDataIndex = -1;
+#endif
 };
