UnversionedPropertySerialization: Declare an FSerializedPropertyScope when save-serializing, to match the VersionedProperty case. This scope is needed when loading and saving in the editor so that FSoftObjectPath::PostLoad or FPackageHarvester::operator<<(FSoftObjectPath) can detect whether an untracked Property is on the serialize stack and mark the property as editor-only if so.
#rnx

diff --git a/Engine/Source/Runtime/CoreUObject/Private/Serialization/UnversionedPropertySerialization.cpp b/Engine/Source/Runtime/CoreUObject/Private/Serialization/UnversionedPropertySerialization.cpp
--- a/Engine/Source/Runtime/CoreUObject/Private/Serialization/UnversionedPropertySerialization.cpp
+++ b/Engine/Source/Runtime/CoreUObject/Private/Serialization/UnversionedPropertySerialization.cpp
@@ -915,6 +915,7 @@ void SerializeUnversionedProperties(const UStruct* Struct, FStructuredArchive::F
 			{
 				if (It.IsNonZero())
 				{
+					FSerializedPropertyScope SerializedProperty(UnderlyingArchive, It.GetSerializer().GetProperty());
 					It.GetSerializer().Serialize(ValueStream.EnterElement(), Data, Defaults);
 				}
 			}
