FAssetData: Callsites need to handle AssetData.GetClass() returning null. Added IsInstanceOf(UClass*) for the common operation GetClass() && GetClass->IsChildOf(BaseClass).
#rnx

diff --git a/Engine/Source/Editor/UnrealEd/Private/Commandlets/AssetRegistryGenerator.cpp b/Engine/Source/Editor/UnrealEd/Private/Commandlets/AssetRegistryGenerator.cpp
--- a/Engine/Source/Editor/UnrealEd/Private/Commandlets/AssetRegistryGenerator.cpp
+++ b/Engine/Source/Editor/UnrealEd/Private/Commandlets/AssetRegistryGenerator.cpp
@@ -1282,7 +1282,8 @@ void FAssetRegistryGenerator::FinalizeChunkIDs(const TSet<FName>& InCookedPackag
 		const_cast<FAssetData&>(AssetData).ChunkIDs.Empty();
 
 		// Update whether the owner package contains a map
-		if (AssetData.GetClass()->IsChildOf(UWorld::StaticClass()) || AssetData.GetClass()->IsChildOf(ULevel::StaticClass()))
+		UClass* AssetClass = AssetData.GetClass();
+		if (AssetClass && (AssetClass->IsChildOf(UWorld::StaticClass()) || AssetClass->IsChildOf(ULevel::StaticClass())))
 		{
 			PackagesContainingMaps.Add(AssetData.PackageName);
 		}
diff --git a/Engine/Source/Editor/UnrealEd/Private/CookOnTheFlyServer.cpp b/Engine/Source/Editor/UnrealEd/Private/CookOnTheFlyServer.cpp
--- a/Engine/Source/Editor/UnrealEd/Private/CookOnTheFlyServer.cpp
+++ b/Engine/Source/Editor/UnrealEd/Private/CookOnTheFlyServer.cpp
@@ -893,7 +893,8 @@ bool UCookOnTheFlyServer::ContainsMap(const FName& PackageName) const
 
 	for (const FAssetData& Asset : Assets)
 	{
-		if (Asset.GetClass()->IsChildOf(UWorld::StaticClass()) || Asset.GetClass()->IsChildOf(ULevel::StaticClass()))
+		UClass* AssetClass = Asset.GetClass();
+		if (AssetClass && (AssetClass->IsChildOf(UWorld::StaticClass()) || AssetClass->IsChildOf(ULevel::StaticClass())))
 		{
 			return true;
 		}
diff --git a/Engine/Source/Runtime/CoreUObject/Public/AssetRegistry/AssetData.h b/Engine/Source/Runtime/CoreUObject/Public/AssetRegistry/AssetData.h
--- a/Engine/Source/Runtime/CoreUObject/Public/AssetRegistry/AssetData.h
+++ b/Engine/Source/Runtime/CoreUObject/Public/AssetRegistry/AssetData.h
@@ -294,6 +294,13 @@ public:
 		}
 		return FoundClass;
 	}
+	
+	/** Returns whether the Asset's class is equal to or a child class of the given class. Returns false if the Asset's class can not be loaded. */
+	bool IsInstanceOf(const UClass* BaseClass) const
+	{
+		UClass* ClassPointer = GetClass();
+		return ClassPointer && ClassPointer->IsChildOf(BaseClass);
+	}
 
 	/** Convert to a SoftObjectPath for loading */
 	FSoftObjectPath ToSoftObjectPath() const
@@ -453,8 +460,7 @@ public:
 		for(int32 AssetIdx=0; AssetIdx<Assets.Num(); AssetIdx++)
 		{
 			const FAssetData& Data = Assets[AssetIdx];
-			UClass* AssetClass = Data.GetClass();
-			if( AssetClass != NULL && AssetClass->IsChildOf(DesiredClass) )
+			if (Data.IsInstanceOf(DesiredClass))
 			{
 				return Data;
 			}
