UnversionedPropertySerialization: declare an FSerializedPropertyScope when load-serializing in the editor, so that EditorDomain loads will correctly notice SoftObjectPath properties that are marked untracked.
#rnx

diff --git a/Engine/Source/Runtime/CoreUObject/Private/Serialization/UnversionedPropertySerialization.cpp b/Engine/Source/Runtime/CoreUObject/Private/Serialization/UnversionedPropertySerialization.cpp
--- a/Engine/Source/Runtime/CoreUObject/Private/Serialization/UnversionedPropertySerialization.cpp
+++ b/Engine/Source/Runtime/CoreUObject/Private/Serialization/UnversionedPropertySerialization.cpp
@@ -861,6 +861,9 @@ void SerializeUnversionedProperties(const UStruct* Struct, FStructuredArchive::F
 				{
 					if (It.IsNonZero())
 					{
+#if WITH_EDITOR // Skip this scope to save time in the runtime; it is only needed for reference collection in editor
+						FSerializedPropertyScope SerializedProperty(UnderlyingArchive, It.GetSerializer().GetProperty());
+#endif
 						It.GetSerializer().Serialize(ValueStream.EnterElement(), Data, Defaults);
 					}
 					else
