#UE AttemptToFindUninitializedScriptStructMembers optimization. Only call constructors in FScriptStructTestWrapper if the data will actually be used.

diff --git a/Engine/Source/Runtime/CoreUObject/Private/UObject/Class.cpp b/Engine/Source/Runtime/CoreUObject/Private/UObject/Class.cpp
--- a/Engine/Source/Runtime/CoreUObject/Private/UObject/Class.cpp
+++ b/Engine/Source/Runtime/CoreUObject/Private/UObject/Class.cpp
@@ -3283,9 +3283,41 @@ struct FScriptStructTestWrapper
 {
 public:
 
-	FScriptStructTestWrapper(UScriptStruct* InStruct, uint8 InitValue = 0xFD, EScriptStructTestCtorSyntax ConstrutorSyntax = EScriptStructTestCtorSyntax::NoInit)
+	FScriptStructTestWrapper(UScriptStruct* InStruct, uint8 InInitValue = 0xFD, EScriptStructTestCtorSyntax InConstrutorSyntax = EScriptStructTestCtorSyntax::NoInit)
 		: ScriptStruct(InStruct)
 		, TempBuffer(nullptr)
+		, InitValue(InInitValue)
+		, ConstrutorSyntax(InConstrutorSyntax)
+		, bTempBufferAttemptedCreate(false)
+	{
+	}
+
+	~FScriptStructTestWrapper()
+	{
+		if (TempBuffer != nullptr)
+		{
+			// Destroy it
+			ScriptStruct->DestroyStruct(TempBuffer);
+			FMemory::Free(TempBuffer);
+		}
+	}
+
+	static bool CanRunTests(UScriptStruct* Struct)
+	{
+		return (Struct != nullptr) && Struct->IsNative() && (!Struct->GetCppStructOps() || !Struct->GetCppStructOps()->HasZeroConstructor());
+	}
+
+	uint8* GetData()
+	{
+		if (!bTempBufferAttemptedCreate)
+		{
+			AttemptCreateTempBuffer();
+		}
+		return TempBuffer;
+	}
+private:
+
+	void AttemptCreateTempBuffer()
 	{
 		if (ScriptStruct->IsNative())
 		{
@@ -3306,7 +3338,7 @@ public:
 					if (!ensure(RequiredAllocSize >= CppAllocSize))
 					{
 						UE_LOG(LogClass, Warning, TEXT("Struct %s%s has Cpp alloc size = %d > ScriptStruct->GetStructureSize() = %d, this could result in allocations that are too small to fit the structure."),
-							InStruct->GetPrefixCPP(), *InStruct->GetName(), CppAllocSize, RequiredAllocSize);
+							ScriptStruct->GetPrefixCPP(), *ScriptStruct->GetName(), CppAllocSize, RequiredAllocSize);
 						// We'd probably crash below, so use a larger allocated size.
 						RequiredAllocSize = StructOps->GetSize();
 					}
@@ -3353,27 +3385,15 @@ public:
 				}
 			}
 		}
-	}
 
-	~FScriptStructTestWrapper()
-	{
-		if (TempBuffer != nullptr)
-		{
-			// Destroy it
-			ScriptStruct->DestroyStruct(TempBuffer);
-			FMemory::Free(TempBuffer);
-		}
+		bTempBufferAttemptedCreate = true;
 	}
 
-	static bool CanRunTests(UScriptStruct* Struct)
-	{
-		return (Struct != nullptr) && Struct->IsNative() && (!Struct->GetCppStructOps() || !Struct->GetCppStructOps()->HasZeroConstructor());
-	}
-
-	uint8* GetData() { return TempBuffer; }
-private:
 	UScriptStruct* ScriptStruct;
 	uint8* TempBuffer;
+	uint8 InitValue;
+	EScriptStructTestCtorSyntax ConstrutorSyntax;
+	bool bTempBufferAttemptedCreate;
 };
 
 static void FindUninitializedScriptStructMembers(UScriptStruct* ScriptStruct, EScriptStructTestCtorSyntax ConstructorSyntax, TSet<const FProperty*>& OutUninitializedProperties)
