Fix blueprint code for replacing placeholder objects so that it is aware of FObjectPtr properties and handles them appropriately when they are unresolved (lazy loaded).

diff --git a/Engine/Source/Runtime/CoreUObject/Private/Blueprint/BlueprintSupport.cpp b/Engine/Source/Runtime/CoreUObject/Private/Blueprint/BlueprintSupport.cpp
--- a/Engine/Source/Runtime/CoreUObject/Private/Blueprint/BlueprintSupport.cpp
+++ b/Engine/Source/Runtime/CoreUObject/Private/Blueprint/BlueprintSupport.cpp
@@ -218,6 +218,14 @@ void FBlueprintSupport::RepairDeferredDependenciesInObject(UObject* Object)
 	{
 		const FObjectProperty* Property = It.Key();
 		void* PropertyValue = (void*)It.Value();
+		if (Property->IsA<FObjectPtrProperty>())
+		{
+			FObjectPtr* PropertyValueAsObjectPtr = ((FObjectPtr*)PropertyValue);
+			if (!PropertyValueAsObjectPtr->IsResolved())
+			{
+				continue;
+			}
+		}
 		UObject* PropertyValueAsObj = *((UObject**)PropertyValue);
 
 		FLinkerPlaceholderBase* Placeholder = nullptr;
