Add a console command that can be used to resolve Soft Object Paths in Editor, and change the warning for broken Soft Object Paths to include the Referencing Package Name.

diff --git a/Engine/Source/Runtime/CoreUObject/Private/Misc/RedirectCollector.cpp b/Engine/Source/Runtime/CoreUObject/Private/Misc/RedirectCollector.cpp
--- a/Engine/Source/Runtime/CoreUObject/Private/Misc/RedirectCollector.cpp
+++ b/Engine/Source/Runtime/CoreUObject/Private/Misc/RedirectCollector.cpp
@@ -13,11 +13,21 @@
 #include "UObject/LinkerLoad.h"
 #include "UObject/UObjectThreadContext.h"
 #include "UObject/SoftObjectPath.h"
+#include "HAL/IConsoleManager.h"
 
 #if WITH_EDITOR
 
 DEFINE_LOG_CATEGORY_STATIC(LogRedirectors, Log, All);
 
+FAutoConsoleCommand CVarResolveAllSoftObjects(
+	TEXT("RedirectCollector.ResolveAllSoftObjectPaths"),
+	TEXT("Attempts to load / resolve all currently referenced Soft Object Paths"),
+	FConsoleCommandDelegate::CreateLambda([]
+	{
+		GRedirectCollector.ResolveAllSoftObjectPaths();
+	})
+);
+
 void FRedirectCollector::OnSoftObjectPathLoaded(const FSoftObjectPath& InPath, FArchive* InArchive)
 {
 	if (InPath.IsNull() || !GIsEditor)
@@ -80,7 +90,7 @@ void FRedirectCollector::CollectSavedSoftPackageReferences(FName ReferencingPack
 
 void FRedirectCollector::ResolveAllSoftObjectPaths(FName FilterPackage)
 {	
-	auto LoadSoftObjectPathLambda = [this](const FSoftObjectPathProperty& SoftObjectPathProperty)
+	auto LoadSoftObjectPathLambda = [this](const FSoftObjectPathProperty& SoftObjectPathProperty, const FName ReferencerPackageName)
 	{
 		const FName& ToLoadFName = SoftObjectPathProperty.GetAssetPathName();
 		const FString ToLoad = ToLoadFName.ToString();
@@ -113,7 +123,7 @@ void FRedirectCollector::ResolveAllSoftObjectPaths(FName FilterPackage)
 			else
 			{
 				const FString Referencer = SoftObjectPathProperty.GetPropertyName().ToString().Len() ? SoftObjectPathProperty.GetPropertyName().ToString() : TEXT("Unknown");
-				UE_LOG(LogRedirectors, Warning, TEXT("Soft Object Path '%s' was not found when resolving paths! (Referencer '%s')"), *ToLoad, *Referencer);
+				UE_LOG(LogRedirectors, Warning, TEXT("Soft Object Path '%s' was not found when resolving paths! (Referencer '%s:%s')"), *ToLoad, *ReferencerPackageName.ToString(), *Referencer);
 			}
 		}
 	};
@@ -145,7 +155,7 @@ void FRedirectCollector::ResolveAllSoftObjectPaths(FName FilterPackage)
 			// This will call LoadObject which may trigger OnSoftObjectPathLoaded and add new soft object paths to the SoftObjectPathMap
 			for (const FSoftObjectPathProperty& SoftObjecPathProperty : SoftObjectPathProperties)
 			{
-				LoadSoftObjectPathLambda(SoftObjecPathProperty);
+				LoadSoftObjectPathLambda(SoftObjecPathProperty, CurrentPackageName);
 			}
 		}
 	}
