Tests converted from 18848115 made to work multi-platform.

diff --git a/Engine/Source/Runtime/CoreUObject/Tests/ObjectHandleTest.cpp b/Engine/Source/Runtime/CoreUObject/Tests/ObjectHandleTest.cpp
new file mode 100644
--- /dev/null
+++ b/Engine/Source/Runtime/CoreUObject/Tests/ObjectHandleTest.cpp
@@ -0,0 +1,215 @@
+// Copyright Epic Games, Inc. All Rights Reserved.
+
+#include "UObject/ObjectHandle.h"
+
+#include "HAL/PlatformProperties.h"
+#include "ObjectRefTrackingTestBase.h"
+#include "IO/IoDispatcher.h"
+
+static_assert(sizeof(FObjectHandle) == sizeof(void*), "FObjectHandle type must always compile to something equivalent to a pointer size.");
+
+#if WITH_DEV_AUTOMATION_TESTS
+
+class FObjectHandleTestBase : public FObjectRefTrackingTestBase
+{
+public:
+	FObjectHandleTestBase(const FString& InName, const bool bInComplexTask)
+	: FObjectRefTrackingTestBase(InName, bInComplexTask)
+	{
+	}
+
+protected:
+	
+	UObject* ResolveHandle(FObjectHandle& TargetHandle)
+	{
+	#if UE_WITH_OBJECT_HANDLE_LATE_RESOLVE
+		// Late resolved handles cannot be null or resolved at this point
+		if (!TestFalse(TEXT("Handle to target is null"), IsObjectHandleNull(TargetHandle)))
+		{
+			return nullptr;
+		}
+		if (!TestFalse(TEXT("Handle to target is resolved"), IsObjectHandleResolved(TargetHandle)))
+		{
+			return nullptr;
+		}
+	#else
+		// Immediately resolved handles may be null (if the target is invalid) and must be resolved at this point
+		if (!TestTrue(TEXT("Handle to target is not resolved"), IsObjectHandleResolved(TargetHandle)))
+		{
+			return nullptr;
+		}
+	#endif
+
+		return ResolveObjectHandle(TargetHandle);
+	}
+
+	UObject* ConstructAndResolveHandle(const ANSICHAR* PackageName, const ANSICHAR* ObjectName, const ANSICHAR* ClassPackageName = nullptr, const ANSICHAR* ClassName = nullptr)
+	{
+		FObjectRef TargetRef{FName(PackageName), FName(ClassPackageName), FName(ClassName), FObjectPathId(ObjectName)};
+		if (!TestFalse(TEXT("Reference to target is null"), IsObjectRefNull(TargetRef)))
+		{
+			return nullptr;
+		}
+
+		FObjectHandle TargetHandle = MakeObjectHandle(TargetRef);
+		return ResolveHandle(TargetHandle);
+	}
+
+	UObject* ConstructAndResolveHandle(const FPackedObjectRef& PackedTargetRef)
+	{
+		if (!TestFalse(TEXT("Reference to target is null"), IsPackedObjectRefNull(PackedTargetRef)))
+		{
+			return nullptr;
+		}
+
+		FObjectHandle TargetHandle = MakeObjectHandle(PackedTargetRef);
+		return ResolveHandle(TargetHandle);
+	}
+
+	bool TestResolvableNonNull(const ANSICHAR* PackageName, const ANSICHAR* ObjectName, const ANSICHAR* ClassPackageName = nullptr, const ANSICHAR* ClassName = nullptr, bool bExpectSubRefReads = false)
+	{
+		FSnapshotObjectRefMetrics ObjectRefMetrics(*this);
+		UObject* ResolvedObject = ConstructAndResolveHandle(PackageName, ObjectName, ClassPackageName, ClassName);
+		ObjectRefMetrics.TestNumResolves(TEXT("NumResolves should be incremented by one after a resolve attempt"), 1);
+		ObjectRefMetrics.TestNumReads(TEXT("NumReads should be incremented by one after a resolve attempt"), 1, bExpectSubRefReads /*bAllowAdditionalReads*/);
+
+		if (!ResolvedObject)
+		{
+			AddError(FString::Printf(TEXT("Expected '%s.%s' to resolve to non null."), ANSI_TO_TCHAR(PackageName), ANSI_TO_TCHAR(ObjectName)), 1);
+			return false;
+		}
+		ObjectRefMetrics.TestNumFailedResolves(TEXT("NumFailedResolves should not change after a successful resolve attempt"), 0);
+
+		return true;
+	}
+
+	bool TestResolveFailure(const ANSICHAR* PackageName, const ANSICHAR* ObjectName, const ANSICHAR* ClassPackageName = nullptr, const ANSICHAR* ClassName = nullptr)
+	{
+		FSnapshotObjectRefMetrics ObjectRefMetrics(*this);
+		UObject* ResolvedObject = ConstructAndResolveHandle(PackageName, ObjectName, ClassPackageName, ClassName);
+		ObjectRefMetrics.TestNumResolves(TEXT("NumResolves should be incremented by one after a resolve attempt"), 1);
+		ObjectRefMetrics.TestNumReads(TEXT("NumReads should be incremented by one after a resolve attempt"), 1);
+
+		if (ResolvedObject)
+		{
+			AddError(FString::Printf(TEXT("Expected '%s.%s' to resolve to null."), ANSI_TO_TCHAR(PackageName), ANSI_TO_TCHAR(ObjectName)), 1);
+			return false;
+		}
+		ObjectRefMetrics.TestNumFailedResolves(TEXT("NumFailedResolves should be incremented by one after a failed resolve attempt"), 1);
+		return true;
+	}
+
+	bool TestResolveFailure(FPackedObjectRef PackedRef)
+	{
+		FSnapshotObjectRefMetrics ObjectRefMetrics(*this);
+		UObject* ResolvedObject = ConstructAndResolveHandle(PackedRef);
+		ObjectRefMetrics.TestNumResolves(TEXT("NumResolves should be incremented by one after a resolve attempt"), 1);
+		ObjectRefMetrics.TestNumReads(TEXT("NumReads should be incremented by one after a resolve attempt"), 1);
+
+		if (ResolvedObject)
+		{
+			AddError(FString::Printf(TEXT("Expected PACKEDREF(%" UPTRINT_X_FMT ") to resolve to null."), PackedRef.EncodedRef), 1);
+			return false;
+		}
+		ObjectRefMetrics.TestNumFailedResolves(TEXT("NumFailedResolves should be incremented by one after a failed resolve attempt"), 1);
+		return true;
+	}
+};
+
+#define TEST_NAME_ROOT TEXT("System.CoreUObject.ObjectHandle")
+constexpr const uint32 ObjectHandleTestFlags = EAutomationTestFlags::ApplicationContextMask | EAutomationTestFlags::EngineFilter;
+
+IMPLEMENT_CUSTOM_SIMPLE_AUTOMATION_TEST(FObjectHandleTestNullBehavior, FObjectHandleTestBase, TEST_NAME_ROOT TEXT(".NullBehavior"), ObjectHandleTestFlags)
+bool FObjectHandleTestNullBehavior::RunTest(const FString& Parameters)
+{
+	FObjectHandle TargetHandle = MakeObjectHandle(nullptr);
+
+	TestTrue(TEXT("Handle to target is null"), IsObjectHandleNull(TargetHandle));
+	TestTrue(TEXT("Handle to target is resolved"), IsObjectHandleResolved(TargetHandle));
+
+	FSnapshotObjectRefMetrics ObjectRefMetrics(*this);
+	UObject* ResolvedObject = ResolveObjectHandle(TargetHandle);
+
+	TestEqual(TEXT("Resolved object is equal to original object"), (UObject*)nullptr, ResolvedObject);
+
+	ObjectRefMetrics.TestNumFailedResolves(TEXT("NumFailedResolves should not change after a resolve attempt on a null handle"), 0);
+	ObjectRefMetrics.TestNumResolves(TEXT("NumResolves should not change after a resolve attempt on a null handle"), 0);
+	ObjectRefMetrics.TestNumReads(TEXT("NumReads should be incremented by one after a resolve attempt on a null handle"), 1);
+
+	return true;
+}
+
+IMPLEMENT_CUSTOM_SIMPLE_AUTOMATION_TEST(FObjectHandleTestPointerBehavior, FObjectHandleTestBase, TEST_NAME_ROOT TEXT(".PointerBehavior"), ObjectHandleTestFlags)
+bool FObjectHandleTestPointerBehavior::RunTest(const FString& Parameters)
+{
+	FObjectHandle TargetHandle = MakeObjectHandle((UObject*)0x0042);
+
+	TestFalse(TEXT("Handle to target is null"), IsObjectHandleNull(TargetHandle));
+	TestTrue(TEXT("Handle to target is resolved"), IsObjectHandleResolved(TargetHandle));
+
+	FSnapshotObjectRefMetrics ObjectRefMetrics(*this);
+	UObject* ResolvedObject = ResolveObjectHandle(TargetHandle);
+
+	TestEqual(TEXT("Resolved object is equal to original object"), (UObject*)0x0042, ResolvedObject);
+
+	ObjectRefMetrics.TestNumResolves(TEXT("NumResolves should not change after a resolve attempt on a pointer handle"), 0);
+	ObjectRefMetrics.TestNumFailedResolves(TEXT("NumFailedResolves should not change after a resolve attempt on a pointer handle"), 0);
+	ObjectRefMetrics.TestNumReads(TEXT("NumReads should be incremented by one after a resolve attempt on a pointer handle"),1);
+
+	return true;
+}
+
+IMPLEMENT_CUSTOM_SIMPLE_AUTOMATION_TEST(FObjectHandleTestResolveEngineContentTarget, FObjectHandleTestBase, TEST_NAME_ROOT TEXT(".ResolveEngineContentTarget"), ObjectHandleTestFlags)
+bool FObjectHandleTestResolveEngineContentTarget::RunTest(const FString& Parameters)
+{
+	// Confirm we successfully resolve a correct reference to engine content
+	TestResolvableNonNull("/Engine/EngineResources/DefaultTexture", "DefaultTexture");
+
+	// @TODO: OBJPTR: These assets aren't in a standard cook of EngineTest, so avoid testing them when using cooked content.  Should look for other assets to use instead.
+	if (!FPlatformProperties::RequiresCookedData())
+	{
+		// Confirm we successfully resolve a correct reference to a subobject in engine content
+		TestResolvableNonNull("/Engine/FunctionalTesting/Blueprints/AITesting_MoveGoal", "AITesting_MoveGoal.EventGraph.K2Node_VariableGet_142", nullptr, nullptr, true);
+
+		// Attempt to load something that uses a User Defined Enum
+		TestResolvableNonNull("/Engine/ArtTools/RenderToTexture/Macros/RenderToTextureMacros", "RenderToTextureMacros:Array to HLSL Float Array.K2Node_Select_1", nullptr, nullptr, true);
+	}
+
+	return true;
+}
+
+IMPLEMENT_CUSTOM_SIMPLE_AUTOMATION_TEST(FObjectHandleTestResolveNonExistentTarget, FObjectHandleTestBase, TEST_NAME_ROOT TEXT(".ResolveNonExistentTarget"), ObjectHandleTestFlags)
+bool FObjectHandleTestResolveNonExistentTarget::RunTest(const FString& Parameters)
+{
+	// Confirm we don't successfully resolve an incorrect reference to engine content
+	TestResolveFailure("/Engine/EngineResources/NonExistentPackageName_0", "DefaultTexture");
+	TestResolveFailure("/Engine/EngineResources/DefaultTexture", "NonExistentObject_0");
+
+	return true;
+}
+
+IMPLEMENT_CUSTOM_SIMPLE_AUTOMATION_TEST(FObjectHandleTestResolveScriptTarget, FObjectHandleTestBase, TEST_NAME_ROOT TEXT(".ResolveScriptTarget"), ObjectHandleTestFlags)
+bool FObjectHandleTestResolveScriptTarget::RunTest(const FString& Parameters)
+{
+	// Confirm we successfully resolve a correct reference to engine content
+	TestResolvableNonNull("/Script/Engine", "Default__Actor");
+
+	TestResolvableNonNull("/Script/Engine", "DefaultPawn");
+
+	return true;
+}
+
+#if UE_WITH_OBJECT_HANDLE_LATE_RESOLVE
+IMPLEMENT_CUSTOM_SIMPLE_AUTOMATION_TEST(FObjectHandleTestResolveMalformedHandle, FObjectHandleTestBase, TEST_NAME_ROOT TEXT(".ResolveMalformedHandle"), ObjectHandleTestFlags)
+bool FObjectHandleTestResolveMalformedHandle::RunTest(const FString& Parameters)
+{
+	TestResolveFailure(FPackedObjectRef { 0xFFFF'FFFF'FFFF'FFFFull });
+	TestResolveFailure(FPackedObjectRef { 0xEFEF'EFEF'EFEF'EFEFull });
+
+	return true;
+}
+#endif // UE_WITH_OBJECT_HANDLE_LATE_RESOLVE
+
+#undef TEST_NAME_ROOT
+
+#endif // WITH_DEV_AUTOMATION_TESTS
diff --git a/Engine/Source/Runtime/CoreUObject/Tests/ObjectPtrTest.cpp b/Engine/Source/Runtime/CoreUObject/Tests/ObjectPtrTest.cpp
new file mode 100644
--- /dev/null
+++ b/Engine/Source/Runtime/CoreUObject/Tests/ObjectPtrTest.cpp
@@ -0,0 +1,374 @@
+// Copyright Epic Games, Inc. All Rights Reserved.
+
+#include "UObject/ObjectPtr.h"
+
+#include "ObjectRefTrackingTestBase.h"
+
+#include "Concepts/EqualityComparable.h"
+#include "Serialization/ArchiveCountMem.h"
+#include "Templates/Models.h"
+#include "UObject/Interface.h"
+#include "UObject/MetaData.h"
+#include "UObject/ObjectRedirector.h"
+#include "UObject/SoftObjectPath.h"
+
+#include <type_traits>
+
+namespace ObjectPtrTest
+{
+using FMutableObjectPtr = TObjectPtr<UObject>;
+using FMutableInterfacePtr = TObjectPtr<UInterface>;
+using FMutablePackagePtr = TObjectPtr<UPackage>;
+using FConstObjectPtr = TObjectPtr<const UObject>;
+using FConstInterfacePtr = TObjectPtr<const UInterface>;
+using FConstPackagePtr = TObjectPtr<const UPackage>;
+
+class UForwardDeclaredObjDerived;
+class FForwardDeclaredNotObjDerived;
+
+using UTestDummyObject = UMetaData;
+
+static_assert(sizeof(FObjectPtr) == sizeof(FObjectHandle), "FObjectPtr type must always compile to something equivalent to an FObjectHandle size.");
+static_assert(sizeof(FObjectPtr) == sizeof(void*), "FObjectPtr type must always compile to something equivalent to a pointer size.");
+static_assert(sizeof(TObjectPtr<UObject>) == sizeof(void*), "TObjectPtr<UObject> type must always compile to something equivalent to a pointer size.");
+
+// Ensure that a TObjectPtr is trivially copyable, (copy/move) constructible, (copy/move) assignable, and destructible
+static_assert(std::is_trivially_copyable<FMutableObjectPtr>::value, "TObjectPtr must be trivially copyable");
+static_assert(std::is_trivially_copy_constructible<FMutableObjectPtr>::value, "TObjectPtr must be trivially copy constructible");
+static_assert(std::is_trivially_move_constructible<FMutableObjectPtr>::value, "TObjectPtr must be trivially move constructible");
+static_assert(std::is_trivially_copy_assignable<FMutableObjectPtr>::value, "TObjectPtr must be trivially copy assignable");
+static_assert(std::is_trivially_move_assignable<FMutableObjectPtr>::value, "TObjectPtr must be trivially move assignable");
+static_assert(std::is_trivially_destructible<FMutableObjectPtr>::value, "TObjectPtr must be trivially destructible");
+static_assert(std::is_trivially_default_constructible<FMutableObjectPtr>::value, "TObjectPtr must be trivially default constructible");
+
+// Ensure that raw pointers can be used to construct wrapped object pointers and that const-ness isn't stripped when constructing or converting with raw pointers
+static_assert(std::is_constructible<FMutableObjectPtr, UObject*>::value, "TObjectPtr<UObject> must be constructible from a raw UObject*");
+static_assert(!std::is_constructible<FMutableObjectPtr, const UObject*>::value, "TObjectPtr<UObject> must not be constructible from a const raw UObject*");
+static_assert(std::is_convertible<FMutableObjectPtr, UObject*>::value, "TObjectPtr<UObject> must be convertible to a raw UObject*");
+static_assert(std::is_convertible<FMutableObjectPtr, const UObject*>::value, "TObjectPtr<UObject> must be convertible to a const raw UObject*");
+
+static_assert(std::is_constructible<FConstObjectPtr, UObject*>::value, "TObjectPtr<const UObject> must be constructible from a raw UObject*");
+static_assert(std::is_constructible<FConstObjectPtr, const UObject*>::value, "TObjectPtr<const UObject> must be constructible from a const raw UObject*");
+static_assert(!std::is_convertible<FConstObjectPtr, UObject*>::value, "TObjectPtr<const UObject> must not be convertible to a raw UObject*");
+static_assert(std::is_convertible<FConstObjectPtr, const UObject*>::value, "TObjectPtr<const UObject> must be convertible to a const raw UObject*");
+
+// Ensure that a TObjectPtr<const UObject> is constructible and assignable from a TObjectPtr<UObject> but not vice versa
+static_assert(std::is_constructible<FConstObjectPtr, const FMutableObjectPtr&>::value, "Missing constructor (TObjectPtr<const UObject> from TObjectPtr<UObject>)");
+static_assert(!std::is_constructible<FMutableObjectPtr, const FConstObjectPtr&>::value, "Invalid constructor (TObjectPtr<UObject> from TObjectPtr<const UObject>)");
+static_assert(std::is_assignable<FConstObjectPtr, const FMutableObjectPtr&>::value, "Missing assignment (TObjectPtr<const UObject> from TObjectPtr<UObject>)");
+static_assert(!std::is_assignable<FMutableObjectPtr, const FConstObjectPtr&>::value, "Invalid assignment (TObjectPtr<UObject> from TObjectPtr<const UObject>)");
+
+static_assert(std::is_constructible<FConstObjectPtr, const FConstObjectPtr&>::value, "Missing constructor (TObjectPtr<const UObject> from TObjectPtr<const UObject>)");
+static_assert(std::is_assignable<FConstObjectPtr, const FConstObjectPtr&>::value, "Missing assignment (TObjectPtr<const UObject> from TObjectPtr<const UObject>)");
+
+// Ensure that a TObjectPtr<UObject> is constructible and assignable from a TObjectPtr<UInterface> but not vice versa
+static_assert(std::is_constructible<FMutableObjectPtr, const FMutableInterfacePtr&>::value, "Missing constructor (TObjectPtr<UObject> from TObjectPtr<UInterface>)");
+static_assert(!std::is_constructible<FMutableInterfacePtr, const FMutableObjectPtr&>::value, "Invalid constructor (TObjectPtr<UInterface> from TObjectPtr<UObject>)");
+static_assert(std::is_constructible<FConstObjectPtr, const FConstInterfacePtr&>::value, "Missing constructor (TObjectPtr<const UObject> from TObjectPtr<const UInterface>)");
+static_assert(std::is_constructible<FConstObjectPtr, const FMutableInterfacePtr&>::value, "Missing constructor (TObjectPtr<const UObject> from TObjectPtr<UInterface>)");
+static_assert(!std::is_constructible<FConstInterfacePtr, const FConstObjectPtr&>::value, "Invalid constructor (TObjectPtr<const UInterface> from TObjectPtr<const UObject>)");
+static_assert(!std::is_constructible<FConstInterfacePtr, const FMutableObjectPtr&>::value, "Invalid constructor (TObjectPtr<const UInterface> from TObjectPtr<UObject>)");
+
+static_assert(std::is_assignable<FMutableObjectPtr, const FMutableInterfacePtr&>::value, "Missing assignment (TObjectPtr<UObject> from TObjectPtr<UInterface>)");
+static_assert(std::is_assignable<FConstObjectPtr, const FMutableInterfacePtr&>::value, "Missing assignment (TObjectPtr<const UObject> from TObjectPtr<UInterface>)");
+static_assert(std::is_assignable<FConstObjectPtr, const FConstInterfacePtr&>::value, "Missing assignment (TObjectPtr<const UObject> from TObjectPtr<const UInterface>)");
+static_assert(!std::is_assignable<FMutableInterfacePtr, const FMutableObjectPtr&>::value, "Invalid assignment (TObjectPtr<UInterface> from TObjectPtr<UObject>)");
+static_assert(!std::is_assignable<FConstInterfacePtr, const FMutableObjectPtr&>::value, "Invalid assignment (TObjectPtr<const UInterface> from TObjectPtr<UObject>)");
+static_assert(!std::is_assignable<FConstInterfacePtr, const FConstObjectPtr&>::value, "Invalid assignment (TObjectPtr<const UInterface> from TObjectPtr<const UObject>)");
+
+// Ensure that TObjectPtr<[const] UObject> is comparable with another TObjectPtr<[const] UObject> regardless of constness
+static_assert(TModels<CEqualityComparableWith, FConstObjectPtr, FConstObjectPtr>::Value, "Must be able to compare equality and inequality bidirectionally between TObjectPtr<const UObject> and TObjectPtr<const UObject>");
+static_assert(TModels<CEqualityComparableWith, FMutableObjectPtr, FConstObjectPtr>::Value, "Must be able to compare equality and inequality bidirectionally between TObjectPtr<UObject> and TObjectPtr<const UObject>");
+
+// Ensure that TObjectPtr<[const] UObject> is comparable with another TObjectPtr<[const] UInterface> regardless of constness
+static_assert(TModels<CEqualityComparableWith, FConstObjectPtr, FConstInterfacePtr>::Value, "Must be able to compare equality and inequality bidirectionally between TObjectPtr<const UObject> and TObjectPtr<const UInterface>");
+static_assert(TModels<CEqualityComparableWith, FMutableObjectPtr, FConstInterfacePtr>::Value, "Must be able to compare equality and inequality bidirectionally between TObjectPtr<UObject> and TObjectPtr<const UInterface>");
+static_assert(TModels<CEqualityComparableWith, FConstObjectPtr, FMutableInterfacePtr>::Value, "Must be able to compare equality and inequality bidirectionally between TObjectPtr<const UObject> and TObjectPtr<UInterface>");
+static_assert(TModels<CEqualityComparableWith, FMutableObjectPtr, FMutableInterfacePtr>::Value, "Must be able to compare equality and inequality bidirectionally between TObjectPtr<UObject> and TObjectPtr<UInterface>");
+
+// Ensure that TObjectPtr<[const] UPackage> is not comparable with a TObjectPtr<[const] UInterface> regardless of constness
+// TODO: This only ensures that at least one of the A==B,B==A,A!=B,B!=A operations fail, not that they all fail.
+#if !(PLATFORM_MICROSOFT) || !defined(_MSC_EXTENSIONS) // MSVC static analyzer is run in non-conformance mode, and that causes these checks to fail.
+static_assert(!TModels<CEqualityComparableWith, FConstPackagePtr, FConstInterfacePtr>::Value, "Must not be able to compare equality and inequality bidirectionally between TObjectPtr<const UPackage> and TObjectPtr<const UInterface>");
+static_assert(!TModels<CEqualityComparableWith, FMutablePackagePtr, FConstInterfacePtr>::Value, "Must not be able to compare equality and inequality bidirectionally between TObjectPtr<UPackage> and TObjectPtr<const UInterface>");
+static_assert(!TModels<CEqualityComparableWith, FConstPackagePtr, FMutableInterfacePtr>::Value, "Must not be able to compare equality and inequality bidirectionally between TObjectPtr<const UPackage> and TObjectPtr<UInterface>");
+static_assert(!TModels<CEqualityComparableWith, FMutablePackagePtr, FMutableInterfacePtr>::Value, "Must not be able to compare equality and inequality bidirectionally between TObjectPtr<UPackage> and TObjectPtr<UInterface>");
+#endif // #if !(PLATFORM_MICROSOFT) || !defined(_MSC_EXTENSIONS)
+
+// Ensure that TObjectPtr<[const] UObject> is comparable with a raw pointer of the same referenced type regardless of constness
+static_assert(TModels<CEqualityComparableWith, FConstObjectPtr, const UObject*>::Value, "Must be able to compare equality and inequality bidirectionally between TObjectPtr<const UObject> and const UObject*");
+static_assert(TModels<CEqualityComparableWith, FMutableObjectPtr, const UObject*>::Value, "Must be able to compare equality and inequality bidirectionally between TObjectPtr<UObject> and const UObject*");
+static_assert(TModels<CEqualityComparableWith, FConstObjectPtr, UObject*>::Value, "Must be able to compare equality and inequality bidirectionally between TObjectPtr<const UObject> and UObject*");
+static_assert(TModels<CEqualityComparableWith, FMutableObjectPtr, UObject*>::Value, "Must be able to compare equality and inequality bidirectionally between TObjectPtr<UObject> and UObject*");
+
+// Ensure that TObjectPtr<[const] UObject> is comparable with a UInterface raw pointer regardless of constness
+static_assert(TModels<CEqualityComparableWith, FConstObjectPtr, const UInterface*>::Value, "Must be able to compare equality and inequality bidirectionally between TObjectPtr<const UObject> and const UInterface*");
+static_assert(TModels<CEqualityComparableWith, FMutableObjectPtr, const UInterface*>::Value, "Must be able to compare equality and inequality bidirectionally between TObjectPtr<UObject> and const UInterface*");
+static_assert(TModels<CEqualityComparableWith, FConstObjectPtr, UInterface*>::Value, "Must be able to compare equality and inequality bidirectionally between TObjectPtr<const UObject> and UInterface*");
+static_assert(TModels<CEqualityComparableWith, FMutableObjectPtr, UInterface*>::Value, "Must be able to compare equality and inequality bidirectionally between TObjectPtr<UObject> and UInterface*");
+
+// Ensure that TObjectPtr<[const] UInterface> is comparable with a UObject raw pointer regardless of constness
+static_assert(TModels<CEqualityComparableWith, FConstInterfacePtr, const UObject*>::Value, "Must be able to compare equality and inequality bidirectionally between TObjectPtr<const UInterface> and const UObject*");
+static_assert(TModels<CEqualityComparableWith, FMutableInterfacePtr, const UObject*>::Value, "Must be able to compare equality and inequality bidirectionally between TObjectPtr<UInterface> and const UObject*");
+static_assert(TModels<CEqualityComparableWith, FConstInterfacePtr, UObject*>::Value, "Must be able to compare equality and inequality bidirectionally between TObjectPtr<const UInterface> and UObject*");
+static_assert(TModels<CEqualityComparableWith, FMutableInterfacePtr, UObject*>::Value, "Must be able to compare equality and inequality bidirectionally between TObjectPtr<UInterface> and UObject*");
+
+// Ensure that TObjectPtr<[const] UInterface> is not comparable with a UPackage raw pointer regardless of constness
+// TODO: This only ensures that at least one of the A==B,B==A,A!=B,B!=A operations fail, not that they all fail.
+static_assert(!TModels<CEqualityComparableWith, FConstInterfacePtr, const UPackage*>::Value, "Must not be able to compare equality and inequality bidirectionally between TObjectPtr<const UInterface> and const UPackage*");
+static_assert(!TModels<CEqualityComparableWith, FMutableInterfacePtr, const UPackage*>::Value, "Must not be able to compare equality and inequality bidirectionally between TObjectPtr<UInterface> and const UPackage*");
+static_assert(!TModels<CEqualityComparableWith, FConstInterfacePtr, UPackage*>::Value, "Must not be able to compare equality and inequality bidirectionally between TObjectPtr<const UInterface> and UPackage*");
+static_assert(!TModels<CEqualityComparableWith, FMutableInterfacePtr, UPackage*>::Value, "Must not be able to compare equality and inequality bidirectionally between TObjectPtr<UInterface> and UPackage*");
+
+// Ensure that TObjectPtr<[const] UInterface> is not comparable with a char raw pointer regardless of constness
+// TODO: This only ensures that at least one of the A==B,B==A,A!=B,B!=A operations fail, not that they all fail.
+static_assert(!TModels<CEqualityComparableWith, FConstObjectPtr, const char*>::Value, "Must not be able to compare equality and inequality bidirectionally between TObjectPtr<const UObject> and const UObject*");
+static_assert(!TModels<CEqualityComparableWith, FMutableObjectPtr, const char*>::Value, "Must not be able to compare equality and inequality bidirectionally between TObjectPtr<UObject> and const UObject*");
+static_assert(!TModels<CEqualityComparableWith, FConstObjectPtr, char*>::Value, "Must not be able to compare equality and inequality bidirectionally between TObjectPtr<const UObject> and UObject*");
+static_assert(!TModels<CEqualityComparableWith, FMutableObjectPtr, char*>::Value, "Must not be able to compare equality and inequality bidirectionally between TObjectPtr<UObject> and UObject*");
+
+// Ensure that TObjectPtr<[const] UObject> is comparable with nullptr regardless of constness
+static_assert(TModels<CEqualityComparableWith, FConstObjectPtr, TYPE_OF_NULLPTR>::Value, "Must be able to compare equality and inequality bidirectionally between TObjectPtr<const UObject> and nullptr");
+static_assert(TModels<CEqualityComparableWith, FMutableObjectPtr, TYPE_OF_NULLPTR>::Value, "Must be able to compare equality and inequality bidirectionally between TObjectPtr<UObject> and nullptr");
+
+#if !UE_OBJECT_PTR_NONCONFORMANCE_SUPPORT // Specialized NULL support causes these checks to fail.
+static_assert(!TModels<CEqualityComparableWith, FConstObjectPtr, long>::Value, "Should not be able to compare equality and inequality bidirectionally between TObjectPtr<const UObject> and long");
+static_assert(!TModels<CEqualityComparableWith, FMutableObjectPtr, long>::Value, "Should not be able to compare equality and inequality bidirectionally between TObjectPtr<UObject> and long");
+#endif // #if !UE_OBJECT_PTR_NONCONFORMANCE_SUPPORT
+
+// Ensure that the use of incomplete types doesn't provide a means to bypass type safety on TObjectPtr
+// NOTE: This is disabled because we're permitting this operation with a deprecation warning.
+//static_assert(!std::is_assignable<TObjectPtr<UForwardDeclaredObjDerived>, UForwardDeclaredObjDerived*>::value, "Should not be able to assign raw pointer of incomplete type that descends from UObject to exactly this type of TObjectPtr");
+//static_assert(!std::is_assignable<TObjectPtr<FForwardDeclaredNotObjDerived>, FForwardDeclaredNotObjDerived*>::value, "Should not be able to assign raw pointer of incomplete type that does not descend from UObject to exactly this type of TObjectPtr");
+
+#if WITH_DEV_AUTOMATION_TESTS
+
+class FObjectPtrTestBase : public FObjectRefTrackingTestBase
+{
+public:
+	FObjectPtrTestBase(const FString& InName, const bool bInComplexTask)
+	: FObjectRefTrackingTestBase(InName, bInComplexTask)
+	{
+	}
+
+protected:
+};
+
+#define TEST_NAME_ROOT TEXT("System.CoreUObject.ObjectPtr")
+constexpr const uint32 ObjectPtrTestFlags = EAutomationTestFlags::ApplicationContextMask | EAutomationTestFlags::EngineFilter;
+
+IMPLEMENT_CUSTOM_SIMPLE_AUTOMATION_TEST(FObjectPtrTestNullBehavior, FObjectPtrTestBase, TEST_NAME_ROOT TEXT(".NullBehavior"), ObjectPtrTestFlags)
+bool FObjectPtrTestNullBehavior::RunTest(const FString& Parameters)
+{
+	TObjectPtr<UObject> NullObjectPtr(nullptr);
+	TestTrue(TEXT("Nullptr should equal a null object pointer"), nullptr == NullObjectPtr);
+	TestTrue(TEXT("A null object pointer should equal nullptr"), NullObjectPtr == nullptr);
+	TestFalse(TEXT("A null object pointer should evaluate to false"), !!NullObjectPtr);
+	TestTrue(TEXT("Negation of a null object pointer should evaluate to true"), !NullObjectPtr);
+
+	return true;
+}
+
+IMPLEMENT_CUSTOM_SIMPLE_AUTOMATION_TEST(FObjectPtrTestDefaultSerialize, FObjectPtrTestBase, TEST_NAME_ROOT TEXT(".DefaultSerialize"), ObjectPtrTestFlags)
+bool FObjectPtrTestDefaultSerialize::RunTest(const FString& Parameters)
+{
+	FSnapshotObjectRefMetrics ObjectRefMetrics(*this);
+	FObjectPtr DefaultTexturePtr(FObjectRef {FName("/Engine/EngineResources/DefaultTexture"), NAME_None, NAME_None, FObjectPathId("DefaultTexture")});
+
+	ObjectRefMetrics.TestNumResolves(TEXT("Unexpected resolve count after initializing an FObjectPtr"), UE_WITH_OBJECT_HANDLE_LATE_RESOLVE ? 0 : 1);
+	ObjectRefMetrics.TestNumFailedResolves(TEXT("Unexpected resolve failure after initializing an FObjectPtr"), 0);
+	ObjectRefMetrics.TestNumReads(TEXT("NumReads should not change when initializing an FObjectPtr"), 0);
+
+	FArchiveUObject Writer;
+	Writer << DefaultTexturePtr;
+
+	ObjectRefMetrics.TestNumResolves(TEXT("Serializing an FObjectPtr should force it to resolve"), 1);
+	ObjectRefMetrics.TestNumFailedResolves(TEXT("Unexpected resolve failure after serializing an FObjectPtr"), 0);
+	ObjectRefMetrics.TestNumReads(TEXT("NumReads should increase after serializing an FObjectPtr"), 1);
+
+	Writer << DefaultTexturePtr;
+
+	ObjectRefMetrics.TestNumResolves(TEXT("Serializing an FObjectPtr twice should only require it to resolve once"), 1);
+	ObjectRefMetrics.TestNumFailedResolves(TEXT("Unexpected resolve failure after serializing an FObjectPtr"), 0);
+	ObjectRefMetrics.TestNumReads(TEXT("NumReads should increase after serializing an FObjectPtr"), 2);
+
+	return true;
+}
+
+IMPLEMENT_CUSTOM_SIMPLE_AUTOMATION_TEST(FObjectPtrTestSoftObjectPath, FObjectPtrTestBase, TEST_NAME_ROOT TEXT(".SoftObjectPath"), ObjectPtrTestFlags)
+bool FObjectPtrTestSoftObjectPath::RunTest(const FString& Parameters)
+{
+	FSnapshotObjectRefMetrics ObjectRefMetrics(*this);
+	FObjectPtr DefaultTexturePtr(FObjectRef {FName("/Engine/EngineResources/DefaultTexture"), NAME_None, NAME_None, FObjectPathId("DefaultTexture")});
+
+	ObjectRefMetrics.TestNumResolves(TEXT("Unexpected resolve count after initializing an FObjectPtr"), UE_WITH_OBJECT_HANDLE_LATE_RESOLVE ? 0 : 1);
+	ObjectRefMetrics.TestNumFailedResolves(TEXT("Unexpected resolve failure after initializing an FObjectPtr"), 0);
+	ObjectRefMetrics.TestNumReads(TEXT("NumReads should not change when initializing an FObjectPtr"), 0);
+
+	FSoftObjectPath DefaultTexturePath(DefaultTexturePtr);
+
+	ObjectRefMetrics.TestNumResolves(TEXT("Unexpected resolve count after initializing an FSoftObjectPath from an FObjectPtr"), UE_WITH_OBJECT_HANDLE_LATE_RESOLVE ? 0 : 1);
+
+	TestEqual(TEXT("Soft object path constructed from an FObjectPtr does not have the expected path value"), *DefaultTexturePath.ToString(), TEXT("/Engine/EngineResources/DefaultTexture.DefaultTexture"));
+
+	return true;
+}
+
+IMPLEMENT_CUSTOM_SIMPLE_AUTOMATION_TEST(FObjectPtrTestForwardDeclared, FObjectPtrTestBase, TEST_NAME_ROOT TEXT(".ForwardDeclared"), ObjectPtrTestFlags)
+bool FObjectPtrTestForwardDeclared::RunTest(const FString& Parameters)
+{
+	UForwardDeclaredObjDerived* PtrFwd = nullptr;
+	TObjectPtr<UForwardDeclaredObjDerived> ObjPtrFwd(MakeObjectPtrUnsafe<UForwardDeclaredObjDerived>(reinterpret_cast<UObject*>(PtrFwd)));
+	TestTrue(TEXT("Null forward declared pointer used to construct a TObjectPtr should result in a null TObjectPtr"), ObjPtrFwd.IsNull());
+	return true;
+}
+
+IMPLEMENT_CUSTOM_SIMPLE_AUTOMATION_TEST(FObjectPtrTestHashConsistency, FObjectPtrTestBase, TEST_NAME_ROOT TEXT(".HashConsistency"), ObjectPtrTestFlags)
+bool FObjectPtrTestHashConsistency::RunTest(const FString& Parameters)
+{
+	const FName TestPackage1Name(TEXT("/Engine/Test/ObjectPtrHashConsistency1/Transient"));
+	UPackage* TestPackage1 = NewObject<UPackage>(nullptr, TestPackage1Name, RF_Transient);
+	TestPackage1->AddToRoot();
+	UObject* TestOuter1 = NewObject<UTestDummyObject>(TestPackage1, TEXT("TestOuter1"));
+	UObject* TestOuter2 = NewObject<UTestDummyObject>(TestPackage1, TEXT("TestOuter2"));
+	UObject* TestOuter3 = NewObject<UTestDummyObject>(TestPackage1, TEXT("TestOuter3"));
+	UObject* TestOuter4 = NewObject<UTestDummyObject>(TestPackage1, TEXT("TestOuter4"));
+
+	UObject* TestPublicObject = NewObject<UTestDummyObject>(TestOuter1, TEXT("TestPublicObject"), RF_Public);
+	UObjectRedirector* TestRedirectorToPublicObject1 = NewObject<UObjectRedirector>(TestOuter3, TEXT("TestPublicRedirector1"), RF_Public);
+	TestRedirectorToPublicObject1->DestinationObject = TestPublicObject;
+	UObjectRedirector* TestRedirectorToPublicObject2 = NewObject<UObjectRedirector>(TestOuter4, TEXT("TestPublicRedirector2"), RF_Public);
+	TestRedirectorToPublicObject2->DestinationObject = TestPublicObject;
+
+	// Perform hash consistency checks on public object reference
+	{
+		// Check that unresolved/resolved pointers produce the same hash
+		FSnapshotObjectRefMetrics ObjectRefMetrics(*this);
+
+		FObjectPtr TestPublicWrappedObjectPtr(FObjectRef{TestPackage1Name, NAME_None, NAME_None, FObjectPathId("TestOuter1.TestPublicObject")});
+		ObjectRefMetrics.TestNumResolves(TEXT("Unexpected resolve count after initializing an FObjectPtr"), UE_WITH_OBJECT_HANDLE_LATE_RESOLVE ? 0 : 1);
+		ObjectRefMetrics.TestNumFailedResolves(TEXT("Unexpected resolve failure after initializing an FObjectPtr"), 0);
+
+		uint32 HashWrapped = GetTypeHash(TestPublicWrappedObjectPtr);
+		ObjectRefMetrics.TestNumResolves(TEXT("Unexpected resolve count after hashing an FObjectPtr"), 1);
+		ObjectRefMetrics.TestNumFailedResolves(TEXT("Unexpected resolve failure after hashing an FObjectPtr"), 0);
+
+		TestPublicWrappedObjectPtr.Get();
+
+		ObjectRefMetrics.TestNumResolves(TEXT("Unexpected resolve count after resolving an FObjectPtr"), 1);
+		ObjectRefMetrics.TestNumFailedResolves(TEXT("Unexpected resolve failure after resolving an FObjectPtr"), 0);
+
+		uint32 HashRaw = GetTypeHash(TestPublicObject);
+		TestEqual(TEXT("Hash of raw public FObjectPtr should equal hash of wrapped public FObjectPtr"), HashRaw, HashWrapped);
+
+		FObjectPtr TestPublicWrappedRedir1(FObjectRef{TestPackage1Name, NAME_None, NAME_None, FObjectPathId("TestOuter3.TestPublicRedirector1")});
+		ObjectRefMetrics.TestNumResolves(TEXT("Unexpected resolve count after initializing an FObjectPtr"), UE_WITH_OBJECT_HANDLE_LATE_RESOLVE ? 1 : 2);
+		ObjectRefMetrics.TestNumFailedResolves(TEXT("Unexpected resolve failure after initializing an FObjectPtr"), 0);
+
+		uint32 HashWrappedRedir1 = GetTypeHash(TestPublicWrappedRedir1);
+		ObjectRefMetrics.TestNumResolves(TEXT("Unexpected resolve count after hashing an FObjectPtr"), 2);
+		ObjectRefMetrics.TestNumFailedResolves(TEXT("Unexpected resolve failure after hashing an FObjectPtr"), 0);
+
+		TestEqual(TEXT("Hash of first wrapped public redirector should equal hash of wrapped public FObjectPtr it references"), HashWrapped, HashWrappedRedir1);
+
+		FObjectPtr TestPublicWrappedRedir2(FObjectRef{TestPackage1Name, NAME_None, NAME_None, FObjectPathId("TestOuter4.TestPublicRedirector2")});
+		ObjectRefMetrics.TestNumResolves(TEXT("Unexpected resolve count after initializing an FObjectPtr"), UE_WITH_OBJECT_HANDLE_LATE_RESOLVE ? 2 : 3);
+		ObjectRefMetrics.TestNumFailedResolves(TEXT("Unexpected resolve failure after initializing an FObjectPtr"), 0);
+
+		uint32 HashWrappedRedir2 = GetTypeHash(TestPublicWrappedRedir2);
+		ObjectRefMetrics.TestNumResolves(TEXT("Unexpected resolve count after hashing an FObjectPtr"), 3);
+		ObjectRefMetrics.TestNumFailedResolves(TEXT("Unexpected resolve failure after hashing an FObjectPtr"), 0);
+
+		TestEqual(TEXT("Hash of first wrapped public redirector should equal hash of wrapped public FObjectPtr it references"), HashWrapped, HashWrappedRedir1);
+
+		//Check that renaming an object doesn't change its hash
+		TestPublicObject->Rename(TEXT("TestPublicObjectRenamed"));
+		uint32 HashWrappedAfterRename = GetTypeHash(TestPublicWrappedObjectPtr);
+		TestEqual(TEXT("Hash of resolved public FObjectPtr before rename should equal hash of resolved public FObjectPtr after rename"), HashWrappedAfterRename, HashWrapped);
+
+		//Check that reparenting an object doesn't change its hash
+		TestPublicObject->Rename(nullptr, TestOuter2);
+		uint32 HashWrappedAfterReparent = GetTypeHash(TestPublicWrappedObjectPtr);
+		TestEqual(TEXT("Hash of resolved public FObjectPtr before reparenting should equal hash of resolved public FObjectPtr after reparenting"), HashWrappedAfterReparent, HashWrapped);
+	}
+
+	TestPackage1->RemoveFromRoot();
+	return true;
+}
+
+IMPLEMENT_CUSTOM_SIMPLE_AUTOMATION_TEST(FObjectPtrTestLongPath, FObjectPtrTestBase, TEST_NAME_ROOT TEXT(".LongPath"), ObjectPtrTestFlags)
+bool FObjectPtrTestLongPath::RunTest(const FString& Parameters)
+{
+	const FName TestPackage1Name(TEXT("/Engine/Test/FObjectPtrTestLongPath/Transient"));
+	UPackage* TestPackage1 = NewObject<UPackage>(nullptr, TestPackage1Name, RF_Transient);
+	TestPackage1->AddToRoot();
+
+	ON_SCOPE_EXIT{
+		TestPackage1->RemoveFromRoot();
+	};
+
+	UObject* TestObject1 = NewObject<UTestDummyObject>(TestPackage1, TEXT("TestObject1"));
+	UObject* TestObject2 = NewObject<UTestDummyObject>(TestObject1, TEXT("TestObject2"));
+	UObject* TestObject3 = NewObject<UTestDummyObject>(TestObject2, TEXT("TestObject3"));
+	UObject* TestObject4 = NewObject<UTestDummyObject>(TestObject3, TEXT("TestObject4"));
+
+	FObjectPathId LongPath(TestObject4);
+	FObjectPathId::ResolvedNameContainerType ResolvedNames;
+	LongPath.Resolve(ResolvedNames);
+
+	TestEqual(TEXT("Resolved path from FObjectPathId should have 4 elements"), ResolvedNames.Num(), 4);
+	TestEqual(TEXT("Resolved path from FObjectPathId should have TestObject1 at element 0"), ResolvedNames[0], TestObject1->GetFName());
+	TestEqual(TEXT("Resolved path from FObjectPathId should have TestObject2 at element 1"), ResolvedNames[1], TestObject2->GetFName());
+	TestEqual(TEXT("Resolved path from FObjectPathId should have TestObject3 at element 2"), ResolvedNames[2], TestObject3->GetFName());
+	TestEqual(TEXT("Resolved path from FObjectPathId should have TestObject4 at element 3"), ResolvedNames[3], TestObject4->GetFName());
+
+	return true;
+}
+
+// @TODO: OBJPTR: We should have a test that ensures that lazy loading of an object with an external package is handled correctly.
+//				  This should also include external packages in the outer chain of the target object.
+// IMPLEMENT_CUSTOM_SIMPLE_AUTOMATION_TEST(FObjectPtrTestExternalPackages, FObjectPtrTestBase, TEST_NAME_ROOT TEXT(".ExternalPackages"), ObjectPtrTestFlags)
+// bool FObjectPtrTestExternalPackages::RunTest(const FString& Parameters)
+// {
+// 	const FName TestExternalPackage1Name(TEXT("/Engine/Test/ObjectPtrExternalPackages1/Transient"));
+// 	UPackage* TestExternalPackage1 = NewObject<UPackage>(nullptr, TestExternalPackage1Name, RF_Transient);
+// 	TestExternalPackage1->SetPackageFlags(PKG_EditorOnly | PKG_ContainsMapData);
+// 	TestExternalPackage1->AddToRoot();
+
+// 	const FName TestPackage1Name(TEXT("/Engine/Test/ObjectPtrExternalPackages1/Transient"));
+// 	UPackage* TestPackage1 = NewObject<UPackage>(nullptr, TestPackage1Name, RF_Transient);
+// 	TestPackage1->AddToRoot();
+// 	UObject* TestOuter1 = NewObject<UTestDummyObject>(TestPackage1, TEXT("TestOuter1"));
+// 	UObject* TestOuter2 = NewObject<UTestDummyObject>(TestPackage1, TEXT("TestOuter2"));
+// 	UObject* TestOuter3 = NewObject<UTestDummyObject>(TestPackage1, TEXT("TestOuter3"));
+// 	UObject* TestOuter4 = NewObject<UTestDummyObject>(TestPackage1, TEXT("TestOuter4"));
+
+// 	UObject* TestPublicObject = NewObject<UTestDummyObject>(TestOuter4, TEXT("TestPublicObject"), RF_Public);
+// 	TestPublicObject->SetExternalPackage(TestExternalPackage1);
+
+// 	TestPackage1->RemoveFromRoot();
+// 	TestExternalPackage1->RemoveFromRoot();
+// 	return true;
+// }
+
+// @TODO: OBJPTR: We should have a test that ensures that we can (de)serialize an FObjectPtr to FLinkerSave/FLinkerLoad and that upon load the object
+//			pointer is not resolved if we are in a configuration that supports lazy load.  This is proving difficult due to the restrictions around how
+//			FLinkerSave/FLinkerLoad is used.
+// IMPLEMENT_CUSTOM_SIMPLE_AUTOMATION_TEST(FObjectPtrTestLinkerSerializeBehavior, FObjectPtrTestBase, TEST_NAME_ROOT TEXT(".LinkerSerialize"), ObjectPtrTestFlags)
+// bool FObjectPtrTestLinkerSerializeBehavior::RunTest(const FString& Parameters)
+// {
+// 	FSnapshotObjectRefMetrics ObjectRefMetrics(*this);
+// 	FObjectPtr DefaultTexturePtr(FObjectRef {FName("/Engine/EngineResources/DefaultTexture"), NAME_None, NAME_None, FObjectPathId("DefaultTexture")});
+
+// 	TUniquePtr<FLinkerSave> Linker = TUniquePtr<FLinkerSave>(new FLinkerSave(nullptr /*InOuter*/, false /*bForceByteSwapping*/, true /*bSaveUnversioned*/));
+// 	return true;
+// }
+
+#undef TEST_NAME_ROOT
+
+#endif // WITH_DEV_AUTOMATION_TESTS
+
+class UForwardDeclaredObjDerived: public UObject {};
+class FForwardDeclaredNotObjDerived {};
+
+}
diff --git a/Engine/Source/Runtime/CoreUObject/Tests/ObjectRefTrackingTestBase.cpp b/Engine/Source/Runtime/CoreUObject/Tests/ObjectRefTrackingTestBase.cpp
new file mode 100644
--- /dev/null
+++ b/Engine/Source/Runtime/CoreUObject/Tests/ObjectRefTrackingTestBase.cpp
@@ -0,0 +1,15 @@
+// Copyright Epic Games, Inc. All Rights Reserved.
+
+#include "ObjectRefTrackingTestBase.h"
+
+#if WITH_DEV_AUTOMATION_TESTS
+
+#if UE_WITH_OBJECT_HANDLE_TRACKING
+ObjectHandleReferenceResolvedFunction* FObjectRefTrackingTestBase::PrevResolvedFunc = nullptr;
+ObjectHandleReadFunction* FObjectRefTrackingTestBase::PrevReadFunc = nullptr;
+#endif
+thread_local uint32 FObjectRefTrackingTestBase::NumResolves = 0;
+thread_local uint32 FObjectRefTrackingTestBase::NumFailedResolves = 0;
+thread_local uint32 FObjectRefTrackingTestBase::NumReads = 0;
+
+#endif // WITH_DEV_AUTOMATION_TESTS
diff --git a/Engine/Source/Runtime/CoreUObject/Tests/ObjectRefTrackingTestBase.h b/Engine/Source/Runtime/CoreUObject/Tests/ObjectRefTrackingTestBase.h
new file mode 100644
--- /dev/null
+++ b/Engine/Source/Runtime/CoreUObject/Tests/ObjectRefTrackingTestBase.h
@@ -0,0 +1,129 @@
+// Copyright Epic Games, Inc. All Rights Reserved.
+
+#pragma once
+
+#include "CoreMinimal.h"
+#include "Misc/AutomationTest.h"
+#include "UObject/ObjectHandle.h"
+
+#if WITH_DEV_AUTOMATION_TESTS
+
+class FObjectRefTrackingTestBase : public FAutomationTestBase
+{
+public:
+	FObjectRefTrackingTestBase(const FString& InName, const bool bInComplexTask)
+	: FAutomationTestBase(InName, bInComplexTask)
+	{
+	}
+
+	uint32 GetNumResolves() const { return NumResolves; }
+	uint32 GetNumFailedResolves() const { return NumFailedResolves; }
+	uint32 GetNumReads() const { return NumReads; }
+
+	struct FSnapshotObjectRefMetrics
+	{
+	public:
+		FSnapshotObjectRefMetrics(FObjectRefTrackingTestBase& InTest)
+		: Test(InTest)
+		, OriginalNumResolves(Test.GetNumResolves())
+		, OriginalNumFailedResolves(Test.GetNumFailedResolves())
+		, OriginalNumReads(Test.GetNumReads())
+		{
+			Test.ConditionalInstallCallbacks();
+		}
+
+		bool TestNumResolves(const TCHAR* What, uint32 ExpectedDelta)
+		{
+#if UE_WITH_OBJECT_HANDLE_TRACKING
+			return Test.TestEqual(What, OriginalNumResolves + ExpectedDelta, Test.GetNumResolves());
+#else
+			return true;
+#endif
+		}
+
+		bool TestNumFailedResolves(const TCHAR* What, uint32 ExpectedDelta)
+		{
+#if UE_WITH_OBJECT_HANDLE_TRACKING
+			return Test.TestEqual(What, OriginalNumFailedResolves + ExpectedDelta, Test.GetNumFailedResolves());
+#else
+			return true;
+#endif
+		}
+
+		bool TestNumReads(const TCHAR* What, uint32 ExpectedDelta, bool bAllowAdditionalReads = false)
+		{
+#if UE_WITH_OBJECT_HANDLE_TRACKING
+			if (bAllowAdditionalReads)
+			{
+				return Test.TestTrue(What, Test.GetNumReads() >= OriginalNumReads + ExpectedDelta);
+			}
+			else
+			{
+				return Test.TestEqual(What, OriginalNumReads + ExpectedDelta, Test.GetNumReads());
+			}
+#else
+			return true;
+#endif
+		}
+	private:
+		FObjectRefTrackingTestBase& Test;
+		uint32 OriginalNumResolves;
+		uint32 OriginalNumFailedResolves;
+		uint32 OriginalNumReads;
+	};
+
+private:
+#if UE_WITH_OBJECT_HANDLE_TRACKING
+	static void OnRefResolved(const FObjectRef& ObjectRef, UPackage* Pkg, UObject* Obj)
+	{
+		NumResolves++;
+		if (!IsObjectRefNull(ObjectRef) && !Obj)
+		{
+			NumFailedResolves++;
+		}
+
+		if (PrevResolvedFunc)
+		{
+			PrevResolvedFunc(ObjectRef, Pkg, Obj);
+		}
+	}
+	static void OnRefRead(UObject* Obj)
+	{
+		NumReads++;
+		if (PrevReadFunc)
+		{
+			PrevReadFunc(Obj);
+		}
+	}
+#endif
+	
+	static void ConditionalInstallCallbacks()
+	{
+		static bool bCallbacksInstalled = false;
+		static FCriticalSection CallbackInstallationLock;
+
+		if (bCallbacksInstalled)
+		{
+			return;
+		}
+
+		FScopeLock ScopeLock(&CallbackInstallationLock);
+		if (!bCallbacksInstalled)
+		{
+#if UE_WITH_OBJECT_HANDLE_TRACKING
+			PrevResolvedFunc = SetObjectHandleReferenceResolvedCallback(OnRefResolved);
+			PrevReadFunc = SetObjectHandleReadCallback(OnRefRead);
+#endif
+			bCallbacksInstalled = true;
+		}
+	}
+
+#if UE_WITH_OBJECT_HANDLE_TRACKING
+	static ObjectHandleReferenceResolvedFunction* PrevResolvedFunc;
+	static ObjectHandleReadFunction* PrevReadFunc;
+#endif
+	static thread_local uint32 NumResolves;
+	static thread_local uint32 NumFailedResolves;
+	static thread_local uint32 NumReads;
+};
+#endif // WITH_DEV_AUTOMATION_TESTS
diff --git a/Engine/Source/Runtime/CoreUObject/Tests/Serialization/BulkDataTests.cpp b/Engine/Source/Runtime/CoreUObject/Tests/Serialization/BulkDataTests.cpp
new file mode 100644
--- /dev/null
+++ b/Engine/Source/Runtime/CoreUObject/Tests/Serialization/BulkDataTests.cpp
@@ -0,0 +1,77 @@
+// Copyright Epic Games, Inc. All Rights Reserved.
+
+#include "Serialization/BulkData.h"
+#include "TestHarness.h"
+
+namespace BulkDataTest
+{
+	// Test code paths for BulkData objects that do not reference a file on disk.
+	TEST_CASE("CoreUObject::Serialization::FByteBulkData::Transient", "[CoreUObject][Serialization][Smoke]")
+	{
+		FByteBulkData BulkData;
+
+		// We should be able to lock for read access but there should be no valid data
+		const void* ReadOnlyDataPtr = BulkData.LockReadOnly();
+		TestNull(TEXT("Locking an empty BulkData object for reading should return nullptr!"), ReadOnlyDataPtr);
+		BulkData.Unlock();
+
+		void* DataPtr = BulkData.Lock(LOCK_READ_WRITE);
+		TestNull(TEXT("Locking an empty BulkData object for writing should return nullptr!"), DataPtr);
+		BulkData.Unlock();
+
+		void* CopyEmptyPtr = nullptr;
+		BulkData.GetCopy(&CopyEmptyPtr, true);
+		TestNull(TEXT("Getting a copy of an empty BulkData object for writing should return nullptr!"), CopyEmptyPtr);
+
+		BulkData.Lock(LOCK_READ_WRITE);
+		DataPtr = BulkData.Realloc(32 * 32 * 4);
+		TestNotNull(TEXT("Reallocating an empty BulkData object should return a valid pointer!"), DataPtr);
+		BulkData.Unlock();
+
+		TestTrue(TEXT("BulkData should be loaded now that it has been reallocated"), BulkData.IsBulkDataLoaded());
+
+		void* CopyWithDiscardPtr = nullptr;
+		BulkData.GetCopy(&CopyWithDiscardPtr, true); // The second parameter should be ignored because the bulkdata cannot be reloaded from disk!
+		TestNotNull(TEXT("GetCopy should return a valid pointer!"), CopyWithDiscardPtr);
+		TestNotEqual(TEXT("GetCopy should return a copy of the data so the pointers should be different!"), DataPtr, CopyWithDiscardPtr);
+		TestTrue(TEXT("BulkData should still loaded after taking a copy"), BulkData.IsBulkDataLoaded());
+
+		// Now try GetCopy again but this time without the discard request
+		void* CopyNoDiscardPtr = nullptr;
+		BulkData.GetCopy(&CopyNoDiscardPtr, false);
+		TestNotNull(TEXT("GetCopy should return a valid pointer!"), CopyNoDiscardPtr);
+		TestNotEqual(TEXT("GetCopy should return a copy of the data so the pointers should be different!"), DataPtr, CopyNoDiscardPtr);
+		TestNotEqual(TEXT("GetCopy should return a copy of the data so the pointers should be different!"), CopyWithDiscardPtr, CopyNoDiscardPtr);
+		TestTrue(TEXT("BulkData should still loaded after taking a copy"), BulkData.IsBulkDataLoaded());
+
+		// Clean up allocations from GetCopy
+		FMemory::Free(CopyWithDiscardPtr);
+		FMemory::Free(CopyNoDiscardPtr);
+
+		// Now do one last lock test after GetCopy
+		DataPtr = BulkData.Lock(LOCK_READ_WRITE);
+		BulkData.Unlock();
+
+		TestTrue(TEXT("BulkData should still loaded after locking for write"), BulkData.IsBulkDataLoaded());
+		TestNotNull(TEXT("Locking for write should return a valid pointer!"), DataPtr);
+
+		// Now remove the bulkdata and make sure that we cannot access the old data anymore
+		BulkData.RemoveBulkData();
+		TestFalse(TEXT("RemoveBulkData should've discarded the BulkData"), BulkData.IsBulkDataLoaded());
+
+#if USE_NEW_BULKDATA // The warning only exists for BulkData2
+		// Both the ::Lock and ::GetCopy calls should warn that we cannot load the missing data (as it will still have a valid size)
+		TestAddWarning(TEXT("Attempting to load a BulkData object that cannot be loaded from disk"));
+#endif //USE_NEW_BULKDATA
+
+		DataPtr = BulkData.Lock(LOCK_READ_WRITE);
+		BulkData.Unlock();
+
+		TestNull(TEXT("Locking for write after calling ::RemoveBulkData should return a nullptr!"), DataPtr);
+
+		CopyEmptyPtr = nullptr;
+		BulkData.GetCopy(&CopyNoDiscardPtr, true);
+		TestNull(TEXT("Getting a copy of BulkData object after calling ::RemoveBulkData should return nullptr!"), DataPtr);
+	}
+}
+
