Fix Replicated Bitfields being incompatible across compilers/platforms

diff --git a/Engine/Source/Runtime/CoreUObject/Private/UObject/CoreNet.cpp b/Engine/Source/Runtime/CoreUObject/Private/UObject/CoreNet.cpp
--- a/Engine/Source/Runtime/CoreUObject/Private/UObject/CoreNet.cpp
+++ b/Engine/Source/Runtime/CoreUObject/Private/UObject/CoreNet.cpp
@@ -230,8 +230,12 @@ const FClassNetCache* FClassNetCacheMgr::GetClassNetCache( UClass* Class )
 			Result->FieldChecksumMap.Add( It->FieldChecksum, &*It );
 		}
 
-		// Initialize class checksum (just use properties for this)
-		SortProperties( Properties );
+		const bool bIsNativeClass = Class->HasAnyClassFlags(CLASS_Native);
+		if (!bIsNativeClass)
+		{
+			// Initialize class checksum (just use properties for this)
+			SortProperties(Properties);
+		}
 
 		for ( auto Property : Properties )
 		{
