Implementing GetObjectPropertyValue_InContainer that can use property getters

diff --git a/Engine/Source/Runtime/CoreUObject/Private/UObject/PropertyBaseObject.cpp b/Engine/Source/Runtime/CoreUObject/Private/UObject/PropertyBaseObject.cpp
--- a/Engine/Source/Runtime/CoreUObject/Private/UObject/PropertyBaseObject.cpp
+++ b/Engine/Source/Runtime/CoreUObject/Private/UObject/PropertyBaseObject.cpp
@@ -567,7 +567,13 @@ FName FObjectPropertyBase::GetID() const
 
 UObject* FObjectPropertyBase::GetObjectPropertyValue(const void* PropertyValueAddress) const
 {
-	check(0);
+	checkf(false, TEXT("%s is missing implementation of GetObjectPropertyValue"), *GetFullName());
+	return nullptr;
+}
+
+UObject* FObjectPropertyBase::GetObjectPropertyValue_InContainer(const void* ContainerAddress, int32 ArrayIndex) const
+{
+	checkf(false, TEXT("%s is missing implementation of GetObjectPropertyValue_InContainer"), *GetFullName());
 	return nullptr;
 }
 
diff --git a/Engine/Source/Runtime/CoreUObject/Private/UObject/PropertyClassPtr.cpp b/Engine/Source/Runtime/CoreUObject/Private/UObject/PropertyClassPtr.cpp
--- a/Engine/Source/Runtime/CoreUObject/Private/UObject/PropertyClassPtr.cpp
+++ b/Engine/Source/Runtime/CoreUObject/Private/UObject/PropertyClassPtr.cpp
@@ -51,6 +51,11 @@ UObject* FClassPtrProperty::GetObjectPropertyValue(const void* PropertyValueAddr
 	return ((FObjectPtr&)GetPropertyValue(PropertyValueAddress)).Get();
 }
 
+UObject* FClassPtrProperty::GetObjectPropertyValue_InContainer(const void* ContainerAddress, int32 ArrayIndex) const
+{
+	return GetWrappedObjectPropertyValue_InContainer<FObjectPtr>(ContainerAddress, ArrayIndex);
+}
+
 void FClassPtrProperty::SetObjectPropertyValue(void* PropertyValueAddress, UObject* Value) const
 {
 	SetPropertyValue(PropertyValueAddress, TCppType(Value));
diff --git a/Engine/Source/Runtime/CoreUObject/Private/UObject/PropertyLazyObjectPtr.cpp b/Engine/Source/Runtime/CoreUObject/Private/UObject/PropertyLazyObjectPtr.cpp
--- a/Engine/Source/Runtime/CoreUObject/Private/UObject/PropertyLazyObjectPtr.cpp
+++ b/Engine/Source/Runtime/CoreUObject/Private/UObject/PropertyLazyObjectPtr.cpp
@@ -110,6 +110,11 @@ UObject* FLazyObjectProperty::GetObjectPropertyValue(const void* PropertyValueAd
 	return GetPropertyValue(PropertyValueAddress).Get();
 }
 
+UObject* FLazyObjectProperty::GetObjectPropertyValue_InContainer(const void* ContainerAddress, int32 ArrayIndex) const
+{
+	return GetWrappedObjectPropertyValue_InContainer<FLazyObjectPtr>(ContainerAddress, ArrayIndex);
+}
+
 void FLazyObjectProperty::SetObjectPropertyValue(void* PropertyValueAddress, UObject* Value) const
 {
 	SetPropertyValue(PropertyValueAddress, TCppType(Value));
diff --git a/Engine/Source/Runtime/CoreUObject/Private/UObject/PropertyObject.cpp b/Engine/Source/Runtime/CoreUObject/Private/UObject/PropertyObject.cpp
--- a/Engine/Source/Runtime/CoreUObject/Private/UObject/PropertyObject.cpp
+++ b/Engine/Source/Runtime/CoreUObject/Private/UObject/PropertyObject.cpp
@@ -183,6 +183,20 @@ UObject* FObjectProperty::GetObjectPropertyValue(const void* PropertyValueAddres
 	return GetPropertyValue(PropertyValueAddress);
 }
 
+UObject* FObjectProperty::GetObjectPropertyValue_InContainer(const void* ContainerAddress, int32 ArrayIndex) const
+{
+	// Dummy struct to satisfy GetWrappedObjectPropertyValue_InContainer API requirements (Get() function)
+	struct FWrappedObjectPtr
+	{
+		UObject* Obj = nullptr;
+		UObject* Get() const
+		{
+			return Obj;
+		}
+	};
+	return GetWrappedObjectPropertyValue_InContainer<FWrappedObjectPtr>(ContainerAddress, ArrayIndex);
+}
+
 void FObjectProperty::SetObjectPropertyValue(void* PropertyValueAddress, UObject* Value) const
 {
 	SetPropertyValue(PropertyValueAddress, Value);
diff --git a/Engine/Source/Runtime/CoreUObject/Private/UObject/PropertyObjectPtr.cpp b/Engine/Source/Runtime/CoreUObject/Private/UObject/PropertyObjectPtr.cpp
--- a/Engine/Source/Runtime/CoreUObject/Private/UObject/PropertyObjectPtr.cpp
+++ b/Engine/Source/Runtime/CoreUObject/Private/UObject/PropertyObjectPtr.cpp
@@ -142,6 +142,11 @@ UObject* FObjectPtrProperty::GetObjectPropertyValue(const void* PropertyValueAdd
 	return ((FObjectPtr&)GetPropertyValue(PropertyValueAddress)).Get();
 }
 
+UObject* FObjectPtrProperty::GetObjectPropertyValue_InContainer(const void* ContainerAddress, int32 ArrayIndex) const
+{
+	return GetWrappedObjectPropertyValue_InContainer<FObjectPtr>(ContainerAddress, ArrayIndex);
+}
+
 void FObjectPtrProperty::SetObjectPropertyValue(void* PropertyValueAddress, UObject* Value) const
 {
 	SetPropertyValue(PropertyValueAddress, TCppType(Value));
diff --git a/Engine/Source/Runtime/CoreUObject/Private/UObject/PropertySoftObjectPtr.cpp b/Engine/Source/Runtime/CoreUObject/Private/UObject/PropertySoftObjectPtr.cpp
--- a/Engine/Source/Runtime/CoreUObject/Private/UObject/PropertySoftObjectPtr.cpp
+++ b/Engine/Source/Runtime/CoreUObject/Private/UObject/PropertySoftObjectPtr.cpp
@@ -236,6 +236,11 @@ UObject* FSoftObjectProperty::GetObjectPropertyValue(const void* PropertyValueAd
 	return GetPropertyValue(PropertyValueAddress).Get();
 }
 
+UObject* FSoftObjectProperty::GetObjectPropertyValue_InContainer(const void* ContainerAddress, int32 ArrayIndex) const
+{
+	return GetWrappedObjectPropertyValue_InContainer<FSoftObjectPtr>(ContainerAddress, ArrayIndex);
+}
+
 void FSoftObjectProperty::SetObjectPropertyValue(void* PropertyValueAddress, UObject* Value) const
 {
 	SetPropertyValue(PropertyValueAddress, TCppType(Value));
diff --git a/Engine/Source/Runtime/CoreUObject/Private/UObject/PropertyWeakObjectPtr.cpp b/Engine/Source/Runtime/CoreUObject/Private/UObject/PropertyWeakObjectPtr.cpp
--- a/Engine/Source/Runtime/CoreUObject/Private/UObject/PropertyWeakObjectPtr.cpp
+++ b/Engine/Source/Runtime/CoreUObject/Private/UObject/PropertyWeakObjectPtr.cpp
@@ -63,6 +63,11 @@ UObject* FWeakObjectProperty::GetObjectPropertyValue(const void* PropertyValueAd
 	return GetPropertyValue(PropertyValueAddress).Get();
 }
 
+UObject* FWeakObjectProperty::GetObjectPropertyValue_InContainer(const void* ContainerAddress, int32 ArrayIndex) const
+{
+	return GetWrappedObjectPropertyValue_InContainer<FWeakObjectPtr>(ContainerAddress, ArrayIndex);
+}
+
 void FWeakObjectProperty::SetObjectPropertyValue(void* PropertyValueAddress, UObject* Value) const
 {
 	SetPropertyValue(PropertyValueAddress, TCppType(Value));
diff --git a/Engine/Source/Runtime/CoreUObject/Public/UObject/UnrealType.h b/Engine/Source/Runtime/CoreUObject/Public/UObject/UnrealType.h
--- a/Engine/Source/Runtime/CoreUObject/Public/UObject/UnrealType.h
+++ b/Engine/Source/Runtime/CoreUObject/Public/UObject/UnrealType.h
@@ -2653,10 +2653,7 @@ public:
 		return LoadObjectPropertyValue(ContainerPtrToValuePtr<void>(PropertyValueAddress, ArrayIndex));
 	}
 	virtual UObject* GetObjectPropertyValue(const void* PropertyValueAddress) const;
-	FORCEINLINE UObject* GetObjectPropertyValue_InContainer(const void* PropertyValueAddress, int32 ArrayIndex = 0) const
-	{
-		return GetObjectPropertyValue(ContainerPtrToValuePtr<void>(PropertyValueAddress, ArrayIndex));
-	}
+	virtual UObject* GetObjectPropertyValue_InContainer(const void* PropertyValueAddress, int32 ArrayIndex = 0) const;
 	virtual void SetObjectPropertyValue(void* PropertyValueAddress, UObject* Value) const;
 	FORCEINLINE void SetObjectPropertyValue_InContainer(void* PropertyValueAddress, UObject* Value, int32 ArrayIndex = 0) const
 	{
@@ -2684,6 +2681,36 @@ protected:
 	virtual bool AllowCrossLevel() const;
 	virtual bool AllowObjectTypeReinterpretationTo(const FObjectPropertyBase* Other) const;
 	// End of FObjectPropertyBase interface
+
+	/* Helper function for UObject property types that wrap the object pointer in a smart pointer */
+	template <typename T>
+	UObject* GetWrappedObjectPropertyValue_InContainer(const void* ContainerAddress, int32 ArrayIndex) const
+	{
+		if (!HasGetter())
+		{
+			// Fast path - direct memory access
+			return GetObjectPropertyValue(ContainerPtrToValuePtr<void>(ContainerAddress, ArrayIndex));
+		}
+		else
+		{
+			T Value;
+			if (ArrayDim == 1)
+			{
+				// Slower but no mallocs. We can copy the value directly to the resulting param
+				GetValue_InContainer(ContainerAddress, &Value);
+			}
+			else
+			{
+				// Malloc a temp value that is the size of the array. Getter will then copy the entire array to the temp value
+				T* ValueArray = (T*)AllocateAndInitializeValue();
+				FProperty::GetValue_InContainer(ContainerAddress, ValueArray);
+				// Grab the item we care about and free the temp array
+				Value = ValueArray[ArrayIndex];
+				DestroyAndFreeValue(ValueArray);
+			}
+			return Value.Get();
+		}
+	}
 };
 
 template<typename InTCppType>
@@ -2814,6 +2841,7 @@ public:
 
 	// FObjectPropertyBase interface
 	virtual UObject* GetObjectPropertyValue(const void* PropertyValueAddress) const override;
+	virtual UObject* GetObjectPropertyValue_InContainer(const void* ContainerAddress, int32 ArrayIndex = 0) const override;
 	virtual void SetObjectPropertyValue(void* PropertyValueAddress, UObject* Value) const override;
 	virtual FString GetCPPTypeCustom(FString* ExtendedTypeText, uint32 CPPExportFlags, const FString& InnerNativeTypeName)  const override;
 	// End of FObjectPropertyBase interface
@@ -2851,6 +2879,7 @@ class COREUOBJECT_API FObjectPtrProperty : public FObjectProperty
 
 	// FObjectProperty interface
 	virtual UObject* GetObjectPropertyValue(const void* PropertyValueAddress) const override;
+	virtual UObject* GetObjectPropertyValue_InContainer(const void* ContainerAddress, int32 ArrayIndex = 0) const override;
 	virtual void SetObjectPropertyValue(void* PropertyValueAddress, UObject* Value) const override;
 	virtual bool AllowCrossLevel() const override;
 	virtual bool AllowObjectTypeReinterpretationTo(const FObjectPropertyBase* Other) const override;
@@ -2911,6 +2940,7 @@ public:
 
 	// FObjectProperty interface
 	virtual UObject* GetObjectPropertyValue(const void* PropertyValueAddress) const override;
+	virtual UObject* GetObjectPropertyValue_InContainer(const void* ContainerAddress, int32 ArrayIndex = 0) const override;
 	virtual void SetObjectPropertyValue(void* PropertyValueAddress, UObject* Value) const override;
 	// End of FObjectProperty interface
 };
@@ -2966,6 +2996,7 @@ PRAGMA_ENABLE_DEPRECATION_WARNINGS
 
 	// FObjectProperty interface
 	virtual UObject* GetObjectPropertyValue(const void* PropertyValueAddress) const override;
+	virtual UObject* GetObjectPropertyValue_InContainer(const void* ContainerAddress, int32 ArrayIndex = 0) const override;
 	virtual void SetObjectPropertyValue(void* PropertyValueAddress, UObject* Value) const override;
 	virtual bool AllowCrossLevel() const override;
 private:
@@ -3037,6 +3068,7 @@ public:
 	// FObjectProperty interface
 	virtual UObject* LoadObjectPropertyValue(const void* PropertyValueAddress) const override;
 	virtual UObject* GetObjectPropertyValue(const void* PropertyValueAddress) const override;
+	virtual UObject* GetObjectPropertyValue_InContainer(const void* ContainerAddress, int32 ArrayIndex = 0) const override;
 	virtual void SetObjectPropertyValue(void* PropertyValueAddress, UObject* Value) const override;
 	virtual bool AllowCrossLevel() const override;
 	virtual FString GetCPPTypeCustom(FString* ExtendedTypeText, uint32 CPPExportFlags, const FString& InnerNativeTypeName)  const override;
@@ -3178,6 +3210,7 @@ class COREUOBJECT_API FClassPtrProperty : public FClassProperty
 
 	// FObjectProperty interface
 	virtual UObject* GetObjectPropertyValue(const void* PropertyValueAddress) const override;
+	virtual UObject* GetObjectPropertyValue_InContainer(const void* ContainerAddress, int32 ArrayIndex = 0) const override;
 	virtual void SetObjectPropertyValue(void* PropertyValueAddress, UObject* Value) const override;
 private:
 	virtual uint32 GetValueTypeHashInternal(const void* Src) const override;
