Comment fix: Fix copy-paste error and add details in IsShortPackageName function comment.
#rnx

diff --git a/Engine/Source/Runtime/CoreUObject/Public/Misc/PackageName.h b/Engine/Source/Runtime/CoreUObject/Public/Misc/PackageName.h
--- a/Engine/Source/Runtime/CoreUObject/Public/Misc/PackageName.h
+++ b/Engine/Source/Runtime/CoreUObject/Public/Misc/PackageName.h
@@ -233,10 +233,12 @@ public:
 	static bool IsValidPath(const FString& InPath);
 
 	/**
-	 * Checks if the given string is a long package name or not.
+	 * Checks if the string is a ShortPackageName. A ShortPackageName is the leaf name after the last
+	 * slash in a LongPackageName. Handling ShortPackageNames is useful for console commands and other UI.
+	 * A ShortPackageName requires a search to convert to a LongPackageName.
 	 *
 	 * @param PossiblyLongName Package name.
-	 * @return true if the given name is a long package name, false otherwise.
+	 * @return true if the given name is a short package name (contains no slashes), false otherwise.
 	 */
 	static bool IsShortPackageName(const FString& PossiblyLongName);
 	static bool IsShortPackageName(const FName PossiblyLongName);
