Adding Get/SetSingleValue_InContainer functions and making SetObjectPropertyValue_InContainer support setters and getters

diff --git a/Engine/Source/Runtime/CoreUObject/Private/UObject/Property.cpp b/Engine/Source/Runtime/CoreUObject/Private/UObject/Property.cpp
--- a/Engine/Source/Runtime/CoreUObject/Private/UObject/Property.cpp
+++ b/Engine/Source/Runtime/CoreUObject/Private/UObject/Property.cpp
@@ -1282,6 +1282,62 @@ void* FProperty::GetValueAddressAtIndex_Direct(const FProperty* Inner, void* InV
 	return (uint8*)InValueAddress + ElementSize * Index;
 }
 
+void FProperty::SetSingleValue_InContainer(void* InContainer, const void* InValue, int32 ArrayIndex) const
+{
+	checkf(ArrayIndex <= ArrayDim, TEXT("ArrayIndex (%d) must be less than the property %s array size (%d)"), ArrayIndex, *GetFullName(), ArrayDim);
+	if (!HasSetter())
+	{
+		// Fast path - direct memory access
+		CopySingleValue(ContainerVoidPtrToValuePtrInternal((void*)InContainer, ArrayIndex), InValue);
+	}
+	else
+	{
+		if (ArrayDim == 1)
+		{
+			// Slower but no mallocs. We can copy the value directly to the resulting param
+			CallSetter(InContainer, InValue);
+		}
+		else
+		{
+			// Malloc a temp value that is the size of the array. We will then copy the entire array to the temp value
+			uint8* ValueArray = (uint8*)AllocateAndInitializeValue();
+			GetValue_InContainer(InContainer, ValueArray);
+			// Replace the value at the specified index in the temp array with the InValue
+			CopySingleValue(ValueArray + ArrayIndex * ElementSize, InValue);
+			// Now call a setter to replace the entire array and then destroy the temp value
+			CallSetter(InContainer, ValueArray);
+			DestroyAndFreeValue(ValueArray);
+		}
+	}
+}
+
+void FProperty::GetSingleValue_InContainer(const void* InContainer, void* OutValue, int32 ArrayIndex) const
+{
+	checkf(ArrayIndex <= ArrayDim, TEXT("ArrayIndex (%d) must be less than the property %s array size (%d)"), ArrayIndex, *GetFullName(), ArrayDim);
+	if (!HasGetter())
+	{
+		// Fast path - direct memory access
+		CopySingleValue(OutValue, ContainerVoidPtrToValuePtrInternal((void*)InContainer, ArrayIndex));
+	}
+	else
+	{
+		if (ArrayDim == 1)
+		{
+			// Slower but no mallocs. We can copy the value directly to the resulting param
+			CallGetter(InContainer, OutValue);
+		}
+		else
+		{
+			// Malloc a temp value that is the size of the array. Getter will then copy the entire array to the temp value
+			uint8* ValueArray = (uint8*)AllocateAndInitializeValue();
+			GetValue_InContainer(InContainer, ValueArray);
+			// Copy the item we care about and free the temp array
+			CopySingleValue(OutValue, ValueArray + ArrayIndex * ElementSize);
+			DestroyAndFreeValue(ValueArray);
+		}
+	}
+}
+
 void FProperty::PerformOperationWithSetter(void* InContainer, void* DirectPropertyAddress, TFunctionRef<void(void*)> DirectValueAccessFunc) const
 {
 	if (InContainer && HasSetterOrGetter()) // If there's a getter we need to allocate a temp value even if there's no setter
diff --git a/Engine/Source/Runtime/CoreUObject/Private/UObject/PropertyBaseObject.cpp b/Engine/Source/Runtime/CoreUObject/Private/UObject/PropertyBaseObject.cpp
--- a/Engine/Source/Runtime/CoreUObject/Private/UObject/PropertyBaseObject.cpp
+++ b/Engine/Source/Runtime/CoreUObject/Private/UObject/PropertyBaseObject.cpp
@@ -577,6 +577,11 @@ UObject* FObjectPropertyBase::GetObjectPropertyValue_InContainer(const void* Con
 	return nullptr;
 }
 
+void FObjectPropertyBase::SetObjectPropertyValue_InContainer(void* ContainerAddress, UObject* Value, int32 ArrayIndex) const
+{
+	checkf(false, TEXT("%s is missing implementation of SetObjectPropertyValue_InContainer"), *GetFullName());
+}
+
 void FObjectPropertyBase::SetObjectPropertyValue(void* PropertyValueAddress, UObject* Value) const
 {
 	check(0);
diff --git a/Engine/Source/Runtime/CoreUObject/Private/UObject/PropertyClassPtr.cpp b/Engine/Source/Runtime/CoreUObject/Private/UObject/PropertyClassPtr.cpp
--- a/Engine/Source/Runtime/CoreUObject/Private/UObject/PropertyClassPtr.cpp
+++ b/Engine/Source/Runtime/CoreUObject/Private/UObject/PropertyClassPtr.cpp
@@ -61,6 +61,11 @@ void FClassPtrProperty::SetObjectPropertyValue(void* PropertyValueAddress, UObje
 	SetPropertyValue(PropertyValueAddress, TCppType(Value));
 }
 
+void FClassPtrProperty::SetObjectPropertyValue_InContainer(void* ContainerAddress, UObject* Value, int32 ArrayIndex) const
+{
+	SetWrappedObjectPropertyValue_InContainer<FObjectPtr>(ContainerAddress, Value, ArrayIndex);
+}
+
 uint32 FClassPtrProperty::GetValueTypeHashInternal(const void* Src) const
 {
 	return GetTypeHash((FObjectPtr&)GetPropertyValue(Src));
diff --git a/Engine/Source/Runtime/CoreUObject/Private/UObject/PropertyLazyObjectPtr.cpp b/Engine/Source/Runtime/CoreUObject/Private/UObject/PropertyLazyObjectPtr.cpp
--- a/Engine/Source/Runtime/CoreUObject/Private/UObject/PropertyLazyObjectPtr.cpp
+++ b/Engine/Source/Runtime/CoreUObject/Private/UObject/PropertyLazyObjectPtr.cpp
@@ -120,6 +120,11 @@ void FLazyObjectProperty::SetObjectPropertyValue(void* PropertyValueAddress, UOb
 	SetPropertyValue(PropertyValueAddress, TCppType(Value));
 }
 
+void FLazyObjectProperty::SetObjectPropertyValue_InContainer(void* ContainerAddress, UObject* Value, int32 ArrayIndex) const
+{
+	SetWrappedObjectPropertyValue_InContainer<FLazyObjectPtr>(ContainerAddress, Value, ArrayIndex);
+}
+
 bool FLazyObjectProperty::AllowCrossLevel() const
 {
 	return true;
diff --git a/Engine/Source/Runtime/CoreUObject/Private/UObject/PropertyNumeric.cpp b/Engine/Source/Runtime/CoreUObject/Private/UObject/PropertyNumeric.cpp
--- a/Engine/Source/Runtime/CoreUObject/Private/UObject/PropertyNumeric.cpp
+++ b/Engine/Source/Runtime/CoreUObject/Private/UObject/PropertyNumeric.cpp
@@ -243,6 +243,11 @@ uint64 FNumericProperty::GetUnsignedIntPropertyValue(void const* Data) const
 	check(0);
 	return 0;
 }
+uint64 FNumericProperty::GetUnsignedIntPropertyValue_InContainer(void const* Data) const
+{
+	check(0);
+	return 0;
+}
 
 /** 
 	* Gets the value of an floating point property type
diff --git a/Engine/Source/Runtime/CoreUObject/Private/UObject/PropertyObject.cpp b/Engine/Source/Runtime/CoreUObject/Private/UObject/PropertyObject.cpp
--- a/Engine/Source/Runtime/CoreUObject/Private/UObject/PropertyObject.cpp
+++ b/Engine/Source/Runtime/CoreUObject/Private/UObject/PropertyObject.cpp
@@ -201,3 +201,18 @@ void FObjectProperty::SetObjectPropertyValue(void* PropertyValueAddress, UObject
 {
 	SetPropertyValue(PropertyValueAddress, Value);
 }
+
+void FObjectProperty::SetObjectPropertyValue_InContainer(void* ContainerAddress, UObject* Value, int32 ArrayIndex) const
+{
+	// Dummy struct to satisfy SetWrappedObjectPropertyValue_InContainer API requirements
+	struct FWrappedObjectPtr
+	{
+		FWrappedObjectPtr() = default;
+		FWrappedObjectPtr(UObject* InObj)
+			: Obj(InObj)
+		{
+		}
+		UObject* Obj = nullptr;
+	};
+	SetWrappedObjectPropertyValue_InContainer<FWrappedObjectPtr>(ContainerAddress, Value, ArrayIndex);
+}
diff --git a/Engine/Source/Runtime/CoreUObject/Private/UObject/PropertyObjectPtr.cpp b/Engine/Source/Runtime/CoreUObject/Private/UObject/PropertyObjectPtr.cpp
--- a/Engine/Source/Runtime/CoreUObject/Private/UObject/PropertyObjectPtr.cpp
+++ b/Engine/Source/Runtime/CoreUObject/Private/UObject/PropertyObjectPtr.cpp
@@ -152,6 +152,11 @@ void FObjectPtrProperty::SetObjectPropertyValue(void* PropertyValueAddress, UObj
 	SetPropertyValue(PropertyValueAddress, TCppType(Value));
 }
 
+void FObjectPtrProperty::SetObjectPropertyValue_InContainer(void* ContainerAddress, UObject* Value, int32 ArrayIndex) const
+{
+	SetWrappedObjectPropertyValue_InContainer<FObjectPtr>(ContainerAddress, Value, ArrayIndex);
+}
+
 bool FObjectPtrProperty::AllowCrossLevel() const
 {
 	return true;
diff --git a/Engine/Source/Runtime/CoreUObject/Private/UObject/PropertySoftObjectPtr.cpp b/Engine/Source/Runtime/CoreUObject/Private/UObject/PropertySoftObjectPtr.cpp
--- a/Engine/Source/Runtime/CoreUObject/Private/UObject/PropertySoftObjectPtr.cpp
+++ b/Engine/Source/Runtime/CoreUObject/Private/UObject/PropertySoftObjectPtr.cpp
@@ -246,6 +246,11 @@ void FSoftObjectProperty::SetObjectPropertyValue(void* PropertyValueAddress, UOb
 	SetPropertyValue(PropertyValueAddress, TCppType(Value));
 }
 
+void FSoftObjectProperty::SetObjectPropertyValue_InContainer(void* ContainerAddress, UObject* Value, int32 ArrayIndex) const
+{
+	SetWrappedObjectPropertyValue_InContainer<FSoftObjectPtr>(ContainerAddress, Value, ArrayIndex);
+}
+
 bool FSoftObjectProperty::AllowCrossLevel() const
 {
 	return true;
diff --git a/Engine/Source/Runtime/CoreUObject/Private/UObject/PropertyWeakObjectPtr.cpp b/Engine/Source/Runtime/CoreUObject/Private/UObject/PropertyWeakObjectPtr.cpp
--- a/Engine/Source/Runtime/CoreUObject/Private/UObject/PropertyWeakObjectPtr.cpp
+++ b/Engine/Source/Runtime/CoreUObject/Private/UObject/PropertyWeakObjectPtr.cpp
@@ -73,6 +73,11 @@ void FWeakObjectProperty::SetObjectPropertyValue(void* PropertyValueAddress, UOb
 	SetPropertyValue(PropertyValueAddress, TCppType(Value));
 }
 
+void FWeakObjectProperty::SetObjectPropertyValue_InContainer(void* ContainerAddress, UObject* Value, int32 ArrayIndex) const
+{
+	SetWrappedObjectPropertyValue_InContainer<FWeakObjectPtr>(ContainerAddress, Value, ArrayIndex);
+}
+
 uint32 FWeakObjectProperty::GetValueTypeHashInternal(const void* Src) const
 {
 	return GetTypeHash(*(FWeakObjectPtr*)Src);
diff --git a/Engine/Source/Runtime/CoreUObject/Public/UObject/UnrealType.h b/Engine/Source/Runtime/CoreUObject/Public/UObject/UnrealType.h
--- a/Engine/Source/Runtime/CoreUObject/Public/UObject/UnrealType.h
+++ b/Engine/Source/Runtime/CoreUObject/Public/UObject/UnrealType.h
@@ -557,6 +557,22 @@ public:
 		}
 	}
 
+	/**
+	* Copies a single value to the property even if the property represents a static array of values
+	* @param InContainer Instance owner of the property
+	* @param InValue Pointer to the memory that the value will be copied from. Must be at least ElementSize big
+	* @param ArrayIndex Index into the static array to copy the value from. If the property is not a static array it should be 0
+	*/
+	void SetSingleValue_InContainer(void* InContainer, const void* InValue, int32 ArrayIndex) const;
+
+	/**
+	* Copies a single value to OutValue even if the property represents a static array of values
+	* @param InContainer Instance owner of the property
+	* @param OutValue Pointer to the memory that the value will be copied to. Must be at least ElementSize big
+	* @param ArrayIndex Index into the static array to copy the value from. If the property is not a static array it should be 0
+	*/
+	void GetSingleValue_InContainer(const void* InContainer, void* OutValue, int32 ArrayIndex) const;
+
 	/** Allocates and initializes memory to hold a value this property represents */
 	void* AllocateAndInitializeValue() const;
 
@@ -1618,6 +1634,7 @@ PRAGMA_ENABLE_DEPRECATION_WARNINGS
 	 * @return Data as an unsigned int
 	**/
 	virtual uint64 GetUnsignedIntPropertyValue(void const* Data) const;
+	virtual uint64 GetUnsignedIntPropertyValue_InContainer(void const* Container) const;
 
 	/** 
 	 * Gets the value of an floating point property type
@@ -1874,6 +1891,13 @@ public:
 		check(TIsIntegral<TCppType>::Value);
 		return (uint64)TTypeFundamentals::GetPropertyValue(Data);
 	}
+	virtual uint64 GetUnsignedIntPropertyValue_InContainer(void const* Container) const override
+	{
+		check(TIsIntegral<TCppType>::Value);
+		TCppType LocalValue{};
+		FNumericProperty::GetValue_InContainer(Container, &LocalValue);
+		return (uint64)LocalValue;
+	}
 	virtual double GetFloatingPointPropertyValue(void const* Data) const override
 	{
 		check(TIsFloatingPoint<TCppType>::Value);
@@ -2653,12 +2677,9 @@ public:
 		return LoadObjectPropertyValue(ContainerPtrToValuePtr<void>(PropertyValueAddress, ArrayIndex));
 	}
 	virtual UObject* GetObjectPropertyValue(const void* PropertyValueAddress) const;
-	virtual UObject* GetObjectPropertyValue_InContainer(const void* PropertyValueAddress, int32 ArrayIndex = 0) const;
+	virtual UObject* GetObjectPropertyValue_InContainer(const void* ContainerAddress, int32 ArrayIndex = 0) const;
 	virtual void SetObjectPropertyValue(void* PropertyValueAddress, UObject* Value) const;
-	FORCEINLINE void SetObjectPropertyValue_InContainer(void* PropertyValueAddress, UObject* Value, int32 ArrayIndex = 0) const
-	{
-		SetObjectPropertyValue(ContainerPtrToValuePtr<void>(PropertyValueAddress, ArrayIndex), Value);
-	}
+	virtual void SetObjectPropertyValue_InContainer(void* ContainerAddress, UObject* Value, int32 ArrayIndex = 0) const;
 
 	/**
 	 * Setter function for this property's PropertyClass member. Favor this 
@@ -2682,7 +2703,7 @@ protected:
 	virtual bool AllowObjectTypeReinterpretationTo(const FObjectPropertyBase* Other) const;
 	// End of FObjectPropertyBase interface
 
-	/* Helper function for UObject property types that wrap the object pointer in a smart pointer */
+	/* Helper functions for UObject property types that wrap the object pointer in a smart pointer */
 	template <typename T>
 	UObject* GetWrappedObjectPropertyValue_InContainer(const void* ContainerAddress, int32 ArrayIndex) const
 	{
@@ -2711,6 +2732,35 @@ protected:
 			return Value.Get();
 		}
 	}
+	template <typename T>
+	void SetWrappedObjectPropertyValue_InContainer(void* ContainerAddress, UObject* InValue, int32 ArrayIndex) const
+	{
+		if (!HasSetter())
+		{
+			// Fast path - direct memory access
+			SetObjectPropertyValue(ContainerPtrToValuePtr<void>(ContainerAddress, ArrayIndex), InValue);
+		}
+		else
+		{			
+			if (ArrayDim == 1)
+			{
+				// Slower but no mallocs. We can copy a local wrapped value directly to the resulting param
+				T WrappedValue(InValue);
+				SetValue_InContainer(ContainerAddress, &WrappedValue);
+			}
+			else
+			{
+				// Malloc a temp value that is the size of the array. Getter will then copy the entire array to the temp value
+				T* ValueArray = (T*)AllocateAndInitializeValue();
+				FProperty::GetValue_InContainer(ContainerAddress, ValueArray);
+				// Replace the item we care about
+				ValueArray[ArrayIndex] = InValue;
+				// Now copy the entire array back to the property using a setter
+				SetValue_InContainer(ContainerAddress, ValueArray);
+				DestroyAndFreeValue(ValueArray);
+			}
+		}
+	}
 };
 
 template<typename InTCppType>
@@ -2843,6 +2893,7 @@ public:
 	virtual UObject* GetObjectPropertyValue(const void* PropertyValueAddress) const override;
 	virtual UObject* GetObjectPropertyValue_InContainer(const void* ContainerAddress, int32 ArrayIndex = 0) const override;
 	virtual void SetObjectPropertyValue(void* PropertyValueAddress, UObject* Value) const override;
+	virtual void SetObjectPropertyValue_InContainer(void* ContainerAddress, UObject* Value, int32 ArrayIndex = 0) const override;
 	virtual FString GetCPPTypeCustom(FString* ExtendedTypeText, uint32 CPPExportFlags, const FString& InnerNativeTypeName)  const override;
 	// End of FObjectPropertyBase interface
 };
@@ -2881,6 +2932,7 @@ class COREUOBJECT_API FObjectPtrProperty : public FObjectProperty
 	virtual UObject* GetObjectPropertyValue(const void* PropertyValueAddress) const override;
 	virtual UObject* GetObjectPropertyValue_InContainer(const void* ContainerAddress, int32 ArrayIndex = 0) const override;
 	virtual void SetObjectPropertyValue(void* PropertyValueAddress, UObject* Value) const override;
+	virtual void SetObjectPropertyValue_InContainer(void* ContainerAddress, UObject* Value, int32 ArrayIndex = 0) const override;
 	virtual bool AllowCrossLevel() const override;
 	virtual bool AllowObjectTypeReinterpretationTo(const FObjectPropertyBase* Other) const override;
 private:
@@ -2942,6 +2994,7 @@ public:
 	virtual UObject* GetObjectPropertyValue(const void* PropertyValueAddress) const override;
 	virtual UObject* GetObjectPropertyValue_InContainer(const void* ContainerAddress, int32 ArrayIndex = 0) const override;
 	virtual void SetObjectPropertyValue(void* PropertyValueAddress, UObject* Value) const override;
+	virtual void SetObjectPropertyValue_InContainer(void* ContainerAddress, UObject* Value, int32 ArrayIndex = 0) const override;
 	// End of FObjectProperty interface
 };
 
@@ -2998,6 +3051,7 @@ PRAGMA_ENABLE_DEPRECATION_WARNINGS
 	virtual UObject* GetObjectPropertyValue(const void* PropertyValueAddress) const override;
 	virtual UObject* GetObjectPropertyValue_InContainer(const void* ContainerAddress, int32 ArrayIndex = 0) const override;
 	virtual void SetObjectPropertyValue(void* PropertyValueAddress, UObject* Value) const override;
+	virtual void SetObjectPropertyValue_InContainer(void* ContainerAddress, UObject* Value, int32 ArrayIndex = 0) const override;
 	virtual bool AllowCrossLevel() const override;
 private:
 	virtual uint32 GetValueTypeHashInternal(const void* Src) const override;
@@ -3070,6 +3124,7 @@ public:
 	virtual UObject* GetObjectPropertyValue(const void* PropertyValueAddress) const override;
 	virtual UObject* GetObjectPropertyValue_InContainer(const void* ContainerAddress, int32 ArrayIndex = 0) const override;
 	virtual void SetObjectPropertyValue(void* PropertyValueAddress, UObject* Value) const override;
+	virtual void SetObjectPropertyValue_InContainer(void* ContainerAddress, UObject* Value, int32 ArrayIndex = 0) const override;
 	virtual bool AllowCrossLevel() const override;
 	virtual FString GetCPPTypeCustom(FString* ExtendedTypeText, uint32 CPPExportFlags, const FString& InnerNativeTypeName)  const override;
 
@@ -3212,6 +3267,7 @@ class COREUOBJECT_API FClassPtrProperty : public FClassProperty
 	virtual UObject* GetObjectPropertyValue(const void* PropertyValueAddress) const override;
 	virtual UObject* GetObjectPropertyValue_InContainer(const void* ContainerAddress, int32 ArrayIndex = 0) const override;
 	virtual void SetObjectPropertyValue(void* PropertyValueAddress, UObject* Value) const override;
+	virtual void SetObjectPropertyValue_InContainer(void* ContainerAddress, UObject* Value, int32 ArrayIndex = 0) const override;
 private:
 	virtual uint32 GetValueTypeHashInternal(const void* Src) const override;
 public:
