Adjusted TWeakObjectPtr to use 0 instead of INDEX_NONE for ObjectIndex, to represent invalid state - so that it correctly implements TIsZeroConstructType in all cases.

diff --git a/Engine/Source/Runtime/CoreUObject/Private/UObject/WeakObjectPtr.cpp b/Engine/Source/Runtime/CoreUObject/Private/UObject/WeakObjectPtr.cpp
--- a/Engine/Source/Runtime/CoreUObject/Private/UObject/WeakObjectPtr.cpp
+++ b/Engine/Source/Runtime/CoreUObject/Private/UObject/WeakObjectPtr.cpp
@@ -50,11 +50,18 @@ bool FWeakObjectPtr::IsValid() const
 
 bool FWeakObjectPtr::IsStale(bool bEvenIfPendingKill, bool bThreadsafeTest) const
 {
+	using namespace UE::Core::Private;
+
 	if (ObjectSerialNumber == 0)
 	{
+#if UE_WEAKOBJECTPTR_ZEROINIT_FIX
+		checkSlow(ObjectIndex == InvalidWeakObjectIndex); // otherwise this is a corrupted weak pointer
+#else
 		checkSlow(ObjectIndex == 0 || ObjectIndex == -1); // otherwise this is a corrupted weak pointer
+#endif
 		return false;
 	}
+
 	if (ObjectIndex < 0)
 	{
 		return true;
@@ -91,7 +98,7 @@ UObject* FWeakObjectPtr::GetEvenIfUnreachable() const
 	UObject* Result = nullptr;
 	if (Internal_IsValid(true, true))
 	{
-		FUObjectItem* ObjectItem = GUObjectArray.IndexToObject(GetObjectIndex(), true);
+		FUObjectItem* ObjectItem = GUObjectArray.IndexToObject(GetObjectIndex_Private(), true);
 		Result = static_cast<UObject*>(ObjectItem->Object);
 	}
 	return Result;
diff --git a/Engine/Source/Runtime/CoreUObject/Public/UObject/ObjectKey.h b/Engine/Source/Runtime/CoreUObject/Public/UObject/ObjectKey.h
--- a/Engine/Source/Runtime/CoreUObject/Public/UObject/ObjectKey.h
+++ b/Engine/Source/Runtime/CoreUObject/Public/UObject/ObjectKey.h
@@ -13,14 +13,14 @@ struct FObjectKey
 public:
 	/** Default constructor */
 	FORCEINLINE FObjectKey()
-		: ObjectIndex(INDEX_NONE)
+		: ObjectIndex(UE::Core::Private::InvalidWeakObjectIndex)
 		, ObjectSerialNumber(0)
 	{
 	}
 
 	/** Construct from an object pointer */
 	FORCEINLINE FObjectKey(const UObject* Object)
-		: ObjectIndex(INDEX_NONE)
+		: ObjectIndex(UE::Core::Private::InvalidWeakObjectIndex)
 		, ObjectSerialNumber(0)
 	{
 		if (Object)
diff --git a/Engine/Source/Runtime/CoreUObject/Public/UObject/WeakObjectPtr.h b/Engine/Source/Runtime/CoreUObject/Public/UObject/WeakObjectPtr.h
--- a/Engine/Source/Runtime/CoreUObject/Public/UObject/WeakObjectPtr.h
+++ b/Engine/Source/Runtime/CoreUObject/Public/UObject/WeakObjectPtr.h
@@ -12,6 +12,21 @@
 
 #include <type_traits>
 
+/** Invalid FWeakObjectPtr ObjectIndex values must be 0 to support zeroed initialization (this used to be INDEX_NONE, leading to subtle bugs). */
+#ifndef UE_WEAKOBJECTPTR_ZEROINIT_FIX
+	#define UE_WEAKOBJECTPTR_ZEROINIT_FIX 1
+#endif
+
+namespace UE::Core::Private
+{
+	/** Specifies the ObjectIndex used for invalid object pointers. */
+#if UE_WEAKOBJECTPTR_ZEROINIT_FIX
+	static constexpr int32 InvalidWeakObjectIndex = 0;
+#else
+	static constexpr int32 InvalidWeakObjectIndex = INDEX_NONE;
+#endif
+}
+
 /**
  * FWeakObjectPtr is a weak pointer to a UObject. 
  * It can return nullptr later if the object is garbage collected.
@@ -24,10 +39,16 @@ struct FWeakObjectPtr
 {
 public:
 
-	template <typename ReferenceProcessorType, typename CollectorType, typename ArrayPoolType, EFastReferenceCollectorOptions Options>
-	friend class TFastReferenceCollector;
 	friend struct FFieldPath;
 
+#if UE_WEAKOBJECTPTR_ZEROINIT_FIX
+	FWeakObjectPtr() = default;
+
+	FORCEINLINE FWeakObjectPtr(TYPE_OF_NULLPTR)
+		: FWeakObjectPtr()
+	{
+	}
+#else
 	/** Null constructor **/
 	FORCEINLINE FWeakObjectPtr()
 	{
@@ -42,6 +63,7 @@ public:
 	{
 		(*this) = nullptr;
 	}
+#endif
 
 	/**  
 	 * Construct from an object pointer or something that can be implicitly converted to an object pointer
@@ -67,7 +89,9 @@ public:
 	 */
 	FORCEINLINE void Reset()
 	{
-		ObjectIndex = INDEX_NONE;
+		using namespace UE::Core::Private;
+
+		ObjectIndex = InvalidWeakObjectIndex;
 		ObjectSerialNumber = 0;
 	}
 
@@ -160,7 +184,13 @@ public:
 	 */
 	FORCEINLINE bool IsExplicitlyNull() const
 	{
-		return ObjectIndex == INDEX_NONE;
+		using namespace UE::Core::Private;
+
+#if UE_WEAKOBJECTPTR_ZEROINIT_FIX
+		return ObjectIndex == InvalidWeakObjectIndex && ObjectSerialNumber == 0;
+#else
+		return ObjectIndex == InvalidWeakObjectIndex;
+#endif
 	}
 
 	/** Hash function. */
@@ -178,12 +208,18 @@ public:
 	COREUOBJECT_API void Serialize(FArchive& Ar);
 
 protected:
-
+	UE_DEPRECATED(5.1, "GetObjectIndex is now deprecated, and will be removed.")
 	FORCEINLINE int32 GetObjectIndex() const
 	{
 		return ObjectIndex;
 	}
 
+private:
+	FORCEINLINE int32 GetObjectIndex_Private() const
+	{
+		return ObjectIndex;
+	}
+
 private:
 	friend struct FObjectKey;
 	
@@ -209,11 +245,19 @@ private:
 
 	FORCEINLINE FUObjectItem* Internal_GetObjectItem() const
 	{
+		using namespace UE::Core::Private;
+
 		if (ObjectSerialNumber == 0)
 		{
+#if UE_WEAKOBJECTPTR_ZEROINIT_FIX
+			checkSlow(ObjectIndex == InvalidWeakObjectIndex); // otherwise this is a corrupted weak pointer
+#else
 			checkSlow(ObjectIndex == 0 || ObjectIndex == -1); // otherwise this is a corrupted weak pointer
+#endif
+
 			return nullptr;
 		}
+
 		if (ObjectIndex < 0)
 		{
 			return nullptr;
@@ -251,8 +295,13 @@ private:
 		return ((ObjectItem != nullptr) && GUObjectArray.IsValid(ObjectItem, bEvenIfPendingKill)) ? (UObject*)ObjectItem->Object : nullptr;
 	}
 
+#if UE_WEAKOBJECTPTR_ZEROINIT_FIX
+	int32		ObjectIndex = UE::Core::Private::InvalidWeakObjectIndex;
+	int32		ObjectSerialNumber = 0;
+#else
 	int32		ObjectIndex;
 	int32		ObjectSerialNumber;
+#endif
 };
 
 template<> struct TIsPODType<FWeakObjectPtr> { enum { Value = true }; };
