Fix for potential link errors caused by obsolete constructor.

diff --git a/Engine/Source/Runtime/CoreUObject/Public/UObject/UnrealType.h b/Engine/Source/Runtime/CoreUObject/Public/UObject/UnrealType.h
--- a/Engine/Source/Runtime/CoreUObject/Public/UObject/UnrealType.h
+++ b/Engine/Source/Runtime/CoreUObject/Public/UObject/UnrealType.h
@@ -1459,7 +1459,7 @@ protected:
 };
 
 template<typename InTCppType, class TInPropertyBaseClass>
-class COREUOBJECT_API TProperty_WithEqualityAndSerializer : public TProperty<InTCppType, TInPropertyBaseClass>
+class TProperty_WithEqualityAndSerializer : public TProperty<InTCppType, TInPropertyBaseClass>
 {
 
 public:
@@ -1502,9 +1502,6 @@ PRAGMA_ENABLE_DEPRECATION_WARNINGS
 	}
 #endif // WITH_EDITORONLY_DATA
 
-	/** DO NOT USE. This constructor is for internal usage only for hot-reload purposes. */
-	TProperty_WithEqualityAndSerializer(FVTableHelper& Helper) : Super(Helper) {};
-
 	// FProperty interface.
 	virtual bool Identical( const void* A, const void* B, uint32 PortFlags=0 ) const override
 	{
