Add new Unit Tests

diff --git a/Engine/Source/Runtime/CoreUObject/Tests/ObjectPtrTest.cpp b/Engine/Source/Runtime/CoreUObject/Tests/ObjectPtrTest.cpp
--- a/Engine/Source/Runtime/CoreUObject/Tests/ObjectPtrTest.cpp
+++ b/Engine/Source/Runtime/CoreUObject/Tests/ObjectPtrTest.cpp
@@ -141,36 +141,24 @@ static_assert(!TModels<CEqualityComparableWith, FMutableObjectPtr, long>::Value,
 //static_assert(!std::is_assignable<TObjectPtr<UForwardDeclaredObjDerived>, UForwardDeclaredObjDerived*>::value, "Should not be able to assign raw pointer of incomplete type that descends from UObject to exactly this type of TObjectPtr");
 //static_assert(!std::is_assignable<TObjectPtr<FForwardDeclaredNotObjDerived>, FForwardDeclaredNotObjDerived*>::value, "Should not be able to assign raw pointer of incomplete type that does not descend from UObject to exactly this type of TObjectPtr");
 
-#if WITH_DEV_AUTOMATION_TESTS
-
 class FObjectPtrTestBase : public FObjectRefTrackingTestBase
 {
 public:
-	FObjectPtrTestBase(const FString& InName, const bool bInComplexTask)
-	: FObjectRefTrackingTestBase(InName, bInComplexTask)
-	{
-	}
 
 protected:
 };
 
-#define TEST_NAME_ROOT TEXT("System.CoreUObject.ObjectPtr")
-constexpr const uint32 ObjectPtrTestFlags = EAutomationTestFlags::ApplicationContextMask | EAutomationTestFlags::EngineFilter;
 
-IMPLEMENT_CUSTOM_SIMPLE_AUTOMATION_TEST(FObjectPtrTestNullBehavior, FObjectPtrTestBase, TEST_NAME_ROOT TEXT(".NullBehavior"), ObjectPtrTestFlags)
-bool FObjectPtrTestNullBehavior::RunTest(const FString& Parameters)
+TEST_CASE_METHOD(FObjectPtrTestBase, "CoreUObject::TObjectPtr::Null Behavior", "[CoreUObject][ObjectPtr][Smoke]")
 {
 	TObjectPtr<UObject> NullObjectPtr(nullptr);
-	TestTrue(TEXT("Nullptr should equal a null object pointer"), nullptr == NullObjectPtr);
-	TestTrue(TEXT("A null object pointer should equal nullptr"), NullObjectPtr == nullptr);
-	TestFalse(TEXT("A null object pointer should evaluate to false"), !!NullObjectPtr);
-	TestTrue(TEXT("Negation of a null object pointer should evaluate to true"), !NullObjectPtr);
-
-	return true;
+	TEST_TRUE(TEXT("Nullptr should equal a null object pointer"), nullptr == NullObjectPtr);
+	TEST_TRUE(TEXT("A null object pointer should equal nullptr"), NullObjectPtr == nullptr);
+	TEST_FALSE(TEXT("A null object pointer should evaluate to false"), !!NullObjectPtr);
+	TEST_TRUE(TEXT("Negation of a null object pointer should evaluate to true"), !NullObjectPtr);
 }
 
-IMPLEMENT_CUSTOM_SIMPLE_AUTOMATION_TEST(FObjectPtrTestDefaultSerialize, FObjectPtrTestBase, TEST_NAME_ROOT TEXT(".DefaultSerialize"), ObjectPtrTestFlags)
-bool FObjectPtrTestDefaultSerialize::RunTest(const FString& Parameters)
+TEST_CASE_METHOD(FObjectPtrTestBase, "CoreUObject::TObjectPtr::Default Serialize", "[CoreUObject][ObjectPtr][.Engine]")
 {
 	FSnapshotObjectRefMetrics ObjectRefMetrics(*this);
 	FObjectPtr DefaultTexturePtr(FObjectRef {FName("/Engine/EngineResources/DefaultTexture"), NAME_None, NAME_None, FObjectPathId("DefaultTexture")});
@@ -191,12 +179,9 @@ bool FObjectPtrTestDefaultSerialize::RunTest(const FString& Parameters)
 	ObjectRefMetrics.TestNumResolves(TEXT("Serializing an FObjectPtr twice should only require it to resolve once"), 1);
 	ObjectRefMetrics.TestNumFailedResolves(TEXT("Unexpected resolve failure after serializing an FObjectPtr"), 0);
 	ObjectRefMetrics.TestNumReads(TEXT("NumReads should increase after serializing an FObjectPtr"), 2);
-
-	return true;
 }
 
-IMPLEMENT_CUSTOM_SIMPLE_AUTOMATION_TEST(FObjectPtrTestSoftObjectPath, FObjectPtrTestBase, TEST_NAME_ROOT TEXT(".SoftObjectPath"), ObjectPtrTestFlags)
-bool FObjectPtrTestSoftObjectPath::RunTest(const FString& Parameters)
+TEST_CASE_METHOD(FObjectPtrTestBase, "CoreUObject::TObjectPtr::Soft Object Path", "[CoreUObject][ObjectPtr][Smoke]")
 {
 	FSnapshotObjectRefMetrics ObjectRefMetrics(*this);
 	FObjectPtr DefaultTexturePtr(FObjectRef {FName("/Engine/EngineResources/DefaultTexture"), NAME_None, NAME_None, FObjectPathId("DefaultTexture")});
@@ -209,22 +194,17 @@ bool FObjectPtrTestSoftObjectPath::RunTest(const FString& Parameters)
 
 	ObjectRefMetrics.TestNumResolves(TEXT("Unexpected resolve count after initializing an FSoftObjectPath from an FObjectPtr"), UE_WITH_OBJECT_HANDLE_LATE_RESOLVE ? 0 : 1);
 
-	TestEqual(TEXT("Soft object path constructed from an FObjectPtr does not have the expected path value"), *DefaultTexturePath.ToString(), TEXT("/Engine/EngineResources/DefaultTexture.DefaultTexture"));
-
-	return true;
+	TEST_TRUE(TEXT("Soft object path constructed from an FObjectPtr does not have the expected path value"), !FCString::Strcmp(*DefaultTexturePath.ToString(), TEXT("/Engine/EngineResources/DefaultTexture.DefaultTexture")));
 }
 
-IMPLEMENT_CUSTOM_SIMPLE_AUTOMATION_TEST(FObjectPtrTestForwardDeclared, FObjectPtrTestBase, TEST_NAME_ROOT TEXT(".ForwardDeclared"), ObjectPtrTestFlags)
-bool FObjectPtrTestForwardDeclared::RunTest(const FString& Parameters)
+TEST_CASE_METHOD(FObjectPtrTestBase, "CoreUObject::TObjectPtr::Forward Declared", "[CoreUObject][ObjectPtr][Smoke]")
 {
 	UForwardDeclaredObjDerived* PtrFwd = nullptr;
 	TObjectPtr<UForwardDeclaredObjDerived> ObjPtrFwd(MakeObjectPtrUnsafe<UForwardDeclaredObjDerived>(reinterpret_cast<UObject*>(PtrFwd)));
-	TestTrue(TEXT("Null forward declared pointer used to construct a TObjectPtr should result in a null TObjectPtr"), ObjPtrFwd.IsNull());
-	return true;
+	TEST_TRUE(TEXT("Null forward declared pointer used to construct a TObjectPtr should result in a null TObjectPtr"), ObjPtrFwd.IsNull());
 }
 
-IMPLEMENT_CUSTOM_SIMPLE_AUTOMATION_TEST(FObjectPtrTestHashConsistency, FObjectPtrTestBase, TEST_NAME_ROOT TEXT(".HashConsistency"), ObjectPtrTestFlags)
-bool FObjectPtrTestHashConsistency::RunTest(const FString& Parameters)
+TEST_CASE_METHOD(FObjectPtrTestBase, "CoreUObject::TObjectPtr::Hash Consistency", "[CoreUObject][ObjectPtr][Smoke]")
 {
 	const FName TestPackage1Name(TEXT("/Engine/Test/ObjectPtrHashConsistency1/Transient"));
 	UPackage* TestPackage1 = NewObject<UPackage>(nullptr, TestPackage1Name, RF_Transient);
@@ -259,7 +239,7 @@ bool FObjectPtrTestHashConsistency::RunTest(const FString& Parameters)
 		ObjectRefMetrics.TestNumFailedResolves(TEXT("Unexpected resolve failure after resolving an FObjectPtr"), 0);
 
 		uint32 HashRaw = GetTypeHash(TestPublicObject);
-		TestEqual(TEXT("Hash of raw public FObjectPtr should equal hash of wrapped public FObjectPtr"), HashRaw, HashWrapped);
+		TEST_EQUAL(TEXT("Hash of raw public FObjectPtr should equal hash of wrapped public FObjectPtr"), HashRaw, HashWrapped);
 
 		FObjectPtr TestPublicWrappedRedir1(FObjectRef{TestPackage1Name, NAME_None, NAME_None, FObjectPathId("TestOuter3.TestPublicRedirector1")});
 		ObjectRefMetrics.TestNumResolves(TEXT("Unexpected resolve count after initializing an FObjectPtr"), UE_WITH_OBJECT_HANDLE_LATE_RESOLVE ? 1 : 2);
@@ -269,7 +249,7 @@ bool FObjectPtrTestHashConsistency::RunTest(const FString& Parameters)
 		ObjectRefMetrics.TestNumResolves(TEXT("Unexpected resolve count after hashing an FObjectPtr"), 2);
 		ObjectRefMetrics.TestNumFailedResolves(TEXT("Unexpected resolve failure after hashing an FObjectPtr"), 0);
 
-		TestEqual(TEXT("Hash of first wrapped public redirector should equal hash of wrapped public FObjectPtr it references"), HashWrapped, HashWrappedRedir1);
+		TEST_EQUAL(TEXT("Hash of first wrapped public redirector should equal hash of wrapped public FObjectPtr it references"), HashWrapped, HashWrappedRedir1);
 
 		FObjectPtr TestPublicWrappedRedir2(FObjectRef{TestPackage1Name, NAME_None, NAME_None, FObjectPathId("TestOuter4.TestPublicRedirector2")});
 		ObjectRefMetrics.TestNumResolves(TEXT("Unexpected resolve count after initializing an FObjectPtr"), UE_WITH_OBJECT_HANDLE_LATE_RESOLVE ? 2 : 3);
@@ -279,25 +259,24 @@ bool FObjectPtrTestHashConsistency::RunTest(const FString& Parameters)
 		ObjectRefMetrics.TestNumResolves(TEXT("Unexpected resolve count after hashing an FObjectPtr"), 3);
 		ObjectRefMetrics.TestNumFailedResolves(TEXT("Unexpected resolve failure after hashing an FObjectPtr"), 0);
 
-		TestEqual(TEXT("Hash of first wrapped public redirector should equal hash of wrapped public FObjectPtr it references"), HashWrapped, HashWrappedRedir1);
+		TEST_EQUAL(TEXT("Hash of first wrapped public redirector should equal hash of wrapped public FObjectPtr it references"), HashWrapped, HashWrappedRedir1);
 
 		//Check that renaming an object doesn't change its hash
 		TestPublicObject->Rename(TEXT("TestPublicObjectRenamed"));
 		uint32 HashWrappedAfterRename = GetTypeHash(TestPublicWrappedObjectPtr);
-		TestEqual(TEXT("Hash of resolved public FObjectPtr before rename should equal hash of resolved public FObjectPtr after rename"), HashWrappedAfterRename, HashWrapped);
+		TEST_EQUAL(TEXT("Hash of resolved public FObjectPtr before rename should equal hash of resolved public FObjectPtr after rename"), HashWrappedAfterRename, HashWrapped);
 
 		//Check that reparenting an object doesn't change its hash
 		TestPublicObject->Rename(nullptr, TestOuter2);
 		uint32 HashWrappedAfterReparent = GetTypeHash(TestPublicWrappedObjectPtr);
-		TestEqual(TEXT("Hash of resolved public FObjectPtr before reparenting should equal hash of resolved public FObjectPtr after reparenting"), HashWrappedAfterReparent, HashWrapped);
+		TEST_EQUAL(TEXT("Hash of resolved public FObjectPtr before reparenting should equal hash of resolved public FObjectPtr after reparenting"), HashWrappedAfterReparent, HashWrapped);
 	}
 
 	TestPackage1->RemoveFromRoot();
-	return true;
 }
 
-IMPLEMENT_CUSTOM_SIMPLE_AUTOMATION_TEST(FObjectPtrTestLongPath, FObjectPtrTestBase, TEST_NAME_ROOT TEXT(".LongPath"), ObjectPtrTestFlags)
-bool FObjectPtrTestLongPath::RunTest(const FString& Parameters)
+
+TEST_CASE_METHOD(FObjectPtrTestBase, "CoreUObject::TObjectPtr::Long Path", "[CoreUObject][ObjectPtr][Smoke]")
 {
 	const FName TestPackage1Name(TEXT("/Engine/Test/FObjectPtrTestLongPath/Transient"));
 	UPackage* TestPackage1 = NewObject<UPackage>(nullptr, TestPackage1Name, RF_Transient);
@@ -316,13 +295,11 @@ bool FObjectPtrTestLongPath::RunTest(const FString& Parameters)
 	FObjectPathId::ResolvedNameContainerType ResolvedNames;
 	LongPath.Resolve(ResolvedNames);
 
-	TestEqual(TEXT("Resolved path from FObjectPathId should have 4 elements"), ResolvedNames.Num(), 4);
-	TestEqual(TEXT("Resolved path from FObjectPathId should have TestObject1 at element 0"), ResolvedNames[0], TestObject1->GetFName());
-	TestEqual(TEXT("Resolved path from FObjectPathId should have TestObject2 at element 1"), ResolvedNames[1], TestObject2->GetFName());
-	TestEqual(TEXT("Resolved path from FObjectPathId should have TestObject3 at element 2"), ResolvedNames[2], TestObject3->GetFName());
-	TestEqual(TEXT("Resolved path from FObjectPathId should have TestObject4 at element 3"), ResolvedNames[3], TestObject4->GetFName());
-
-	return true;
+	TEST_EQUAL(TEXT("Resolved path from FObjectPathId should have 4 elements"), ResolvedNames.Num(), 4);
+	TEST_EQUAL(TEXT("Resolved path from FObjectPathId should have TestObject1 at element 0"), ResolvedNames[0], TestObject1->GetFName());
+	TEST_EQUAL(TEXT("Resolved path from FObjectPathId should have TestObject2 at element 1"), ResolvedNames[1], TestObject2->GetFName());
+	TEST_EQUAL(TEXT("Resolved path from FObjectPathId should have TestObject3 at element 2"), ResolvedNames[2], TestObject3->GetFName());
+	TEST_EQUAL(TEXT("Resolved path from FObjectPathId should have TestObject4 at element 3"), ResolvedNames[3], TestObject4->GetFName());
 }
 
 // @TODO: OBJPTR: We should have a test that ensures that lazy loading of an object with an external package is handled correctly.
@@ -364,10 +341,6 @@ bool FObjectPtrTestLongPath::RunTest(const FString& Parameters)
 // 	return true;
 // }
 
-#undef TEST_NAME_ROOT
-
-#endif // WITH_DEV_AUTOMATION_TESTS
-
 class UForwardDeclaredObjDerived: public UObject {};
 class FForwardDeclaredNotObjDerived {};
 
diff --git a/Engine/Source/Runtime/CoreUObject/Tests/ObjectRefTrackingTestBase.h b/Engine/Source/Runtime/CoreUObject/Tests/ObjectRefTrackingTestBase.h
--- a/Engine/Source/Runtime/CoreUObject/Tests/ObjectRefTrackingTestBase.h
+++ b/Engine/Source/Runtime/CoreUObject/Tests/ObjectRefTrackingTestBase.h
@@ -3,18 +3,12 @@
 #pragma once
 
 #include "CoreMinimal.h"
-#include "Misc/AutomationTest.h"
+#include "TestFixtures/CoreTestFixture.h"
 #include "UObject/ObjectHandle.h"
 
-#if WITH_DEV_AUTOMATION_TESTS
-
-class FObjectRefTrackingTestBase : public FAutomationTestBase
+class FObjectRefTrackingTestBase : public FCoreTestFixture
 {
 public:
-	FObjectRefTrackingTestBase(const FString& InName, const bool bInComplexTask)
-	: FAutomationTestBase(InName, bInComplexTask)
-	{
-	}
 
 	uint32 GetNumResolves() const { return NumResolves; }
 	uint32 GetNumFailedResolves() const { return NumFailedResolves; }
@@ -35,19 +29,18 @@ public:
 		bool TestNumResolves(const TCHAR* What, uint32 ExpectedDelta)
 		{
 #if UE_WITH_OBJECT_HANDLE_TRACKING
-			return Test.TestEqual(What, OriginalNumResolves + ExpectedDelta, Test.GetNumResolves());
-#else
-			return true;
+			TEST_EQUAL(What, OriginalNumResolves + ExpectedDelta, Test.GetNumResolves());
 #endif
+			return true;
 		}
 
 		bool TestNumFailedResolves(const TCHAR* What, uint32 ExpectedDelta)
 		{
 #if UE_WITH_OBJECT_HANDLE_TRACKING
-			return Test.TestEqual(What, OriginalNumFailedResolves + ExpectedDelta, Test.GetNumFailedResolves());
-#else
-			return true;
+			TEST_EQUAL(What, OriginalNumFailedResolves + ExpectedDelta, Test.GetNumFailedResolves());
 #endif
+			return true;
+
 		}
 
 		bool TestNumReads(const TCHAR* What, uint32 ExpectedDelta, bool bAllowAdditionalReads = false)
@@ -55,15 +48,15 @@ public:
 #if UE_WITH_OBJECT_HANDLE_TRACKING
 			if (bAllowAdditionalReads)
 			{
-				return Test.TestTrue(What, Test.GetNumReads() >= OriginalNumReads + ExpectedDelta);
+				TEST_TRUE(What, Test.GetNumReads() >= OriginalNumReads + ExpectedDelta);
 			}
 			else
 			{
-				return Test.TestEqual(What, OriginalNumReads + ExpectedDelta, Test.GetNumReads());
+				TEST_EQUAL(What, OriginalNumReads + ExpectedDelta, Test.GetNumReads());
 			}
-#else
-			return true;
 #endif
+			return true;
+
 		}
 	private:
 		FObjectRefTrackingTestBase& Test;
@@ -126,4 +119,4 @@ private:
 	static thread_local uint32 NumFailedResolves;
 	static thread_local uint32 NumReads;
 };
-#endif // WITH_DEV_AUTOMATION_TESTS
+
diff --git a/Engine/Source/Runtime/CoreUObject/Tests/Serialization/BulkDataTests.cpp b/Engine/Source/Runtime/CoreUObject/Tests/Serialization/BulkDataTests.cpp
--- a/Engine/Source/Runtime/CoreUObject/Tests/Serialization/BulkDataTests.cpp
+++ b/Engine/Source/Runtime/CoreUObject/Tests/Serialization/BulkDataTests.cpp
@@ -6,43 +6,43 @@
 namespace BulkDataTest
 {
 	// Test code paths for BulkData objects that do not reference a file on disk.
-	TEST_CASE("CoreUObject::Serialization::FByteBulkData::Transient", "[CoreUObject][Serialization][Smoke]")
+	TEST_CASE_METHOD(FAutomationTestFixture, "CoreUObject::Serialization::FByteBulkData::Transient", "[CoreUObject][Serialization][Smoke]")
 	{
 		FByteBulkData BulkData;
 
 		// We should be able to lock for read access but there should be no valid data
 		const void* ReadOnlyDataPtr = BulkData.LockReadOnly();
-		TestNull(TEXT("Locking an empty BulkData object for reading should return nullptr!"), ReadOnlyDataPtr);
+		TEST_NULL(TEXT("Locking an empty BulkData object for reading should return nullptr!"), ReadOnlyDataPtr);
 		BulkData.Unlock();
 
 		void* DataPtr = BulkData.Lock(LOCK_READ_WRITE);
-		TestNull(TEXT("Locking an empty BulkData object for writing should return nullptr!"), DataPtr);
+		TEST_NULL(TEXT("Locking an empty BulkData object for writing should return nullptr!"), DataPtr);
 		BulkData.Unlock();
 
 		void* CopyEmptyPtr = nullptr;
 		BulkData.GetCopy(&CopyEmptyPtr, true);
-		TestNull(TEXT("Getting a copy of an empty BulkData object for writing should return nullptr!"), CopyEmptyPtr);
+		TEST_NULL(TEXT("Getting a copy of an empty BulkData object for writing should return nullptr!"), CopyEmptyPtr);
 
 		BulkData.Lock(LOCK_READ_WRITE);
 		DataPtr = BulkData.Realloc(32 * 32 * 4);
-		TestNotNull(TEXT("Reallocating an empty BulkData object should return a valid pointer!"), DataPtr);
+		TEST_NOT_NULL(TEXT("Reallocating an empty BulkData object should return a valid pointer!"), DataPtr);
 		BulkData.Unlock();
 
-		TestTrue(TEXT("BulkData should be loaded now that it has been reallocated"), BulkData.IsBulkDataLoaded());
+		TEST_TRUE(TEXT("BulkData should be loaded now that it has been reallocated"), BulkData.IsBulkDataLoaded());
 
 		void* CopyWithDiscardPtr = nullptr;
 		BulkData.GetCopy(&CopyWithDiscardPtr, true); // The second parameter should be ignored because the bulkdata cannot be reloaded from disk!
-		TestNotNull(TEXT("GetCopy should return a valid pointer!"), CopyWithDiscardPtr);
-		TestNotEqual(TEXT("GetCopy should return a copy of the data so the pointers should be different!"), DataPtr, CopyWithDiscardPtr);
-		TestTrue(TEXT("BulkData should still loaded after taking a copy"), BulkData.IsBulkDataLoaded());
+		TEST_NOT_NULL(TEXT("GetCopy should return a valid pointer!"), CopyWithDiscardPtr);
+		TEST_NOT_EQUAL(TEXT("GetCopy should return a copy of the data so the pointers should be different!"), DataPtr, CopyWithDiscardPtr);
+		TEST_TRUE(TEXT("BulkData should still loaded after taking a copy"), BulkData.IsBulkDataLoaded());
 
 		// Now try GetCopy again but this time without the discard request
 		void* CopyNoDiscardPtr = nullptr;
 		BulkData.GetCopy(&CopyNoDiscardPtr, false);
-		TestNotNull(TEXT("GetCopy should return a valid pointer!"), CopyNoDiscardPtr);
-		TestNotEqual(TEXT("GetCopy should return a copy of the data so the pointers should be different!"), DataPtr, CopyNoDiscardPtr);
-		TestNotEqual(TEXT("GetCopy should return a copy of the data so the pointers should be different!"), CopyWithDiscardPtr, CopyNoDiscardPtr);
-		TestTrue(TEXT("BulkData should still loaded after taking a copy"), BulkData.IsBulkDataLoaded());
+		TEST_NOT_NULL(TEXT("GetCopy should return a valid pointer!"), CopyNoDiscardPtr);
+		TEST_NOT_EQUAL(TEXT("GetCopy should return a copy of the data so the pointers should be different!"), DataPtr, CopyNoDiscardPtr);
+		TEST_NOT_EQUAL(TEXT("GetCopy should return a copy of the data so the pointers should be different!"), CopyWithDiscardPtr, CopyNoDiscardPtr);
+		TEST_TRUE(TEXT("BulkData should still loaded after taking a copy"), BulkData.IsBulkDataLoaded());
 
 		// Clean up allocations from GetCopy
 		FMemory::Free(CopyWithDiscardPtr);
@@ -52,26 +52,26 @@ namespace BulkDataTest
 		DataPtr = BulkData.Lock(LOCK_READ_WRITE);
 		BulkData.Unlock();
 
-		TestTrue(TEXT("BulkData should still loaded after locking for write"), BulkData.IsBulkDataLoaded());
-		TestNotNull(TEXT("Locking for write should return a valid pointer!"), DataPtr);
+		TEST_TRUE(TEXT("BulkData should still loaded after locking for write"), BulkData.IsBulkDataLoaded());
+		TEST_NOT_NULL(TEXT("Locking for write should return a valid pointer!"), DataPtr);
 
 		// Now remove the bulkdata and make sure that we cannot access the old data anymore
 		BulkData.RemoveBulkData();
-		TestFalse(TEXT("RemoveBulkData should've discarded the BulkData"), BulkData.IsBulkDataLoaded());
+		TEST_FALSE(TEXT("RemoveBulkData should've discarded the BulkData"), BulkData.IsBulkDataLoaded());
 
 #if USE_NEW_BULKDATA // The warning only exists for BulkData2
 		// Both the ::Lock and ::GetCopy calls should warn that we cannot load the missing data (as it will still have a valid size)
-		TestAddWarning(TEXT("Attempting to load a BulkData object that cannot be loaded from disk"));
+		AddExpectedError(TEXT("Attempting to load a BulkData object that cannot be loaded from disk"), EAutomationExpectedErrorFlags::Exact, 2);
 #endif //USE_NEW_BULKDATA
 
 		DataPtr = BulkData.Lock(LOCK_READ_WRITE);
 		BulkData.Unlock();
 
-		TestNull(TEXT("Locking for write after calling ::RemoveBulkData should return a nullptr!"), DataPtr);
+		TEST_NULL(TEXT("Locking for write after calling ::RemoveBulkData should return a nullptr!"), DataPtr);
 
 		CopyEmptyPtr = nullptr;
 		BulkData.GetCopy(&CopyNoDiscardPtr, true);
-		TestNull(TEXT("Getting a copy of BulkData object after calling ::RemoveBulkData should return nullptr!"), DataPtr);
+		TEST_NULL(TEXT("Getting a copy of BulkData object after calling ::RemoveBulkData should return nullptr!"), DataPtr);
 	}
 }
 
diff --git a/Engine/Source/Runtime/CoreUObject/Tests/Serialization/EditorBulkDataTests.cpp b/Engine/Source/Runtime/CoreUObject/Tests/Serialization/EditorBulkDataTests.cpp
--- a/Engine/Source/Runtime/CoreUObject/Tests/Serialization/EditorBulkDataTests.cpp
+++ b/Engine/Source/Runtime/CoreUObject/Tests/Serialization/EditorBulkDataTests.cpp
@@ -10,13 +10,16 @@
 #include "Serialization/MemoryReader.h"
 #include "Serialization/MemoryWriter.h"
 #include "Templates/UniquePtr.h"
-#include "TestHarness.h"
 
-#if WITH_EDITORONLY_DATA
+#include "Misc/AutomationTest.h"
+
+#if WITH_DEV_AUTOMATION_TESTS && WITH_EDITORONLY_DATA
 
 namespace UE::Serialization
 {
 	
+constexpr const uint32 TestFlags = EAutomationTestFlags::ApplicationContextMask | EAutomationTestFlags::EngineFilter;
+
 /** Creates a buffer full of random data to make it easy to have something to test against. */
 TUniquePtr<uint8[]> CreateRandomData(int64 BufferSize)
 {
@@ -35,12 +38,13 @@ TUniquePtr<uint8[]> CreateRandomData(int64 BufferSize)
  * payload via both TFuture and callback methods. It then creates copies of the object and makes sure that we can get the payload from
  * the copies, even when the original source object has been reset.
  */
-TEST_CASE("CoreUObject::Serialization::FEditorBulkData::Basic", "[Core][UObject][Smoke]")
+IMPLEMENT_SIMPLE_AUTOMATION_TEST(FEditorBulkDataTestBasic, TEXT("System.CoreUObject.Serialization.EditorBulkData.Basic"), TestFlags)
+bool FEditorBulkDataTestBasic::RunTest(const FString& Parameters)
 {
 	const int64 BufferSize = 1024;
 	TUniquePtr<uint8[]> SourceBuffer = CreateRandomData(BufferSize);
 
-	auto ValidateBulkData = [&SourceBuffer, BufferSize](const FEditorBulkData& BulkDataToValidate, const TCHAR* Label)
+	auto ValidateBulkData = [this, &SourceBuffer, BufferSize](const FEditorBulkData& BulkDataToValidate, const TCHAR* Label)
 		{
 			FSharedBuffer RetrievedBuffer = BulkDataToValidate.GetPayload().Get();
 			TestEqual(FString::Printf(TEXT("%s buffer length"),Label), (int64)RetrievedBuffer.GetSize(), BufferSize);
@@ -72,46 +76,51 @@ TEST_CASE("CoreUObject::Serialization::FEditorBulkData::Basic", "[Core][UObject]
 	// Test both bulkdata objects again now that we reset the data
 	ValidateBulkData(BulkDataCopy, TEXT("Copy Constructor (after data reset)"));
 	ValidateBulkData(BulkDataAssignment, TEXT("Copy Assignment (after data reset)"));
+
+	return true;
 }
 
 /**
  * This test will validate how FEditorBulkData behaves when it has no associated payload and make sure 
  * that our assumptions are correct.
  */
-TEST_CASE("CoreUObject::Serialization::FEditorBulkData::Empty", "[Core][UObject][Smoke]")
+IMPLEMENT_SIMPLE_AUTOMATION_TEST(FEditorBulkDataTestEmpty, TEXT("System.CoreUObject.Serialization.EditorBulkData.Empty"), TestFlags)
+bool FEditorBulkDataTestEmpty::RunTest(const FString& Parameters)
 {
-	auto Validate = [](const TCHAR* Id,  const FEditorBulkData& BulkData)
+	auto Validate = [](const TCHAR* Id, FAutomationTestBase& Test, const FEditorBulkData& BulkData)
 	{
 		// Validate the general accessors
-		TestEqual(FString::Printf(TEXT("(%s) Return value of ::GetBulkDataSize()"), Id), BulkData.GetPayloadSize(), (int64)0);
-		TestTrue(FString::Printf(TEXT("(%s) Payload key is invalid"), Id), BulkData.GetPayloadId().IsZero());
-		TestFalse(FString::Printf(TEXT("(%s) Return value of ::DoesPayloadNeedLoading()"), Id), BulkData.DoesPayloadNeedLoading());
+		Test.TestEqual(FString::Printf(TEXT("(%s) Return value of ::GetBulkDataSize()"), Id), BulkData.GetPayloadSize(), (int64)0);
+		Test.TestTrue(FString::Printf(TEXT("(%s) Payload key is invalid"), Id), BulkData.GetPayloadId().IsZero());
+		Test.TestFalse(FString::Printf(TEXT("(%s) Return value of ::DoesPayloadNeedLoading()"), Id), BulkData.DoesPayloadNeedLoading());
 
 		// Validate the payload accessors
 		FSharedBuffer Payload = BulkData.GetPayload().Get();
-		TestTrue(FString::Printf(TEXT("(%s) The payload from the GetPayload TFuture is null"), Id), Payload.IsNull());
+		Test.TestTrue(FString::Printf(TEXT("(%s) The payload from the GetPayload TFuture is null"), Id), Payload.IsNull());
 
 		FCompressedBuffer CompressedPayload = BulkData.GetCompressedPayload().Get();
-		TestTrue(FString::Printf(TEXT("(%s) The payload from the GetCompressedPayload TFuture is null"), Id), Payload.IsNull());
+		Test.TestTrue(FString::Printf(TEXT("(%s) The payload from the GetCompressedPayload TFuture is null"), Id), Payload.IsNull());
 	};
 
 	FEditorBulkData DefaultBulkData;
-	Validate(TEXT("DefaultBulkData"), DefaultBulkData);
+	Validate(TEXT("DefaultBulkData"), *this, DefaultBulkData);
 		
 	FEditorBulkData NullPayloadBulkData;
 	NullPayloadBulkData.UpdatePayload(FSharedBuffer());
-	Validate(TEXT("NullPayloadBulkData"), NullPayloadBulkData);
+	Validate(TEXT("NullPayloadBulkData"), *this, NullPayloadBulkData);
 
 	FEditorBulkData ZeroLengthPayloadBulkData;
 	ZeroLengthPayloadBulkData.UpdatePayload(FUniqueBuffer::Alloc(0).MoveToShared());
-	Validate(TEXT("ZeroLengthPayloadBulkData"), ZeroLengthPayloadBulkData);
+	Validate(TEXT("ZeroLengthPayloadBulkData"), *this, ZeroLengthPayloadBulkData);
+
+	return true;
 }
 
 /**
  * Test the various methods for updating the payload that a FEditorBulkData owns
  */
-
-TEST_CASE("CoreUObject::Serialization::FEditorBulkData::Update Payload", "[Core][UObject][Smoke]")
+IMPLEMENT_SIMPLE_AUTOMATION_TEST(FEditorBulkDataTestUpdatePayload, TEXT("System.CoreUObject.Serialization.EditorBulkData.UpdatePayload"), TestFlags)
+bool FEditorBulkDataTestUpdatePayload::RunTest(const FString& Parameters)
 {
 	// Create a memory buffer of all zeros
 	const int64 BufferSize = 1024;
@@ -177,6 +186,8 @@ TEST_CASE("CoreUObject::Serialization::FEditorBulkData::Update Payload", "[Core]
 
 		TestTrue(TEXT("All payload elements correctly updated"), bAllElementsCorrect);
 	}
+
+	return true;
 }
 
 /**
@@ -184,8 +195,8 @@ TEST_CASE("CoreUObject::Serialization::FEditorBulkData::Update Payload", "[Core]
  * Then we will serialize the FEditorBulkData object back to a second buffer and compare the results.
  * If the reader and writers are working then the ReplicatedBuffer should be the same as the original SourceBuffer.
  */
-
-TEST_CASE("CoreUObject::Serialization::FEditorBulkData::Reader/Writer", "[Core][UObject][Smoke]")
+IMPLEMENT_SIMPLE_AUTOMATION_TEST(FEditorBulkDataTestReaderWriter, TEXT("System.CoreUObject.Serialization.EditorBulkData.Reader/Writer"), TestFlags)
+bool FEditorBulkDataTestReaderWriter::RunTest(const FString& Parameters)
 {
 	const int64 BufferSize = 1024;
 
@@ -217,6 +228,7 @@ TEST_CASE("CoreUObject::Serialization::FEditorBulkData::Reader/Writer", "[Core][
 		FEditorBulkDataWriter WriterAr(EmptyBulkData);
 		FEditorBulkDataReader ReaderAr(EmptyBulkData);
 	}
+	return true;
 }
 
 /** 
@@ -224,8 +236,8 @@ TEST_CASE("CoreUObject::Serialization::FEditorBulkData::Reader/Writer", "[Core][
  * as bulkdata objects again. This should help find problems where the data saved and loaded are mismatched.
  * Note that this does not check the package saving code paths, only direct serialization to buffers.
  */
-
-TEST_CASE("CoreUObject::Serialization::FEditorBulkData::Serialization To Memory", "[Core][UObject][Smoke]")
+IMPLEMENT_SIMPLE_AUTOMATION_TEST(FEditorBulkDataTestSerializationToMemory, TEXT("System.CoreUObject.Serialization.EditorBulkData.SerializationToMemory"), TestFlags)
+bool FEditorBulkDataTestSerializationToMemory::RunTest(const FString& Parameters)
 {
 	const int64 BufferSize = 1024;
 
@@ -270,14 +282,15 @@ TEST_CASE("CoreUObject::Serialization::FEditorBulkData::Serialization To Memory"
 		TestEqual(TEXT("Bulkdata identifier should remain the same"), SerializedBulkData.GetIdentifier(), EmptyBulkData.GetIdentifier());
 		TestTrue(TEXT("Bulkdata should not have a payload"), SerializedBulkData.GetPayload().Get().IsNull());
 	}
+	return true;
 }
 
 /**
  * This set of tests validate that the BulkData's identifier works how we expect it too. It should remain unique in all cases except
  * move semantics.
  */
-
-TEST_CASE("CoreUObject::Serialization::FEditorBulkData::Identifiers", "[Core][UObject][Smoke]")
+IMPLEMENT_SIMPLE_AUTOMATION_TEST(FEditorBulkDataTestIdentifiers, TEXT("System.CoreUObject.Serialization.EditorBulkData.Identifiers"), TestFlags)
+bool FEditorBulkDataTestIdentifiers::RunTest(const FString& Parameters)
 {
 	// Some basic tests with an invalid id
 	{
@@ -459,9 +472,12 @@ TEST_CASE("CoreUObject::Serialization::FEditorBulkData::Identifiers", "[Core][UO
 		DstData.UpdatePayload(FUniqueBuffer::Alloc(32).MoveToShared());
 		TestEqual(TEXT("After adding a new payload the object should have the original identifier"), DstData.GetIdentifier(), OriginalIdentifier);
 	}
+
+	return true;
 }
 
-TEST_CASE("CoreUObject::Serialization::FEditorBulkData::Zero Sized Allocs", "[Core][UObject][Smoke]")
+IMPLEMENT_SIMPLE_AUTOMATION_TEST(FEditorBulkDataTestZeroSizedAllocs, TEXT("System.CoreUObject.Serialization.EditorBulkData.ZeroSizedAllocs"), TestFlags)
+bool FEditorBulkDataTestZeroSizedAllocs::RunTest(const FString& Parameters)
 {
 	{
 		FEditorBulkData Empty;
@@ -500,12 +516,13 @@ TEST_CASE("CoreUObject::Serialization::FEditorBulkData::Zero Sized Allocs", "[Co
 		TestFalse(TEXT("ZeroAlloc payload has payload data"), ZeroAlloc.HasPayloadData());
 		TestFalse(TEXT("ZeroAlloc payload needs loading"), ZeroAlloc.DoesPayloadNeedLoading());
 	}
+
+	return true;
 }
 
 /** This tests the function UE::Serialization::IoHashToGuid which is closely used with the FEditorBulkData system */
-
-
-TEST_CASE("CoreUObject::Serialization::FEditorBulkData::Io Hash To Guid", "[Core][UObject][Smoke]")
+IMPLEMENT_SIMPLE_AUTOMATION_TEST(FEditorBulkDataIoHashToGuid, TEXT("System.CoreUObject.Serialization.EditorBulkData.IoHashToGuid"), TestFlags)
+bool FEditorBulkDataIoHashToGuid::RunTest(const FString& Parameters)
 {
 	// Test that an empty hash will give an invalid FGuid
 	FIoHash DefaultHash;
@@ -523,10 +540,13 @@ TEST_CASE("CoreUObject::Serialization::FEditorBulkData::Io Hash To Guid", "[Core
 	FGuid KnownResult(TEXT("04030201-0807-0605-0C0B-0A09100F0E0D"));
 
 	TestEqual(TEXT("That the result of hashing known data will be unchanged"), KnownGuid, KnownResult);
+
+	return true;
 }
 
 /** This tests that updating a payload via a FSharedBufferWithID will have the same results as if the payload was applied to the bulkdata object directly */
-TEST_CASE("CoreUObject::Serialization::FEditorBulkData::Shared Buffer With ID", "[Core][UObject][Smoke]")
+IMPLEMENT_SIMPLE_AUTOMATION_TEST(FEditorBulkDataSharedBufferWithID, TEXT("System.CoreUObject.Serialization.EditorBulkData.SharedBufferWithID"), TestFlags)
+bool FEditorBulkDataSharedBufferWithID::RunTest(const FString& Parameters)
 {
 	// Test FSharedBufferWithID/FEditorBulkData in their default states
 	{
@@ -586,9 +606,11 @@ TEST_CASE("CoreUObject::Serialization::FEditorBulkData::Shared Buffer With ID",
 		TestEqual(TEXT("3: PayloadId"), BulkDataFromSharedId.GetPayloadId(), BulkData.GetPayloadId());
 		TestEqual(TEXT("3: PayloadSize"), BulkDataFromSharedId.GetPayloadSize(), BulkData.GetPayloadSize());
 	}
+
+	return true;
 }
 
 
 } // namespace UE::Serialization
 
-#endif //WITH_EDITORONLY_DATA
+#endif //WITH_DEV_AUTOMATION_TESTS && WITH_EDITORONLY_DATA
