Eliminated some redundant stats since they're generally all similar in size and match the parent scope, we probably don't get a lot of value out of this.

Added a UObject scope for widget creation.

diff --git a/Engine/Source/Runtime/CoreUObject/Private/UObject/UObjectGlobals.cpp b/Engine/Source/Runtime/CoreUObject/Private/UObject/UObjectGlobals.cpp
--- a/Engine/Source/Runtime/CoreUObject/Private/UObject/UObjectGlobals.cpp
+++ b/Engine/Source/Runtime/CoreUObject/Private/UObject/UObjectGlobals.cpp
@@ -3531,7 +3531,6 @@ void FObjectInitializer::InitProperties(UObject* Obj, UClass* DefaultsClass, UOb
 		{
 			if (Class->GetDefaultObject(false) != DefaultData)
 			{
-				QUICK_SCOPE_CYCLE_COUNTER(STAT_InitProperties_FromTemplate);
 				for (FProperty* P = Class->PropertyLink; P; P = P->PropertyLinkNext)
 				{
 					bool bIsTransient = P->HasAnyPropertyFlags(CPF_Transient | CPF_DuplicateTransient | CPF_NonPIEDuplicateTransient);
@@ -3546,7 +3545,6 @@ void FObjectInitializer::InitProperties(UObject* Obj, UClass* DefaultsClass, UOb
 			}
 			else
 			{
-				QUICK_SCOPE_CYCLE_COUNTER(STAT_InitProperties_ConfigEtcOnly);
 				// Copy all properties that require additional initialization (e.g. CPF_Config).
 				for (FProperty* P = Class->PostConstructLink; P; P = P->PostConstructLinkNext)
 				{
@@ -3564,8 +3562,6 @@ void FObjectInitializer::InitProperties(UObject* Obj, UClass* DefaultsClass, UOb
 	}
 	else
 	{
-		QUICK_SCOPE_CYCLE_COUNTER(STAT_InitProperties_Blueprint);
-
 		// As with native classes, we must iterate through all properties (slow path) if default data is pointing at something other than the CDO.
 		bCanUsePostConstructLink &= (DefaultData == Class->GetDefaultObject(false));
 
diff --git a/Engine/Source/Runtime/UMG/Public/Blueprint/UserWidget.h b/Engine/Source/Runtime/UMG/Public/Blueprint/UserWidget.h
--- a/Engine/Source/Runtime/UMG/Public/Blueprint/UserWidget.h
+++ b/Engine/Source/Runtime/UMG/Public/Blueprint/UserWidget.h
@@ -1486,6 +1486,7 @@ WidgetT* CreateWidget(OwnerT* OwningObject, TSubclassOf<UUserWidget> UserWidgetC
 		|| TIsDerivedFrom<OwnerT, UWorld>::IsDerived, "The given OwningObject is not of a supported type for use with CreateWidget.");
 
 	SCOPE_CYCLE_COUNTER(STAT_CreateWidget);
+	FScopeCycleCounterUObject WidgetObjectCycleCounter(UserWidgetClass, GET_STATID(STAT_CreateWidget));
 
 	if (OwningObject)
 	{
