Fixing setting properties with setters in script using wrong container pointer

diff --git a/Engine/Source/Runtime/CoreUObject/Private/UObject/ScriptCore.cpp b/Engine/Source/Runtime/CoreUObject/Private/UObject/ScriptCore.cpp
--- a/Engine/Source/Runtime/CoreUObject/Private/UObject/ScriptCore.cpp
+++ b/Engine/Source/Runtime/CoreUObject/Private/UObject/ScriptCore.cpp
@@ -2685,6 +2685,7 @@ DEFINE_FUNCTION(UObject::execLet)
 
 	uint8* LocalTempResult = nullptr;
 	uint8* PreviousPropertyAddress = nullptr;
+	uint8* LocalPropertyContainer = Stack.MostRecentPropertyContainer;
 
 	if (Stack.MostRecentPropertyAddress == nullptr)
 	{
@@ -2721,8 +2722,8 @@ DEFINE_FUNCTION(UObject::execLet)
 	{
 		if (LocallyKnownProperty->HasSetter())
 		{
-			check(Stack.MostRecentPropertyContainer != nullptr);
-			LocallyKnownProperty->SetValue_InContainer(Stack.MostRecentPropertyContainer, LocalTempResult);
+			check(LocalPropertyContainer != nullptr);
+			LocallyKnownProperty->SetValue_InContainer(LocalPropertyContainer, LocalTempResult);
 			Stack.MostRecentPropertyAddress = PreviousPropertyAddress;
 		}
 
