Exposed a helper function in FLinkerInstancingContext used by LinkerLoad, that builds Instanced package names (will be used by world partition).

diff --git a/Engine/Source/Runtime/CoreUObject/Private/UObject/LinkerLoad.cpp b/Engine/Source/Runtime/CoreUObject/Private/UObject/LinkerLoad.cpp
--- a/Engine/Source/Runtime/CoreUObject/Private/UObject/LinkerLoad.cpp
+++ b/Engine/Source/Runtime/CoreUObject/Private/UObject/LinkerLoad.cpp
@@ -1912,7 +1912,7 @@ FLinkerLoad::ELinkerStatus FLinkerLoad::PopulateInstancingContext()
 				// if there's isn't already a remapping for that package, create one
 				if (InstancedName.IsNone())
 				{
-					InstancedName = *FString::Printf(TEXT("%s_InstanceOf_%s"), *InstancingName.ToString(), *LinkerPackageName);
+					InstancedName = *FLinkerInstancingContext::GetInstancedPackageName(LinkerPackageName, InstancingName.ToString());
 				}
 			}
 		}
diff --git a/Engine/Source/Runtime/CoreUObject/Public/UObject/LinkerInstancingContext.h b/Engine/Source/Runtime/CoreUObject/Public/UObject/LinkerInstancingContext.h
--- a/Engine/Source/Runtime/CoreUObject/Public/UObject/LinkerInstancingContext.h
+++ b/Engine/Source/Runtime/CoreUObject/Public/UObject/LinkerInstancingContext.h
@@ -41,6 +41,12 @@ public:
 		Mapping.Append(NewMapping);
 	}
 
+	/** Return the instanced package name for a given instanced outer package and an object package name */
+	static FString GetInstancedPackageName(const FString& InOuterPackageName, const FString& InPackageName)
+	{
+		return FString::Printf(TEXT("%s_InstanceOf_%s"), *InOuterPackageName, *InPackageName);
+	}
+
 private:
 	friend class FLinkerLoad;
 
