Added GetImplementingOuter, an interface counterpart to GetTypedOuter

diff --git a/Engine/Source/Runtime/CoreUObject/Private/UObject/UObjectBaseUtility.cpp b/Engine/Source/Runtime/CoreUObject/Private/UObject/UObjectBaseUtility.cpp
--- a/Engine/Source/Runtime/CoreUObject/Private/UObject/UObjectBaseUtility.cpp
+++ b/Engine/Source/Runtime/CoreUObject/Private/UObject/UObjectBaseUtility.cpp
@@ -287,6 +287,20 @@ UObject* UObjectBaseUtility::GetTypedOuter(UClass* Target) const
 	return Result;
 }
 
+UObjectBaseUtility* UObjectBaseUtility::GetImplementingOuterObject(const UClass* InInterfaceClass) const
+{
+	for (UObject* NextOuter = GetOuter(); NextOuter != nullptr; NextOuter = NextOuter->GetOuter())
+	{
+		UClass* OuterClass = NextOuter->GetClass();
+		if (OuterClass && OuterClass->ImplementsInterface(InInterfaceClass))
+		{
+			return OuterClass;
+		}
+	}
+
+	return nullptr;
+}
+
 /*-----------------------------------------------------------------------------
 	UObject accessors that depend on UClass.
 -----------------------------------------------------------------------------*/
diff --git a/Engine/Source/Runtime/CoreUObject/Public/UObject/UObjectBaseUtility.h b/Engine/Source/Runtime/CoreUObject/Public/UObject/UObjectBaseUtility.h
--- a/Engine/Source/Runtime/CoreUObject/Public/UObject/UObjectBaseUtility.h
+++ b/Engine/Source/Runtime/CoreUObject/Public/UObject/UObjectBaseUtility.h
@@ -604,6 +604,30 @@ public:
 		return (T *)GetTypedOuter(T::StaticClass());
 	}
 
+	/** 
+	 * Traverses the outer chain looking for the next object that implements the specified IInterface (InterfaceClass must be an IInterface)
+	 * 
+	 * @return	a pointer to the interface on the first object in this object's Outer chain which implements the specified interface.
+	 */
+	template<typename InterfaceClassType>
+	InterfaceClassType* GetImplementingOuter() const
+	{
+		UClass* InterfaceClass = InterfaceClassType::UClassType::StaticClass();
+		if(UObjectBaseUtility* ImplementingOuter = GetImplementingOuterObject(InterfaceClass))
+		{
+			return static_cast<InterfaceClassType*>(ImplementingOuter->GetInterfaceAddress(InterfaceClass));
+		}
+		return nullptr;
+	}
+
+	/** 
+	 * Traverses the outer chain looking for the next object that implements the specified UInterface (InInterfaceClass must be a subclass of UInterface)
+	 *
+	 * @param	InInterfaceClass	Target interface to search for
+	 * @return	a pointer to the first object in this object's Outer chain which implements the specified interface.
+	 */
+	UObjectBaseUtility* GetImplementingOuterObject(const UClass* InInterfaceClass) const;
+
 	/** 
 	 * Return the dispatch to `IsInOuter` or `IsInPackage` depending on SomeOuter's class. 
 	 * Legacy function, preferably use IsInOuter or IsInPackage depending on use case.
