Fix-up FSoftObjectPath/FSoftObjectPtr during load of an instanced package

#rnx

diff --git a/Engine/Source/Runtime/CoreUObject/Private/Serialization/AsyncLoading2.cpp b/Engine/Source/Runtime/CoreUObject/Private/Serialization/AsyncLoading2.cpp
--- a/Engine/Source/Runtime/CoreUObject/Private/Serialization/AsyncLoading2.cpp
+++ b/Engine/Source/Runtime/CoreUObject/Private/Serialization/AsyncLoading2.cpp
@@ -1140,7 +1140,6 @@ public:
 
 	//~ Begin FArchive::FArchiveUObject Interface
 	virtual FArchive& operator<<(FObjectPtr& Value) override { return FArchiveUObject::SerializeObjectPtr(*this, Value); }
-	virtual FArchive& operator<<(FSoftObjectPath& Value) override { return FArchiveUObject::SerializeSoftObjectPath(*this, Value); }
 	virtual FArchive& operator<<(FWeakObjectPtr& Value) override { return FArchiveUObject::SerializeWeakObjectPtr(*this, Value); }
 	//~ End FArchive::FArchiveUObject Interface
 
@@ -1233,10 +1232,18 @@ public:
 		FArchive& Ar = *this;
 		FSoftObjectPath ID;
 		ID.Serialize(Ar);
+		FixupSoftObjectPathForInstancedPackage(ID);
 		Value = ID;
 		return Ar;
 	}
 
+	inline virtual FArchive& operator<<(FSoftObjectPath& Value) override
+	{
+		FArchive& Ar = FArchiveUObject::SerializeSoftObjectPath(*this, Value);
+		FixupSoftObjectPathForInstancedPackage(Value);
+		return Ar;
+	}
+
 	FORCENOINLINE void HandleBadNameIndex(int32 NameIndex, FName& Name)
 	{
 		UE_ASYNC_PACKAGE_LOG(Fatal, *PackageDesc, TEXT("ObjectSerializationError"),
@@ -1283,6 +1290,27 @@ private:
 	uint64 CookedSerialOffset = 0;
 	uint64 CookedSerialSize = 0;
 	uint64 BufferSerialOffset = 0;
+
+	/** Set when the package is being loaded as an instance; empty otherwise. */
+	FNameBuilder InstancedPackageSourceName;
+	FNameBuilder InstancedPackageInstanceName;
+	void FixupSoftObjectPathForInstancedPackage(FSoftObjectPath& InOutSoftObjectPath)
+	{
+		if (InstancedPackageSourceName.Len() > 0 && InstancedPackageInstanceName.Len() > 0)
+		{
+			FNameBuilder TmpSoftObjectPathBuilder;
+			InOutSoftObjectPath.ToString(TmpSoftObjectPathBuilder);
+
+			FStringView InstancedPackageSourceNameView = InstancedPackageSourceName.ToView();
+			FStringView TmpSoftObjectPathView = TmpSoftObjectPathBuilder.ToView();
+
+			if (TmpSoftObjectPathView.StartsWith(InstancedPackageSourceNameView) && (TmpSoftObjectPathView.Len() == InstancedPackageSourceNameView.Len() || TmpSoftObjectPathView[InstancedPackageSourceNameView.Len()] == TEXT('.')))
+			{
+				TmpSoftObjectPathBuilder.ReplaceAt(0, InstancedPackageSourceNameView.Len(), InstancedPackageInstanceName.ToView());
+				InOutSoftObjectPath.SetPath(TmpSoftObjectPathBuilder.ToView());
+			}
+		}
+	}
 };
 
 enum class EAsyncPackageLoadingState2 : uint8
@@ -3622,6 +3650,13 @@ EAsyncPackageState::Type FAsyncPackage2::Event_ProcessExportBundle(FAsyncLoading
 			Ar.Exports = Package->Data.Exports;
 			Ar.ExportMap = Package->ExportMap.GetData();
 			Ar.ExternalReadDependencies = &Package->ExternalReadDependencies;
+
+			// Check if the package is instanced
+			if (Package->Desc.UPackageName != Package->Desc.PackageNameToLoad)
+			{
+				Package->Desc.PackageNameToLoad.ToString(Ar.InstancedPackageSourceName);
+				Package->Desc.UPackageName.ToString(Ar.InstancedPackageInstanceName);
+			}
 		}
 		const FExportBundleEntry* BundleEntries = Package->Data.ExportBundleEntries + ExportBundle->FirstEntryIndex;
 		const FExportBundleEntry* BundleEntry = BundleEntries + Package->ExportBundleEntryIndex;
diff --git a/Engine/Source/Runtime/CoreUObject/Private/UObject/LinkerLoad.cpp b/Engine/Source/Runtime/CoreUObject/Private/UObject/LinkerLoad.cpp
--- a/Engine/Source/Runtime/CoreUObject/Private/UObject/LinkerLoad.cpp
+++ b/Engine/Source/Runtime/CoreUObject/Private/UObject/LinkerLoad.cpp
@@ -913,14 +913,17 @@ FLinkerLoad::FLinkerLoad(UPackage* InParent, const FPackagePath& InPackagePath,
 #endif
 	OwnerThread = FPlatformTLS::GetCurrentThreadId();
 
-#if WITH_EDITOR
 	// Check if the linker is instanced @todo: pass through a load flag?
 	FName PackageNameToLoad = GetPackagePath().GetPackageFName();
 	if (LinkerRoot->GetFName() != PackageNameToLoad)
 	{
+		PackageNameToLoad.ToString(InstancedPackageSourceName);
+		LinkerRoot->GetFName().ToString(InstancedPackageInstanceName);
+
+#if WITH_EDITOR
 		InstancingContext.AddMapping(PackageNameToLoad, LinkerRoot->GetFName());
-	}
 #endif
+	}
 }
 
 FLinkerLoad::~FLinkerLoad()
@@ -6105,6 +6108,24 @@ FName FLinkerLoad::InstancingContextRemap(FName ObjectName) const
 #endif
 }
 
+void FLinkerLoad::FixupSoftObjectPathForInstancedPackage(FSoftObjectPath& InOutSoftObjectPath)
+{
+	if (InstancedPackageSourceName.Len() > 0 && InstancedPackageInstanceName.Len() > 0)
+	{
+		FNameBuilder TmpSoftObjectPathBuilder;
+		InOutSoftObjectPath.ToString(TmpSoftObjectPathBuilder);
+
+		FStringView InstancedPackageSourceNameView = InstancedPackageSourceName.ToView();
+		FStringView TmpSoftObjectPathView = TmpSoftObjectPathBuilder.ToView();
+
+		if (TmpSoftObjectPathView.StartsWith(InstancedPackageSourceNameView) && (TmpSoftObjectPathView.Len() == InstancedPackageSourceNameView.Len() || TmpSoftObjectPathView[InstancedPackageSourceNameView.Len()] == TEXT('.')))
+		{
+			TmpSoftObjectPathBuilder.ReplaceAt(0, InstancedPackageSourceNameView.Len(), InstancedPackageInstanceName.ToView());
+			InOutSoftObjectPath.SetPath(TmpSoftObjectPathBuilder.ToView());
+		}
+	}
+}
+
 #if WITH_EDITORONLY_DATA
 /** Performs a fixup on packages' editor-only flag */
 void FixupPackageEditorOnlyFlag(FName PackageThatGotEditorOnlyFlagCleared, bool bRecursive)
diff --git a/Engine/Source/Runtime/CoreUObject/Public/UObject/LinkerLoad.h b/Engine/Source/Runtime/CoreUObject/Public/UObject/LinkerLoad.h
--- a/Engine/Source/Runtime/CoreUObject/Public/UObject/LinkerLoad.h
+++ b/Engine/Source/Runtime/CoreUObject/Public/UObject/LinkerLoad.h
@@ -200,6 +200,11 @@ private:
 	// return ObjectName directly  if WITH_EDITOR isn't defined.
 	FName InstancingContextRemap(FName ObjectName) const;
 
+	/** Set when the package is being loaded as an instance; empty otherwise. */
+	FNameBuilder InstancedPackageSourceName;
+	FNameBuilder InstancedPackageInstanceName;
+	void FixupSoftObjectPathForInstancedPackage(FSoftObjectPath& InOutSoftObjectPath);
+
 protected:
 
 	void SetLoader(FArchive* InLoader, bool bInLoaderNeedsEngineVersionChecks);
@@ -904,9 +909,18 @@ private:
 		FArchive& Ar = *this;
 		FSoftObjectPath ID;
 		ID.Serialize(Ar);
+		FixupSoftObjectPathForInstancedPackage(ID);
 		Value = ID;
 		return Ar;
 	}
+
+	FORCEINLINE virtual FArchive& operator<<(FSoftObjectPath& Value) override
+	{
+		FArchive& Ar = FArchiveUObject::operator<<(Value);
+		FixupSoftObjectPathForInstancedPackage(Value);
+		return Ar;
+	}
+
 	void BadNameIndexError(int32 NameIndex);
 	FORCEINLINE virtual FArchive& operator<<(FName& Name) override
 	{
