Modified FlushTree to FlushRequest and being a simple wrapper around a package loading request
Automatically propagate explicit package load request association to a package dependencies which allow flush request to process only related packages while skipping others
This should also allow FlushAsyncLoading to be re-entrant down the line
Suppress the transaction system while ticking async loading in the main thread
Added an editor config to enabled the async loading thread. Currently non functional

diff --git a/Engine/Source/Runtime/CoreUObject/Private/Serialization/AsyncLoading.cpp b/Engine/Source/Runtime/CoreUObject/Private/Serialization/AsyncLoading.cpp
--- a/Engine/Source/Runtime/CoreUObject/Private/Serialization/AsyncLoading.cpp
+++ b/Engine/Source/Runtime/CoreUObject/Private/Serialization/AsyncLoading.cpp
@@ -387,23 +387,12 @@ void FAsyncLoadingThread::QueuePackage(FAsyncPackageDesc& Package)
 	QueuedRequestsEvent->Trigger();
 }
 
-void FAsyncPackage::PopulateFlushTree(struct FFlushTree* FlushTree)
-{
-	if (FlushTree->AddPackage(GetPackageName()))
-	{
-		for (TTuple<FName, FAsyncPackage*>& PendingImport : PendingImportedPackages)
-		{
-			PendingImport.Value->PopulateFlushTree(FlushTree);
-		}
-	}
-}
-
 FUObjectSerializeContext* FAsyncPackage::GetSerializeContext()
 {
 	return FUObjectThreadContext::Get().GetSerializeContext();
 }
 
-FAsyncPackage* FAsyncLoadingThread::FindExistingPackageAndAddCompletionCallback(FAsyncPackageDesc* PackageRequest, TMap<FName, FAsyncPackage*>& PackageList, FFlushTree* FlushTree)
+FAsyncPackage* FAsyncLoadingThread::FindExistingPackageAndAddCompletionCallback(FAsyncPackageDesc* PackageRequest, TMap<FName, FAsyncPackage*>& PackageList)
 {
 	checkSlow(IsInAsyncLoadThread());
 	FAsyncPackage* Result = PackageList.FindRef(PackageRequest->Name);
@@ -414,11 +403,8 @@ FAsyncPackage* FAsyncLoadingThread::FindExistingPackageAndAddCompletionCallback(
 			const bool bInternalCallback = false;
 			Result->AddCompletionCallback(MoveTemp(PackageRequest->PackageLoadedDelegate), bInternalCallback);
 		}
-			Result->AddRequestID(PackageRequest->RequestID);
-		if (FlushTree)
-		{
-			Result->PopulateFlushTree(FlushTree);
-		}
+		Result->AddRequestID(PackageRequest->RequestID);
+		
 		const int32 QueuedPackagesCount = QueuedPackagesCounter.Decrement();
 		check(QueuedPackagesCount >= 0);
 		NotifyAsyncLoadingStateHasMaybeChanged();
@@ -448,9 +434,9 @@ void FAsyncLoadingThread::UpdateExistingPackagePriorities(FAsyncPackage* InPacka
 	}
 }
 
-void FAsyncLoadingThread::ProcessAsyncPackageRequest(FAsyncPackageDesc* InRequest, FAsyncPackage* InRootPackage, FFlushTree* FlushTree)
+void FAsyncLoadingThread::ProcessAsyncPackageRequest(FAsyncPackageDesc* InRequest, FAsyncPackage* InRootPackage)
 {
-	FAsyncPackage* Package = FindExistingPackageAndAddCompletionCallback(InRequest, AsyncPackageNameLookup, FlushTree);
+	FAsyncPackage* Package = FindExistingPackageAndAddCompletionCallback(InRequest, AsyncPackageNameLookup);
 
 	if (Package)
 	{
@@ -464,7 +450,7 @@ void FAsyncLoadingThread::ProcessAsyncPackageRequest(FAsyncPackageDesc* InReques
 #if THREADSAFE_UOBJECTS
 		FScopeLock LoadedLock(&LoadedPackagesCritical);
 #endif
-		Package = FindExistingPackageAndAddCompletionCallback(InRequest, LoadedPackagesNameLookup, FlushTree);
+		Package = FindExistingPackageAndAddCompletionCallback(InRequest, LoadedPackagesNameLookup);
 	}
 
 	if (!Package)
@@ -473,7 +459,7 @@ void FAsyncLoadingThread::ProcessAsyncPackageRequest(FAsyncPackageDesc* InReques
 #if THREADSAFE_UOBJECTS
 		FScopeLock LoadedLock(&LoadedPackagesToProcessCritical);
 #endif
-		Package = FindExistingPackageAndAddCompletionCallback(InRequest, LoadedPackagesToProcessNameLookup, FlushTree);
+		Package = FindExistingPackageAndAddCompletionCallback(InRequest, LoadedPackagesToProcessNameLookup);
 	}
 
 	if (!Package)
@@ -490,10 +476,6 @@ void FAsyncLoadingThread::ProcessAsyncPackageRequest(FAsyncPackageDesc* InReques
 			Package->AddCompletionCallback(MoveTemp(InRequest->PackageLoadedDelegate), bInternalCallback);
 		}
 		Package->SetDependencyRootPackage(InRootPackage);
-		if (FlushTree)
-		{
-			Package->PopulateFlushTree(FlushTree);
-		}
 
 		// Add to queue according to priority.
 		InsertPackage(Package, false, EAsyncPackageInsertMode::InsertAfterMatchingPriorities);
@@ -505,7 +487,7 @@ void FAsyncLoadingThread::ProcessAsyncPackageRequest(FAsyncPackageDesc* InReques
 	}
 }
 
-int32 FAsyncLoadingThread::CreateAsyncPackagesFromQueue(bool bUseTimeLimit, bool bUseFullTimeLimit, float TimeLimit, FFlushTree* FlushTree)
+int32 FAsyncLoadingThread::CreateAsyncPackagesFromQueue(bool bUseTimeLimit, bool bUseFullTimeLimit, float TimeLimit)
 {
 	SCOPED_LOADTIMER(CreateAsyncPackagesFromQueueTime);
 
@@ -564,7 +546,7 @@ int32 FAsyncLoadingThread::CreateAsyncPackagesFromQueue(bool bUseTimeLimit, bool
 				SCOPE_SECONDS_COUNTER(Timer);
 				for (FAsyncPackageDesc* PackageRequest : QueueCopy)
 				{
-					ProcessAsyncPackageRequest(PackageRequest, nullptr, FlushTree);
+					ProcessAsyncPackageRequest(PackageRequest, nullptr);
 					delete PackageRequest;
 				}
 			}
@@ -1818,6 +1800,7 @@ EAsyncPackageState::Type FAsyncPackage::LoadImports_Event()
 		}
 		if (PendingPackage)
 		{
+			AddDependentRequestsTo(PendingPackage);
 			if (int32(PendingPackage->AsyncPackageLoadingState) <= int32(EAsyncPackageLoadingState::WaitingForSummary))
 			{
 				FEventLoadNodePtr PrereqisiteNode;
@@ -3868,7 +3851,7 @@ struct FScopedRecursionNotAllowed
 };
 #endif
 
-EAsyncPackageState::Type FAsyncLoadingThread::ProcessAsyncLoading(int32& OutPackagesProcessed, bool bUseTimeLimit /*= false*/, bool bUseFullTimeLimit /*= false*/, float TimeLimit /*= 0.0f*/, FFlushTree* FlushTree)
+EAsyncPackageState::Type FAsyncLoadingThread::ProcessAsyncLoading(int32& OutPackagesProcessed, bool bUseTimeLimit /*= false*/, bool bUseFullTimeLimit /*= false*/, float TimeLimit /*= 0.0f*/, FFlushRequest& FlushRequest)
 {
 	SCOPED_LOADTIMER(AsyncLoadingTime);
 	check(!IsInGameThread() || !IsMultithreaded());
@@ -3965,7 +3948,7 @@ EAsyncPackageState::Type FAsyncLoadingThread::ProcessAsyncLoading(int32& OutPack
 				if (Package->HasFinishedLoading() == false)
 				{
 					float RemainingTimeLimit = FMath::Max(0.0f, TimeLimit - (float)(FPlatformTime::Seconds() - TickStartTime));
-					LocalLoadingState = Package->TickAsyncPackage(bUseTimeLimit, bUseFullTimeLimit, RemainingTimeLimit, FlushTree);
+					LocalLoadingState = Package->TickAsyncPackage(bUseTimeLimit, bUseFullTimeLimit, RemainingTimeLimit);
 					if (LocalLoadingState == EAsyncPackageState::TimeOut)
 					{
 						if (IsTimeLimitExceeded(TickStartTime, bUseTimeLimit, TimeLimit, TEXT("TickAsyncPackage")))
@@ -4062,7 +4045,7 @@ EAsyncPackageState::Type FAsyncLoadingThread::ProcessAsyncLoading(int32& OutPack
 
 			// Package to be loaded.
 			FAsyncPackage* Package = AsyncPackages[PackageIndex];
-			if (FlushTree && !FlushTree->Contains(Package->GetPackageName()))
+			if (!ShouldProcessPackage(Package, FlushRequest))
 			{
 				LoadingState = EAsyncPackageState::PendingImports;
 			}
@@ -4070,7 +4053,7 @@ EAsyncPackageState::Type FAsyncLoadingThread::ProcessAsyncLoading(int32& OutPack
 			{
 				// Package tick returns EAsyncPackageState::Complete on completion.
 				// We only tick packages that have not yet been loaded.
-				LoadingState = Package->TickAsyncPackage(bUseTimeLimit, bUseFullTimeLimit, TimeLimit, FlushTree);
+				LoadingState = Package->TickAsyncPackage(bUseTimeLimit, bUseFullTimeLimit, TimeLimit);
 			}
 			else
 			{
@@ -4099,6 +4082,7 @@ EAsyncPackageState::Type FAsyncLoadingThread::ProcessAsyncLoading(int32& OutPack
 
 			}
 
+			//@note FH: if we are flushing should we be doing this? are those exclusive external request or could be dependencies being pulled in?
 			{ // maybe we shouldn't do this if we are already out of time?
 				// Check if there's any new packages in the queue.
 				const float RemainingTimeLimit = FMath::Max(0.0f, TimeLimit - (float)(FPlatformTime::Seconds() - TickStartTime));
@@ -4117,6 +4101,8 @@ EAsyncPackageState::Type FAsyncLoadingThread::ProcessAsyncLoading(int32& OutPack
 
 bool FAsyncPackage::AreAllDependenciesFullyLoadedInternal(FAsyncPackage* Package, TSet<UPackage*>& VisitedPackages, FString& OutError)
 {
+	// @note FH: this isn't properly thread-safe at all...
+
 	for (UPackage* ImportPackage : Package->ImportedPackages)
 	{
 		if (!ImportPackage || VisitedPackages.Contains(ImportPackage))
@@ -4160,6 +4146,33 @@ bool FAsyncPackage::AreAllDependenciesFullyLoadedInternal(FAsyncPackage* Package
 	return true;
 }
 
+
+bool FAsyncPackage::IsDependencyOf(int32 RequestID)
+{
+	FReadScopeLock Lock(DependentRequestsLock);
+	return DependentRequests.Contains(RequestID);
+}
+
+void FAsyncPackage::AddDependentRequestsTo(FAsyncPackage* InDependency)
+{
+	TArray<int32> Requests;
+	{
+		FReadScopeLock Lock(DependentRequestsLock);
+		Requests = DependentRequests.Array();
+	}
+	InDependency->AddDependentRequests(Requests);
+}
+
+void FAsyncPackage::AddDependentRequests(TArrayView<int32 const> Requests)
+{
+	FWriteScopeLock Lock(DependentRequestsLock);
+	DependentRequests.Reserve(DependentRequests.Num() + Requests.Num());
+	for (const int32 Request : Requests)
+	{
+		DependentRequests.Add(Request);
+	}
+}
+
 bool FAsyncPackage::AreAllDependenciesFullyLoaded(TSet<UPackage*>& VisitedPackages)
 {
 	VisitedPackages.Reset();
@@ -4172,7 +4185,85 @@ bool FAsyncPackage::AreAllDependenciesFullyLoaded(TSet<UPackage*>& VisitedPackag
 	return bLoaded;
 }
 
-EAsyncPackageState::Type FAsyncLoadingThread::ProcessLoadedPackages(bool bUseTimeLimit, bool bUseFullTimeLimit, float TimeLimit, bool& bDidSomething, FFlushTree* FlushTree)
+void FAsyncLoadingThread::CompleteFlushRequests(TArray<int32>& RequestIDs)
+{
+	FScopeLock Lock(&FlushRequestCritical);
+	for (int32 RequestID : RequestIDs)
+	{
+		FlushRequests.RemoveAll([RequestID](const FFlushRequest& Request)
+			{
+				return Request.GetId() == RequestID;
+			});
+		DuplicateRequestMap.Remove(RequestID);
+	}
+}
+
+void FAsyncLoadingThread::MapDuplicateRequestID(int32 DuplicateRequestId, int32 MainRequestId)
+{
+	DuplicateRequestMap.Add(DuplicateRequestId, MainRequestId);
+}
+
+FFlushRequest FAsyncLoadingThread::AddFlushRequest(int32 RequestID)
+{
+	FFlushRequest FlushRequest = FFlushRequest(RequestID);
+	{
+		FScopeLock Lock(&FlushRequestCritical);
+		FlushRequests.Push(FlushRequest);
+	}
+	return FlushRequest;
+}
+
+
+FFlushRequest FAsyncLoadingThread::PeekFlushRequest()
+{
+	FFlushRequest FlushRequest;
+	{
+		FScopeLock Lock(&FlushRequestCritical);
+		if (FlushRequests.Num())
+		{
+			FlushRequest = FlushRequests.Top();
+		}
+	}
+
+	if (FlushRequest)
+	{
+		// Validate that the flush request is still valid and pending
+		if (!ContainsRequestInternal(FlushRequest.GetId()))
+		{
+			FScopeLock Lock(&FlushRequestCritical);
+			FFlushRequest Top = FlushRequests.Pop();
+			ensure(Top == FlushRequest);
+			FlushRequest.Reset();
+		}
+	}
+
+	// Adjust request id
+	AdjustFlushRequest(FlushRequest);
+	return FlushRequest;
+}
+
+void FAsyncLoadingThread::AdjustFlushRequest(FFlushRequest& FlushRequest)
+{
+	if (FlushRequest)
+	{
+		if (int32* MainRequest = DuplicateRequestMap.Find(FlushRequest.GetId()))
+		{
+			FlushRequest.SetId(*MainRequest);
+		}
+	}
+}
+
+bool FAsyncLoadingThread::ShouldProcessPackage(FAsyncPackage* InAsyncPackage, const FFlushRequest& FlushRequest)
+{
+	if (!FlushRequest)
+	{
+		return true;
+	}
+
+	return InAsyncPackage->IsDependencyOf(FlushRequest.GetId());
+}
+
+EAsyncPackageState::Type FAsyncLoadingThread::ProcessLoadedPackages(bool bUseTimeLimit, bool bUseFullTimeLimit, float TimeLimit, bool& bDidSomething, const FFlushRequest& FlushRequest)
 {
 	EAsyncPackageState::Type Result = EAsyncPackageState::Complete;
 
@@ -4218,13 +4309,20 @@ EAsyncPackageState::Type FAsyncLoadingThread::ProcessLoadedPackages(bool bUseTim
 		FPlatformMisc::PumpEssentialAppMessages();
 
 		FAsyncPackage* Package = LoadedPackagesToProcess[PackageIndex];
-		if (Package->GetDependencyRefCount() == 0)
+		// if we are currently flushing a specific load request, only post load related package objects
+		if (ShouldProcessPackage(Package, FlushRequest))
 		{
 			SCOPED_LOADTIMER(ProcessLoadedPackagesTime);
-
 			Result = Package->PostLoadDeferredObjects(TickStartTime, bUseTimeLimit, TimeLimit);
-			if (Result == EAsyncPackageState::Complete)
+			// if we didn't complete the post-load, we ran out of time for this slice so break immediately
+			if (Result != EAsyncPackageState::Complete)
+			{
+				break;
+			}
+			// if we are done with post load and there isn't any package currently depending on this package we can start to clean it up
+			else if (Package->GetDependencyRefCount() == 0)
 			{
+				check(Result == EAsyncPackageState::Complete);
 				// Remove the package from the list before we trigger the callbacks, 
 				// this is to ensure we can re-enter FlushAsyncLoading from any of the callbacks
 				{
@@ -4268,22 +4366,12 @@ EAsyncPackageState::Type FAsyncLoadingThread::ProcessLoadedPackages(bool bUseTim
 
 				TRACE_LOADTIME_END_LOAD_ASYNC_PACKAGE(Package);
 
-				if (IsTimeLimitExceeded(TickStartTime, bUseTimeLimit, TimeLimit, TEXT("ProcessLoadedPackages Misc")) || (FlushTree && !ContainsRequestID(FlushTree->RequestId)))
+				if (IsTimeLimitExceeded(TickStartTime, bUseTimeLimit, TimeLimit, TEXT("ProcessLoadedPackages Misc")) || (FlushRequest && !ContainsRequestInternal(FlushRequest.GetId())))
 				{
 					// The only package we care about has finished loading, so we're good to exit
 					break;
 				}
 			}
-			else
-			{
-				break;
-			}
-		}
-		else
-		{
-			Result = EAsyncPackageState::PendingImports;
-			// Break immediately, we want to keep the order of processing when packages get here
-			break;
 		}
 	}
 
@@ -4406,13 +4494,19 @@ EAsyncPackageState::Type FAsyncLoadingThread::ProcessLoadedPackages(bool bUseTim
 
 
 
-EAsyncPackageState::Type FAsyncLoadingThread::TickAsyncLoading(bool bUseTimeLimit, bool bUseFullTimeLimit, float TimeLimit, FFlushTree* FlushTree)
+EAsyncPackageState::Type FAsyncLoadingThread::TickAsyncLoading(bool bUseTimeLimit, bool bUseFullTimeLimit, float TimeLimit, FFlushRequest FlushRequest)
 {
 	LLM_SCOPE(ELLMTag::AsyncLoading);
 
 	check(IsInGameThread());
 	check(!IsGarbageCollecting());
 
+#if WITH_EDITOR
+	// In the editor loading cannot be part of a transaction as it cannot be undone, and may result in recording half-loaded objects. So we suppress any active transaction while in this stack, and set the editor loading flag
+	TGuardValue<ITransaction*> SuppressTransaction(GUndo, nullptr);
+	TGuardValue<bool> IsEditorLoadingPackage(GIsEditorLoadingPackage, GIsEditor || GIsEditorLoadingPackage);
+#endif
+
 	const bool bLoadingSuspended = IsAsyncLoadingSuspendedInternal();
 	EAsyncPackageState::Type Result = bLoadingSuspended ? EAsyncPackageState::PendingImports : EAsyncPackageState::Complete;
 
@@ -4440,7 +4534,7 @@ EAsyncPackageState::Type FAsyncLoadingThread::TickAsyncLoading(bool bUseTimeLimi
 
 		bool bDidSomething = false;
 		{
-			Result = ProcessLoadedPackages(bUseTimeLimit, bUseFullTimeLimit, TimeLimit, bDidSomething, FlushTree);
+			Result = ProcessLoadedPackages(bUseTimeLimit, bUseFullTimeLimit, TimeLimit, bDidSomething, FlushRequest);
 			TimeLimitUsedForProcessLoaded = FPlatformTime::Seconds() - TickStartTime;
 			UE_CLOG(!GIsEditor && bUseTimeLimit && TimeLimitUsedForProcessLoaded > .1f, LogStreaming, Warning, TEXT("Took %6.2fms to ProcessLoadedPackages"), float(TimeLimitUsedForProcessLoaded) * 1000.0f);
 		}
@@ -4448,7 +4542,8 @@ EAsyncPackageState::Type FAsyncLoadingThread::TickAsyncLoading(bool bUseTimeLimi
 		if (!bIsMultithreaded && Result != EAsyncPackageState::TimeOut && !IsTimeLimitExceeded(TickStartTime, bUseTimeLimit, TimeLimit, TEXT("ProcessLoadedPackages")))
 		{
 			double RemainingTimeLimit = FMath::Max(0.0, TimeLimit - TimeLimitUsedForProcessLoaded);
-			Result = TickAsyncThread(bUseTimeLimit, bUseFullTimeLimit, RemainingTimeLimit, bDidSomething, FlushTree);
+			AdjustFlushRequest(FlushRequest);
+			Result = TickAsyncThread(bUseTimeLimit, bUseFullTimeLimit, RemainingTimeLimit, bDidSomething, FlushRequest);
 		}
 
 		if (Result != EAsyncPackageState::TimeOut && !IsTimeLimitExceeded(TickStartTime, bUseTimeLimit, TimeLimit, TEXT("TickAsyncThread")))
@@ -4650,7 +4745,10 @@ uint32 FAsyncLoadingThread::Run()
 			if (!IsGarbageCollectionWaiting())
 			{
 				bool bDidSomething = false;
-				TickAsyncThread(true, false, 0.033f, bDidSomething);
+
+				// Peek if we currently have a flush request going on.
+				FFlushRequest FlushRequest = PeekFlushRequest();
+				TickAsyncThread(true, false, 0.033f, bDidSomething, FlushRequest);
 			}
 		}
 		else if (!bWasSuspendedLastFrame)
@@ -4697,7 +4795,7 @@ void FAsyncLoadingThread::CheckForCycles()
 }
 
 
-EAsyncPackageState::Type FAsyncLoadingThread::TickAsyncThread(bool bUseTimeLimit, bool bUseFullTimeLimit, float TimeLimit, bool& bDidSomething, FFlushTree* FlushTree)
+EAsyncPackageState::Type FAsyncLoadingThread::TickAsyncThread(bool bUseTimeLimit, bool bUseFullTimeLimit, float TimeLimit, bool& bDidSomething, FFlushRequest& FlushRequest)
 {
 	check(!IsInGameThread() || !IsMultithreaded());
 	EAsyncPackageState::Type Result = EAsyncPackageState::Complete;
@@ -4712,8 +4810,9 @@ EAsyncPackageState::Type FAsyncLoadingThread::TickAsyncThread(bool bUseTimeLimit
 				EDLBootNotificationManager.FireCompletedCompiledInImports();
 			}
 			{
+				//@note FH: pretty sure the GCScopeGuard isn't needed here, nor is it useful, it is also retaken inside CreateAsyncPackagesFromQueue which isn't necessary
 				FGCScopeGuard GCGuard;
-				CreateAsyncPackagesFromQueue(bUseTimeLimit, bUseFullTimeLimit, TimeLimit, FlushTree);
+				CreateAsyncPackagesFromQueue(bUseTimeLimit, bUseFullTimeLimit, TimeLimit); 
 			}
 			float TimeUsed = (float)(FPlatformTime::Seconds() - TickStartTime);
 			const float RemainingTimeLimit = FMath::Max(0.0f, TimeLimit - TimeUsed);
@@ -4724,7 +4823,7 @@ EAsyncPackageState::Type FAsyncLoadingThread::TickAsyncThread(bool bUseTimeLimit
 			else
 			{
 				FGCScopeGuard GCGuard;
-				Result = ProcessAsyncLoading(ProcessedRequests, bUseTimeLimit, bUseFullTimeLimit, RemainingTimeLimit, FlushTree);
+				Result = ProcessAsyncLoading(ProcessedRequests, bUseTimeLimit, bUseFullTimeLimit, RemainingTimeLimit, FlushRequest);
 				bDidSomething = bDidSomething || ProcessedRequests > 0;
 			}
 		}
@@ -5090,6 +5189,7 @@ FAsyncPackage::FAsyncPackage(FAsyncLoadingThread& InThread, const FAsyncPackageD
 {
 	TRACE_LOADTIME_NEW_ASYNC_PACKAGE(this);
 	AddRequestID(InDesc.RequestID);
+	AddDependentRequests(MakeArrayView<const int32>(&InDesc.RequestID, 1));
 }
 
 FAsyncPackage::~FAsyncPackage()
@@ -5177,6 +5277,12 @@ void FAsyncPackage::AddRequestID(int32 Id)
 			// For debug readability
 			Desc.RequestID = Id;
 		}
+		else if (Desc.RequestID != Id)
+		{
+			// If this is a duplicate id, map it to the original request.
+			// this is to properly handle flush request which may not have propagated request dependencies to a duplicate request made later
+			AsyncLoadingThread.MapDuplicateRequestID(Id, Desc.RequestID);
+		}
 		RequestIDs.Add(Id);
 		AsyncLoadingThread.AddPendingRequest(Id);
 		TRACE_LOADTIME_ASYNC_PACKAGE_REQUEST_ASSOCIATION(this, Id);
@@ -5186,6 +5292,7 @@ void FAsyncPackage::AddRequestID(int32 Id)
 void FAsyncPackage::MarkRequestIDsAsComplete()
 {
 	AsyncLoadingThread.RemovePendingRequests(RequestIDs);
+	AsyncLoadingThread.CompleteFlushRequests(RequestIDs);
 	RequestIDs.Reset();
 }
 
@@ -5321,6 +5428,7 @@ void FAsyncPackage::EndAsyncLoad()
 
 	if (!bLoadHasFailed)
 	{
+		// @note FH: This is actually unsafer if multiple thread plays with the object flags...
 		// Mark the package as loaded, if we succeeded
 		LinkerRoot->SetFlags(RF_WasLoaded);
 	}
@@ -5336,11 +5444,13 @@ void FAsyncPackage::EndAsyncLoad()
  * @return	true if package has finished loading, false otherwise
  */
 
-EAsyncPackageState::Type FAsyncPackage::TickAsyncPackage(bool InbUseTimeLimit, bool InbUseFullTimeLimit, float& InOutTimeLimit, FFlushTree* FlushTree)
+EAsyncPackageState::Type FAsyncPackage::TickAsyncPackage(bool InbUseTimeLimit, bool InbUseFullTimeLimit, float& InOutTimeLimit)
 {
 
 	// We want this check only with EDL enabled
 	check(!GEventDrivenLoaderEnabled || int32(AsyncPackageLoadingState) > int32(EAsyncPackageLoadingState::ProcessNewImportsAndExports));
+	// @note FH: We shouldn't be ticking this package from the async thread once we are have finished loading it
+	checkSlow(!HasFinishedLoading());
 
 	ReentryCount++;
 
@@ -5409,7 +5519,7 @@ EAsyncPackageState::Type FAsyncPackage::TickAsyncPackage(bool InbUseTimeLimit, b
 			if (LoadingState == EAsyncPackageState::Complete)
 			{
 				SCOPED_LOADTIMER(Package_LoadImports);
-				LoadingState = LoadImports(FlushTree);
+				LoadingState = LoadImports();
 			}
 
 			// Create imports from linker import table.
@@ -5493,6 +5603,7 @@ EAsyncPackageState::Type FAsyncPackage::TickAsyncPackage(bool InbUseTimeLimit, b
 
 	check(bUseTimeLimit || LoadingState != EAsyncPackageState::TimeOut || AsyncLoadingThread.IsAsyncLoadingSuspendedInternal() || IsGarbageCollectionWaiting());
 
+	// @note FH: Should the package be marked as fully loaded if there are still pending post-load?
 	if (LinkerRoot && LoadingState == EAsyncPackageState::Complete)
 	{
 		LinkerRoot->MarkAsFullyLoaded();
@@ -5781,24 +5892,13 @@ EAsyncPackageState::Type FAsyncPackage::FinishLinker()
 	return Result;
 }
 
-/**
- * Adds a package to the list of pending import packages.
- *
- * @param PendingImport Name of the package imported either directly or by one of the imported packages
- */
-void FAsyncPackage::AddImportDependency(const FName& PendingImport, FFlushTree* FlushTree)
-{
-	AddImportDependency(PendingImport, NAME_None, FlushTree, FLinkerInstancingContext());
-}
-
 /**
  * Adds a package to the list of pending import packages.
  *
  * @param PendingImport Name of the package imported either directly or by one of the imported packages
  * @param PackageToLoad Name of the package to load into PendingImport
- * @param FlushTree Package dependency tree to be flushed
  */
-void FAsyncPackage::AddImportDependency(const FName& PendingImport, const FName& PackageToLoad, FFlushTree* FlushTree, FLinkerInstancingContext InstancingContext)
+void FAsyncPackage::AddImportDependency(const FName& PendingImport, const FName& PackageToLoad, FLinkerInstancingContext InstancingContext)
 {
 	FAsyncPackage* PackageToStream = AsyncLoadingThread.FindAsyncPackage(PendingImport);
 	const bool bReinsert = PackageToStream != nullptr;
@@ -5823,18 +5923,17 @@ void FAsyncPackage::AddImportDependency(const FName& PendingImport, const FName&
 		AsyncLoadingThread.InsertPackage(PackageToStream, bReinsert);
 	}
 
+	// Add dependent requests (any request that depends on this package will also depends on this package dependencies)
+	AddDependentRequestsTo(PackageToStream);
+
 	if (!PackageToStream->HasFinishedLoading() && 
 		!PackageToStream->bLoadHasFailed)
 	{
 		const bool bInternalCallback = true;
-		TUniquePtr<FLoadPackageAsyncDelegate> InternalDelegate = MakeUnique<FLoadPackageAsyncDelegate>(FLoadPackageAsyncDelegate::CreateRaw(this, &FAsyncPackage::ImportFullyLoadedCallback));
+		TUniquePtr<FLoadPackageAsyncDelegate> InternalDelegate = MakeUnique<FLoadPackageAsyncDelegate>(FLoadPackageAsyncDelegate::CreateRaw(this, &FAsyncPackage::ImportPreLoadedAsyncCallback));
 		PackageToStream->AddCompletionCallback(MoveTemp(InternalDelegate), bInternalCallback);
 		PackageToStream->DependencyRefCount.Increment();
 		PendingImportedPackages.Emplace(PackageToStream->GetPackageName(), PackageToStream);
-		if (FlushTree)
-		{
-			PackageToStream->PopulateFlushTree(FlushTree);
-		}
 	}
 	else
 	{
@@ -5849,14 +5948,14 @@ void FAsyncPackage::AddImportDependency(const FName& PendingImport, const FName&
  *
  * @param PendingImport Package imported either directly or by one of the imported packages
  */
-bool FAsyncPackage::AddUniqueLinkerDependencyPackage(FAsyncPackage& PendingImport, FFlushTree* FlushTree)
+bool FAsyncPackage::AddUniqueLinkerDependencyPackage(FAsyncPackage& PendingImport)
 {
 	if (!PendingImportedPackages.Contains(PendingImport.GetPackageName()))
 	{
 		FLinkerLoad* PendingImportLinker = PendingImport.Linker;
 		if (PendingImportLinker == nullptr || !PendingImportLinker->HasFinishedInitialization())
 		{
-			AddImportDependency(PendingImport.GetPackageName(), FlushTree);
+			AddImportDependency(PendingImport.GetPackageName(), NAME_None, FLinkerInstancingContext());
 			UE_LOG(LogStreaming, Verbose, TEXT("  Adding linker dependency %s"), *PendingImport.GetPackageName().ToString());
 		}
 		else if (this != &PendingImport)
@@ -5872,7 +5971,7 @@ bool FAsyncPackage::AddUniqueLinkerDependencyPackage(FAsyncPackage& PendingImpor
  *
  * @param ImportedPackage Package imported either directly or by one of the imported packages
  */
-void FAsyncPackage::AddDependencyTree(FAsyncPackage& ImportedPackage, TSet<FAsyncPackage*>& SearchedPackages, FFlushTree* FlushTree)
+void FAsyncPackage::AddDependencyTree(FAsyncPackage& ImportedPackage, TSet<FAsyncPackage*>& SearchedPackages)
 {
 	if (SearchedPackages.Contains(&ImportedPackage))
 	{
@@ -5882,9 +5981,9 @@ void FAsyncPackage::AddDependencyTree(FAsyncPackage& ImportedPackage, TSet<FAsyn
 	for (TTuple<FName, FAsyncPackage*>& Kvp : ImportedPackage.PendingImportedPackages)
 	{
 		FAsyncPackage& PendingImport = *Kvp.Value;
-		if (!AddUniqueLinkerDependencyPackage(PendingImport, FlushTree))
+		if (!AddUniqueLinkerDependencyPackage(PendingImport))
 		{
-			AddDependencyTree(PendingImport, SearchedPackages, FlushTree);
+			AddDependencyTree(PendingImport, SearchedPackages);
 		}
 	}
 	// Mark this package as searched
@@ -5896,7 +5995,7 @@ void FAsyncPackage::AddDependencyTree(FAsyncPackage& ImportedPackage, TSet<FAsyn
  *
  * @return true if we finished load all imports, false otherwise
  */
-EAsyncPackageState::Type FAsyncPackage::LoadImports(FFlushTree* FlushTree)
+EAsyncPackageState::Type FAsyncPackage::LoadImports()
 {
 	TRACE_CPUPROFILER_EVENT_SCOPE(FAsyncPackage::LoadImports);
 	SCOPE_CYCLE_COUNTER(STAT_FAsyncPackage_LoadImports);
@@ -5961,7 +6060,7 @@ EAsyncPackageState::Type FAsyncPackage::LoadImports(FFlushTree* FlushTree)
 				if (PendingPackageLinker == nullptr || !PendingPackageLinker->HasFinishedInitialization())
 				{
 					// Add this import to the dependency list.
-					AddUniqueLinkerDependencyPackage(*PendingPackage, FlushTree);
+					AddUniqueLinkerDependencyPackage(*PendingPackage);
 				}
 				else
 				{
@@ -5971,7 +6070,7 @@ EAsyncPackageState::Type FAsyncPackage::LoadImports(FFlushTree* FlushTree)
 					ReferencedImports.Add(PendingPackage);
 					// Check if we need to add its dependencies too.
 					TSet<FAsyncPackage*> SearchedPackages;
-					AddDependencyTree(*PendingPackage, SearchedPackages, FlushTree);
+					AddDependencyTree(*PendingPackage, SearchedPackages);
 				}
 			}
 		}
@@ -5983,7 +6082,7 @@ EAsyncPackageState::Type FAsyncPackage::LoadImports(FFlushTree* FlushTree)
 			if (!FPackageName::IsShortPackageName(ImportPackageName))
 			{
 				UE_LOG(LogStreaming, Verbose, TEXT("FAsyncPackage::LoadImports for %s: Loading %s"), *Desc.PackagePath.GetDebugName(), *ImportPackageName);
-				AddImportDependency(ImportPackageFName, ImportToLoad, FlushTree, InstancingContext);
+				AddImportDependency(ImportPackageFName, ImportToLoad, InstancingContext);
 			}
 			else
 			{
@@ -6004,10 +6103,12 @@ EAsyncPackageState::Type FAsyncPackage::LoadImports(FFlushTree* FlushTree)
 }
 
 /**
- * Function called when pending import package has been fully loaded.
+ * this callback is triggered once the package exports are created and de-serialized but not post loaded
  */
-void FAsyncPackage::ImportFullyLoadedCallback(const FName& InPackageName, UPackage* LoadedPackage, EAsyncLoadingResult::Type Result)
+void FAsyncPackage::ImportPreLoadedAsyncCallback(const FName& InPackageName, UPackage* LoadedPackage, EAsyncLoadingResult::Type Result)
 {
+	// @note FH: the array here should only be accessed from the async thread
+	check(IsInAsyncLoadingThread());
 	if (Result != EAsyncLoadingResult::Canceled)
 	{
 		UE_LOG(LogStreaming, Verbose, TEXT("FAsyncPackage::LoadImports for %s: Loaded %s"), *Desc.PackagePath.GetDebugName(), *InPackageName.ToString());
@@ -6363,6 +6464,7 @@ EAsyncPackageState::Type FAsyncPackage::PostLoadDeferredObjects(double InTickSta
 	SCOPED_LOADTIMER(PostLoadDeferredObjectsTime);
 
 	FAsyncPackageScope PackageScope(this);
+	++ReentryCount;
 
 	EAsyncPackageState::Type Result = EAsyncPackageState::Complete;
 	FAsyncLoadingTickScope InAsyncLoadingTick(AsyncLoadingThread);
@@ -6467,7 +6569,7 @@ EAsyncPackageState::Type FAsyncPackage::PostLoadDeferredObjects(double InTickSta
 		TArray<UObject*> CDODefaultSubobjects;
 		// Clear async loading flags (we still want RF_Async, but EInternalObjectFlags::AsyncLoading can be cleared)
 		while (DeferredFinalizeIndex < DeferredFinalizeObjects.Num() &&
-			(DeferredPostLoadIndex % 100 != 0 || (!AsyncLoadingThread.IsAsyncLoadingSuspendedInternal() && !::IsTimeLimitExceeded(InTickStartTime, bInUseTimeLimit, InOutTimeLimit, LastTypeOfWorkPerformed, LastObjectWorkWasPerformedOn))))
+			((!AsyncLoadingThread.IsAsyncLoadingSuspendedInternal() && !::IsTimeLimitExceeded(InTickStartTime, bInUseTimeLimit, InOutTimeLimit, LastTypeOfWorkPerformed, LastObjectWorkWasPerformedOn))))
 		{
 			UObject* Object = DeferredFinalizeObjects[DeferredFinalizeIndex++];
 			if (Object)
@@ -6493,6 +6595,8 @@ EAsyncPackageState::Type FAsyncPackage::PostLoadDeferredObjects(double InTickSta
 			}
 		}
 		::IsTimeLimitExceeded(InTickStartTime, bInUseTimeLimit, InOutTimeLimit, LastTypeOfWorkPerformed, LastObjectWorkWasPerformedOn);
+
+		//@ note FH: this might currently get called multiple time, would need to had a check for being already finalized 
 		if (DeferredFinalizeIndex == DeferredFinalizeObjects.Num())
 		{
 			DeferredFinalizeIndex = 0;
@@ -6524,6 +6628,7 @@ EAsyncPackageState::Type FAsyncPackage::PostLoadDeferredObjects(double InTickSta
 		FUniqueObjectGuid::InvalidateTag();
 	}
 
+	--ReentryCount;
 	return Result;
 }
 
@@ -6847,7 +6952,7 @@ void FAsyncLoadingThread::FlushLoading(int32 PackageID)
 
 		SCOPE_CYCLE_COUNTER(STAT_FAsyncPackage_FlushAsyncLoadingGameThread);
 
-		if (PackageID != INDEX_NONE && !ContainsRequestID(PackageID))
+		if (PackageID != INDEX_NONE && !ContainsRequestInternal(PackageID))
 		{
 			return;
 		}
@@ -6858,16 +6963,17 @@ void FAsyncLoadingThread::FlushLoading(int32 PackageID)
 
 		// Flush async loaders without using a time limit. Needed for e.g. garbage collection.
 		{
-			TUniquePtr<FFlushTree> FlushTree;
+			FFlushRequest FlushRequest;
 			if (PackageID != INDEX_NONE)
 			{
-				FlushTree = MakeUnique<FFlushTree>(PackageID);
+				// This create and add a flush request to use on the async thread
+				FlushRequest = AddFlushRequest(PackageID);
 			}
 			SCOPE_CYCLE_COUNTER(STAT_FAsyncPackage_TickAsyncLoadingGameThread);
 			while (IsAsyncLoadingPackages())
 			{
-				EAsyncPackageState::Type Result = TickAsyncLoading(false, false, 0, FlushTree.Get());
-				if (PackageID != INDEX_NONE && !ContainsRequestID(PackageID))
+				EAsyncPackageState::Type Result = TickAsyncLoading(false, false, 0, FlushRequest);
+				if (PackageID != INDEX_NONE && !ContainsRequestInternal(PackageID))
 				{
 					break;
 				}
diff --git a/Engine/Source/Runtime/CoreUObject/Private/Serialization/AsyncLoading.h b/Engine/Source/Runtime/CoreUObject/Private/Serialization/AsyncLoading.h
--- a/Engine/Source/Runtime/CoreUObject/Private/Serialization/AsyncLoading.h
+++ b/Engine/Source/Runtime/CoreUObject/Private/Serialization/AsyncLoading.h
@@ -22,7 +22,7 @@
 #include "Async/AsyncFileHandle.h"
 
 struct FAsyncPackage;
-struct FFlushTree;
+struct FFlushRequest;
 class FAsyncLoadingThread;
 
 
@@ -336,7 +336,6 @@ public:
 * Structure containing intermediate data required for async loading of all imports and exports of a
 * FLinkerLoad.
 */
-
 struct FAsyncPackage : public FGCObject
 {
 	friend struct FScopedAsyncPackageEvent;
@@ -356,10 +355,7 @@ struct FAsyncPackage : public FGCObject
 	 *
 	 * @return	true if package has finished loading, false otherwise
 	 */
-	EAsyncPackageState::Type TickAsyncPackage(bool bUseTimeLimit, bool bInbUseFullTimeLimit, float& InOutTimeLimit, FFlushTree* FlushTree);
-
-	/** Fills the package dependency tree required to flush a specific package */
-	void PopulateFlushTree(struct FFlushTree* FlushTree);	
+	EAsyncPackageState::Type TickAsyncPackage(bool bUseTimeLimit, bool bInbUseFullTimeLimit, float& InOutTimeLimit);
 	
 	/** Marks a specific request as complete */
 	void MarkRequestIDsAsComplete();
@@ -427,6 +423,12 @@ struct FAsyncPackage : public FGCObject
 		return bLoadHasFinished;
 	}
 
+	/** Returns package loading priority. */
+	FORCEINLINE int32 GetRequestID() const
+	{
+		return Desc.RequestID;
+	}
+
 	/** Returns package loading priority. */
 	FORCEINLINE TAsyncLoadPriority GetPriority() const
 	{
@@ -512,7 +514,12 @@ struct FAsyncPackage : public FGCObject
 	/** Creates GC clusters from loaded objects */
 	EAsyncPackageState::Type CreateClusters(double InTickStartTime, bool bInUseTimeLimit, float& InOutTimeLimit);
 
-private:	
+	/** Returns true if this package is a dependency to the explicit request id RequestID. */
+	bool IsDependencyOf(int32 RequestID);
+
+private:
+
+	friend class FLinkerLoad;
 
 	/** Checks if all dependencies (imported packages) of this package have been fully loaded */
 	static bool AreAllDependenciesFullyLoadedInternal(FAsyncPackage* Package, TSet<UPackage*>& VisitedPackages, FString& OutError);
@@ -536,6 +543,12 @@ private:
 		}
 	};
 
+	/** Propagate dependent requests with a package dependency. */
+	void AddDependentRequestsTo(FAsyncPackage* InDependency);
+
+	/** Add the list of passed in Requests to the list of requests that are dependent on this package. */
+	void AddDependentRequests(TArrayView<int32 const> Requests);
+
 	/** Basic information associated with this package */
 	FAsyncPackageDesc Desc;
 	/** Linker which is going to have its exports and imports loaded									*/
@@ -544,7 +557,7 @@ private:
 	UPackage*				LinkerRoot;
 	/** Call backs called when we finished loading this package											*/
 	TArray<FCompletionCallback>	CompletionCallbacks;
-	/** Pending Import packages - we wait until all of them have been fully loaded. */
+	/** Pending Import packages - we wait until all of them have been pre loaded. */
 	TMap<FName, FAsyncPackage*> PendingImportedPackages;
 	/** Referenced imports - list of packages we need until we finish loading this package. */
 	TArray<FAsyncPackage*> ReferencedImports;
@@ -584,10 +597,10 @@ private:
 	bool						bTimeLimitExceeded;
 	/** True if our load has failed */
 	bool						bLoadHasFailed;
-	/** True if our load has finished */
+	/** True if our load on the async thread has finished */
 	bool						bLoadHasFinished;
 	/** True if threaded loading has finished for this package */
-	bool						bThreadedLoadingFinished;
+	bool						bThreadedLoadingFinished; /** @note FH: seems to serve a purpose almost identical to bLoadHasFinished. */
 	/** True if this package was created by this async package */
 	bool						bCreatedLinkerRoot;
 	/** The time taken when we started the tick.														*/
@@ -611,14 +624,19 @@ private:
 	/** Objects to create GC clusters from */
 	TArray<UObject*> DeferredClusterObjects;
 
-	/** List of all request handles */
+	/** List of all request handles. Those handles refer to explicitly made request. */
 	TArray<int32> RequestIDs;
+
+	FRWLock DependentRequestsLock;
+	/** Requests that currently depends on this package. */
+	TSet<int32> DependentRequests;
+
 #if WITH_EDITORONLY_DATA
 	/** Index of the meta-data object within the linkers export table (unset if not yet processed, although may still be INDEX_NONE if there is no meta-data) */
 	TOptional<int32> MetaDataIndex;
 #endif // WITH_EDITORONLY_DATA
-	/** Number of times we recursed to load this package. */
-	int32 ReentryCount;
+	/** Number of times we recursed to load or post load this package. could be split into 2 counters not to use atomics. */
+	std::atomic<int32> ReentryCount;
 	/** List of objects referenced by this package */
 	TSet<UObject*> ReferencedObjects;
 	/** Critical section for referenced objects list */
@@ -853,10 +871,9 @@ private:
 	/**
 	 * Loads imported packages..
 	 *
-	 * @param FlushTree Package dependency tree to be flushed
 	 * @return true if we finished loading all imports, false otherwise
 	 */
-	EAsyncPackageState::Type LoadImports(FFlushTree* FlushTree);
+	EAsyncPackageState::Type LoadImports();
 	/**
 	 * Create imports till time limit is exceeded.
 	 *
@@ -905,38 +922,31 @@ private:
 	EAsyncPackageState::Type FinishExternalReadDependencies();
 
 	/**
-	 * Function called when pending import package has been loaded.
+	 * Function called when pending import package has been preloaded. called from the async loading thread
 	 */
-	void ImportFullyLoadedCallback(const FName& PackageName, UPackage* LoadedPackage, EAsyncLoadingResult::Type Result);
+	void ImportPreLoadedAsyncCallback(const FName& PackageName, UPackage* LoadedPackage, EAsyncLoadingResult::Type Result);
+
 	/**
 	 * Adds dependency tree to the list if packages to wait for until their linkers have been created.
 	 *
 	 * @param ImportedPackage Package imported either directly or by one of the imported packages
-	 * @param FlushTree Package dependency tree to be flushed
 	 */
-	void AddDependencyTree(FAsyncPackage& ImportedPackage, TSet<FAsyncPackage*>& SearchedPackages, FFlushTree* FlushTree);
+	void AddDependencyTree(FAsyncPackage& ImportedPackage, TSet<FAsyncPackage*>& SearchedPackages);
+
 	/**
 	 * Adds a unique package to the list of packages to wait for until their linkers have been created.
 	 *
 	 * @param PendingImport Package imported either directly or by one of the imported packages
-	 * @param FlushTree Package dependency tree to be flushed
-	 */
-	bool AddUniqueLinkerDependencyPackage(FAsyncPackage& PendingImport, FFlushTree* FlushTree);
-	/**
-	 * Adds a package to the list of pending import packages.
-	 *
-	 * @param PendingImport Name of the package imported either directly or by one of the imported packages
-	 * @param FlushTree Package dependency tree to be flushed
 	 */
-	void AddImportDependency(const FName& PendingImport, FFlushTree* FlushTree);
+	bool AddUniqueLinkerDependencyPackage(FAsyncPackage& PendingImport);
+
 	/**
 	 * Adds a package to the list of pending import packages.
 	 *
 	 * @param PendingImport Name of the package imported either directly or by one of the imported packages
 	 * @param PackageToLoad Name of the package to load into PendingImport
-	 * @param FlushTree Package dependency tree to be flushed
 	 */
-	void AddImportDependency(const FName& PendingImport, const FName& PackageToLoad, FFlushTree* FlushTree, FLinkerInstancingContext InstancingContext);
+	void AddImportDependency(const FName& PendingImport, const FName& PackageToLoad, FLinkerInstancingContext InstancingContext);
 
 	/**
 	 * Removes references to any imported packages.
diff --git a/Engine/Source/Runtime/CoreUObject/Private/Serialization/AsyncLoadingThread.h b/Engine/Source/Runtime/CoreUObject/Private/Serialization/AsyncLoadingThread.h
--- a/Engine/Source/Runtime/CoreUObject/Private/Serialization/AsyncLoadingThread.h
+++ b/Engine/Source/Runtime/CoreUObject/Private/Serialization/AsyncLoadingThread.h
@@ -19,6 +19,8 @@
 #include "HAL/LowLevelMemTracker.h"
 #include "Serialization/AsyncPackageLoader.h"
 
+#include <atomic>
+
 struct FPrecacheCallbackHandler;
 
 /** [EDL] Event Driven loader event */
@@ -108,31 +110,55 @@ struct FAsyncLoadEventQueue
 	}
 };
 
-/** Package dependency tree used for flushing specific packages */
-struct FFlushTree
+/** 
+ * Request to flush a specific package using the request id 
+ */
+struct FFlushRequest
 {
-	int32 RequestId;
-	TSet<FName> PackagesToFlush;
+	FFlushRequest() = default;
 
-	FFlushTree(int32 InRequestId)
-		: RequestId(InRequestId)
+	FFlushRequest(int32 InRequestId)
+		: RequestId(MakeShared<std::atomic<int32>, ESPMode::ThreadSafe>(InRequestId))
 	{
 	}
 
-	bool AddPackage(const FName& Package)
+	int32 GetId() const
 	{
-		if (!PackagesToFlush.Contains(Package))
-		{
-			PackagesToFlush.Add(Package);
-			return true;
-		}
-		return false;
+		return RequestId->load(std::memory_order_relaxed);
+	}
+
+	bool IsValid() const
+	{
+		return RequestId.IsValid();
+	}
+
+	operator bool() const
+	{
+		return IsValid();
+	}
+
+	void Reset()
+	{
+		RequestId.Reset();
 	}
 
-	bool Contains(const FName& Package)
+private:
+	void SetId(int32 InId)
 	{
-		return PackagesToFlush.Contains(Package);
+		RequestId->store(InId, std::memory_order_relaxed);
 	}
+
+	friend class FAsyncLoadingThread;
+
+	/** 
+	 * Shared request id state between main thread and async loading thread. 
+	 * @note because request id are assigned to package request after being generated
+	 * multiple request ids can identify a specific package request however only the original request id
+	 * will properly identify unloaded dependencies. The shared nature of the request allows the async loading thread to
+	 * fix up the main thread flush request id to identify the original request id instead of duplicate when trying to identify
+	 * which main thread packages should be processed when flushing.
+	 */
+	TSharedPtr<std::atomic<int32>, ESPMode::ThreadSafe> RequestId;
 };
 
 /** Holds the maximum package summary size that can be set via ini files
@@ -149,6 +175,8 @@ struct FMaxPackageSummarySize
  */
 class FAsyncLoadingThread final : public FRunnable, public IAsyncPackageLoader
 {
+	friend struct FAsyncPackage;
+
 	IEDLBootNotificationManager& EDLBootNotificationManager;
 
 	/** Thread to run the worker FRunnable on */
@@ -198,6 +226,17 @@ class FAsyncLoadingThread final : public FRunnable, public IAsyncPackageLoader
 	FCriticalSection LoadedPackagesToProcessCritical;
 #endif
 
+	/** [ASYNC/GAME THREAD] FIFO of current flush requests from the game thread, managed by the async thread.
+	 * producer can only be the game thread.
+	 */
+	TArray<FFlushRequest, TInlineAllocator<4>> FlushRequests;
+#if THREADSAFE_UOBJECTS
+	/** [ASYNC/GAME THREAD] Critical section for the FlushRequests stack. needed if running with the async loading thread. */
+	FCriticalSection FlushRequestCritical;
+#endif
+	/** [ASYNC THREAD] Map containing a mapping of duplicate request to their main request. Use to operate flush request accurately. */
+	TMap<int32, int32> DuplicateRequestMap;
+
 	/** [ASYNC THREAD] Array of packages that are being preloaded */
 	TArray<FAsyncPackage*> AsyncPackages;
 	TMap<FName, FAsyncPackage*> AsyncPackageNameLookup;
@@ -485,10 +524,10 @@ public:
 	* @param bUseTimeLimit True if time limit should be used [time-slicing].
 	* @param bUseFullTimeLimit True if full time limit should be used [time-slicing].
 	* @param TimeLimit Maximum amount of time that can be spent in this call [time-slicing].
-	* @param FlushTree Package dependency tree to be flushed
+	* @param FlushRequest The request to flush a package and its depedencies. May be invalid.
 	* @return The current state of async loading
 	*/
-	EAsyncPackageState::Type ProcessAsyncLoading(int32& OutPackagesProcessed, bool bUseTimeLimit = false, bool bUseFullTimeLimit = false, float TimeLimit = 0.0f, FFlushTree* FlushTree = nullptr);
+	EAsyncPackageState::Type ProcessAsyncLoading(int32& OutPackagesProcessed, bool bUseTimeLimit, bool bUseFullTimeLimit, float TimeLimit, FFlushRequest& FlushRequest);
 
 	/**
 	* [EDL] [ASYNC* THREAD] Checks fopr cycles in the event driven loader and does fatal errors in that case
@@ -501,10 +540,10 @@ public:
 	* @param bUseTimeLimit True if time limit should be used [time-slicing].
 	* @param bUseFullTimeLimit True if full time limit should be used [time-slicing].
 	* @param TimeLimit Maximum amount of time that can be spent in this call [time-slicing].
-	* @param FlushTree Package dependency tree to be flushed
+	* @param FlushRequest The request to flush a package and its depedencies. May be invalid.
 	* @return The current state of async loading
 	*/
-	EAsyncPackageState::Type TickAsyncLoading(bool bUseTimeLimit, bool bUseFullTimeLimit, float TimeLimit, FFlushTree* FlushTree = nullptr);
+	EAsyncPackageState::Type TickAsyncLoading(bool bUseTimeLimit, bool bUseFullTimeLimit, float TimeLimit, FFlushRequest FlushRequest = FFlushRequest());
 
 	/**
 	* [ASYNC THREAD] Main thread loop
@@ -512,9 +551,9 @@ public:
 	* @param bUseTimeLimit True if time limit should be used [time-slicing].
 	* @param bUseFullTimeLimit True if full time limit should be used [time-slicing].
 	* @param TimeLimit Maximum amount of time that can be spent in this call [time-slicing].
-	* @param FlushTree Package dependency tree to be flushed
+	* @param FlushRequest The request to flush a package and its depedencies. May be invalid.
 	*/
-	EAsyncPackageState::Type TickAsyncThread(bool bUseTimeLimit, bool bUseFullTimeLimit, float TimeLimit, bool& bDidSomething, FFlushTree* FlushTree = nullptr);
+	EAsyncPackageState::Type TickAsyncThread(bool bUseTimeLimit, bool bUseFullTimeLimit, float TimeLimit, bool& bDidSomething, FFlushRequest& FlushRequest);
 
 	/** Initializes async loading thread */
 	void InitializeLoading() override;
@@ -526,17 +565,8 @@ public:
 	 */
 	float GetAsyncLoadPercentage(const FName& PackageName) override;
 
-	/** 
-	 * [ASYNC/GAME THREAD] Checks if a request ID already is added to the loading queue
-	 */
-	bool ContainsRequestID(int32 RequestID)
-	{
-#if THREADSAFE_UOBJECTS
-		FScopeLock Lock(&PendingRequestsCritical);
-#endif
-		return PendingRequests.Contains(RequestID);
-	}
 
+private:
 	/** 
 	 * [ASYNC/GAME THREAD] Adds a request ID to the list of pending requests
 	 */
@@ -566,6 +596,20 @@ public:
 		}		
 	}
 
+	/**
+	 * [ASYNC/GAME THREAD] Pop flush request from the flush request stacks.
+	 * @note: one of the request id passed should normally match the top of the request stack if any.
+	 */
+	void CompleteFlushRequests(TArray<int32>& RequestIDs);
+
+	/**
+	 * [ASYNC THREAD] Add a mapping being duplicate request and the main/first explicit request made for a package
+	 * @param DuplicateRequestId the duplicate request id
+	 * @param MainRequestId the main request id that will be used for processing flush request
+	 */
+	void MapDuplicateRequestID(int32 DuplicateRequestId, int32 MainRequestId);
+
+
 	/** [ASYNC/GAME THREAD] Number of package load requests in the async loading queue */
 	int32 GetQueuedPackagesCount() const
 	{
@@ -577,31 +621,66 @@ public:
 		return ExistingAsyncPackagesCounter.GetValue();
 	}
 
-private:
+	/**
+	 * [ASYNC/GAME THREAD] Checks if a request ID already is added to the loading queue
+	 */
+	bool ContainsRequestInternal(int32 RequestID)
+	{
+#if THREADSAFE_UOBJECTS
+		FScopeLock Lock(&PendingRequestsCritical);
+#endif
+		return PendingRequests.Contains(RequestID);
+	}
+
+	/** 
+	 * [GAME THREAD] Add a Flush request to the flush stack
+	 * @param RequestID the request id of the package to flush
+	 * @returns the flushtree for the flush request to use on the game thread
+	 * @note the FlushRequest is a shared resource between the async loading thread and the main thread
+	 */
+	FFlushRequest AddFlushRequest(int32 RequestID);
+
+	/** 
+	 * [ASYNC THREAD] Peek the top of the flush request stack
+	 * @returns the current flush tree to consider 
+	 */
+	FFlushRequest PeekFlushRequest();
+
+	/** 
+	 * [ASYNC THREAD] Adjust the flush request if needed, in case of duplicate request
+	 * @param  InFlushRequest the flush request to adjust
+	 */
+	void AdjustFlushRequest(FFlushRequest& FlushRequest);
 
+	/**
+	 * [GAME THREAD/ASYNC THREAD] Use by both the game thread and async thread to determine if the current package should be processed in relation to the current flush tree
+	 * @param InAsyncPackage The package to potentially process
+	 * @param FlushRequest the current flush request
+	 * @returns true if we should process the package
+	 */
+	bool ShouldProcessPackage(FAsyncPackage* InAsyncPackage, const FFlushRequest& FlushRequest);
+	
 	/**
 	* [GAME THREAD] Performs game-thread specific operations on loaded packages (not-thread-safe PostLoad, callbacks)
 	*
 	* @param bUseTimeLimit True if time limit should be used [time-slicing].
 	* @param bUseFullTimeLimit True if full time limit should be used [time-slicing].
 	* @param TimeLimit Maximum amount of time that can be spent in this call [time-slicing].
-	* @param FlushTree Package dependency tree to be flushed
+	* @param FlushRequest Package dependency tree to be flushed
 	* @return The current state of async loading
 	*/
-	EAsyncPackageState::Type ProcessLoadedPackages(bool bUseTimeLimit, bool bUseFullTimeLimit, float TimeLimit, bool& bDidSomething, FFlushTree* FlushTree = nullptr);
+	EAsyncPackageState::Type ProcessLoadedPackages(bool bUseTimeLimit, bool bUseFullTimeLimit, float TimeLimit, bool& bDidSomething, const FFlushRequest& FlushRequest);
 
 	/**
 	* [ASYNC THREAD] Creates async packages from the queued requests
-	* @param FlushTree Package dependency tree to be flushed
 	*/
-	int32 CreateAsyncPackagesFromQueue(bool bUseTimeLimit, bool bUseFullTimeLimit, float TimeLimit, FFlushTree* FlushTree = nullptr);
+	int32 CreateAsyncPackagesFromQueue(bool bUseTimeLimit, bool bUseFullTimeLimit, float TimeLimit);
 
 	/**
 	* [ASYNC THREAD] Internal helper function for processing a package load request. If dependency preloading is enabled, 
 	* it will call itself recursively for all the package dependencies
-	* @param FlushTree Package dependency tree to be flushed
 	*/
-	void ProcessAsyncPackageRequest(FAsyncPackageDesc* InRequest, FAsyncPackage* InRootPackage, FFlushTree* FlushTree);
+	void ProcessAsyncPackageRequest(FAsyncPackageDesc* InRequest, FAsyncPackage* InRootPackage);
 
 	/**
 	* [ASYNC THREAD] Internal helper function for updating the priorities of an existing package and all its dependencies
@@ -610,9 +689,10 @@ private:
 
 	/**
 	* [ASYNC THREAD] Finds existing async package and adds the new request's completion callback to it.
-	* @param FlushTree Package dependency tree to be flushed
+	* @param PackageRequest The package description for the request
+	* @param PackageList The package list to look into
 	*/
-	FAsyncPackage* FindExistingPackageAndAddCompletionCallback(FAsyncPackageDesc* PackageRequest, TMap<FName, FAsyncPackage*>& PackageList, FFlushTree* FlushTree);
+	FAsyncPackage* FindExistingPackageAndAddCompletionCallback(FAsyncPackageDesc* PackageRequest, TMap<FName, FAsyncPackage*>& PackageList);
 
 	/**
 	* [ASYNC THREAD] Adds a package to a list of packages that have finished loading on the async thread
@@ -659,4 +739,4 @@ public:
 #endif
 
 	static FAsyncLoadingThread* Instance;
-};
+};
\ No newline at end of file
diff --git a/Engine/Source/Runtime/CoreUObject/Private/Serialization/AsyncPackageLoader.cpp b/Engine/Source/Runtime/CoreUObject/Private/Serialization/AsyncPackageLoader.cpp
--- a/Engine/Source/Runtime/CoreUObject/Private/Serialization/AsyncPackageLoader.cpp
+++ b/Engine/Source/Runtime/CoreUObject/Private/Serialization/AsyncPackageLoader.cpp
@@ -465,30 +465,30 @@ static FEDLBootNotificationManager& GetGEDLBootNotificationManager()
 }
 
 FAsyncLoadingThreadSettings::FAsyncLoadingThreadSettings()
+	: bAsyncLoadingThreadEnabled(false)
+	, bAsyncPostLoadEnabled(false)
 {
 #if THREADSAFE_UOBJECTS
-	if (FPlatformProperties::RequiresCookedData())
-	{
-		check(GConfig);
-
-		bool bConfigValue = true;
-		GConfig->GetBool(TEXT("/Script/Engine.StreamingSettings"), TEXT("s.AsyncLoadingThreadEnabled"), bConfigValue, GEngineIni);
-		bool bCommandLineDisable = FParse::Param(FCommandLine::Get(), TEXT("NoAsyncLoadingThread"));
-		bool bCommandLineEnable = FParse::Param(FCommandLine::Get(), TEXT("AsyncLoadingThread"));
-		bAsyncLoadingThreadEnabled = bCommandLineEnable || (bConfigValue && FApp::ShouldUseThreadingForPerformance() && !bCommandLineDisable);
-
-		bConfigValue = true;
-		GConfig->GetBool(TEXT("/Script/Engine.StreamingSettings"), TEXT("s.AsyncPostLoadEnabled"), bConfigValue, GEngineIni);
-		bCommandLineDisable = FParse::Param(FCommandLine::Get(), TEXT("NoAsyncPostLoad"));
-		bCommandLineEnable = FParse::Param(FCommandLine::Get(), TEXT("AsyncPostLoad"));
-		bAsyncPostLoadEnabled = bCommandLineEnable || (bConfigValue && FApp::ShouldUseThreadingForPerformance() && !bCommandLineDisable);
-	}
-	else
+	check(GConfig);
+	// if we are on an uncooked platform query the editor streaming settings
+	const TCHAR* SectionName = FPlatformProperties::RequiresCookedData() ? TEXT("/Script/Engine.StreamingSettings") : TEXT("/Script/Engine.EditorStreamingSettings");
+
+	// if we are on an uncooked platform the default value for async loading is disabled while it is enabled for cooked platform
+	bool bConfigValue = FPlatformProperties::RequiresCookedData() ? true : false;
+
+	GConfig->GetBool(SectionName, TEXT("s.AsyncLoadingThreadEnabled"), bConfigValue, GEngineIni);
+	bool bCommandLineDisable = FParse::Param(FCommandLine::Get(), TEXT("NoAsyncLoadingThread"));
+	bool bCommandLineEnable = FParse::Param(FCommandLine::Get(), TEXT("AsyncLoadingThread"));
+	bAsyncLoadingThreadEnabled = bCommandLineEnable || (bConfigValue && FApp::ShouldUseThreadingForPerformance() && !bCommandLineDisable);
+
+	// if we are on an uncooked platform the default value for async loading is disabled while it is enabled for cooked platform
+	bConfigValue = FPlatformProperties::RequiresCookedData() ? true : false;
+
+	GConfig->GetBool(SectionName, TEXT("s.AsyncPostLoadEnabled"), bConfigValue, GEngineIni);
+	bCommandLineDisable = FParse::Param(FCommandLine::Get(), TEXT("NoAsyncPostLoad"));
+	bCommandLineEnable = FParse::Param(FCommandLine::Get(), TEXT("AsyncPostLoad"));
+	bAsyncPostLoadEnabled = bCommandLineEnable || (bConfigValue && FApp::ShouldUseThreadingForPerformance() && !bCommandLineDisable);
 #endif
-	{
-		bAsyncLoadingThreadEnabled = false;
-		bAsyncPostLoadEnabled = false;
-	}
 }
 
 FAsyncLoadingThreadSettings& FAsyncLoadingThreadSettings::Get()
