Fix typo in UObjectBaseUtility::GetImplementingOuterObject

diff --git a/Engine/Source/Runtime/CoreUObject/Private/UObject/UObjectBaseUtility.cpp b/Engine/Source/Runtime/CoreUObject/Private/UObject/UObjectBaseUtility.cpp
--- a/Engine/Source/Runtime/CoreUObject/Private/UObject/UObjectBaseUtility.cpp
+++ b/Engine/Source/Runtime/CoreUObject/Private/UObject/UObjectBaseUtility.cpp
@@ -294,7 +294,7 @@ UObjectBaseUtility* UObjectBaseUtility::GetImplementingOuterObject(const UClass*
 		UClass* OuterClass = NextOuter->GetClass();
 		if (OuterClass && OuterClass->ImplementsInterface(InInterfaceClass))
 		{
-			return OuterClass;
+			return NextOuter;
 		}
 	}
 
