Add some insight scope/counters that would've been useful to hunt down a problem

#rnx

- Someone was complaining about seing a large number of long calls to FEditorBulkData::DetachFromDisk that were taking a long time. Although from context he was able to determine that the problem was caused by renaming a package this was not obvious from the provided trace.
-- Added some trace scopes to  UObject::Rename, FLinkerManager::ResetLoaders and around the decompresison of payloads when detaching an editorbulkdata. This would've m,ade it very obvious as to where the time was going.
-- Added two new counters to FEditorBulkData to track the amount of data (in bytes) either loaded from disk or pulled from a virtualized backend. This would've allowed us to see how much data was being loaded during the rename (in this case near 32gb...)

diff --git a/Engine/Source/Runtime/CoreUObject/Private/Serialization/EditorBulkData.cpp b/Engine/Source/Runtime/CoreUObject/Private/Serialization/EditorBulkData.cpp
--- a/Engine/Source/Runtime/CoreUObject/Private/Serialization/EditorBulkData.cpp
+++ b/Engine/Source/Runtime/CoreUObject/Private/Serialization/EditorBulkData.cpp
@@ -9,6 +9,7 @@
 #include "Misc/PackageSegment.h"
 #include "Misc/SecureHash.h"
 #include "Misc/ScopeLock.h"
+#include "ProfilingDebugging/CountersTrace.h"
 #include "Serialization/BulkData.h"
 #include "Serialization/BulkDataRegistry.h"
 #include "UObject/LinkerLoad.h"
@@ -35,6 +36,9 @@
 	#define UE_CORRUPTED_DATA_SEVERITY Error
 #endif // VBD_CORRUPTED_PAYLOAD_IS_FATAL
 
+TRACE_DECLARE_INT_COUNTER(EditorBulkData_PayloadDataLoaded, TEXT("EditorBulkData/PayloadDataLoaded"));
+TRACE_DECLARE_INT_COUNTER(EditorBulkData_PayloadDataPulled, TEXT("EditorBulkData/PayloadDatPulled"));
+
 namespace UE::Serialization
 {
 /** This console variable should only exist for testing */
@@ -1308,6 +1312,7 @@ void FEditorBulkData::DetachFromDisk(FArchive* Ar, bool bEnsurePayloadIsLoaded)
 		if (Payload.IsNull() && bEnsurePayloadIsLoaded)
 		{
 			FCompressedBuffer CompressedPayload = GetDataInternal();
+			TRACE_CPUPROFILER_EVENT_SCOPE(DecompressPayload);
 			Payload = CompressedPayload.Decompress();
 		}
 
@@ -1521,6 +1526,7 @@ FCompressedBuffer FEditorBulkData::GetDataInternal() const
 		UE_CLOG(CompressedPayload.IsNull(), LogSerialization, Error, TEXT("Failed to pull payload '%s'"), *LexToString(PayloadContentId));
 		UE_CLOG(!IsDataValid(*this, CompressedPayload), LogSerialization, UE_CORRUPTED_DATA_SEVERITY, TEXT("Virtualized payload '%s' is corrupt! Check the backend storage."), *LexToString(PayloadContentId));
 		
+		TRACE_COUNTER_ADD(EditorBulkData_PayloadDataPulled, (int64)CompressedPayload.GetCompressedSize());
 		return CompressedPayload;
 	}
 	else
@@ -1534,6 +1540,7 @@ FCompressedBuffer FEditorBulkData::GetDataInternal() const
 			*LexToString(PayloadContentId),
 			*PackagePath.GetDebugName());
 
+		TRACE_COUNTER_ADD(EditorBulkData_PayloadDataLoaded, (int64)CompressedPayload.GetCompressedSize());
 		return CompressedPayload;
 	}
 }
diff --git a/Engine/Source/Runtime/CoreUObject/Private/UObject/LinkerManager.cpp b/Engine/Source/Runtime/CoreUObject/Private/UObject/LinkerManager.cpp
--- a/Engine/Source/Runtime/CoreUObject/Private/UObject/LinkerManager.cpp
+++ b/Engine/Source/Runtime/CoreUObject/Private/UObject/LinkerManager.cpp
@@ -109,6 +109,8 @@ void FLinkerManager::ResetLinkerExports(UPackage* InPackage)
 
 void FLinkerManager::ResetLoaders(UObject* InPkg)
 {
+	TRACE_CPUPROFILER_EVENT_SCOPE(FLinkerManager::ResetLoaders);
+
 	// Top level package to reset loaders for.
 	UObject*		TopLevelPackage = InPkg ? InPkg->GetOutermost() : nullptr;
 
diff --git a/Engine/Source/Runtime/CoreUObject/Private/UObject/Obj.cpp b/Engine/Source/Runtime/CoreUObject/Private/UObject/Obj.cpp
--- a/Engine/Source/Runtime/CoreUObject/Private/UObject/Obj.cpp
+++ b/Engine/Source/Runtime/CoreUObject/Private/UObject/Obj.cpp
@@ -167,6 +167,8 @@ UObject* UObject::GetDefaultSubobjectByName(FName ToFind)
 
 bool UObject::Rename( const TCHAR* InName, UObject* NewOuter, ERenameFlags Flags )
 {
+	TRACE_CPUPROFILER_EVENT_SCOPE(UObject::Rename);
+
 #if WITH_EDITOR
 	// This guarantees that if this UObject is actually renamed and changes packages
 	// the metadata will be moved with it.
