Enforce Public/Private dependencies on Save

diff --git a/Engine/Source/Runtime/CoreUObject/Private/UObject/SavePackage2.cpp b/Engine/Source/Runtime/CoreUObject/Private/UObject/SavePackage2.cpp
--- a/Engine/Source/Runtime/CoreUObject/Private/UObject/SavePackage2.cpp
+++ b/Engine/Source/Runtime/CoreUObject/Private/UObject/SavePackage2.cpp
@@ -762,6 +762,26 @@ ESavePackageResult ValidateImports(FSaveContext& SaveContext)
 		{
 			PrivateObjects.Add(Import);
 		}
+
+		// Enforce that Private content can only be referenced by something within the same Mount Point
+		if (!ImportPackage->IsExternallyReferenceable())
+		{
+			FName MountPointName = FPackageName::GetPackageMountPoint(PackageName);
+
+			FName ImportMountPointName = FPackageName::GetPackageMountPoint(ImportPackage->GetName());
+
+			if (!MountPointName.IsNone() && !ImportMountPointName.IsNone())
+			{
+				if (MountPointName != ImportMountPointName)
+				{
+					PrivateObjects.Add(Import);
+				}
+			}
+			else
+			{
+				PrivateObjects.Add(Import);
+			}
+		}
 	}
 	if (PrivateObjects.Num() > 0 || ObjectsInOtherMaps.Num() > 0)
 	{
