Calling ::Rename on a UPackage will no longer reset the loaders for all currently loaded packages.

- UObject::Rename was calling this->GetOuter to pass to ::ResetLoaders, when 'this' is a UPackage this would return nullptr and when you pass a nullptr to ::ResetLoaders it will cause all loaders in the processed to be reset.
- If we call this->GetPackage() instead we do not get this problem as calling that on a UPackage will return a pointer to itself. Infact, internally ::ResetLoaders will call GetPackage on the given UObject to find the loader to reset (or if it returns nullptr at that point, reset all existing loaders)

diff --git a/Engine/Source/Runtime/CoreUObject/Private/UObject/Obj.cpp b/Engine/Source/Runtime/CoreUObject/Private/UObject/Obj.cpp
--- a/Engine/Source/Runtime/CoreUObject/Private/UObject/Obj.cpp
+++ b/Engine/Source/Runtime/CoreUObject/Private/UObject/Obj.cpp
@@ -231,7 +231,7 @@ bool UObject::Rename( const TCHAR* InName, UObject* NewOuter, ERenameFlags Flags
 
 	if (!(Flags & REN_ForceNoResetLoaders))
 	{
-		ResetLoaders( GetOuter() );
+		ResetLoaders(GetPackage());
 	}
 
 	FName OldName = GetFName();
