CoreUObject test conversion damage control.
Correcting tests commited in CL 19498448 that were converted and tagged [Smoke] but weren't marked with EAUtomationTestFlags::SmokeFilter in the original source.
EditorBulkDataTests was left unconverted, moved back to original Private/Tests/Serialization folder
ObjectPtrTest -> CoreUObject::TObjectPtr::Soft Object Path fails at line 197 on non-desktop platforms but it didn't failed before the conversion.

diff --git a/Engine/Source/Developer/LowLevelTestsRunner/Public/TestHarness.h b/Engine/Source/Developer/LowLevelTestsRunner/Public/TestHarness.h
--- a/Engine/Source/Developer/LowLevelTestsRunner/Public/TestHarness.h
+++ b/Engine/Source/Developer/LowLevelTestsRunner/Public/TestHarness.h
@@ -100,144 +100,66 @@ struct Catch::StringMaker<TTuple<KeyType, ValueType>>
 
 /** Use with TEST_CASE_METHOD to quickly migrate automation tests to Catch. */
 class FAutomationTestFixture
-{
-public:
+{	
+};
 
-	template <typename WhatType>
-	inline bool TestTrue(const WhatType What, const bool bValue)
-	{
-		CAPTURE(What);
-		CHECK(bValue);
-		return bValue;
-	}
+#define VERIFY(What, Actual)\
+	CAPTURE(What);\
+	CHECK(Actual == true)
 
-	template <typename WhatType>
-	inline bool TestFalse(const WhatType What, const bool bValue)
-	{
-		CAPTURE(What);
-		CHECK(!bValue);
-		return !bValue;
-	}
+#define ADD_WARNING(What)\
+	WARN(What)
 
-	template <typename WhatType, typename ActualType, typename ExpectedType>
-	inline bool TestEqual(const WhatType What, const ActualType& Actual, const ExpectedType& Expected)
-	{
-		CAPTURE(What);
-		CHECK(Actual == Expected);
-		return Actual == Expected;
-	}
-
-	template <typename WhatType, typename ActualType, typename ExpectedType>
-	inline bool TestNotEqual(const WhatType& What, const ActualType& Actual, const ExpectedType& Expected)
-	{
-		CAPTURE(What);
-		CHECK(Actual != Expected);
-		return Actual != Expected;
-	}
-	template <typename WhatType, typename ActualType>
-	inline bool TestNull(const WhatType& What, const ActualType& Actual)
-	{
-		CAPTURE(What);
-		CHECK(Actual == nullptr);
-		return Actual == nullptr;
-	}
-	template <typename WhatType, typename ActualType>
-	inline bool TestNotNull(const WhatType& What, const ActualType& Actual)
-	{
-		CAPTURE(What);
-		CHECK(Actual != nullptr);
-		return Actual != nullptr;
-	}
-	template <typename WhatType, typename ActualType>
-	inline bool TestValid(const WhatType& What, const ActualType& Actual)
-	{
-		CAPTURE(What);
-		CHECK(Actual.IsValid() == true);
-		return Actual.IsValid() == true;
-	}
-	template <typename WhatType, typename ActualType>
-	inline bool TestInvalid(const WhatType& What, const ActualType& Actual)
-	{
-		CAPTURE(What);
-		CHECK(Actual.IsValid() == false);
-		return Actual.IsValid() == false;
-	}
-	template <typename WhatType, typename ActualType>
-	inline bool Verify(const WhatType What, const ActualType& Actual)
-	{
-		CAPTURE(What);
-		CHECK(Actual == true);
-		return Actual == true;
-	}
-	template <typename WhatType>
-	inline void AddWarning(const WhatType& What)
-	{
-		CAPTURE(What);
-		return;
-	}
-	template <typename WhatType>
-	inline void AddError(const WhatType& What)
-	{
-		CAPTURE(What);
-		return;
-	}
-	
-};
+#define ADD_ERROR(What)\
+	FAIL(What)
 
 #define CHECK_EQUAL(Actual, Expected)\
-	CHECK(Actual == Expected);
+	CHECK(Actual == Expected)
 
 #define CHECK_NOT_EQUAL(Actual, Expected)\
-	CHECK(Actual != Expected);
+	CHECK(Actual != Expected)
 
 #define REQUIRE_EQUAL(Actual, Expected)\
-	REQUIRE(Actual == Expected);
+	REQUIRE(Actual == Expected)
 
 #define REQUIRE_NOT_EQUAL(Actual, Expected)\
-	REQUIRE(Actual != Expected);
-
-
+	REQUIRE(Actual != Expected)
 
 #define TEST_TRUE(What, Value)\
-{\
 	INFO(What);\
-	CHECK((Value)==true);\
-}
+	CHECK((Value)==true)
+
 #define TEST_FALSE(What, Value)\
-{\
 	INFO(What);\
-	CHECK((Value)==false);\
-}
+	CHECK((Value)==false)
+
 #define TEST_EQUAL(What, Actual, Expected)\
-{\
 	INFO(What);\
-	CHECK((Actual) == (Expected));\
-}
+	CHECK((Actual) == (Expected))
+
+#define TEST_EQUAL_STR(What, Expected, Actual)\
+	INFO(What);\
+	CAPTURE(Actual);\
+	CHECK(FCString::Strcmp(ToCStr((Expected)), ToCStr((Actual))) == 0)
 
 #define TEST_NOT_EQUAL(What, Actual, Expected)\
-{\
 	INFO(What);\
-	CHECK((Actual) != (Expected));\
-}
+	CHECK((Actual) != (Expected))
+
 #define TEST_NULL(What, Value)\
-{\
 	INFO(What);\
-	CHECK((Value)==nullptr);\
-}
+	CHECK((Value)==nullptr)
+
 #define TEST_NOT_NULL(What, Value)\
-{\
 	INFO(What);\
-	CHECK((Value)!=nullptr);\
-}
+	CHECK((Value)!=nullptr)
+
 #define TEST_VALID(What, Value)\
-{\
 	INFO(What);\
-	CHECK(Value.IsValid()==true);\
-}
+	CHECK(Value.IsValid()==true)
+
 #define TEST_INVALID(What, Value)\
-{\
 	INFO(What);\
-	CHECK(Value.IsValid()==false);\
-}
+	CHECK(Value.IsValid()==false)
 
 #endif // UE_ENABLE_TESTHARNESS_ENGINE_SUPPORT
\ No newline at end of file
diff --git a/Engine/Source/Runtime/CoreUObject/Private/Tests/Serialization/EditorBulkDataTests.cpp b/Engine/Source/Runtime/CoreUObject/Private/Tests/Serialization/EditorBulkDataTests.cpp
new file mode 100644
--- /dev/null
+++ b/Engine/Source/Runtime/CoreUObject/Private/Tests/Serialization/EditorBulkDataTests.cpp
@@ -0,0 +1,616 @@
+// Copyright Epic Games, Inc. All Rights Reserved.
+
+#include "Algo/AllOf.h"
+#include "Memory/SharedBuffer.h"
+#include "Serialization/EditorBulkData.h"
+#include "Serialization/EditorBulkDataReader.h"
+#include "Serialization/EditorBulkDataWriter.h"
+#include "Serialization/LargeMemoryReader.h"
+#include "Serialization/LargeMemoryWriter.h"
+#include "Serialization/MemoryReader.h"
+#include "Serialization/MemoryWriter.h"
+#include "Templates/UniquePtr.h"
+
+#include "Misc/AutomationTest.h"
+
+#if WITH_DEV_AUTOMATION_TESTS && WITH_EDITORONLY_DATA
+
+namespace UE::Serialization
+{
+	
+constexpr const uint32 TestFlags = EAutomationTestFlags::ApplicationContextMask | EAutomationTestFlags::EngineFilter;
+
+/** Creates a buffer full of random data to make it easy to have something to test against. */
+TUniquePtr<uint8[]> CreateRandomData(int64 BufferSize)
+{
+	TUniquePtr<uint8[]> Buffer = MakeUnique<uint8[]>(BufferSize);
+
+	for (int64 Index = 0; Index < BufferSize; ++Index)
+	{
+		Buffer[Index] = (uint8)(FMath::Rand() % 255);
+	}
+
+	return Buffer;
+}
+
+/**
+ * This test creates a very basic FEditorBulkData object with in memory payload and validates that we are able to retrieve the 
+ * payload via both TFuture and callback methods. It then creates copies of the object and makes sure that we can get the payload from
+ * the copies, even when the original source object has been reset.
+ */
+IMPLEMENT_SIMPLE_AUTOMATION_TEST(FEditorBulkDataTestBasic, TEXT("System.CoreUObject.Serialization.EditorBulkData.Basic"), TestFlags)
+bool FEditorBulkDataTestBasic::RunTest(const FString& Parameters)
+{
+	const int64 BufferSize = 1024;
+	TUniquePtr<uint8[]> SourceBuffer = CreateRandomData(BufferSize);
+
+	auto ValidateBulkData = [this, &SourceBuffer, BufferSize](const FEditorBulkData& BulkDataToValidate, const TCHAR* Label)
+		{
+			FSharedBuffer RetrievedBuffer = BulkDataToValidate.GetPayload().Get();
+			TestEqual(FString::Printf(TEXT("%s buffer length"),Label), (int64)RetrievedBuffer.GetSize(), BufferSize);
+			TestTrue(FString::Printf(TEXT("SourceBuffer values == %s values"),Label), FMemory::Memcmp(SourceBuffer.Get(), RetrievedBuffer.GetData(), BufferSize) == 0);
+		};
+
+	// Create a basic bulkdata (but retain ownership of the buffer!)
+	FEditorBulkData BulkData;
+	BulkData.UpdatePayload(FSharedBuffer::MakeView(SourceBuffer.Get(), BufferSize));
+
+	// Test accessing the data from the bulkdata object
+	ValidateBulkData(BulkData, TEXT("Retrieved"));
+
+	// Create a new bulkdata object via the copy constructor
+	FEditorBulkData BulkDataCopy(BulkData);
+
+	// Create a new bulkdata object and copy by assignment (note we assign some junk data that will get overwritten)
+	FEditorBulkData BulkDataAssignment;
+	BulkDataAssignment.UpdatePayload(FUniqueBuffer::Alloc(128).MoveToShared());
+	BulkDataAssignment = BulkData;
+
+	// Test both bulkdata objects
+	ValidateBulkData(BulkDataCopy, TEXT("Copy Constructor"));
+	ValidateBulkData(BulkDataAssignment, TEXT("Copy Assignment"));
+
+	// Should not affect BulkDataAssignment/BulkDataCopy 
+	BulkData.Reset();
+
+	// Test both bulkdata objects again now that we reset the data
+	ValidateBulkData(BulkDataCopy, TEXT("Copy Constructor (after data reset)"));
+	ValidateBulkData(BulkDataAssignment, TEXT("Copy Assignment (after data reset)"));
+
+	return true;
+}
+
+/**
+ * This test will validate how FEditorBulkData behaves when it has no associated payload and make sure 
+ * that our assumptions are correct.
+ */
+IMPLEMENT_SIMPLE_AUTOMATION_TEST(FEditorBulkDataTestEmpty, TEXT("System.CoreUObject.Serialization.EditorBulkData.Empty"), TestFlags)
+bool FEditorBulkDataTestEmpty::RunTest(const FString& Parameters)
+{
+	auto Validate = [](const TCHAR* Id, FAutomationTestBase& Test, const FEditorBulkData& BulkData)
+	{
+		// Validate the general accessors
+		Test.TestEqual(FString::Printf(TEXT("(%s) Return value of ::GetBulkDataSize()"), Id), BulkData.GetPayloadSize(), (int64)0);
+		Test.TestTrue(FString::Printf(TEXT("(%s) Payload key is invalid"), Id), BulkData.GetPayloadId().IsZero());
+		Test.TestFalse(FString::Printf(TEXT("(%s) Return value of ::DoesPayloadNeedLoading()"), Id), BulkData.DoesPayloadNeedLoading());
+
+		// Validate the payload accessors
+		FSharedBuffer Payload = BulkData.GetPayload().Get();
+		Test.TestTrue(FString::Printf(TEXT("(%s) The payload from the GetPayload TFuture is null"), Id), Payload.IsNull());
+
+		FCompressedBuffer CompressedPayload = BulkData.GetCompressedPayload().Get();
+		Test.TestTrue(FString::Printf(TEXT("(%s) The payload from the GetCompressedPayload TFuture is null"), Id), Payload.IsNull());
+	};
+
+	FEditorBulkData DefaultBulkData;
+	Validate(TEXT("DefaultBulkData"), *this, DefaultBulkData);
+		
+	FEditorBulkData NullPayloadBulkData;
+	NullPayloadBulkData.UpdatePayload(FSharedBuffer());
+	Validate(TEXT("NullPayloadBulkData"), *this, NullPayloadBulkData);
+
+	FEditorBulkData ZeroLengthPayloadBulkData;
+	ZeroLengthPayloadBulkData.UpdatePayload(FUniqueBuffer::Alloc(0).MoveToShared());
+	Validate(TEXT("ZeroLengthPayloadBulkData"), *this, ZeroLengthPayloadBulkData);
+
+	return true;
+}
+
+/**
+ * Test the various methods for updating the payload that a FEditorBulkData owns
+ */
+IMPLEMENT_SIMPLE_AUTOMATION_TEST(FEditorBulkDataTestUpdatePayload, TEXT("System.CoreUObject.Serialization.EditorBulkData.UpdatePayload"), TestFlags)
+bool FEditorBulkDataTestUpdatePayload::RunTest(const FString& Parameters)
+{
+	// Create a memory buffer of all zeros
+	const int64 BufferSize = 1024;
+	TUniquePtr<uint8[]> OriginalData = MakeUnique<uint8[]>(BufferSize);
+	FMemory::Memzero(OriginalData.Get(), BufferSize);
+
+	// Pass the buffer to to bulkdata but retain ownership
+	FEditorBulkData BulkData;
+	BulkData.UpdatePayload(FSharedBuffer::MakeView(OriginalData.Get(), BufferSize));
+
+	// Access the payload, edit it and push it back into the bulkdata object
+	{
+		// The payload should be the same size and same contents as the original buffer but a different
+		// memory address since we retained ownership in the TUniquePtr, so the bulkdata object should 
+		// have created it's own copy.
+		FSharedBuffer Payload = BulkData.GetPayload().Get();
+		TestEqual(TEXT("Payload length"), (int64)Payload.GetSize(), BufferSize);
+		TestNotEqual(TEXT("OriginalData and the payload should have different memory addresses"), (uint8*)OriginalData.Get(), (uint8*)Payload.GetData());
+		TestTrue(TEXT("Orginal buffer == Payload data"), FMemory::Memcmp(OriginalData.Get(), Payload.GetData(), Payload.GetSize()) == 0);
+
+		// Make a copy of the payload that we can edit
+		const uint8 NewValue = 255;
+		FSharedBuffer EditedPayload;
+		{
+			FUniqueBuffer EditablePayload = FUniqueBuffer::Clone(Payload);
+			FMemory::Memset(EditablePayload.GetData(), NewValue, EditablePayload.GetSize());
+			EditedPayload = EditablePayload.MoveToShared();
+		}
+
+		// Update the bulkdata object with the new edited payload
+		BulkData.UpdatePayload(EditedPayload);
+
+		Payload = BulkData.GetPayload().Get();
+		TestEqual(TEXT("Updated payload length"), (int64)Payload.GetSize(), BufferSize);
+		TestEqual(TEXT("Payload and EditablePayload should have the same memory addresses"), (uint8*)Payload.GetData(), (uint8*)EditedPayload.GetData());
+
+		const bool bAllElementsCorrect = Algo::AllOf(TArrayView64<uint8>((uint8*)Payload.GetData(), (int64)Payload.GetSize()), [NewValue](uint8 Val)
+			{ 
+				return Val == NewValue; 
+			});
+
+		TestTrue(TEXT("All payload elements correctly updated"), bAllElementsCorrect);
+	}
+
+	{
+		// Store the original data pointer address so we can test against it later, we should not actually use
+		// this pointer though as once we pass it to the bulkdata object we cannot be sure what happens to it.
+		uint8* OriginalDataPtr = OriginalData.Get();
+
+		// Update the bulkdata object with the original data but this time we give ownership of the buffer to
+		// the bulkdata object.
+		BulkData.UpdatePayload(FSharedBuffer::TakeOwnership(OriginalData.Release(), BufferSize, [](void* Ptr, uint64) { delete[] (uint8*)Ptr; }));
+
+		FSharedBuffer Payload = BulkData.GetPayload().Get();
+		TestEqual(TEXT("Updated payload length"), (int64)Payload.GetSize(), BufferSize);
+		TestEqual(TEXT("Payload and OriginalDataPtr should have the same memory addresses"), (uint8*)Payload.GetData(), OriginalDataPtr);
+
+		// The original data was all zeros, so we can test for that to make sure that the contents are correct.
+		const bool bAllElementsCorrect = Algo::AllOf(TArrayView64<uint8>((uint8*)Payload.GetData(), (int64)Payload.GetSize()), [](uint8 Val)
+			{
+				return Val == 0;
+			});
+
+		TestTrue(TEXT("All payload elements correctly updated"), bAllElementsCorrect);
+	}
+
+	return true;
+}
+
+/**
+ * This test will create a buffer, then serialize it to a FEditorBulkData object via FEditorBulkDataWriter.
+ * Then we will serialize the FEditorBulkData object back to a second buffer and compare the results.
+ * If the reader and writers are working then the ReplicatedBuffer should be the same as the original SourceBuffer.
+ */
+IMPLEMENT_SIMPLE_AUTOMATION_TEST(FEditorBulkDataTestReaderWriter, TEXT("System.CoreUObject.Serialization.EditorBulkData.Reader/Writer"), TestFlags)
+bool FEditorBulkDataTestReaderWriter::RunTest(const FString& Parameters)
+{
+	const int64 BufferSize = 1024;
+
+	TUniquePtr<uint8[]> SourceBuffer = CreateRandomData(BufferSize);
+	TUniquePtr<uint8[]> ReplicatedBuffer = MakeUnique<uint8[]>(BufferSize);
+
+	FEditorBulkData BulkData;
+
+	// Serialize the SourceBuffer to BulkData 
+	{
+		FEditorBulkDataWriter WriterAr(BulkData);
+		WriterAr.Serialize(SourceBuffer.Get(), BufferSize);
+	}
+
+	// Serialize BulkData back to ReplicatedBuffer
+	{
+		FEditorBulkDataReader ReaderAr(BulkData);
+		ReaderAr.Serialize(ReplicatedBuffer.Get(), BufferSize);
+	}
+
+	// Now test that the buffer was restored to the original values
+	const bool bMemCmpResult = FMemory::Memcmp(SourceBuffer.Get(), ReplicatedBuffer.Get(), BufferSize) == 0;
+	TestTrue(TEXT("SourceBuffer values == ReplicatedBuffer values"), bMemCmpResult);
+
+	// Test writing nothing to an empty bulkdata object and then reading that bulkdata object
+	// to make sure that we deal with null buffers properly.
+	{
+		FEditorBulkData EmptyBulkData;
+		FEditorBulkDataWriter WriterAr(EmptyBulkData);
+		FEditorBulkDataReader ReaderAr(EmptyBulkData);
+	}
+	return true;
+}
+
+/** 
+ * This test will serialize a number of empty and valid bulkdata objects to a memory buffer and then serialize them back 
+ * as bulkdata objects again. This should help find problems where the data saved and loaded are mismatched.
+ * Note that this does not check the package saving code paths, only direct serialization to buffers.
+ */
+IMPLEMENT_SIMPLE_AUTOMATION_TEST(FEditorBulkDataTestSerializationToMemory, TEXT("System.CoreUObject.Serialization.EditorBulkData.SerializationToMemory"), TestFlags)
+bool FEditorBulkDataTestSerializationToMemory::RunTest(const FString& Parameters)
+{
+	const int64 BufferSize = 1024;
+
+	TUniquePtr<uint8[]> SourceBuffer = CreateRandomData(BufferSize);
+	
+	FEditorBulkData EmptyBulkData;
+	FEditorBulkData ValidBulkData;
+	ValidBulkData.UpdatePayload(FSharedBuffer::Clone(SourceBuffer.Get(), BufferSize));
+	
+	TArray<uint8> MemoryBuffer;
+	const bool bIsArPersistent = true;
+
+	// Write out
+	{
+		FMemoryWriter WriterAr(MemoryBuffer, bIsArPersistent);
+
+		ValidBulkData.Serialize(WriterAr, nullptr);
+		EmptyBulkData.Serialize(WriterAr, nullptr);
+		ValidBulkData.Serialize(WriterAr, nullptr);
+		EmptyBulkData.Serialize(WriterAr, nullptr);
+	}
+
+	// Read in
+	{
+		FEditorBulkData SerializedBulkData;
+
+		FMemoryReader ReaderAr(MemoryBuffer, bIsArPersistent);
+
+		SerializedBulkData.Serialize(ReaderAr, nullptr);
+		TestEqual(TEXT("Bulkdata identifier should remain the same"), SerializedBulkData.GetIdentifier(), ValidBulkData.GetIdentifier());
+		TestTrue(TEXT("Bulkdata payload should remain the same"), FMemory::Memcmp(SourceBuffer.Get(), SerializedBulkData.GetPayload().Get().GetData(), BufferSize) == 0);
+
+		SerializedBulkData.Serialize(ReaderAr, nullptr);
+		TestEqual(TEXT("Bulkdata identifier should remain the same"), SerializedBulkData.GetIdentifier(), EmptyBulkData.GetIdentifier());
+		TestTrue(TEXT("Bulkdata should not have a payload"), SerializedBulkData.GetPayload().Get().IsNull());
+
+		SerializedBulkData.Serialize(ReaderAr, nullptr);
+		TestEqual(TEXT("Bulkdata identifier should remain the same"), SerializedBulkData.GetIdentifier(), ValidBulkData.GetIdentifier());
+		TestTrue(TEXT("Bulkdata payload should remain the same"), FMemory::Memcmp(SourceBuffer.Get(), SerializedBulkData.GetPayload().Get().GetData(), BufferSize) == 0);
+
+		SerializedBulkData.Serialize(ReaderAr, nullptr);
+		TestEqual(TEXT("Bulkdata identifier should remain the same"), SerializedBulkData.GetIdentifier(), EmptyBulkData.GetIdentifier());
+		TestTrue(TEXT("Bulkdata should not have a payload"), SerializedBulkData.GetPayload().Get().IsNull());
+	}
+	return true;
+}
+
+/**
+ * This set of tests validate that the BulkData's identifier works how we expect it too. It should remain unique in all cases except
+ * move semantics.
+ */
+IMPLEMENT_SIMPLE_AUTOMATION_TEST(FEditorBulkDataTestIdentifiers, TEXT("System.CoreUObject.Serialization.EditorBulkData.Identifiers"), TestFlags)
+bool FEditorBulkDataTestIdentifiers::RunTest(const FString& Parameters)
+{
+	// Some basic tests with an invalid id
+	{
+		FEditorBulkData BulkData;
+		TestFalse(TEXT("BulkData with no payload should return an invalid identifier"), BulkData.GetIdentifier().IsValid());
+
+		FEditorBulkData CopiedBulkData(BulkData);
+		TestFalse(TEXT("Copying a bulkdata with an invalid id should result in an invalid id"), CopiedBulkData.GetIdentifier().IsValid());
+
+		FEditorBulkData AssignedBulkData;
+		AssignedBulkData = BulkData;
+		TestFalse(TEXT("Assigning a a bulkdata with an invalid id should result in an invalid id"), AssignedBulkData.GetIdentifier().IsValid());
+		
+		// Check that we did not change the initial object at any point
+		TestFalse(TEXT("Being copied and assigned to other objects should not affect the identifier"), BulkData.GetIdentifier().IsValid());
+	}
+
+	// Some basic tests with a valid id
+	{
+		FEditorBulkData BulkData;
+		BulkData.UpdatePayload(FUniqueBuffer::Alloc(32).MoveToShared()); // Assigning this payload should cause BulkData to gain an identifier
+		TestTrue(TEXT("BulkData with a payload should returns a valid identifier"), BulkData.GetIdentifier().IsValid());
+
+		const FGuid OriginalGuid = BulkData.GetIdentifier();
+
+		FEditorBulkData CopiedBulkData(BulkData);
+		TestNotEqual(TEXT("Copying a bulkdata with a valid id should result in a unique identifier"), BulkData.GetIdentifier(), CopiedBulkData.GetIdentifier());
+
+		FEditorBulkData AssignedBulkData;
+		AssignedBulkData = BulkData;
+		TestNotEqual(TEXT("Assignment operator creates different identifiers"), BulkData.GetIdentifier(), AssignedBulkData.GetIdentifier());
+
+		// Check that we did not change the initial object at any point
+		TestEqual(TEXT("Being copied and assigned to other objects should not affect the identifier"), BulkData.GetIdentifier(), OriginalGuid);
+
+		// Now that AssignedBulkData has a valid identifier, make sure that it is not changed if we assign something else to it.
+		const FGuid OriginalAssignedGuid = AssignedBulkData.GetIdentifier();
+		AssignedBulkData = CopiedBulkData;
+		TestEqual(TEXT("Being copied and assigned to other objects should not affect the identifier"), AssignedBulkData.GetIdentifier(), OriginalAssignedGuid);
+	}
+
+	// Test move constructor
+	{	
+		FEditorBulkData BulkData;
+		BulkData.UpdatePayload(FUniqueBuffer::Alloc(32).MoveToShared());
+
+		FGuid OriginalGuid = BulkData.GetIdentifier();
+
+		FEditorBulkData MovedBulkData = MoveTemp(BulkData);
+
+		TestEqual(TEXT("Move constructor should preserve the identifier"), MovedBulkData.GetIdentifier(), OriginalGuid);
+	}
+
+	// Test move assignment
+	{
+		FEditorBulkData BulkData;
+		BulkData.UpdatePayload(FUniqueBuffer::Alloc(32).MoveToShared());
+
+		FGuid OriginalGuid = BulkData.GetIdentifier();
+
+		FEditorBulkData MovedBulkData;
+		MovedBulkData = MoveTemp(BulkData);
+
+		TestEqual(TEXT("Move assignment should preserve the identifier"), MovedBulkData.GetIdentifier(), OriginalGuid);
+	}
+
+	// Check that resizing an array will not change the internals
+	{
+		const uint32 NumToTest = 10;
+
+		TArray<FEditorBulkData> BulkDataArray;
+		TArray<FGuid> GuidArray;
+
+		for (uint32 Index = 0; Index < NumToTest; ++Index)
+		{
+			BulkDataArray.Add(FEditorBulkData());
+
+			// Leave some with invalid ids and some with valid ones
+			if (Index % 2 == 0)
+			{
+				BulkDataArray[Index].UpdatePayload(FUniqueBuffer::Alloc(32).MoveToShared());
+			}
+
+			GuidArray.Add(BulkDataArray[Index].GetIdentifier());
+		}
+
+		// Force an internal reallocation and make sure that the identifiers are unchanged.
+		// Note that it is possible that the allocation is just resized and not reallocated.
+		BulkDataArray.Reserve(BulkDataArray.Max() * 4);
+
+		for (uint32 Index = 0; Index < NumToTest; ++Index)
+		{
+			TestEqual(TEXT(""), BulkDataArray[Index].GetIdentifier(), GuidArray[Index]);
+		}
+
+		// Now insert a new item, moving all of the existing entries and make sure that
+		// the identifiers are unchanged.
+		BulkDataArray.Insert(FEditorBulkData(), 0);
+
+		for (uint32 Index = 0; Index < NumToTest; ++Index)
+		{
+			TestEqual(TEXT(""), BulkDataArray[Index+1].GetIdentifier(), GuidArray[Index]);
+		}
+	}
+
+	// Test that adding a payload to a reset bulkdata object or one that has had a zero length payload applied
+	// will correctly show the original id once it has a valid payload.
+	{
+		FEditorBulkData BulkData;
+		BulkData.UpdatePayload(FUniqueBuffer::Alloc(32).MoveToShared());
+		
+		const FGuid OriginalGuid = BulkData.GetIdentifier();
+
+		BulkData.Reset();
+		TestEqual(TEXT("BulkData with no payload should return its original identifier"), BulkData.GetIdentifier(), OriginalGuid);
+
+		BulkData.UpdatePayload(FUniqueBuffer::Alloc(32).MoveToShared());
+		TestEqual(TEXT("Removing a payload then adding a new one should return the original identifier"), BulkData.GetIdentifier(), OriginalGuid);
+
+		BulkData.UpdatePayload(FUniqueBuffer::Alloc(0).MoveToShared());
+		TestEqual(TEXT("Setting a zero length payload should keep the original identifier"), BulkData.GetIdentifier(), OriginalGuid);
+		
+		BulkData.UpdatePayload(FUniqueBuffer::Alloc(32).MoveToShared());
+		TestEqual(TEXT("Restoring a payload should return the original identifier"), BulkData.GetIdentifier(), OriginalGuid);
+	}
+
+	// Test that serialization does not change the identifier (in this case serializing to and from a memory buffer) 
+	{
+		FEditorBulkData SrcData;
+		SrcData.UpdatePayload(FUniqueBuffer::Alloc(32).MoveToShared());
+
+		TArray<uint8> MemoryBuffer;
+		const bool bIsArPersistent = true;
+
+		FEditorBulkData DstData;
+
+		// Serialize the SourceBuffer to BulkData 
+		{
+			FMemoryWriter WriterAr(MemoryBuffer, bIsArPersistent);
+			SrcData.Serialize(WriterAr, nullptr);
+		}
+
+		// Serialize BulkData back to ReplicatedBuffer
+		{
+			FMemoryReader ReaderAr(MemoryBuffer, bIsArPersistent);
+			DstData.Serialize(ReaderAr, nullptr);
+		}
+
+		TestEqual(TEXT("Serialization should preserve the identifier"), SrcData.GetIdentifier(), DstData.GetIdentifier());
+	}
+
+	// Test that serialization does not change the identifier (in this case serializing to and from a memory buffer) 
+	{
+		FEditorBulkData SrcData;
+		SrcData.UpdatePayload(FUniqueBuffer::Alloc(32).MoveToShared());
+
+		const FGuid OriginalIdentifier = SrcData.GetIdentifier();
+		SrcData.UpdatePayload(FUniqueBuffer::Alloc(0).MoveToShared());
+
+		TArray<uint8> MemoryBuffer;
+		const bool bIsArPersistent = true;
+
+		FEditorBulkData DstData;
+
+		// Serialize the SourceBuffer to BulkData 
+		{
+			FMemoryWriter WriterAr(MemoryBuffer, bIsArPersistent);
+			SrcData.Serialize(WriterAr, nullptr);
+		}
+
+		// Serialize BulkData back to ReplicatedBuffer
+		{
+			FMemoryReader ReaderAr(MemoryBuffer, bIsArPersistent);
+			DstData.Serialize(ReaderAr, nullptr);
+		}
+
+		TestEqual(TEXT("After serialization the identifier should keep its identifier "), DstData.GetIdentifier(), OriginalIdentifier);
+
+		DstData.UpdatePayload(FUniqueBuffer::Alloc(32).MoveToShared());
+		TestEqual(TEXT("After adding a new payload the object should have the original identifier"), DstData.GetIdentifier(), OriginalIdentifier);
+	}
+
+	return true;
+}
+
+IMPLEMENT_SIMPLE_AUTOMATION_TEST(FEditorBulkDataTestZeroSizedAllocs, TEXT("System.CoreUObject.Serialization.EditorBulkData.ZeroSizedAllocs"), TestFlags)
+bool FEditorBulkDataTestZeroSizedAllocs::RunTest(const FString& Parameters)
+{
+	{
+		FEditorBulkData Empty;
+		TestNull(TEXT("Empty bulkdata returning a payload"), Empty.GetPayload().Get().GetData());
+		TestNull(TEXT("Empty bulkdata returning a compressed payload"), Empty.GetCompressedPayload().Get().Decompress().GetData());
+
+		TestFalse(TEXT("Empty payload has payload data"), Empty.HasPayloadData());
+		TestFalse(TEXT("Empty payload needs loading"), Empty.DoesPayloadNeedLoading());
+	}
+
+	{
+		FEditorBulkData EmptySrc;
+		
+		FLargeMemoryWriter ArWrite(0, true);
+		EmptySrc.Serialize(ArWrite, nullptr);
+
+		FEditorBulkData EmptyDst;
+
+		FLargeMemoryReader ArRead(ArWrite.GetData(), ArWrite.TotalSize(), ELargeMemoryReaderFlags::Persistent);
+		EmptyDst.Serialize(ArRead, nullptr);
+
+		TestNull(TEXT("Serialized bulkdata returning a payload"), EmptyDst.GetPayload().Get().GetData());
+		TestNull(TEXT("Serialized bulkdata returning a compressed payload"), EmptyDst.GetCompressedPayload().Get().Decompress().GetData());
+
+		TestFalse(TEXT("Serialized payload has payload data"), EmptyDst.HasPayloadData());
+		TestFalse(TEXT("Serialized payload needs loading"), EmptyDst.DoesPayloadNeedLoading());
+	}
+
+	{
+		FEditorBulkData ZeroAlloc;
+		ZeroAlloc.UpdatePayload(FUniqueBuffer::Alloc(0).MoveToShared());
+
+		TestNull(TEXT("ZeroAlloc bulkdata returning a payload"), ZeroAlloc.GetPayload().Get().GetData());
+		TestNull(TEXT("ZeroAlloc bulkdata returning a compressed payload"), ZeroAlloc.GetCompressedPayload().Get().Decompress().GetData());
+
+		TestFalse(TEXT("ZeroAlloc payload has payload data"), ZeroAlloc.HasPayloadData());
+		TestFalse(TEXT("ZeroAlloc payload needs loading"), ZeroAlloc.DoesPayloadNeedLoading());
+	}
+
+	return true;
+}
+
+/** This tests the function UE::Serialization::IoHashToGuid which is closely used with the FEditorBulkData system */
+IMPLEMENT_SIMPLE_AUTOMATION_TEST(FEditorBulkDataIoHashToGuid, TEXT("System.CoreUObject.Serialization.EditorBulkData.IoHashToGuid"), TestFlags)
+bool FEditorBulkDataIoHashToGuid::RunTest(const FString& Parameters)
+{
+	// Test that an empty hash will give an invalid FGuid
+	FIoHash DefaultHash;
+	FGuid InvalidGuid = IoHashToGuid(DefaultHash);
+
+	TestTrue(TEXT("Calling::IoHashToGuid on a all zero FIoHash should result in an invalid FGuid"), !InvalidGuid.IsValid());
+	
+	// Test that finding the FGuid of a known hash results in the FGuid we expect. If not then the generation algorithm has
+	// changed. The failing test should remind whoever changed the algorithm to double check that changing the results 
+	// will not have knock on effects.
+	const uint8 KnownHashData[20] = {1, 2, 3, 4, 5, 6, 7, 8, 9, 10, 11, 12, 13, 14, 15, 16, 17, 18 , 19};
+
+	FIoHash KnownHash(KnownHashData);
+	FGuid KnownGuid = IoHashToGuid(KnownHash);
+	FGuid KnownResult(TEXT("04030201-0807-0605-0C0B-0A09100F0E0D"));
+
+	TestEqual(TEXT("That the result of hashing known data will be unchanged"), KnownGuid, KnownResult);
+
+	return true;
+}
+
+/** This tests that updating a payload via a FSharedBufferWithID will have the same results as if the payload was applied to the bulkdata object directly */
+IMPLEMENT_SIMPLE_AUTOMATION_TEST(FEditorBulkDataSharedBufferWithID, TEXT("System.CoreUObject.Serialization.EditorBulkData.SharedBufferWithID"), TestFlags)
+bool FEditorBulkDataSharedBufferWithID::RunTest(const FString& Parameters)
+{
+	// Test FSharedBufferWithID/FEditorBulkData in their default states
+	{
+		FEditorBulkData::FSharedBufferWithID SharedId;
+		FEditorBulkData BulkDataFromSharedId;
+		BulkDataFromSharedId.UpdatePayload(MoveTemp(SharedId));
+
+		FEditorBulkData BulkData;
+
+		TestEqual(TEXT("0: PayloadId"), BulkDataFromSharedId.GetPayloadId(), BulkData.GetPayloadId());
+		TestEqual(TEXT("0: PayloadSize"), BulkDataFromSharedId.GetPayloadSize(), BulkData.GetPayloadSize());
+	}
+
+	// Test FSharedBufferWithID/FEditorBulkData updated with a null FSharedBuffer
+	{
+		FSharedBuffer NullBuffer;
+
+		FEditorBulkData::FSharedBufferWithID SharedId(NullBuffer);
+		FEditorBulkData BulkDataFromSharedId;
+		BulkDataFromSharedId.UpdatePayload(MoveTemp(SharedId));
+
+		FEditorBulkData BulkData;
+		BulkData.UpdatePayload(NullBuffer);
+
+		TestEqual(TEXT("1: PayloadId"), BulkDataFromSharedId.GetPayloadId(), BulkData.GetPayloadId());
+		TestEqual(TEXT("1: PayloadSize"), BulkDataFromSharedId.GetPayloadSize(), BulkData.GetPayloadSize());
+	}
+
+	// Test FSharedBufferWithID/FEditorBulkData updated with a zero length FSharedBuffer
+	{
+		FSharedBuffer ZeroLengthBuffer = FUniqueBuffer::Alloc(0).MoveToShared();
+
+		FEditorBulkData::FSharedBufferWithID SharedId(ZeroLengthBuffer);
+		FEditorBulkData BulkDataFromSharedId;
+		BulkDataFromSharedId.UpdatePayload(MoveTemp(SharedId));
+
+		FEditorBulkData BulkData;
+		BulkData.UpdatePayload(ZeroLengthBuffer);
+
+		TestEqual(TEXT("2: PayloadId"), BulkDataFromSharedId.GetPayloadId(), BulkData.GetPayloadId());
+		TestEqual(TEXT("2: PayloadSize"), BulkDataFromSharedId.GetPayloadSize(), BulkData.GetPayloadSize());
+	}
+
+	// Test FSharedBufferWithID/FEditorBulkData updated with a random set of data
+	{
+		const int64 BufferSize = 1024;
+		TUniquePtr<uint8[]> SourceBuffer = CreateRandomData(BufferSize);
+		FSharedBuffer RandomData = FSharedBuffer::MakeView(SourceBuffer.Get(), BufferSize);
+
+		FEditorBulkData::FSharedBufferWithID SharedId(RandomData);
+		FEditorBulkData BulkDataFromSharedId;
+		BulkDataFromSharedId.UpdatePayload(MoveTemp(SharedId));
+
+		FEditorBulkData BulkData;
+		BulkData.UpdatePayload(RandomData);
+
+		TestEqual(TEXT("3: PayloadId"), BulkDataFromSharedId.GetPayloadId(), BulkData.GetPayloadId());
+		TestEqual(TEXT("3: PayloadSize"), BulkDataFromSharedId.GetPayloadSize(), BulkData.GetPayloadSize());
+	}
+
+	return true;
+}
+
+
+} // namespace UE::Serialization
+
+#endif //WITH_DEV_AUTOMATION_TESTS && WITH_EDITORONLY_DATA
diff --git a/Engine/Source/Runtime/CoreUObject/Tests/ObjectHandleTest.cpp b/Engine/Source/Runtime/CoreUObject/Tests/ObjectHandleTest.cpp
--- a/Engine/Source/Runtime/CoreUObject/Tests/ObjectHandleTest.cpp
+++ b/Engine/Source/Runtime/CoreUObject/Tests/ObjectHandleTest.cpp
@@ -6,6 +6,8 @@
 #include "ObjectRefTrackingTestBase.h"
 #include "IO/IoDispatcher.h"
 
+#include "TestHarness.h"
+
 static_assert(sizeof(FObjectHandle) == sizeof(void*), "FObjectHandle type must always compile to something equivalent to a pointer size.");
 
 class FObjectHandleTestBase : public FObjectRefTrackingTestBase
@@ -18,17 +20,23 @@ protected:
 	{
 	#if UE_WITH_OBJECT_HANDLE_LATE_RESOLVE
 		// Late resolved handles cannot be null or resolved at this point
-		if (!TestFalse(TEXT("Handle to target is null"), IsObjectHandleNull(TargetHandle)))
+		bool bValue = IsObjectHandleNull(TargetHandle);
+		TEST_FALSE(TEXT("Handle to target is null"), bValue);
+		if (bValue)
 		{
 			return nullptr;
 		}
-		if (!TestFalse(TEXT("Handle to target is resolved"), IsObjectHandleResolved(TargetHandle)))
+		bValue = IsObjectHandleResolved(TargetHandle);
+		TEST_FALSE(TEXT("Handle to target is resolved"), bValue);
+		if (bValue)
 		{
 			return nullptr;
 		}
 	#else
+		bool bValue = IsObjectHandleResolved(TargetHandle);
 		// Immediately resolved handles may be null (if the target is invalid) and must be resolved at this point
-		if (!TestFalse(TEXT("Handle to target is not resolved"), IsObjectHandleResolved(TargetHandle)))
+		TEST_TRUE(TEXT("Handle to target is not resolved"), bValue);
+		if (!bValue)
 		{
 			return nullptr;
 		}
@@ -40,7 +48,9 @@ protected:
 	UObject* ConstructAndResolveHandle(const ANSICHAR* PackageName, const ANSICHAR* ObjectName, const ANSICHAR* ClassPackageName = nullptr, const ANSICHAR* ClassName = nullptr)
 	{
 		FObjectRef TargetRef{FName(PackageName), FName(ClassPackageName), FName(ClassName), FObjectPathId(ObjectName)};
-		if (!TestFalse(TEXT("Reference to target is null"), IsObjectRefNull(TargetRef)))
+		bool bValue = IsObjectRefNull(TargetRef);
+		TEST_FALSE(TEXT("Reference to target is null"), bValue);
+		if (bValue)
 		{
 			return nullptr;
 		}
@@ -51,7 +61,9 @@ protected:
 
 	UObject* ConstructAndResolveHandle(const FPackedObjectRef& PackedTargetRef)
 	{
-		if (!TestFalse(TEXT("Reference to target is null"), IsPackedObjectRefNull(PackedTargetRef)))
+		bool bValue = IsPackedObjectRefNull(PackedTargetRef);
+		TEST_FALSE(TEXT("Reference to target is null"), bValue);
+		if (bValue)
 		{
 			return nullptr;
 		}
@@ -110,7 +122,7 @@ protected:
 	}
 };
 
-TEST_CASE_METHOD(FObjectHandleTestBase, "CoreUObject::FObjectHandle::Null Behavior", "[CoreUObject][ObjectHandle][Smoke]")
+TEST_CASE_METHOD(FObjectHandleTestBase, "CoreUObject::FObjectHandle::Null Behavior", "[CoreUObject][ObjectHandle]")
 {
 	FObjectHandle TargetHandle = MakeObjectHandle(nullptr);
 
@@ -127,7 +139,7 @@ TEST_CASE_METHOD(FObjectHandleTestBase, "CoreUObject::FObjectHandle::Null Behavi
 	ObjectRefMetrics.TestNumReads(TEXT("NumReads should be incremented by one after a resolve attempt on a null handle"), 1);
 }
 
-TEST_CASE_METHOD(FObjectHandleTestBase, "CoreUObject::FObjectHandle::Pointer Behavior", "[CoreUObject][ObjectHandle][Smoke]")
+TEST_CASE_METHOD(FObjectHandleTestBase, "CoreUObject::FObjectHandle::Pointer Behavior", "[CoreUObject][ObjectHandle]")
 {
 	FObjectHandle TargetHandle = MakeObjectHandle((UObject*)0x0042);
 
@@ -160,7 +172,7 @@ TEST_CASE_METHOD(FObjectHandleTestBase, "CoreUObject::FObjectHandle::Resolve Eng
 	}
 }
 
-TEST_CASE_METHOD(FObjectHandleTestBase, "CoreUObject::FObjectHandle::Resolve Non Existent Target", "[CoreUObject][ObjectHandle][Smoke]")
+TEST_CASE_METHOD(FObjectHandleTestBase, "CoreUObject::FObjectHandle::Resolve Non Existent Target", "[CoreUObject][ObjectHandle]")
 {
 	// Confirm we don't successfully resolve an incorrect reference to engine content
 	TestResolveFailure("/Engine/EngineResources/NonExistentPackageName_0", "DefaultTexture");
@@ -177,7 +189,7 @@ TEST_CASE_METHOD(FObjectHandleTestBase, "CoreUObject::FObjectHandle::Resolve Scr
 
 #if UE_WITH_OBJECT_HANDLE_LATE_RESOLVE
 
-TEST_CASE_METHOD(FObjectHandleTestBase, "CoreUObject::FObjectHandle::Resolve Malformed Handle", "[CoreUObject][ObjectHandle][Smoke]")
+TEST_CASE_METHOD(FObjectHandleTestBase, "CoreUObject::FObjectHandle::Resolve Malformed Handle", "[CoreUObject][ObjectHandle]")
 {
 	TestResolveFailure(FPackedObjectRef { 0xFFFF'FFFF'FFFF'FFFFull });
 	TestResolveFailure(FPackedObjectRef { 0xEFEF'EFEF'EFEF'EFEFull });
diff --git a/Engine/Source/Runtime/CoreUObject/Tests/ObjectPtrTest.cpp b/Engine/Source/Runtime/CoreUObject/Tests/ObjectPtrTest.cpp
--- a/Engine/Source/Runtime/CoreUObject/Tests/ObjectPtrTest.cpp
+++ b/Engine/Source/Runtime/CoreUObject/Tests/ObjectPtrTest.cpp
@@ -149,7 +149,7 @@ protected:
 };
 
 
-TEST_CASE_METHOD(FObjectPtrTestBase, "CoreUObject::TObjectPtr::Null Behavior", "[CoreUObject][ObjectPtr][Smoke]")
+TEST_CASE_METHOD(FObjectPtrTestBase, "CoreUObject::TObjectPtr::Null Behavior", "[CoreUObject][ObjectPtr]")
 {
 	TObjectPtr<UObject> NullObjectPtr(nullptr);
 	TEST_TRUE(TEXT("Nullptr should equal a null object pointer"), nullptr == NullObjectPtr);
@@ -181,7 +181,7 @@ TEST_CASE_METHOD(FObjectPtrTestBase, "CoreUObject::TObjectPtr::Default Serialize
 	ObjectRefMetrics.TestNumReads(TEXT("NumReads should increase after serializing an FObjectPtr"), 2);
 }
 
-TEST_CASE_METHOD(FObjectPtrTestBase, "CoreUObject::TObjectPtr::Soft Object Path", "[CoreUObject][ObjectPtr][Smoke]")
+TEST_CASE_METHOD(FObjectPtrTestBase, "CoreUObject::TObjectPtr::Soft Object Path", "[CoreUObject][ObjectPtr]")
 {
 	FSnapshotObjectRefMetrics ObjectRefMetrics(*this);
 	FObjectPtr DefaultTexturePtr(FObjectRef {FName("/Engine/EngineResources/DefaultTexture"), NAME_None, NAME_None, FObjectPathId("DefaultTexture")});
@@ -194,17 +194,17 @@ TEST_CASE_METHOD(FObjectPtrTestBase, "CoreUObject::TObjectPtr::Soft Object Path"
 
 	ObjectRefMetrics.TestNumResolves(TEXT("Unexpected resolve count after initializing an FSoftObjectPath from an FObjectPtr"), UE_WITH_OBJECT_HANDLE_LATE_RESOLVE ? 0 : 1);
 
-	TEST_TRUE(TEXT("Soft object path constructed from an FObjectPtr does not have the expected path value"), !FCString::Strcmp(*DefaultTexturePath.ToString(), TEXT("/Engine/EngineResources/DefaultTexture.DefaultTexture")));
+	TEST_EQUAL_STR(TEXT("Soft object path constructed from an FObjectPtr does not have the expected path value"), TEXT("/Engine/EngineResources/DefaultTexture.DefaultTexture"), *DefaultTexturePath.ToString());
 }
 
-TEST_CASE_METHOD(FObjectPtrTestBase, "CoreUObject::TObjectPtr::Forward Declared", "[CoreUObject][ObjectPtr][Smoke]")
+TEST_CASE_METHOD(FObjectPtrTestBase, "CoreUObject::TObjectPtr::Forward Declared", "[CoreUObject][ObjectPtr]")
 {
 	UForwardDeclaredObjDerived* PtrFwd = nullptr;
 	TObjectPtr<UForwardDeclaredObjDerived> ObjPtrFwd(MakeObjectPtrUnsafe<UForwardDeclaredObjDerived>(reinterpret_cast<UObject*>(PtrFwd)));
 	TEST_TRUE(TEXT("Null forward declared pointer used to construct a TObjectPtr should result in a null TObjectPtr"), ObjPtrFwd.IsNull());
 }
 
-TEST_CASE_METHOD(FObjectPtrTestBase, "CoreUObject::TObjectPtr::Hash Consistency", "[CoreUObject][ObjectPtr][Smoke]")
+TEST_CASE_METHOD(FObjectPtrTestBase, "CoreUObject::TObjectPtr::Hash Consistency", "[CoreUObject][ObjectPtr]")
 {
 	const FName TestPackage1Name(TEXT("/Engine/Test/ObjectPtrHashConsistency1/Transient"));
 	UPackage* TestPackage1 = NewObject<UPackage>(nullptr, TestPackage1Name, RF_Transient);
@@ -276,7 +276,7 @@ TEST_CASE_METHOD(FObjectPtrTestBase, "CoreUObject::TObjectPtr::Hash Consistency"
 }
 
 
-TEST_CASE_METHOD(FObjectPtrTestBase, "CoreUObject::TObjectPtr::Long Path", "[CoreUObject][ObjectPtr][Smoke]")
+TEST_CASE_METHOD(FObjectPtrTestBase, "CoreUObject::TObjectPtr::Long Path", "[CoreUObject][ObjectPtr]")
 {
 	const FName TestPackage1Name(TEXT("/Engine/Test/FObjectPtrTestLongPath/Transient"));
 	UPackage* TestPackage1 = NewObject<UPackage>(nullptr, TestPackage1Name, RF_Transient);
diff --git a/Engine/Source/Runtime/CoreUObject/Tests/ObjectRefTrackingTestBase.h b/Engine/Source/Runtime/CoreUObject/Tests/ObjectRefTrackingTestBase.h
--- a/Engine/Source/Runtime/CoreUObject/Tests/ObjectRefTrackingTestBase.h
+++ b/Engine/Source/Runtime/CoreUObject/Tests/ObjectRefTrackingTestBase.h
@@ -29,7 +29,9 @@ public:
 		bool TestNumResolves(const TCHAR* What, uint32 ExpectedDelta)
 		{
 #if UE_WITH_OBJECT_HANDLE_TRACKING
-			TEST_EQUAL(What, OriginalNumResolves + ExpectedDelta, Test.GetNumResolves());
+			bool bValue = OriginalNumResolves + ExpectedDelta == Test.GetNumResolves();
+			TEST_TRUE(What, bValue);
+			return bValue;
 #endif
 			return true;
 		}
@@ -37,7 +39,9 @@ public:
 		bool TestNumFailedResolves(const TCHAR* What, uint32 ExpectedDelta)
 		{
 #if UE_WITH_OBJECT_HANDLE_TRACKING
-			TEST_EQUAL(What, OriginalNumFailedResolves + ExpectedDelta, Test.GetNumFailedResolves());
+			bool bValue = OriginalNumFailedResolves + ExpectedDelta == Test.GetNumFailedResolves();
+			TEST_TRUE(What, bValue);
+			return bValue;
 #endif
 			return true;
 
@@ -46,14 +50,18 @@ public:
 		bool TestNumReads(const TCHAR* What, uint32 ExpectedDelta, bool bAllowAdditionalReads = false)
 		{
 #if UE_WITH_OBJECT_HANDLE_TRACKING
+			bool bValue = false;
 			if (bAllowAdditionalReads)
 			{
-				TEST_TRUE(What, Test.GetNumReads() >= OriginalNumReads + ExpectedDelta);
+				bValue = Test.GetNumReads() >= OriginalNumReads + ExpectedDelta;
+				TEST_TRUE(What, bValue);
 			}
 			else
 			{
-				TEST_EQUAL(What, OriginalNumReads + ExpectedDelta, Test.GetNumReads());
+				bValue = OriginalNumReads + ExpectedDelta == Test.GetNumReads();
+				TEST_TRUE(What, bValue);
 			}
+			return bValue;
 #endif
 			return true;
 
diff --git a/Engine/Source/Runtime/CoreUObject/Tests/Serialization/BulkDataTests.cpp b/Engine/Source/Runtime/CoreUObject/Tests/Serialization/BulkDataTests.cpp
--- a/Engine/Source/Runtime/CoreUObject/Tests/Serialization/BulkDataTests.cpp
+++ b/Engine/Source/Runtime/CoreUObject/Tests/Serialization/BulkDataTests.cpp
@@ -6,7 +6,7 @@
 namespace BulkDataTest
 {
 	// Test code paths for BulkData objects that do not reference a file on disk.
-	TEST_CASE_METHOD(FAutomationTestFixture, "CoreUObject::Serialization::FByteBulkData::Transient", "[CoreUObject][Serialization][Smoke]")
+	TEST_CASE_METHOD(FAutomationTestFixture, "CoreUObject::Serialization::FByteBulkData::Transient", "[CoreUObject][Serialization]")
 	{
 		FByteBulkData BulkData;
 
@@ -61,7 +61,7 @@ namespace BulkDataTest
 
 #if USE_NEW_BULKDATA // The warning only exists for BulkData2
 		// Both the ::Lock and ::GetCopy calls should warn that we cannot load the missing data (as it will still have a valid size)
-		FAIL(TEXT("Attempting to load a BulkData object that cannot be loaded from disk"));
+		WARN(TEXT("Attempting to load a BulkData object that cannot be loaded from disk"));
 #endif //USE_NEW_BULKDATA
 
 		DataPtr = BulkData.Lock(LOCK_READ_WRITE);
diff --git a/Engine/Source/Runtime/CoreUObject/Tests/Serialization/EditorBulkDataTests.cpp b/Engine/Source/Runtime/CoreUObject/Tests/Serialization/EditorBulkDataTests.cpp
deleted file mode 100644
