Fix Crash when adding Event Dispatcher to a blueprint currently being viewed in the Header Preview

diff --git a/Engine/Source/Runtime/CoreUObject/Private/UObject/PropertyMulticastDelegate.cpp b/Engine/Source/Runtime/CoreUObject/Private/UObject/PropertyMulticastDelegate.cpp
--- a/Engine/Source/Runtime/CoreUObject/Private/UObject/PropertyMulticastDelegate.cpp
+++ b/Engine/Source/Runtime/CoreUObject/Private/UObject/PropertyMulticastDelegate.cpp
@@ -130,14 +130,12 @@ bool FMulticastDelegateProperty::NetSerializeItem( FArchive& Ar, UPackageMap* Ma
 
 FString FMulticastDelegateProperty::GetCPPType( FString* ExtendedTypeText/*=NULL*/, uint32 CPPExportFlags/*=0*/ ) const
 {
-#if HACK_HEADER_GENERATOR
 	// We have this test because sometimes the delegate hasn't been set up by FixupDelegateProperties at the time
 	// we need the type for an error message.  We deliberately format it so that it's unambiguously not CPP code, but is still human-readable.
 	if (!SignatureFunction)
 	{
 		return FString(TEXT("{multicast delegate type}"));
 	}
-#endif
 
 	FString UnmangledFunctionName = SignatureFunction->GetName().LeftChop( FString( HEADER_GENERATED_DELEGATE_SIGNATURE_SUFFIX ).Len() );
 	const UClass* OwnerClass = SignatureFunction->GetOwnerClass();
