Removal of SA_FIELD_NAME.

diff --git a/Engine/Plugins/Experimental/StructUtils/Source/StructUtils/Private/InstancedStruct.cpp b/Engine/Plugins/Experimental/StructUtils/Source/StructUtils/Private/InstancedStruct.cpp
--- a/Engine/Plugins/Experimental/StructUtils/Source/StructUtils/Private/InstancedStruct.cpp
+++ b/Engine/Plugins/Experimental/StructUtils/Source/StructUtils/Private/InstancedStruct.cpp
@@ -294,7 +294,7 @@ bool FInstancedStruct::SerializeFromMismatchedTag(const FPropertyTag& Tag, FStru
 			{
 				if (StructTypePtr)
 				{
-					StructTypePtr->SerializeItem(Record.EnterField(SA_FIELD_NAME(TEXT("StructInstance"))), GetMutableMemory(), nullptr);
+					StructTypePtr->SerializeItem(Record.EnterField(TEXT("StructInstance")), GetMutableMemory(), nullptr);
 				}
 			};
 
diff --git a/Engine/Plugins/FX/Niagara/Source/Niagara/Private/NiagaraSpriteRendererProperties.cpp b/Engine/Plugins/FX/Niagara/Source/Niagara/Private/NiagaraSpriteRendererProperties.cpp
--- a/Engine/Plugins/FX/Niagara/Source/Niagara/Private/NiagaraSpriteRendererProperties.cpp
+++ b/Engine/Plugins/FX/Niagara/Source/Niagara/Private/NiagaraSpriteRendererProperties.cpp
@@ -189,7 +189,7 @@ void UNiagaraSpriteRendererProperties::Serialize(FStructuredArchive::FRecord Rec
 	FArchive& UnderlyingArchive = Record.GetUnderlyingArchive();
 	if (UnderlyingArchive.IsCooking() || (FPlatformProperties::RequiresCookedData() && UnderlyingArchive.IsLoading()) || bIsCookedForEditor)
 	{
-		DerivedData.Serialize(Record.EnterField(SA_FIELD_NAME(TEXT("DerivedData"))));
+		DerivedData.Serialize(Record.EnterField(TEXT("DerivedData")));
 	}
 }
 
diff --git a/Engine/Source/Editor/UnrealEd/Private/Commandlets/TextAssetCommandlet.cpp b/Engine/Source/Editor/UnrealEd/Private/Commandlets/TextAssetCommandlet.cpp
--- a/Engine/Source/Editor/UnrealEd/Private/Commandlets/TextAssetCommandlet.cpp
+++ b/Engine/Source/Editor/UnrealEd/Private/Commandlets/TextAssetCommandlet.cpp
@@ -103,18 +103,18 @@ namespace StringConstants
 
 inline void WriteSimpleSchemaField(FStructuredArchiveRecord Record, const TCHAR* FieldName, FString& Type, FSimpleSchemaFieldPropertyGenerator PropertiesCallback = FSimpleSchemaFieldPropertyGenerator())
 {
-	FStructuredArchiveRecord FieldRecord = Record.EnterField(SA_FIELD_NAME(FieldName)).EnterRecord();
+	FStructuredArchiveRecord FieldRecord = Record.EnterField(FieldName).EnterRecord();
 	FieldRecord << SA_VALUE(TEXT("type"), Type);
 
 	if (PropertiesCallback)
 	{
-		FStructuredArchiveRecord PropertiesRecord = FieldRecord.EnterField(SA_FIELD_NAME(TEXT("properties"))).EnterRecord();
+		FStructuredArchiveRecord PropertiesRecord = FieldRecord.EnterField(TEXT("properties")).EnterRecord();
 		TArray<FString> Required;
 		PropertiesCallback(PropertiesRecord, Required);
 		if (Required.Num() > 0)
 		{
 			int32 NumRequired = Required.Num();
-			FStructuredArchiveArray RequiredArray = FieldRecord.EnterField(SA_FIELD_NAME(TEXT("required"))).EnterArray(NumRequired);
+			FStructuredArchiveArray RequiredArray = FieldRecord.EnterField(TEXT("required")).EnterArray(NumRequired);
 			for (FString& RequiredProperty : Required)
 			{
 				RequiredArray.EnterElement() << RequiredProperty;
@@ -140,7 +140,7 @@ void GeneratePropertySchema(FProperty* Property, FStructuredArchiveRecord Record
 		WriteSimpleSchemaField(Record, TEXT("__InnerStructName"), StringConstants::String);
 		WriteSimpleSchemaField(Record, TEXT("__Value"), StringConstants::Array);
 
-		FStructuredArchiveRecord ItemsRecord = Record.EnterRecord(SA_FIELD_NAME(TEXT("items")));
+		FStructuredArchiveRecord ItemsRecord = Record.EnterRecord(TEXT("items"));
 	}
 	else
 	{
@@ -239,7 +239,7 @@ void GenerateSchema()
 		{
 			if (NAME_SpecificClass == NAME_None || Class->GetFName() == NAME_SpecificClass)
 			{
-				FStructuredArchiveRecord ClassRecord = RootRecord.EnterRecord(SA_FIELD_NAME(*Class->GetFullName()));
+				FStructuredArchiveRecord ClassRecord = RootRecord.EnterRecord(*Class->GetFullName());
 				ClassRecord << SA_VALUE(*StringConstants::Type, StringConstants::Object);
 
 				WriteSimpleSchemaField(ClassRecord, TEXT("__Class"), StringConstants::Object);
diff --git a/Engine/Source/Runtime/Core/Private/Internationalization/InternationalizationMetadata.cpp b/Engine/Source/Runtime/Core/Private/Internationalization/InternationalizationMetadata.cpp
--- a/Engine/Source/Runtime/Core/Private/Internationalization/InternationalizationMetadata.cpp
+++ b/Engine/Source/Runtime/Core/Private/Internationalization/InternationalizationMetadata.cpp
@@ -296,7 +296,7 @@ namespace
 		Record << SA_VALUE(TEXT("Type"), MetaDataTypeAsInt);
 		MetaDataType = static_cast<ELocMetadataType>(MetaDataTypeAsInt);
 
-		FStructuredArchive::FSlot ValueSlot = Record.EnterField(SA_FIELD_NAME(TEXT("Value")));
+		FStructuredArchive::FSlot ValueSlot = Record.EnterField(TEXT("Value"));
 
 		switch (MetaDataType)
 		{
@@ -359,7 +359,7 @@ void operator<<(FStructuredArchive::FSlot Slot, FLocMetadataObject& Object)
 		Object.Values.Reserve(ValueCount);
 	}
 
-	FStructuredArchive::FStream Stream = Record.EnterStream(SA_FIELD_NAME(TEXT("Values")));
+	FStructuredArchive::FStream Stream = Record.EnterStream(TEXT("Values"));
 	TArray<FString> MapKeys;
 	Object.Values.GetKeys(MapKeys);
 	for (int32 i = 0; i < ValueCount; ++i)
@@ -375,7 +375,7 @@ void operator<<(FStructuredArchive::FSlot Slot, FLocMetadataObject& Object)
 
 		ValueRecord << SA_VALUE(TEXT("Key"), Key);
 
-		FStructuredArchive::FSlot ValueSlot = ValueRecord.EnterField(SA_FIELD_NAME(TEXT("Value")));
+		FStructuredArchive::FSlot ValueSlot = ValueRecord.EnterField(TEXT("Value"));
 		if (bIsLoading)
 		{
 			TSharedPtr<FLocMetadataValue> Value;
diff --git a/Engine/Source/Runtime/Core/Private/Internationalization/Text.cpp b/Engine/Source/Runtime/Core/Private/Internationalization/Text.cpp
--- a/Engine/Source/Runtime/Core/Private/Internationalization/Text.cpp
+++ b/Engine/Source/Runtime/Core/Private/Internationalization/Text.cpp
@@ -841,10 +841,10 @@ void FText::SerializeText(FStructuredArchive::FSlot Slot, FText& Value)
 		if (UnderlyingArchive.UEVer() >= VER_UE4_ADDED_NAMESPACE_AND_KEY_DATA_TO_FTEXT)
 		{
 			FTextKey Namespace;
-			Namespace.SerializeAsString(Record.EnterField(SA_FIELD_NAME(TEXT("Namespace"))));
+			Namespace.SerializeAsString(Record.EnterField(TEXT("Namespace")));
 
 			FTextKey Key;
-			Key.SerializeAsString(Record.EnterField(SA_FIELD_NAME(TEXT("Key"))));
+			Key.SerializeAsString(Record.EnterField(TEXT("Key")));
 
 			// Get the DisplayString using the namespace, key, and source string.
 			TextId = FTextId(Namespace, Key);
diff --git a/Engine/Source/Runtime/Core/Private/Internationalization/TextHistory.cpp b/Engine/Source/Runtime/Core/Private/Internationalization/TextHistory.cpp
--- a/Engine/Source/Runtime/Core/Private/Internationalization/TextHistory.cpp
+++ b/Engine/Source/Runtime/Core/Private/Internationalization/TextHistory.cpp
@@ -795,10 +795,10 @@ void FTextHistory_Base::Serialize(FStructuredArchive::FRecord Record)
 	if(BaseArchive.IsLoading())
 	{
 		FTextKey Namespace;
-		Namespace.SerializeAsString(Record.EnterField(SA_FIELD_NAME(TEXT("Namespace"))));
+		Namespace.SerializeAsString(Record.EnterField(TEXT("Namespace")));
 
 		FTextKey Key;
-		Key.SerializeAsString(Record.EnterField(SA_FIELD_NAME(TEXT("Key"))));
+		Key.SerializeAsString(Record.EnterField(TEXT("Key")));
 
 		Record << SA_VALUE(TEXT("SourceString"), SourceString);
 
@@ -885,10 +885,10 @@ void FTextHistory_Base::Serialize(FStructuredArchive::FRecord Record)
 		}
 
 		// Serialize the Namespace
-		Namespace.SerializeAsString(Record.EnterField(SA_FIELD_NAME(TEXT("Namespace"))));
+		Namespace.SerializeAsString(Record.EnterField(TEXT("Namespace")));
 
 		// Serialize the Key
-		Key.SerializeAsString(Record.EnterField(SA_FIELD_NAME(TEXT("Key"))));
+		Key.SerializeAsString(Record.EnterField(TEXT("Key")));
 
 		// Serialize the SourceString
 		Record << SA_VALUE(TEXT("SourceString"), SourceString);
@@ -1148,12 +1148,12 @@ void FTextHistory_NamedFormat::Serialize(FStructuredArchive::FRecord Record)
 	if (BaseArchive.IsSaving())
 	{
 		FText FormatText = SourceFmt.GetSourceText();
-		Record.EnterField(SA_FIELD_NAME(TEXT("FormatText"))) << FormatText;
+		Record.EnterField(TEXT("FormatText")) << FormatText;
 	}
 	else if (BaseArchive.IsLoading())
 	{
 		FText FormatText;
-		Record.EnterField(SA_FIELD_NAME(TEXT("FormatText"))) << FormatText;
+		Record.EnterField(TEXT("FormatText")) << FormatText;
 		SourceFmt = FTextFormat(FormatText);
 	}
 
@@ -1302,12 +1302,12 @@ void FTextHistory_OrderedFormat::Serialize(FStructuredArchive::FRecord Record)
 	if (BaseArchive.IsSaving())
 	{
 		FText FormatText = SourceFmt.GetSourceText();
-		Record.EnterField(SA_FIELD_NAME(TEXT("FormatText"))) << FormatText;
+		Record.EnterField(TEXT("FormatText")) << FormatText;
 	}
 	else if (BaseArchive.IsLoading())
 	{
 		FText FormatText;
-		Record.EnterField(SA_FIELD_NAME(TEXT("FormatText"))) << FormatText;
+		Record.EnterField(TEXT("FormatText")) << FormatText;
 		SourceFmt = FTextFormat(FormatText);
 	}
 
@@ -1455,12 +1455,12 @@ void FTextHistory_ArgumentDataFormat::Serialize(FStructuredArchive::FRecord Reco
 	if (BaseArchive.IsSaving())
 	{
 		FText FormatText = SourceFmt.GetSourceText();
-		Record.EnterField(SA_FIELD_NAME(TEXT("FormatText"))) << FormatText;
+		Record.EnterField(TEXT("FormatText")) << FormatText;
 	}
 	else if (BaseArchive.IsLoading())
 	{
 		FText FormatText;
-		Record.EnterField(SA_FIELD_NAME(TEXT("FormatText"))) << FormatText;
+		Record.EnterField(TEXT("FormatText")) << FormatText;
 		SourceFmt = FTextFormat(FormatText);
 	}
 
@@ -2422,7 +2422,7 @@ void FTextHistory_StringTableEntry::Serialize(FStructuredArchive::FRecord Record
 		FName TableId;
 		FTextKey Key;
 		Record << SA_VALUE(TEXT("TableId"), TableId);
-		Key.SerializeAsString(Record.EnterField(SA_FIELD_NAME(TEXT("Key"))));
+		Key.SerializeAsString(Record.EnterField(TEXT("Key")));
 
 		// String Table assets should already have been created via dependency loading when using the EDL (although they may not be fully loaded yet)
 		const bool bIsLoadingViaEDL = GEventDrivenLoaderEnabled && EVENT_DRIVEN_ASYNC_LOAD_ACTIVE_AT_RUNTIME && BaseArchive.GetLinker();
@@ -2440,7 +2440,7 @@ void FTextHistory_StringTableEntry::Serialize(FStructuredArchive::FRecord Record
 		}
 
 		Record << SA_VALUE(TEXT("TableId"), TableId);
-		Key.SerializeAsString(Record.EnterField(SA_FIELD_NAME(TEXT("Key"))));
+		Key.SerializeAsString(Record.EnterField(TEXT("Key")));
 	}
 
 	// Collect string table asset references
@@ -2599,7 +2599,7 @@ void FTextHistory_StringTableEntry::FStringTableReferenceData::CollectStringTabl
 		FScopeLock ScopeLock(&DataCS);
 
 		const FName OldTableId = TableId;
-		IStringTableEngineBridge::CollectStringTableAssetReferences(TableId, Record.EnterField(SA_FIELD_NAME(TEXT("AssetReferences"))));
+		IStringTableEngineBridge::CollectStringTableAssetReferences(TableId, Record.EnterField(TEXT("AssetReferences")));
 
 		if (TableId != OldTableId)
 		{
diff --git a/Engine/Source/Runtime/Core/Private/Serialization/ArchiveFromStructuredArchive.cpp b/Engine/Source/Runtime/Core/Private/Serialization/ArchiveFromStructuredArchive.cpp
--- a/Engine/Source/Runtime/Core/Private/Serialization/ArchiveFromStructuredArchive.cpp
+++ b/Engine/Source/Runtime/Core/Private/Serialization/ArchiveFromStructuredArchive.cpp
@@ -151,7 +151,7 @@ FArchive& FArchiveFromStructuredArchiveImpl::operator<<(class UObject*& Value)
 			}
 			else
 			{
-				FStructuredArchive::FStream Stream = Pimpl->Root->EnterStream(SA_FIELD_NAME(TEXT("Objects")));
+				FStructuredArchive::FStream Stream = Pimpl->Root->EnterStream(TEXT("Objects"));
 
 				// Skip earlier elements
 				FString Str;
@@ -245,10 +245,10 @@ bool FArchiveFromStructuredArchiveImpl::Finalize(FStructuredArchive::FRecord Rec
 
 	if (bShouldSerialize)
 	{
-		FStructuredArchive::FSlot DataSlot = Record.EnterField(SA_FIELD_NAME(TEXT("Data")));
+		FStructuredArchive::FSlot DataSlot = Record.EnterField(TEXT("Data"));
 		DataSlot.Serialize(Pimpl->Buffer);
 
-		TOptional<FStructuredArchive::FSlot> ObjectsSlot = Record.TryEnterField(SA_FIELD_NAME(TEXT("Objects")), Pimpl->Objects.Num() > 0);
+		TOptional<FStructuredArchive::FSlot> ObjectsSlot = Record.TryEnterField(TEXT("Objects"), Pimpl->Objects.Num() > 0);
 		if (ObjectsSlot.IsSet())
 		{
 			if (IsLoading())
@@ -273,7 +273,7 @@ bool FArchiveFromStructuredArchiveImpl::Finalize(FStructuredArchive::FRecord Rec
 			}
 		}
 
-		TOptional<FStructuredArchive::FSlot> NamesSlot = Record.TryEnterField(SA_FIELD_NAME(TEXT("Names")), Pimpl->Names.Num() > 0);
+		TOptional<FStructuredArchive::FSlot> NamesSlot = Record.TryEnterField(TEXT("Names"), Pimpl->Names.Num() > 0);
 		if (NamesSlot.IsSet())
 		{
 			NamesSlot.GetValue() << Pimpl->Names;
diff --git a/Engine/Source/Runtime/Core/Public/Internationalization/TextKey.h b/Engine/Source/Runtime/Core/Public/Internationalization/TextKey.h
--- a/Engine/Source/Runtime/Core/Public/Internationalization/TextKey.h
+++ b/Engine/Source/Runtime/Core/Public/Internationalization/TextKey.h
@@ -198,24 +198,24 @@ public:
 	void SerializeAsString(FStructuredArchiveSlot Slot)
 	{
 		FStructuredArchiveRecord Record = Slot.EnterRecord();
-		Namespace.SerializeAsString(Record.EnterField(SA_FIELD_NAME(TEXT("Namespace"))));
-		Key.SerializeAsString(Record.EnterField(SA_FIELD_NAME(TEXT("Key"))));
+		Namespace.SerializeAsString(Record.EnterField(TEXT("Namespace")));
+		Key.SerializeAsString(Record.EnterField(TEXT("Key")));
 	}
 
 	/** Serialize this text identity including its hash values (this method is sensitive to hashing algorithm changes, so only use it for generated files that can be rebuilt from another source) */
 	void SerializeWithHash(FStructuredArchiveSlot Slot)
 	{
 		FStructuredArchiveRecord Record = Slot.EnterRecord();
-		Namespace.SerializeWithHash(Record.EnterField(SA_FIELD_NAME(TEXT("Namespace"))));
-		Key.SerializeWithHash(Record.EnterField(SA_FIELD_NAME(TEXT("Key"))));
+		Namespace.SerializeWithHash(Record.EnterField(TEXT("Namespace")));
+		Key.SerializeWithHash(Record.EnterField(TEXT("Key")));
 	}
 
 	/** Serialize this text identity including its hash values, discarding the hash on load (to upgrade from an older hashing algorithm) */
 	void SerializeDiscardHash(FStructuredArchiveSlot Slot)
 	{
 		FStructuredArchiveRecord Record = Slot.EnterRecord();
-		Namespace.SerializeDiscardHash(Record.EnterField(SA_FIELD_NAME(TEXT("Namespace"))));
-		Key.SerializeDiscardHash(Record.EnterField(SA_FIELD_NAME(TEXT("Key"))));
+		Namespace.SerializeDiscardHash(Record.EnterField(TEXT("Namespace")));
+		Key.SerializeDiscardHash(Record.EnterField(TEXT("Key")));
 	}
 
 	/** Is this text identity empty? */
diff --git a/Engine/Source/Runtime/Core/Public/Serialization/StructuredArchiveNameHelpers.h b/Engine/Source/Runtime/Core/Public/Serialization/StructuredArchiveNameHelpers.h
--- a/Engine/Source/Runtime/Core/Public/Serialization/StructuredArchiveNameHelpers.h
+++ b/Engine/Source/Runtime/Core/Public/Serialization/StructuredArchiveNameHelpers.h
@@ -12,7 +12,7 @@ struct FArchiveFieldName
 	const TCHAR* Name;
 #endif
 
-	explicit FArchiveFieldName(const TCHAR* InName)
+	FArchiveFieldName(const TCHAR* InName)
 #if WITH_TEXT_ARCHIVE_SUPPORT
 		: Name(InName)
 #endif
@@ -73,7 +73,7 @@ namespace UE::StructuredArchive::Private
 	}
 }
 
-#define SA_FIELD_NAME(x) FArchiveFieldName(x)
+#define SA_FIELD_NAME(x) /*DEPRECATED_MACRO(5.1, "Field names no longer need to be wrapped in the SA_FIELD_NAME macro.")*/ FArchiveFieldName(x)
 
 /** Construct a TNamedValue given an ANSI string and value reference. */
 #define SA_VALUE(Name, Value) UE::StructuredArchive::Private::MakeNamedValue(FArchiveFieldName(Name), Value)
diff --git a/Engine/Source/Runtime/CoreUObject/Private/Serialization/ArchiveUObjectFromStructuredArchive.cpp b/Engine/Source/Runtime/CoreUObject/Private/Serialization/ArchiveUObjectFromStructuredArchive.cpp
--- a/Engine/Source/Runtime/CoreUObject/Private/Serialization/ArchiveUObjectFromStructuredArchive.cpp
+++ b/Engine/Source/Runtime/CoreUObject/Private/Serialization/ArchiveUObjectFromStructuredArchive.cpp
@@ -189,25 +189,25 @@ bool FArchiveUObjectFromStructuredArchiveImpl::Finalize(FStructuredArchive::FRec
 	bool bShouldSerialize = Super::Finalize(Record);
 	if (bShouldSerialize)
 	{
-		TOptional<FStructuredArchive::FSlot> LazyObjectPtrsSlot = Record.TryEnterField(SA_FIELD_NAME(TEXT("LazyObjectPtrs")), LazyObjectPtrs.Num() > 0);
+		TOptional<FStructuredArchive::FSlot> LazyObjectPtrsSlot = Record.TryEnterField(TEXT("LazyObjectPtrs"), LazyObjectPtrs.Num() > 0);
 		if (LazyObjectPtrsSlot.IsSet())
 		{
 			LazyObjectPtrsSlot.GetValue() << LazyObjectPtrs;
 		}
 
-		TOptional<FStructuredArchive::FSlot> SoftObjectPtrsSlot = Record.TryEnterField(SA_FIELD_NAME(TEXT("SoftObjectPtrs")), SoftObjectPtrs.Num() > 0);
+		TOptional<FStructuredArchive::FSlot> SoftObjectPtrsSlot = Record.TryEnterField(TEXT("SoftObjectPtrs"), SoftObjectPtrs.Num() > 0);
 		if (SoftObjectPtrsSlot.IsSet())
 		{
 			SoftObjectPtrsSlot.GetValue() << SoftObjectPtrs;
 		}
 
-		TOptional<FStructuredArchive::FSlot> SoftObjectPathsSlot = Record.TryEnterField(SA_FIELD_NAME(TEXT("SoftObjectPaths")), SoftObjectPaths.Num() > 0);
+		TOptional<FStructuredArchive::FSlot> SoftObjectPathsSlot = Record.TryEnterField(TEXT("SoftObjectPaths"), SoftObjectPaths.Num() > 0);
 		if (SoftObjectPathsSlot.IsSet())
 		{
 			SoftObjectPathsSlot.GetValue() << SoftObjectPaths;
 		}
 
-		TOptional<FStructuredArchive::FSlot> WeakObjectPtrsSlot = Record.TryEnterField(SA_FIELD_NAME(TEXT("WeakObjectPtrs")), WeakObjectPtrs.Num() > 0);
+		TOptional<FStructuredArchive::FSlot> WeakObjectPtrsSlot = Record.TryEnterField(TEXT("WeakObjectPtrs"), WeakObjectPtrs.Num() > 0);
 		if (WeakObjectPtrsSlot.IsSet())
 		{
 			WeakObjectPtrsSlot.GetValue() << WeakObjectPtrs;
diff --git a/Engine/Source/Runtime/CoreUObject/Private/Serialization/UnversionedPropertySerialization.cpp b/Engine/Source/Runtime/CoreUObject/Private/Serialization/UnversionedPropertySerialization.cpp
--- a/Engine/Source/Runtime/CoreUObject/Private/Serialization/UnversionedPropertySerialization.cpp
+++ b/Engine/Source/Runtime/CoreUObject/Private/Serialization/UnversionedPropertySerialization.cpp
@@ -846,7 +846,7 @@ void SerializeUnversionedProperties(const UStruct* Struct, FStructuredArchive::F
 		check(CanUseUnversionedPropertySerialization());
 
 		FUnversionedHeader Header;
-		Header.Load(StructRecord.EnterStream(SA_FIELD_NAME(TEXT("Header"))));
+		Header.Load(StructRecord.EnterStream(TEXT("Header")));
 
 		if (Header.HasValues())
 		{
@@ -856,7 +856,7 @@ void SerializeUnversionedProperties(const UStruct* Struct, FStructuredArchive::F
 			{
 				FDefaultStruct Defaults(DefaultsData, DefaultsStruct);
 
-				FStructuredArchive::FStream ValueStream = StructRecord.EnterStream(SA_FIELD_NAME(TEXT("Values")));
+				FStructuredArchive::FStream ValueStream = StructRecord.EnterStream(TEXT("Values"));
 				for (FUnversionedHeader::FIterator It(Header, Schema); It; It.Next())
 				{
 					if (It.IsNonZero())
@@ -910,10 +910,10 @@ void SerializeUnversionedProperties(const UStruct* Struct, FStructuredArchive::F
 		Header.Finalize();
 
 		// Save header and non-zero values 
-		Header.Save(StructRecord.EnterStream(SA_FIELD_NAME(TEXT("Header"))));
+		Header.Save(StructRecord.EnterStream(TEXT("Header")));
 		if (Header.HasNonZeroValues())
 		{
-			FStructuredArchive::FStream ValueStream = StructRecord.EnterStream(SA_FIELD_NAME(TEXT("Values")));
+			FStructuredArchive::FStream ValueStream = StructRecord.EnterStream(TEXT("Values"));
 			for (FUnversionedHeader::FIterator It(Header, Schema); It; It.Next())
 			{
 				if (It.IsNonZero())
diff --git a/Engine/Source/Runtime/CoreUObject/Private/UObject/Class.cpp b/Engine/Source/Runtime/CoreUObject/Private/UObject/Class.cpp
--- a/Engine/Source/Runtime/CoreUObject/Private/UObject/Class.cpp
+++ b/Engine/Source/Runtime/CoreUObject/Private/UObject/Class.cpp
@@ -1129,7 +1129,7 @@ void UStruct::LoadTaggedPropertiesFromText(FStructuredArchive::FSlot Slot, uint8
 		FName PropertyName = *PropertyNameString;
 
 		// If this property has a guid attached then we need to resolve it to the right name before we start loading
-		TOptional<FStructuredArchiveSlot> PropertyGuidSlot = PropertySlot.TryEnterAttribute(SA_FIELD_NAME(TEXT("PropertyGuid")), false);
+		TOptional<FStructuredArchiveSlot> PropertyGuidSlot = PropertySlot.TryEnterAttribute(TEXT("PropertyGuid"), false);
 		if (PropertyGuidSlot.IsSet())
 		{
 			FGuid PropertyGuid;
@@ -1445,7 +1445,7 @@ void UStruct::SerializeVersionedTaggedProperties(FStructuredArchive::FSlot Slot,
 					}
 					else
 					{
-						FStructuredArchive::FSlot ValueSlot = PropertyRecord.EnterField(SA_FIELD_NAME(TEXT("Value")));
+						FStructuredArchive::FSlot ValueSlot = PropertyRecord.EnterField(TEXT("Value"));
 
 						switch (Property->ConvertFromType(Tag, ValueSlot, Data, DefaultsStruct))
 						{
@@ -1528,7 +1528,7 @@ void UStruct::SerializeVersionedTaggedProperties(FStructuredArchive::FSlot Slot,
 				if (((LoopMax - 1) > LoopMin) && UnderlyingArchive.IsTextFormat())
 				{
 					int32 NumItems = LoopMax - LoopMin;
-					StaticArrayContainer.Emplace(PropertiesRecord.EnterArray(SA_FIELD_NAME((*Property->GetName())), NumItems));
+					StaticArrayContainer.Emplace(PropertiesRecord.EnterArray(*Property->GetName(), NumItems));
 				}
 
 				for (int32 Idx = LoopMin; Idx < LoopMax; Idx++)
@@ -1558,7 +1558,7 @@ void UStruct::SerializeVersionedTaggedProperties(FStructuredArchive::FSlot Slot,
 
 						TStringBuilder<256> TagName;
 						Tag.Name.ToString(TagName);
-						FStructuredArchive::FSlot PropertySlot = StaticArrayContainer.IsSet() ? StaticArrayContainer->EnterElement() : PropertiesRecord.EnterField(SA_FIELD_NAME(TagName.ToString()));
+						FStructuredArchive::FSlot PropertySlot = StaticArrayContainer.IsSet() ? StaticArrayContainer->EnterElement() : PropertiesRecord.EnterField(TagName.ToString());
 
 						PropertySlot << Tag;
 
diff --git a/Engine/Source/Runtime/CoreUObject/Private/UObject/LazyObjectPtr.cpp b/Engine/Source/Runtime/CoreUObject/Private/UObject/LazyObjectPtr.cpp
--- a/Engine/Source/Runtime/CoreUObject/Private/UObject/LazyObjectPtr.cpp
+++ b/Engine/Source/Runtime/CoreUObject/Private/UObject/LazyObjectPtr.cpp
@@ -106,7 +106,7 @@ void FLazyObjectPtr::PossiblySerializeObjectGuid(UObject *Object, FStructuredArc
 	if (UnderlyingArchive.IsSaving() || UnderlyingArchive.IsCountingMemory())
 	{
 		FUniqueObjectGuid Guid = GuidAnnotation.GetAnnotation(Object);
-		TOptional<FStructuredArchiveSlot> GuidSlot = Record.TryEnterField(SA_FIELD_NAME(TEXT("Guid")), Guid.IsValid());
+		TOptional<FStructuredArchiveSlot> GuidSlot = Record.TryEnterField(TEXT("Guid"), Guid.IsValid());
 		if (GuidSlot.IsSet())
 		{
 			if (UnderlyingArchive.GetPortFlags() & PPF_DuplicateForPIE)
@@ -119,7 +119,7 @@ void FLazyObjectPtr::PossiblySerializeObjectGuid(UObject *Object, FStructuredArc
 	}
 	else if (UnderlyingArchive.IsLoading())
 	{
-		TOptional<FStructuredArchiveSlot> GuidSlot = Record.TryEnterField(SA_FIELD_NAME(TEXT("Guid")), false);
+		TOptional<FStructuredArchiveSlot> GuidSlot = Record.TryEnterField(TEXT("Guid"), false);
 		if (GuidSlot.IsSet())
 		{
 			FUniqueObjectGuid Guid;
diff --git a/Engine/Source/Runtime/CoreUObject/Private/UObject/LinkerLoad.cpp b/Engine/Source/Runtime/CoreUObject/Private/UObject/LinkerLoad.cpp
--- a/Engine/Source/Runtime/CoreUObject/Private/UObject/LinkerLoad.cpp
+++ b/Engine/Source/Runtime/CoreUObject/Private/UObject/LinkerLoad.cpp
@@ -1638,7 +1638,7 @@ FLinkerLoad::ELinkerStatus FLinkerLoad::SerializeGatherableTextDataMap(bool bFor
 		Seek( Summary.GatherableTextDataOffset );
 	}
 
-	FStructuredArchive::FStream Stream = StructuredArchiveRootRecord->EnterStream(SA_FIELD_NAME(TEXT("GatherableTextData")));
+	FStructuredArchive::FStream Stream = StructuredArchiveRootRecord->EnterStream(TEXT("GatherableTextData"));
 	while (GatherableTextDataMapIndex < Summary.GatherableTextDataCount && !IsTimeLimitExceeded(TEXT("serializing gatherable text data map"), 100))
 	{
 		FGatherableTextData* GatherableTextData = new(GatherableTextDataMap)FGatherableTextData;
@@ -1664,7 +1664,7 @@ FLinkerLoad::ELinkerStatus FLinkerLoad::SerializeImportMap()
 		Seek( Summary.ImportOffset );
 	}
 
-	FStructuredArchive::FStream Stream = StructuredArchiveRootRecord->EnterStream(SA_FIELD_NAME(TEXT("ImportTable")));
+	FStructuredArchive::FStream Stream = StructuredArchiveRootRecord->EnterStream(TEXT("ImportTable"));
 
 	while( ImportMapIndex < Summary.ImportCount && !IsTimeLimitExceeded(TEXT("serializing import map"),100) )
 	{
@@ -1939,7 +1939,7 @@ FLinkerLoad::ELinkerStatus FLinkerLoad::SerializeExportMap()
 		Seek( Summary.ExportOffset );
 	}
 
-	FStructuredArchive::FStream Stream = StructuredArchiveRootRecord->EnterStream(SA_FIELD_NAME(TEXT("ExportTable")));
+	FStructuredArchive::FStream Stream = StructuredArchiveRootRecord->EnterStream(TEXT("ExportTable"));
 
 	while( ExportMapIndex < Summary.ExportCount && !IsTimeLimitExceeded(TEXT("serializing export map"),100) )
 	{
@@ -2076,7 +2076,7 @@ FLinkerLoad::ELinkerStatus FLinkerLoad::ConstructExportsReaders()
 	if (!bHasConstructedExportsReaders && IsTextFormat())
 	{
 		int32 NumExports = 0;
-		FStructuredArchiveMap PackageExports = StructuredArchiveRootRecord->EnterMap(SA_FIELD_NAME(TEXT("Exports")), NumExports);
+		FStructuredArchiveMap PackageExports = StructuredArchiveRootRecord->EnterMap(TEXT("Exports"), NumExports);
 
 		ExportReaders.AddDefaulted(NumExports);
 		for (int32 ExportIndex = 0; ExportIndex < NumExports; ++ExportIndex)
@@ -2129,7 +2129,7 @@ FLinkerLoad::ELinkerStatus FLinkerLoad::SerializeDependsMap()
 		DependsMap.AddZeroed(Summary.ExportCount);
 	}
 
-	FStructuredArchive::FStream Stream = StructuredArchiveRootRecord->EnterStream(SA_FIELD_NAME(TEXT("DependsMap")));
+	FStructuredArchive::FStream Stream = StructuredArchiveRootRecord->EnterStream(TEXT("DependsMap"));
 	while (DependsMapIndex < Summary.ExportCount && !IsTimeLimitExceeded(TEXT("serializing depends map"), 100))
 	{
 		TArray<FPackageIndex>& Depends = DependsMap[DependsMapIndex];
@@ -2168,7 +2168,7 @@ FLinkerLoad::ELinkerStatus FLinkerLoad::SerializePreloadDependencies()
 		)
 	{
 		//@todoio check endiness and fastpath this as a single serialize
-		FStructuredArchive::FStream Stream = StructuredArchiveRootRecord->EnterStream(SA_FIELD_NAME(TEXT("PreloadDependencies")));
+		FStructuredArchive::FStream Stream = StructuredArchiveRootRecord->EnterStream(TEXT("PreloadDependencies"));
 		for (int32 Index = 0; Index < Summary.PreloadDependencyCount; Index++)
 		{
 			FPackageIndex Idx;
@@ -2205,7 +2205,7 @@ FLinkerLoad::ELinkerStatus FLinkerLoad::SerializeThumbnails( bool bForceEnableIn
 
 	if (IsTextFormat())
 	{
-		ThumbnailsSlot = StructuredArchiveRootRecord->TryEnterField(SA_FIELD_NAME(TEXT("Thumbnails")), false);
+		ThumbnailsSlot = StructuredArchiveRootRecord->TryEnterField(TEXT("Thumbnails"), false);
 		if (!ThumbnailsSlot.IsSet())
 		{
 			return LINKER_Loaded;
@@ -2215,7 +2215,7 @@ FLinkerLoad::ELinkerStatus FLinkerLoad::SerializeThumbnails( bool bForceEnableIn
 	{
 		if(Summary.ThumbnailTableOffset > 0)
 		{
-			ThumbnailsSlot = StructuredArchiveRootRecord->EnterField(SA_FIELD_NAME(TEXT("Thumbnails")));
+			ThumbnailsSlot = StructuredArchiveRootRecord->EnterField(TEXT("Thumbnails"));
 		}
 	}
 
@@ -2226,13 +2226,13 @@ FLinkerLoad::ELinkerStatus FLinkerLoad::SerializeThumbnails( bool bForceEnableIn
 
 		if (IsTextFormat())
 		{
-			IndexSlot = Record.TryEnterField(SA_FIELD_NAME(TEXT("Index")), false);
+			IndexSlot = Record.TryEnterField(TEXT("Index"), false);
 		}
 		else
 		{
 			// Seek to the thumbnail table of contents
 			Seek(Summary.ThumbnailTableOffset);
-			IndexSlot.Emplace(Record.EnterField(SA_FIELD_NAME(TEXT("Index"))));
+			IndexSlot.Emplace(Record.EnterField(TEXT("Index")));
 		}
 
 		if (IndexSlot.IsSet())
@@ -2284,7 +2284,7 @@ FLinkerLoad::ELinkerStatus FLinkerLoad::SerializeThumbnails( bool bForceEnableIn
 			}
 
 
-			FStructuredArchive::FStream DataStream = Record.EnterStream(SA_FIELD_NAME(TEXT("Thumbnails")));
+			FStructuredArchive::FStream DataStream = Record.EnterStream(TEXT("Thumbnails"));
 
 			// Now go and load and cache all of the thumbnails
 			for (int32 CurObjectIndex = 0; CurObjectIndex < ThumbnailInfoArray.Num(); ++CurObjectIndex)
diff --git a/Engine/Source/Runtime/CoreUObject/Private/UObject/Obj.cpp b/Engine/Source/Runtime/CoreUObject/Private/UObject/Obj.cpp
--- a/Engine/Source/Runtime/CoreUObject/Private/UObject/Obj.cpp
+++ b/Engine/Source/Runtime/CoreUObject/Private/UObject/Obj.cpp
@@ -1431,7 +1431,7 @@ void UObject::Serialize(FStructuredArchive::FRecord Record)
 		// Handle derived UClass objects (exact UClass objects are native only and shouldn't be touched)
 		if (ObjClass != UClass::StaticClass())
 		{
-			SerializeScriptProperties(Record.EnterField(SA_FIELD_NAME(TEXT("Properties"))));
+			SerializeScriptProperties(Record.EnterField(TEXT("Properties")));
 		}
 
 		// Keep track of pending kill
@@ -1477,7 +1477,7 @@ void UObject::Serialize(FStructuredArchive::FRecord Record)
 
 			if (SerializedSparseClassDataStruct)
 			{
-				ObjClass->SerializeSparseClassData(Record.EnterField(SA_FIELD_NAME(TEXT("SparseClassData"))));
+				ObjClass->SerializeSparseClassData(Record.EnterField(TEXT("SparseClassData")));
 			}
 		}
 
diff --git a/Engine/Source/Runtime/CoreUObject/Private/UObject/PackageFileSummary.cpp b/Engine/Source/Runtime/CoreUObject/Private/UObject/PackageFileSummary.cpp
--- a/Engine/Source/Runtime/CoreUObject/Private/UObject/PackageFileSummary.cpp
+++ b/Engine/Source/Runtime/CoreUObject/Private/UObject/PackageFileSummary.cpp
@@ -134,7 +134,7 @@ void operator<<(FStructuredArchive::FSlot Slot, FPackageFileSummary& Sum)
 
 				if (LegacyFileVersion <= -2)
 				{
-					Sum.CustomVersionContainer.Serialize(Record.EnterField(SA_FIELD_NAME(TEXT("CustomVersions"))), GetCustomVersionFormatForArchive(LegacyFileVersion));
+					Sum.CustomVersionContainer.Serialize(Record.EnterField(TEXT("CustomVersions")), GetCustomVersionFormatForArchive(LegacyFileVersion));
 				}
 
 				if (!Sum.FileVersionUE.FileVersionUE4 && !Sum.FileVersionUE.FileVersionUE5 && !Sum.FileVersionLicenseeUE)
@@ -176,7 +176,7 @@ void operator<<(FStructuredArchive::FSlot Slot, FPackageFileSummary& Sum)
 				Record << SA_VALUE(TEXT("FileVersionLicenseeUE4"), Zero); // VersionLicenseeUE4
 
 				FCustomVersionContainer NoCustomVersions;
-				NoCustomVersions.Serialize(Record.EnterField(SA_FIELD_NAME(TEXT("CustomVersions"))));
+				NoCustomVersions.Serialize(Record.EnterField(TEXT("CustomVersions")));
 			}
 			else
 			{
@@ -188,7 +188,7 @@ void operator<<(FStructuredArchive::FSlot Slot, FPackageFileSummary& Sum)
 				Record << SA_VALUE(TEXT("FileVersionLicenseeUE4"), Sum.FileVersionLicenseeUE);
 
 				// Serialise custom version map.
-				Sum.CustomVersionContainer.Serialize(Record.EnterField(SA_FIELD_NAME(TEXT("CustomVersions"))));
+				Sum.CustomVersionContainer.Serialize(Record.EnterField(TEXT("CustomVersions")));
 			}
 		}
 		Record << SA_VALUE(TEXT("TotalHeaderSize"), Sum.TotalHeaderSize);
@@ -280,7 +280,7 @@ void operator<<(FStructuredArchive::FSlot Slot, FPackageFileSummary& Sum)
 			Sum.Generations.AddZeroed(GenerationCount);
 		}
 
-		FStructuredArchive::FStream GenerationsStream = Record.EnterStream(SA_FIELD_NAME(TEXT("Generations")));
+		FStructuredArchive::FStream GenerationsStream = Record.EnterStream(TEXT("Generations"));
 		for (int32 i = 0; i<GenerationCount; i++)
 		{
 			Sum.Generations[i].Serialize(GenerationsStream.EnterElement(), Sum);
diff --git a/Engine/Source/Runtime/CoreUObject/Private/UObject/PropertyMap.cpp b/Engine/Source/Runtime/CoreUObject/Private/UObject/PropertyMap.cpp
--- a/Engine/Source/Runtime/CoreUObject/Private/UObject/PropertyMap.cpp
+++ b/Engine/Source/Runtime/CoreUObject/Private/UObject/PropertyMap.cpp
@@ -343,7 +343,7 @@ void FMapProperty::SerializeItem(FStructuredArchive::FSlot Slot, void* Value, co
 
 		// Delete any explicitly-removed elements
 		int32 NumKeysToRemove = 0;
-		FStructuredArchive::FArray KeysToRemoveArray = Record.EnterArray(SA_FIELD_NAME(TEXT("KeysToRemove")), NumKeysToRemove);
+		FStructuredArchive::FArray KeysToRemoveArray = Record.EnterArray(TEXT("KeysToRemove"), NumKeysToRemove);
 
 		if (!Defaults || MapHelper.Num() == 0) // Faster loading path when loading into an empty map
 		{
@@ -364,7 +364,7 @@ void FMapProperty::SerializeItem(FStructuredArchive::FSlot Slot, void* Value, co
 			}
 
 			int32 NumEntries = 0;
-			FStructuredArchive::FArray EntriesArray = Record.EnterArray(SA_FIELD_NAME(TEXT("Entries")), NumEntries);
+			FStructuredArchive::FArray EntriesArray = Record.EnterArray(TEXT("Entries"), NumEntries);
 
 			// Empty and reserve then deserialize pairs directly into map memory
 			MapHelper.EmptyValues(NumEntries);
@@ -375,11 +375,11 @@ void FMapProperty::SerializeItem(FStructuredArchive::FSlot Slot, void* Value, co
 
 				{
 					FSerializedPropertyScope SerializedProperty(UnderlyingArchive, KeyProp, this);
-					KeyProp->SerializeItem(EntryRecord.EnterField(SA_FIELD_NAME(TEXT("Key"))), MapHelper.GetKeyPtr(Index));
+					KeyProp->SerializeItem(EntryRecord.EnterField(TEXT("Key")), MapHelper.GetKeyPtr(Index));
 				}
 				{
 					FSerializedPropertyScope SerializedProperty(UnderlyingArchive, ValueProp, this);
-					ValueProp->SerializeItem(EntryRecord.EnterField(SA_FIELD_NAME(TEXT("Value"))), MapHelper.GetValuePtr(Index));
+					ValueProp->SerializeItem(EntryRecord.EnterField(TEXT("Value")), MapHelper.GetValuePtr(Index));
 				}
 			}
 
@@ -417,7 +417,7 @@ void FMapProperty::SerializeItem(FStructuredArchive::FSlot Slot, void* Value, co
 			}
 
 			int32 NumEntries = 0;
-			FStructuredArchive::FArray EntriesArray = Record.EnterArray(SA_FIELD_NAME(TEXT("Entries")), NumEntries);
+			FStructuredArchive::FArray EntriesArray = Record.EnterArray(TEXT("Entries"), NumEntries);
 
 			// Allocate temporary key space if we haven't allocated it already above
 			if (NumEntries != 0 && !TempKeyValueStorage)
@@ -434,7 +434,7 @@ void FMapProperty::SerializeItem(FStructuredArchive::FSlot Slot, void* Value, co
 				// Read key into temporary storage
 				{
 					FSerializedPropertyScope SerializedProperty(UnderlyingArchive, KeyProp, this);
-					KeyProp->SerializeItem(EntryRecord.EnterField(SA_FIELD_NAME(TEXT("Key"))), TempKeyValueStorage);
+					KeyProp->SerializeItem(EntryRecord.EnterField(TEXT("Key")), TempKeyValueStorage);
 				}
 
 				void* ValuePtr = MapHelper.FindOrAdd(TempKeyValueStorage);
@@ -442,7 +442,7 @@ void FMapProperty::SerializeItem(FStructuredArchive::FSlot Slot, void* Value, co
 				// Deserialize value into hash map-owned memory
 				{
 					FSerializedPropertyScope SerializedProperty(UnderlyingArchive, ValueProp, this);
-					ValueProp->SerializeItem(EntryRecord.EnterField(SA_FIELD_NAME(TEXT("Value"))), ValuePtr);
+					ValueProp->SerializeItem(EntryRecord.EnterField(TEXT("Value")), ValuePtr);
 				}
 			}
 		}
@@ -475,7 +475,7 @@ void FMapProperty::SerializeItem(FStructuredArchive::FSlot Slot, void* Value, co
 
 		// Write out the missing keys
 		int32 MissingKeysNum = Indices.Num();
-		FStructuredArchive::FArray KeysToRemoveArray = Record.EnterArray(SA_FIELD_NAME(TEXT("KeysToRemove")), MissingKeysNum);
+		FStructuredArchive::FArray KeysToRemoveArray = Record.EnterArray(TEXT("KeysToRemove"), MissingKeysNum);
 		{
 			FSerializedPropertyScope SerializedProperty(UnderlyingArchive, KeyProp, this);
 			for (int32 Index : Indices)
@@ -506,7 +506,7 @@ void FMapProperty::SerializeItem(FStructuredArchive::FSlot Slot, void* Value, co
 
 			// Write out differences from defaults
 			int32 Num = Indices.Num();
-			FStructuredArchive::FArray EntriesArray = Record.EnterArray(SA_FIELD_NAME(TEXT("Entries")), Num);
+			FStructuredArchive::FArray EntriesArray = Record.EnterArray(TEXT("Entries"), Num);
 			for (int32 Index : Indices)
 			{
 				uint8* ValuePairPtr = MapHelper.GetPairPtrWithoutCheck(Index);
@@ -514,18 +514,18 @@ void FMapProperty::SerializeItem(FStructuredArchive::FSlot Slot, void* Value, co
 
 				{
 					FSerializedPropertyScope SerializedProperty(UnderlyingArchive, KeyProp, this);
-					KeyProp->SerializeItem(EntryRecord.EnterField(SA_FIELD_NAME(TEXT("Key"))), ValuePairPtr);
+					KeyProp->SerializeItem(EntryRecord.EnterField(TEXT("Key")), ValuePairPtr);
 				}
 				{
 					FSerializedPropertyScope SerializedProperty(UnderlyingArchive, ValueProp, this);
-					ValueProp->SerializeItem(EntryRecord.EnterField(SA_FIELD_NAME(TEXT("Value"))), ValuePairPtr + MapLayout.ValueOffset);
+					ValueProp->SerializeItem(EntryRecord.EnterField(TEXT("Value")), ValuePairPtr + MapLayout.ValueOffset);
 				}
 			}
 		}
 		else
 		{
 			int32 Num = MapHelper.Num();
-			FStructuredArchive::FArray EntriesArray = Record.EnterArray(SA_FIELD_NAME(TEXT("Entries")), Num);
+			FStructuredArchive::FArray EntriesArray = Record.EnterArray(TEXT("Entries"), Num);
 
 			for (int32 Index = 0; Num; ++Index)
 			{
@@ -537,11 +537,11 @@ void FMapProperty::SerializeItem(FStructuredArchive::FSlot Slot, void* Value, co
 
 					{
 						FSerializedPropertyScope SerializedProperty(UnderlyingArchive, KeyProp, this);
-						KeyProp->SerializeItem(EntryRecord.EnterField(SA_FIELD_NAME(TEXT("Key"))), ValuePairPtr);
+						KeyProp->SerializeItem(EntryRecord.EnterField(TEXT("Key")), ValuePairPtr);
 					}
 					{
 						FSerializedPropertyScope SerializedProperty(UnderlyingArchive, ValueProp, this);
-						ValueProp->SerializeItem(EntryRecord.EnterField(SA_FIELD_NAME(TEXT("Value"))), ValuePairPtr + MapLayout.ValueOffset);
+						ValueProp->SerializeItem(EntryRecord.EnterField(TEXT("Value")), ValuePairPtr + MapLayout.ValueOffset);
 					}
 
 					--Num;
@@ -1124,7 +1124,7 @@ EConvertFromTypeResult FMapProperty::ConvertFromType(const FPropertyTag& Tag, FS
 			// instance that was being written. Presumably we were constructed from our defaults and must now remove 
 			// any of the elements that were not present when we saved this Map:
 			int32 NumKeysToRemove = 0;
-			FStructuredArchive::FArray KeysToRemoveArray = ValueRecord.EnterArray(SA_FIELD_NAME(TEXT("KeysToRemove")), NumKeysToRemove);
+			FStructuredArchive::FArray KeysToRemoveArray = ValueRecord.EnterArray(TEXT("KeysToRemove"), NumKeysToRemove);
 
 			if( NumKeysToRemove != 0 )
 			{
@@ -1158,7 +1158,7 @@ EConvertFromTypeResult FMapProperty::ConvertFromType(const FPropertyTag& Tag, FS
 			}
 
 			int32 NumEntries = 0;
-			FStructuredArchive::FArray EntriesArray = ValueRecord.EnterArray(SA_FIELD_NAME(TEXT("Entries")), NumEntries);
+			FStructuredArchive::FArray EntriesArray = ValueRecord.EnterArray(TEXT("Entries"), NumEntries);
 
 			if( bConversionSucceeded )
 			{
@@ -1172,7 +1172,7 @@ EConvertFromTypeResult FMapProperty::ConvertFromType(const FPropertyTag& Tag, FS
 
 					FStructuredArchive::FRecord FirstPropertyRecord = EntriesArray.EnterElement().EnterRecord();
 
-					if( SerializeOrConvert( KeyProp, KeyPropertyTag, FirstPropertyRecord.EnterField(SA_FIELD_NAME(TEXT("Key"))), TempKeyValueStorage, DefaultsStruct ) )
+					if( SerializeOrConvert( KeyProp, KeyPropertyTag, FirstPropertyRecord.EnterField(TEXT("Key")), TempKeyValueStorage, DefaultsStruct ) )
 					{
 						// Add a new default value if the key doesn't currently exist in the map
 						bool bKeyAlreadyPresent = true;
@@ -1188,14 +1188,14 @@ EConvertFromTypeResult FMapProperty::ConvertFromType(const FPropertyTag& Tag, FS
 						KeyProp->CopyCompleteValue_InContainer(NextPairPtr, TempKeyValueStorage);
 
 						// Deserialize value
-						if( SerializeOrConvert( ValueProp, ValuePropertyTag, FirstPropertyRecord.EnterField(SA_FIELD_NAME(TEXT("Value"))), NextPairPtr, DefaultsStruct ) )
+						if( SerializeOrConvert( ValueProp, ValuePropertyTag, FirstPropertyRecord.EnterField(TEXT("Value")), NextPairPtr, DefaultsStruct ) )
 						{
 							// first entry went fine, convert the rest:
 							for(int32 I = 1; I < NumEntries; ++I)
 							{
 								FStructuredArchive::FRecord PropertyRecord = EntriesArray.EnterElement().EnterRecord();
 
-								verify( SerializeOrConvert( KeyProp, KeyPropertyTag, PropertyRecord.EnterField(SA_FIELD_NAME(TEXT("Key"))), TempKeyValueStorage, DefaultsStruct ) );
+								verify( SerializeOrConvert( KeyProp, KeyPropertyTag, PropertyRecord.EnterField(TEXT("Key")), TempKeyValueStorage, DefaultsStruct ) );
 								NextPairIndex = MapHelper.FindMapIndexWithKey(TempKeyValueStorage);
 								if (NextPairIndex == INDEX_NONE)
 								{
@@ -1205,7 +1205,7 @@ EConvertFromTypeResult FMapProperty::ConvertFromType(const FPropertyTag& Tag, FS
 								NextPairPtr = MapHelper.GetPairPtrWithoutCheck(NextPairIndex);
 								// This copy is unnecessary when the key was already in the map:
 								KeyProp->CopyCompleteValue_InContainer(NextPairPtr, TempKeyValueStorage);
-								verify( SerializeOrConvert( ValueProp, ValuePropertyTag, PropertyRecord.EnterField(SA_FIELD_NAME(TEXT("Value"))), NextPairPtr, DefaultsStruct ) );
+								verify( SerializeOrConvert( ValueProp, ValuePropertyTag, PropertyRecord.EnterField(TEXT("Value")), NextPairPtr, DefaultsStruct ) );
 							}
 						}
 						else
diff --git a/Engine/Source/Runtime/CoreUObject/Private/UObject/PropertySet.cpp b/Engine/Source/Runtime/CoreUObject/Private/UObject/PropertySet.cpp
--- a/Engine/Source/Runtime/CoreUObject/Private/UObject/PropertySet.cpp
+++ b/Engine/Source/Runtime/CoreUObject/Private/UObject/PropertySet.cpp
@@ -297,7 +297,7 @@ void FSetProperty::SerializeItem(FStructuredArchive::FSlot Slot, void* Value, co
 
 		// Delete any explicitly-removed elements
 		int32 NumElementsToRemove = 0;
-		FStructuredArchive::FArray ElementsToRemoveArray = Record.EnterArray(SA_FIELD_NAME(TEXT("ElementsToRemove")), NumElementsToRemove);
+		FStructuredArchive::FArray ElementsToRemoveArray = Record.EnterArray(TEXT("ElementsToRemove"), NumElementsToRemove);
 
 		if (!Defaults || SetHelper.Num() == 0) // Faster loading path when loading elements into an empty set
 		{
@@ -318,7 +318,7 @@ void FSetProperty::SerializeItem(FStructuredArchive::FSlot Slot, void* Value, co
 			}
 
 			int32 Num = 0;
-			FStructuredArchive::FArray ElementsArray = Record.EnterArray(SA_FIELD_NAME(TEXT("Elements")), Num);
+			FStructuredArchive::FArray ElementsArray = Record.EnterArray(TEXT("Elements"), Num);
 
 			// Empty and reserve then deserialize elements directly into set memory
 			SetHelper.EmptyElements(Num);
@@ -361,7 +361,7 @@ void FSetProperty::SerializeItem(FStructuredArchive::FSlot Slot, void* Value, co
 			}
 
 			int32 Num = 0;
-			FStructuredArchive::FArray ElementsArray = Record.EnterArray(SA_FIELD_NAME(TEXT("Elements")), Num);
+			FStructuredArchive::FArray ElementsArray = Record.EnterArray(TEXT("Elements"), Num);
 
 			// Allocate temporary key space if we haven't allocated it already above
 			if (Num != 0 && !TempElementStorage)
@@ -419,7 +419,7 @@ void FSetProperty::SerializeItem(FStructuredArchive::FSlot Slot, void* Value, co
 
 		// Write out the removed elements
 		int32 RemovedElementsNum = Indices.Num();
-		FStructuredArchive::FArray RemovedElementsArray = Record.EnterArray(SA_FIELD_NAME(TEXT("ElementsToRemove")), RemovedElementsNum);
+		FStructuredArchive::FArray RemovedElementsArray = Record.EnterArray(TEXT("ElementsToRemove"), RemovedElementsNum);
 		
 		{
 			FSerializedPropertyScope SerializedProperty(UnderlyingArchive, ElementProp, this);
@@ -451,7 +451,7 @@ void FSetProperty::SerializeItem(FStructuredArchive::FSlot Slot, void* Value, co
 
 			// Write out differences from defaults
 			int32 Num = Indices.Num();
-			FStructuredArchive::FArray ElementsArray = Record.EnterArray(SA_FIELD_NAME(TEXT("Elements")), Num);
+			FStructuredArchive::FArray ElementsArray = Record.EnterArray(TEXT("Elements"), Num);
 
 			FSerializedPropertyScope SerializedProperty(UnderlyingArchive, ElementProp, this);
 			for (int32 Index : Indices)
@@ -464,7 +464,7 @@ void FSetProperty::SerializeItem(FStructuredArchive::FSlot Slot, void* Value, co
 		else
 		{
 			int32 Num = SetHelper.Num();
-			FStructuredArchive::FArray ElementsArray = Record.EnterArray(SA_FIELD_NAME(TEXT("Elements")), Num);
+			FStructuredArchive::FArray ElementsArray = Record.EnterArray(TEXT("Elements"), Num);
 
 			FSerializedPropertyScope SerializedProperty(UnderlyingArchive, ElementProp, this);
 			for (int32 Index = 0; Num; ++Index)
@@ -947,7 +947,7 @@ EConvertFromTypeResult FSetProperty::ConvertFromType(const FPropertyTag& Tag, FS
 			// instance that was being written. Presumably we were constructed from our defaults and must now remove 
 			// any of the elements that were not present when we saved this Set:
 			int32 NumElementsToRemove = 0;
-			FStructuredArchive::FArray ElementsToRemoveArray = ValueRecord.EnterArray(SA_FIELD_NAME(TEXT("ElementsToRemove")), NumElementsToRemove);
+			FStructuredArchive::FArray ElementsToRemoveArray = ValueRecord.EnterArray(TEXT("ElementsToRemove"), NumElementsToRemove);
 
 			if(NumElementsToRemove)
 			{
@@ -980,7 +980,7 @@ EConvertFromTypeResult FSetProperty::ConvertFromType(const FPropertyTag& Tag, FS
 			}
 
 			int32 Num = 0;
-			FStructuredArchive::FArray ElementsArray = ValueRecord.EnterArray(SA_FIELD_NAME(TEXT("Elements")), Num);
+			FStructuredArchive::FArray ElementsArray = ValueRecord.EnterArray(TEXT("Elements"), Num);
 
 			if(bConversionSucceeded)
 			{
diff --git a/Engine/Source/Runtime/CoreUObject/Private/UObject/SavePackage.cpp b/Engine/Source/Runtime/CoreUObject/Private/UObject/SavePackage.cpp
--- a/Engine/Source/Runtime/CoreUObject/Private/UObject/SavePackage.cpp
+++ b/Engine/Source/Runtime/CoreUObject/Private/UObject/SavePackage.cpp
@@ -2811,7 +2811,7 @@ FSavePackageResultStruct UPackage::Save(UPackage* InOuter, UObject* InAsset, con
 				}
 				SlowTask.EnterProgressFrame();
 
-				FStructuredArchive::FStream Stream = StructuredArchiveRoot.EnterStream(SA_FIELD_NAME(TEXT("GatherableTextData")));
+				FStructuredArchive::FStream Stream = StructuredArchiveRoot.EnterStream(TEXT("GatherableTextData"));
 				Linker->Summary.GatherableTextDataOffset = 0;
 				Linker->Summary.GatherableTextDataCount = 0;
 				if (!(Linker->Summary.GetPackageFlags() & PKG_FilterEditorOnly))
@@ -3240,7 +3240,7 @@ FSavePackageResultStruct UPackage::Save(UPackage* InOuter, UObject* InAsset, con
 				{
 					SCOPED_SAVETIMER(UPackage_Save_WriteDependsMap);
 
-					FStructuredArchive::FStream DependsStream = StructuredArchiveRoot.EnterStream(SA_FIELD_NAME(TEXT("DependsMap")));
+					FStructuredArchive::FStream DependsStream = StructuredArchiveRoot.EnterStream(TEXT("DependsMap"));
 					if (Linker->IsCooking())
 					{
 #if WITH_EDITOR
@@ -3284,7 +3284,7 @@ FSavePackageResultStruct UPackage::Save(UPackage* InOuter, UObject* InAsset, con
 					Linker->Summary.SoftPackageReferencesCount = Linker->SoftPackageReferenceList.Num();
 					if (!bTextFormat)
 					{
-						FStructuredArchive::FStream SoftReferenceStream = StructuredArchiveRoot.EnterStream(SA_FIELD_NAME(TEXT("SoftReferences")));
+						FStructuredArchive::FStream SoftReferenceStream = StructuredArchiveRoot.EnterStream(TEXT("SoftReferences"));
 						for (FName& SoftPackageName : Linker->SoftPackageReferenceList)
 						{
 							SoftReferenceStream.EnterElement() << SoftPackageName;
@@ -3292,7 +3292,7 @@ FSavePackageResultStruct UPackage::Save(UPackage* InOuter, UObject* InAsset, con
 
 						// Save searchable names map
 						Linker->Summary.SearchableNamesOffset = Linker->Tell();
-						Linker->SerializeSearchableNamesMap(StructuredArchiveRoot.EnterField(SA_FIELD_NAME(TEXT("SearchableNames"))));
+						Linker->SerializeSearchableNamesMap(StructuredArchiveRoot.EnterField(TEXT("SearchableNames")));
 					}
 				}
 				else
@@ -3310,7 +3310,7 @@ FSavePackageResultStruct UPackage::Save(UPackage* InOuter, UObject* InAsset, con
 					// Save thumbnails
 					{
 						SCOPED_SAVETIMER(UPackage_Save_SaveThumbnails);
-						SavePackageUtilities::SaveThumbnails(InOuter, Linker.Get(), StructuredArchiveRoot.EnterField(SA_FIELD_NAME(TEXT("Thumbnails"))));
+						SavePackageUtilities::SaveThumbnails(InOuter, Linker.Get(), StructuredArchiveRoot.EnterField(TEXT("Thumbnails")));
 					}
 
 					if (!bTextFormat)
@@ -3490,7 +3490,7 @@ FSavePackageResultStruct UPackage::Save(UPackage* InOuter, UObject* InAsset, con
 						}
 					};
 
-					FStructuredArchive::FStream DepedenciesStream = StructuredArchiveRoot.EnterStream(SA_FIELD_NAME(TEXT("PreloadDependencies")));
+					FStructuredArchive::FStream DepedenciesStream = StructuredArchiveRoot.EnterStream(TEXT("PreloadDependencies"));
 					TArray<UObject*> Subobjects;
 					TArray<UObject*> Deps;
 					TSet<FPackageIndex> SerializationBeforeCreateDependencies;
@@ -3718,7 +3718,7 @@ FSavePackageResultStruct UPackage::Save(UPackage* InOuter, UObject* InAsset, con
 					SCOPED_SAVETIMER(UPackage_Save_SaveExports);
 					FScopedSlowTask ExportScope(Linker->ExportMap.Num());
 
-					FStructuredArchive::FRecord ExportsRecord = StructuredArchiveRoot.EnterRecord(SA_FIELD_NAME(TEXT("Exports")));
+					FStructuredArchive::FRecord ExportsRecord = StructuredArchiveRoot.EnterRecord(TEXT("Exports"));
 
 					// Save exports.
 					int32 LastExportSaveStep = 0;
@@ -3739,7 +3739,7 @@ FSavePackageResultStruct UPackage::Save(UPackage* InOuter, UObject* InAsset, con
 							// UE_LOG(LogSavePackage, Log, TEXT("export %s for %s"), *Export.Object->GetFullName(), *Linker->CookingTarget()->PlatformName());
 
 							FString ObjectName = Export.Object->GetPathName(InOuter);
-							FStructuredArchive::FSlot ExportSlot = ExportsRecord.EnterField(SA_FIELD_NAME(*ObjectName));
+							FStructuredArchive::FSlot ExportSlot = ExportsRecord.EnterField(*ObjectName);
 
 							if (bTextFormat)
 							{
@@ -3924,7 +3924,7 @@ FSavePackageResultStruct UPackage::Save(UPackage* InOuter, UObject* InAsset, con
 						Linker->Seek(Linker->Summary.ImportOffset);
 
 						int32 NumImports = Linker->ImportMap.Num();
-						FStructuredArchive::FStream ImportTableStream = StructuredArchiveRoot.EnterStream(SA_FIELD_NAME(TEXT("ImportTable")));
+						FStructuredArchive::FStream ImportTableStream = StructuredArchiveRoot.EnterStream(TEXT("ImportTable"));
 
 						for (int32 i = 0; i < Linker->ImportMap.Num(); i++)
 						{
@@ -3996,7 +3996,7 @@ FSavePackageResultStruct UPackage::Save(UPackage* InOuter, UObject* InAsset, con
 					Linker->Seek(Linker->Summary.ExportOffset);
 
 					int32 NumExports = Linker->ExportMap.Num();
-					FStructuredArchive::FStream ExportTableStream = StructuredArchiveRoot.EnterStream(SA_FIELD_NAME(TEXT("ExportTable")));
+					FStructuredArchive::FStream ExportTableStream = StructuredArchiveRoot.EnterStream(TEXT("ExportTable"));
 					{
 #if WITH_EDITOR
 						FArchiveStackTraceIgnoreScope IgnoreSummaryDiffsScope(bIgnoreHeaderDiffs);
@@ -4065,7 +4065,7 @@ FSavePackageResultStruct UPackage::Save(UPackage* InOuter, UObject* InAsset, con
 #if WITH_EDITOR
 					FArchiveStackTraceIgnoreScope IgnoreSummaryDiffsScope(bIgnoreHeaderDiffs);
 #endif // WITH_EDITOR
-					StructuredArchiveRoot.EnterField(SA_FIELD_NAME(TEXT("Summary"))) << Linker->Summary;
+					StructuredArchiveRoot.EnterField(TEXT("Summary")) << Linker->Summary;
 					SerializedPackageFlags = Linker->Summary.GetPackageFlags();
 
 					// Currently the PKG_ContainsNoAsset flag is not serialized as part of the summary
diff --git a/Engine/Source/Runtime/CoreUObject/Private/UObject/SavePackage/SavePackageUtilities.cpp b/Engine/Source/Runtime/CoreUObject/Private/UObject/SavePackage/SavePackageUtilities.cpp
--- a/Engine/Source/Runtime/CoreUObject/Private/UObject/SavePackage/SavePackageUtilities.cpp
+++ b/Engine/Source/Runtime/CoreUObject/Private/UObject/SavePackage/SavePackageUtilities.cpp
@@ -1787,7 +1787,7 @@ void SaveThumbnails(UPackage* InOuter, FLinkerSave* Linker, FStructuredArchive::
 		if( ObjectsWithThumbnails.Num() > 0 )
 		{
 			// Save out the image data for the thumbnails
-			FStructuredArchive::FStream ThumbnailStream = Record.EnterStream(SA_FIELD_NAME(TEXT("Thumbnails")));
+			FStructuredArchive::FStream ThumbnailStream = Record.EnterStream(TEXT("Thumbnails"));
 
 			for( int32 CurObjectIndex = 0; CurObjectIndex < ObjectsWithThumbnails.Num(); ++CurObjectIndex )
 			{
@@ -1808,7 +1808,7 @@ void SaveThumbnails(UPackage* InOuter, FLinkerSave* Linker, FStructuredArchive::
 
 				// Save number of thumbnails
 				int32 ThumbnailCount = ObjectsWithThumbnails.Num();
-				FStructuredArchive::FArray IndexArray = Record.EnterField(SA_FIELD_NAME(TEXT("Index"))).EnterArray(ThumbnailCount);
+				FStructuredArchive::FArray IndexArray = Record.EnterField(TEXT("Index")).EnterArray(ThumbnailCount);
 
 				// Store a list of object names along with the offset in the file where the thumbnail is stored
 				for( int32 CurObjectIndex = 0; CurObjectIndex < ObjectsWithThumbnails.Num(); ++CurObjectIndex )
@@ -2427,7 +2427,7 @@ namespace AssetRegistry
 
 		// Store the asset registry offset in the file and enter a record for the asset registry data
 		Linker->Summary.AssetRegistryDataOffset = BinaryArchive.Tell();
-		FStructuredArchiveRecord AssetRegistryRecord = ParentRecord.EnterField(SA_FIELD_NAME(TEXT("AssetRegistry"))).EnterRecord();
+		FStructuredArchiveRecord AssetRegistryRecord = ParentRecord.EnterField(TEXT("AssetRegistry")).EnterRecord();
 
 		// Offset to Dependencies
 		int64 OffsetToAssetRegistryDependencyDataOffset = INDEX_NONE;
@@ -2455,7 +2455,7 @@ namespace AssetRegistry
 			}
 		}
 		int32 ObjectCount = AssetObjects.Num();
-		FStructuredArchive::FArray AssetArray = AssetRegistryRecord.EnterArray(SA_FIELD_NAME(TEXT("TagMap")), ObjectCount);
+		FStructuredArchive::FArray AssetArray = AssetRegistryRecord.EnterArray(TEXT("TagMap"), ObjectCount);
 		for (int32 ObjectIdx = 0; ObjectIdx < AssetObjects.Num(); ++ObjectIdx)
 		{
 			const UObject* Object = AssetObjects[ObjectIdx];
@@ -2487,7 +2487,7 @@ namespace AssetRegistry
 			FStructuredArchive::FRecord AssetRecord = AssetArray.EnterElement().EnterRecord();
 			AssetRecord << SA_VALUE(TEXT("Path"), ObjectPath) << SA_VALUE(TEXT("Class"), ObjectClassName);
 
-			FStructuredArchive::FMap TagMap = AssetRecord.EnterField(SA_FIELD_NAME(TEXT("Tags"))).EnterMap(TagCount);
+			FStructuredArchive::FMap TagMap = AssetRecord.EnterField(TEXT("Tags")).EnterMap(TagCount);
 
 			for (TArray<UObject::FAssetRegistryTag>::TConstIterator TagIter(Tags); TagIter; ++TagIter)
 			{
@@ -2510,7 +2510,7 @@ namespace AssetRegistry
 			BinaryArchive << AssetRegistryDependencyDataOffset;
 			BinaryArchive.Seek(AssetRegistryDependencyDataOffset);
 		}
-		FStructuredArchiveRecord DependencyDataRecord = ParentRecord.EnterField(SA_FIELD_NAME(TEXT("AssetRegistryDependencyData"))).EnterRecord();
+		FStructuredArchiveRecord DependencyDataRecord = ParentRecord.EnterField(TEXT("AssetRegistryDependencyData")).EnterRecord();
 
 		// Convert the IsUsedInGame sets into a bitarray with a value per import/softpackagereference
 		TBitArray<> ImportUsedInGameBits;
diff --git a/Engine/Source/Runtime/CoreUObject/Private/UObject/SavePackage2.cpp b/Engine/Source/Runtime/CoreUObject/Private/UObject/SavePackage2.cpp
--- a/Engine/Source/Runtime/CoreUObject/Private/UObject/SavePackage2.cpp
+++ b/Engine/Source/Runtime/CoreUObject/Private/UObject/SavePackage2.cpp
@@ -1316,7 +1316,7 @@ void SavePreloadDependencies(FStructuredArchive::FRecord& StructuredArchiveRoot,
 	{
 		Linker->Summary.PreloadDependencyCount = 0;
 
-		FStructuredArchive::FStream DepedenciesStream = StructuredArchiveRoot.EnterStream(SA_FIELD_NAME(TEXT("PreloadDependencies")));
+		FStructuredArchive::FStream DepedenciesStream = StructuredArchiveRoot.EnterStream(TEXT("PreloadDependencies"));
 
 		TArray<UObject*> Subobjects;
 		TArray<UObject*> Deps;
@@ -1521,7 +1521,7 @@ void SavePreloadDependencies(FStructuredArchive::FRecord& StructuredArchiveRoot,
 
 void WriteGatherableText(FStructuredArchive::FRecord& StructuredArchiveRoot, FSaveContext& SaveContext)
 {
-	FStructuredArchive::FStream Stream = StructuredArchiveRoot.EnterStream(SA_FIELD_NAME(TEXT("GatherableTextData")));
+	FStructuredArchive::FStream Stream = StructuredArchiveRoot.EnterStream(TEXT("GatherableTextData"));
 	if (!SaveContext.IsFilterEditorOnly()
 		// We can only cache packages that:
 		//	1) Don't contain script data, as script data is very volatile and can only be safely gathered after it's been compiled (which happens automatically on asset load).
@@ -1599,7 +1599,7 @@ ESavePackageResult WritePackageHeader(FStructuredArchive::FRecord& StructuredArc
 	{
 		SCOPED_SAVETIMER(UPackage_Save_WriteDependsMap);
 
-		FStructuredArchive::FStream DependsStream = StructuredArchiveRoot.EnterStream(SA_FIELD_NAME(TEXT("DependsMap")));
+		FStructuredArchive::FStream DependsStream = StructuredArchiveRoot.EnterStream(TEXT("DependsMap"));
 		Linker->Summary.DependsOffset = Linker->Tell();
 		if (SaveContext.IsCooking())
 		{
@@ -1630,7 +1630,7 @@ ESavePackageResult WritePackageHeader(FStructuredArchive::FRecord& StructuredArc
 		Linker->Summary.SoftPackageReferencesOffset = Linker->Tell();
 		Linker->Summary.SoftPackageReferencesCount = Linker->SoftPackageReferenceList.Num();
 		{
-			FStructuredArchive::FStream SoftReferenceStream = StructuredArchiveRoot.EnterStream(SA_FIELD_NAME(TEXT("SoftReferences")));
+			FStructuredArchive::FStream SoftReferenceStream = StructuredArchiveRoot.EnterStream(TEXT("SoftReferences"));
 			for (FName& SoftPackageName : Linker->SoftPackageReferenceList)
 			{
 				SoftReferenceStream.EnterElement() << SoftPackageName;
@@ -1638,7 +1638,7 @@ ESavePackageResult WritePackageHeader(FStructuredArchive::FRecord& StructuredArc
 
 			// Save searchable names map
 			Linker->Summary.SearchableNamesOffset = Linker->Tell();
-			Linker->SerializeSearchableNamesMap(StructuredArchiveRoot.EnterField(SA_FIELD_NAME(TEXT("SearchableNames"))));
+			Linker->SerializeSearchableNamesMap(StructuredArchiveRoot.EnterField(TEXT("SearchableNames")));
 		}
 	}
 	else
@@ -1651,7 +1651,7 @@ ESavePackageResult WritePackageHeader(FStructuredArchive::FRecord& StructuredArc
 	// Save thumbnails
 	{
 		SCOPED_SAVETIMER(UPackage_Save_SaveThumbnails);
-		SavePackageUtilities::SaveThumbnails(SaveContext.GetPackage(), Linker, StructuredArchiveRoot.EnterField(SA_FIELD_NAME(TEXT("Thumbnails"))));
+		SavePackageUtilities::SaveThumbnails(SaveContext.GetPackage(), Linker, StructuredArchiveRoot.EnterField(TEXT("Thumbnails")));
 	}
 	{
 		// Save asset registry data so the editor can search for information about assets in this package
@@ -1686,7 +1686,7 @@ ESavePackageResult WritePackageTextHeader(FStructuredArchive::FRecord& Structure
 	// Write ImportTable
 	{
 		SCOPED_SAVETIMER(UPackage_Save_WriteImportTable);
-		FStructuredArchive::FStream ImportTableStream = StructuredArchiveRoot.EnterStream(SA_FIELD_NAME(TEXT("ImportTable")));
+		FStructuredArchive::FStream ImportTableStream = StructuredArchiveRoot.EnterStream(TEXT("ImportTable"));
 		for (FObjectImport& Import : Linker->ImportMap)
 		{
 			ImportTableStream.EnterElement() << Import;
@@ -1696,7 +1696,7 @@ ESavePackageResult WritePackageTextHeader(FStructuredArchive::FRecord& Structure
 	// Write ExportTable
 	{
 		SCOPED_SAVETIMER(UPackage_Save_WriteExportTable);
-		FStructuredArchive::FStream ExportTableStream = StructuredArchiveRoot.EnterStream(SA_FIELD_NAME(TEXT("ExportTable")));
+		FStructuredArchive::FStream ExportTableStream = StructuredArchiveRoot.EnterStream(TEXT("ExportTable"));
 		for (FObjectExport& Export : Linker->ExportMap)
 		{
 			ExportTableStream.EnterElement() << Export;
@@ -1706,7 +1706,7 @@ ESavePackageResult WritePackageTextHeader(FStructuredArchive::FRecord& Structure
 	// Save thumbnails
 	{
 		SCOPED_SAVETIMER(UPackage_Save_SaveThumbnails);
-		SavePackageUtilities::SaveThumbnails(SaveContext.GetPackage(), Linker, StructuredArchiveRoot.EnterField(SA_FIELD_NAME(TEXT("Thumbnails"))));
+		SavePackageUtilities::SaveThumbnails(SaveContext.GetPackage(), Linker, StructuredArchiveRoot.EnterField(TEXT("Thumbnails")));
 	}
 	// Save level information used by World browser
 	{
@@ -1724,7 +1724,7 @@ ESavePackageResult WriteExports(FStructuredArchive::FRecord& StructuredArchiveRo
 	FLinkerSave* Linker = SaveContext.GetLinker();
 	FScopedSlowTask SlowTask(Linker->ExportMap.Num(), FText(), SaveContext.IsUsingSlowTask());
 
-	FStructuredArchive::FRecord ExportsRecord = StructuredArchiveRoot.EnterRecord(SA_FIELD_NAME(TEXT("Exports")));
+	FStructuredArchive::FRecord ExportsRecord = StructuredArchiveRoot.EnterRecord(TEXT("Exports"));
 
 	// Save exports.
 	int32 LastExportSaveStep = 0;
@@ -1746,7 +1746,7 @@ ESavePackageResult WriteExports(FStructuredArchive::FRecord& StructuredArchiveRo
 			Linker->CurrentlySavingExport = FPackageIndex::FromExport(i);
 
 			FString ObjectName = Export.Object->GetPathName(SaveContext.GetPackage());
-			FStructuredArchive::FSlot ExportSlot = ExportsRecord.EnterField(SA_FIELD_NAME(*ObjectName));
+			FStructuredArchive::FSlot ExportSlot = ExportsRecord.EnterField(*ObjectName);
 
 #if WITH_EDITOR
 			bool bSupportsText = UClass::IsSafeToSerializeToStructuredArchives(Export.Object->GetClass());
@@ -1875,7 +1875,7 @@ ESavePackageResult UpdatePackageHeader(FStructuredArchive::FRecord& StructuredAr
 	if (!SaveContext.IsTextFormat())
 	{
 		Linker->Seek(Linker->Summary.ImportOffset);
-		FStructuredArchive::FStream ImportTableStream = StructuredArchiveRoot.EnterStream(SA_FIELD_NAME(TEXT("ImportTable")));
+		FStructuredArchive::FStream ImportTableStream = StructuredArchiveRoot.EnterStream(TEXT("ImportTable"));
 		for (FObjectImport& Import : Linker->ImportMap)
 		{
 			ImportTableStream.EnterElement() << Import;
@@ -1886,7 +1886,7 @@ ESavePackageResult UpdatePackageHeader(FStructuredArchive::FRecord& StructuredAr
 	{
 		check(Linker->Tell() == SaveContext.OffsetAfterImportMap);
 		Linker->Seek(Linker->Summary.ExportOffset);
-		FStructuredArchive::FStream ExportTableStream = StructuredArchiveRoot.EnterStream(SA_FIELD_NAME(TEXT("ExportTable")));
+		FStructuredArchive::FStream ExportTableStream = StructuredArchiveRoot.EnterStream(TEXT("ExportTable"));
 
 		for (FObjectExport& Export : Linker->ExportMap)
 		{
@@ -1952,7 +1952,7 @@ ESavePackageResult UpdatePackageHeader(FStructuredArchive::FRecord& StructuredAr
 			Linker->Seek(0);
 		}
 		{
-			StructuredArchiveRoot.EnterField(SA_FIELD_NAME(TEXT("Summary"))) << Linker->Summary;
+			StructuredArchiveRoot.EnterField(TEXT("Summary")) << Linker->Summary;
 		}
 
 		if (!SaveContext.IsTextFormat())
diff --git a/Engine/Source/Runtime/CoreUObject/Public/UObject/ObjectMacros.h b/Engine/Source/Runtime/CoreUObject/Public/UObject/ObjectMacros.h
--- a/Engine/Source/Runtime/CoreUObject/Public/UObject/ObjectMacros.h
+++ b/Engine/Source/Runtime/CoreUObject/Public/UObject/ObjectMacros.h
@@ -1656,7 +1656,7 @@ namespace UM
 	}
 
 #define IMPLEMENT_FARCHIVE_SERIALIZER( TClass ) void TClass::Serialize(FArchive& Ar) { TClass::Serialize(FStructuredArchiveFromArchive(Ar).GetSlot().EnterRecord()); }
-#define IMPLEMENT_FSTRUCTUREDARCHIVE_SERIALIZER( TClass ) void TClass::Serialize(FStructuredArchive::FRecord Record) { FArchiveUObjectFromStructuredArchive Ar(Record.EnterField(SA_FIELD_NAME(TEXT("BaseClassAutoGen")))); TClass::Serialize(Ar.GetArchive()); Ar.Close(); }
+#define IMPLEMENT_FSTRUCTUREDARCHIVE_SERIALIZER( TClass ) void TClass::Serialize(FStructuredArchive::FRecord Record) { FArchiveUObjectFromStructuredArchive Ar(Record.EnterField(TEXT("BaseClassAutoGen"))); TClass::Serialize(Ar.GetArchive()); Ar.Close(); }
 #define DECLARE_FARCHIVE_SERIALIZER( TClass, API ) virtual API void Serialize(FArchive& Ar) override;
 #define DECLARE_FSTRUCTUREDARCHIVE_SERIALIZER( TClass, API ) virtual API void Serialize(FStructuredArchive::FRecord Record) override;
 
diff --git a/Engine/Source/Runtime/Engine/Private/DataTable.cpp b/Engine/Source/Runtime/Engine/Private/DataTable.cpp
--- a/Engine/Source/Runtime/Engine/Private/DataTable.cpp
+++ b/Engine/Source/Runtime/Engine/Private/DataTable.cpp
@@ -113,7 +113,7 @@ void UDataTable::LoadStructData(FStructuredArchiveSlot Slot)
 		// And be sure to call DestroyScriptStruct later
 		LoadUsingStruct->InitializeStruct(RowData);
 
-		LoadUsingStruct->SerializeItem(RowRecord.EnterField(SA_FIELD_NAME(TEXT("Value"))), RowData, nullptr);
+		LoadUsingStruct->SerializeItem(RowRecord.EnterField(TEXT("Value")), RowData, nullptr);
 
 		// Add to map
 		RowMap.Add(RowName, RowData);
@@ -146,7 +146,7 @@ void UDataTable::SaveStructData(FStructuredArchiveSlot Slot)
 		// Save out data
 		uint8* RowData = RowIt.Value();
 
-		SaveUsingStruct->SerializeItem(Row.EnterField(SA_FIELD_NAME(TEXT("Value"))), RowData, nullptr);
+		SaveUsingStruct->SerializeItem(Row.EnterField(TEXT("Value")), RowData, nullptr);
 	}
 }
 
@@ -248,11 +248,11 @@ void UDataTable::Serialize(FStructuredArchiveRecord Record)
 	{
 		DATATABLE_CHANGE_SCOPE();
 		EmptyTable();
-		LoadStructData(Record.EnterField(SA_FIELD_NAME(TEXT("Data"))));
+		LoadStructData(Record.EnterField(TEXT("Data")));
 	}
 	else if(BaseArchive.IsSaving())
 	{
-		SaveStructData(Record.EnterField(SA_FIELD_NAME(TEXT("Data"))));
+		SaveStructData(Record.EnterField(TEXT("Data")));
 	}
 }
 
diff --git a/Engine/Source/Runtime/Engine/Private/Particles/ParticleModules.cpp b/Engine/Source/Runtime/Engine/Private/Particles/ParticleModules.cpp
--- a/Engine/Source/Runtime/Engine/Private/Particles/ParticleModules.cpp
+++ b/Engine/Source/Runtime/Engine/Private/Particles/ParticleModules.cpp
@@ -1125,7 +1125,7 @@ void UParticleModuleRequired::Serialize(FStructuredArchive::FRecord Record)
 
 		if (bCooked)
 		{
-			DerivedData.Serialize(Record.EnterField(SA_FIELD_NAME(TEXT("DerivedData"))));
+			DerivedData.Serialize(Record.EnterField(TEXT("DerivedData")));
 		}
 	}
 }
diff --git a/Engine/Source/Runtime/Engine/Private/Particles/SubUVAnimation.cpp b/Engine/Source/Runtime/Engine/Private/Particles/SubUVAnimation.cpp
--- a/Engine/Source/Runtime/Engine/Private/Particles/SubUVAnimation.cpp
+++ b/Engine/Source/Runtime/Engine/Private/Particles/SubUVAnimation.cpp
@@ -90,7 +90,7 @@ void USubUVAnimation::Serialize(FStructuredArchive::FRecord Record)
 
 	if (bCooked)
 	{
-		DerivedData.Serialize(Record.EnterField(SA_FIELD_NAME(TEXT("DerivedData"))));
+		DerivedData.Serialize(Record.EnterField(TEXT("DerivedData")));
 	}
 }
 
diff --git a/Engine/Source/Runtime/Engine/Private/SoundCue.cpp b/Engine/Source/Runtime/Engine/Private/SoundCue.cpp
--- a/Engine/Source/Runtime/Engine/Private/SoundCue.cpp
+++ b/Engine/Source/Runtime/Engine/Private/SoundCue.cpp
@@ -169,7 +169,7 @@ void USoundCue::Serialize(FStructuredArchive::FRecord Record)
 
 	if (UnderlyingArchive.UEVer() >= VER_UE4_COOKED_ASSETS_IN_EDITOR_SUPPORT)
 	{
-		FStripDataFlags StripFlags(Record.EnterField(SA_FIELD_NAME(TEXT("SoundCueStripFlags"))));
+		FStripDataFlags StripFlags(Record.EnterField(TEXT("SoundCueStripFlags")));
 #if WITH_EDITORONLY_DATA
 		if (!StripFlags.IsEditorDataStripped())
 		{
diff --git a/Engine/Source/Runtime/Engine/Private/UserDefinedStruct.cpp b/Engine/Source/Runtime/Engine/Private/UserDefinedStruct.cpp
--- a/Engine/Source/Runtime/Engine/Private/UserDefinedStruct.cpp
+++ b/Engine/Source/Runtime/Engine/Private/UserDefinedStruct.cpp
@@ -84,7 +84,7 @@ void UUserDefinedStruct::Serialize(FStructuredArchive::FRecord Record)
 			uint8* StructData = DefaultStructInstance.GetStructMemory();
 
 			FScopedPlaceholderRawContainerTracker TrackStruct(StructData);
-			SerializeItem(Record.EnterField(SA_FIELD_NAME(TEXT("Data"))), StructData, nullptr);
+			SerializeItem(Record.EnterField(TEXT("Data")), StructData, nullptr);
 
 			// Now that defaults have been loaded we can inspect our properties
 			// and default values and set the StructFlags accordingly:
