EditorPackageLoader: Fixed issue with cooked packages not being imported correctly from uncooked packages loaded with LoadPackageAsync
#rnx

diff --git a/Engine/Source/Runtime/CoreUObject/Private/Serialization/AsyncLoading.cpp b/Engine/Source/Runtime/CoreUObject/Private/Serialization/AsyncLoading.cpp
--- a/Engine/Source/Runtime/CoreUObject/Private/Serialization/AsyncLoading.cpp
+++ b/Engine/Source/Runtime/CoreUObject/Private/Serialization/AsyncLoading.cpp
@@ -6081,8 +6081,28 @@ EAsyncPackageState::Type FAsyncPackage::LoadImports()
 			// The package doesn't exist and this import is not in the dependency list so add it now.
 			if (!FPackageName::IsShortPackageName(ImportPackageName))
 			{
-				UE_LOG(LogStreaming, Verbose, TEXT("FAsyncPackage::LoadImports for %s: Loading %s"), *Desc.PackagePath.GetDebugName(), *ImportPackageName);
-				AddImportDependency(ImportPackageFName, ImportToLoad, InstancingContext);
+				bool bLoadedFromIoStore = false;
+				if (IAsyncPackageLoader* IoStorePackageLoader = AsyncLoadingThread.GetIoStorePackageLoader())
+				{
+					FPackagePath ImportedPackagePath = FPackagePath::FromPackageNameUnchecked(ImportPackageFName);
+					FPackageName::EPackageLocationFilter PackageLocation = FPackageName::EPackageLocationFilter(uint8(FPackageName::EPackageLocationFilter::IoDispatcher) | uint8(FPackageName::EPackageLocationFilter::FileSystem));
+					PackageLocation = FPackageName::DoesPackageExistEx(ImportedPackagePath, PackageLocation, false /* bMatchCaseOnDisk */, nullptr);
+					if (PackageLocation == FPackageName::EPackageLocationFilter::IoDispatcher)
+					{
+						UE_LOG(LogStreaming, Verbose, TEXT("FAsyncPackage::LoadImports for %s: Loading %s (0x%llX) using IoStore package loader"), *Desc.PackagePath.GetDebugName(), *ImportPackageName, FPackageId::FromName(ImportPackageFName).ValueForDebugging());
+						check(IsInGameThread());
+						FAsyncLoadingThread::LeaveAsyncLoadingTick(AsyncLoadingThread.GetThreadIndex()); // Fix thread check that triggers if we recurse back into the uncooked package loader from here
+						int32 ImportRequestId = IoStorePackageLoader->LoadPackage(ImportedPackagePath, NAME_None, FLoadPackageAsyncDelegate(), PKG_None, INDEX_NONE, 0, nullptr);
+						IoStorePackageLoader->FlushLoading(ImportRequestId);
+						FAsyncLoadingThread::EnterAsyncLoadingTick(AsyncLoadingThread.GetThreadIndex());
+						bLoadedFromIoStore = true;
+					}
+				}
+				if (!bLoadedFromIoStore)
+				{
+					UE_LOG(LogStreaming, Verbose, TEXT("FAsyncPackage::LoadImports for %s: Loading %s"), *Desc.PackagePath.GetDebugName(), *ImportPackageName);
+					AddImportDependency(ImportPackageFName, ImportToLoad, InstancingContext);
+				}
 			}
 			else
 			{
@@ -6782,7 +6802,7 @@ void FAsyncPackage::CloseDelayedLinkers()
 
 void FAsyncPackage::CallCompletionCallbacks(bool bInternal, EAsyncLoadingResult::Type LoadingResult)
 {
-	checkSlow(bInternal || !IsInAsyncLoadingThread());
+	checkSlow(bInternal || !GetOwnerThread().IsInAsyncLoadThread());
 
 	UPackage* LoadedPackage = (!bLoadHasFailed) ? LinkerRoot : nullptr;
 	for (FCompletionCallback& CompletionCallback : CompletionCallbacks)
diff --git a/Engine/Source/Runtime/CoreUObject/Private/Serialization/AsyncLoading2.cpp b/Engine/Source/Runtime/CoreUObject/Private/Serialization/AsyncLoading2.cpp
--- a/Engine/Source/Runtime/CoreUObject/Private/Serialization/AsyncLoading2.cpp
+++ b/Engine/Source/Runtime/CoreUObject/Private/Serialization/AsyncLoading2.cpp
@@ -2049,10 +2049,8 @@ private:
 
 	/** [ASYNC/GAME THREAD] Number of package load requests in the async loading queue */
 	TAtomic<uint32> QueuedPackagesCounter { 0 };
-	/** [ASYNC/GAME THREAD] Number of packages being loaded on the async thread and post loaded on the game thread */
-	FThreadSafeCounter ExistingAsyncPackagesCounter;
 	/** [ASYNC/GAME THREAD] Number of packages being loaded on the async thread and post loaded on the game thread. Excludes packages in the deferred delete queue*/
-	FThreadSafeCounter ActiveAsyncPackagesCounter;
+	TAtomic<uint32> LoadingPackagesCounter { 0 };
 
 	FThreadSafeCounter AsyncThreadReady;
 
@@ -2128,7 +2126,7 @@ public:
 	/** Returns true if packages are currently being loaded on the async thread */
 	inline virtual bool IsAsyncLoadingPackages() override
 	{
-		return QueuedPackagesCounter != 0 || ExistingAsyncPackagesCounter.GetValue() != 0;
+		return QueuedPackagesCounter + LoadingPackagesCounter != 0;
 	}
 
 	/** Returns true this codes runs on the async loading thread */
@@ -2273,7 +2271,7 @@ public:
 
 	virtual int32 GetNumAsyncPackages() override
 	{
-		return ActiveAsyncPackagesCounter.GetValue();
+		return LoadingPackagesCounter;
 	}
 
 	/**
@@ -2377,7 +2375,6 @@ private:
 	{
 		UE_ASYNC_PACKAGE_DEBUG(Desc);
 
-		ExistingAsyncPackagesCounter.Increment();
 		return new FAsyncPackage2(Desc, *this, GraphAllocator, EventSpecs.GetData());
 	}
 
@@ -2447,7 +2444,6 @@ private:
 		TRACE_CPUPROFILER_EVENT_SCOPE(DeleteAsyncPackage);
 		UE_ASYNC_PACKAGE_DEBUG(Package->Desc);
 		delete Package;
-		ExistingAsyncPackagesCounter.Decrement();
 	}
 
 	/** Number of times we re-entered the async loading tick, mostly used by singlethreaded ticking. Debug purposes only. */
@@ -2577,7 +2573,6 @@ FAsyncPackage2* FAsyncLoadingThread2::FindOrInsertPackage(FAsyncPackageDesc2& De
 			Package = CreateAsyncPackage(Desc);
 			checkf(Package, TEXT("Failed to create async package %s"), *Desc.UPackageName.ToString());
 			Package->AddRef();
-			ActiveAsyncPackagesCounter.Increment();
 			AsyncPackageLookup.Add(Desc.UPackageId, Package);
 			bInserted = true;
 		}
@@ -2633,6 +2628,9 @@ bool FAsyncLoadingThread2::CreateAsyncPackagesFromQueue(FAsyncLoadingThreadState
 			{
 				break;
 			}
+
+			++LoadingPackagesCounter;
+			--QueuedPackagesCounter;
 			++NumDequeued;
 		
 			FPackageRequest& Request = OptionalRequest.GetValue();
@@ -2730,9 +2728,8 @@ bool FAsyncLoadingThread2::CreateAsyncPackagesFromQueue(FAsyncLoadingThreadState
 				{
 					UE_ASYNC_PACKAGE_LOG_VERBOSE(Verbose, PackageDesc, TEXT("CreateAsyncPackages: UpdatePackage"),
 						TEXT("Package is alreay being loaded."));
+					--LoadingPackagesCounter;
 				}
-
-				--QueuedPackagesCounter;
 			}
 		}
 		
@@ -3284,10 +3281,13 @@ void FAsyncPackage2::ImportPackagesRecursive(FIoBatch& IoBatch, IPackageStore& P
 			UPackage* UncookedPackage = nullptr;
 			if (!PackageRef.AreAllPublicExportsLoaded())
 			{
+				UE_ASYNC_PACKAGE_LOG(Verbose, Desc, TEXT("ImportPackages: LoadUncookedImport"), TEXT("Loading imported uncooked package '%s' '0x%llX'"), *ImportedPackageEntry.UncookedPackageName.ToString(), ImportedPackageId.ValueForDebugging());
+				check(IsInGameThread());
+				IoBatch.Issue(); // The batch might already contain requests for packages being imported from the uncooked one we're going to load so make sure that those are started before blocking
 				FPackagePath ImportedPackagePath = FPackagePath::FromPackageNameUnchecked(ImportedPackageEntry.UncookedPackageName);
 				ImportedPackagePath.SetHeaderExtension(static_cast<EPackageExtension>(ImportedPackageEntry.UncookedPackageHeaderExtension));
-				AsyncLoadingThread.UncookedPackageLoader->LoadPackage(ImportedPackagePath, NAME_None, FLoadPackageAsyncDelegate(), PKG_None, INDEX_NONE, 0, nullptr);
-				AsyncLoadingThread.UncookedPackageLoader->FlushLoading(INDEX_NONE);
+				int32 ImportRequestId = AsyncLoadingThread.UncookedPackageLoader->LoadPackage(ImportedPackagePath, NAME_None, FLoadPackageAsyncDelegate(), PKG_None, INDEX_NONE, 0, nullptr);
+				AsyncLoadingThread.UncookedPackageLoader->FlushLoading(ImportRequestId);
 				UncookedPackage = FindObjectFast<UPackage>(nullptr, ImportedPackagePath.GetPackageFName());
 				if (UncookedPackage)
 				{
@@ -3369,13 +3369,14 @@ void FAsyncPackage2::ImportPackagesRecursive(FIoBatch& IoBatch, IPackageStore& P
 
 		if (bInserted)
 		{
-			UE_ASYNC_PACKAGE_LOG(Verbose, PackageDesc, TEXT("ImportPackages: AddPackage"),
-				TEXT("Start loading imported package."));
+			UE_ASYNC_PACKAGE_LOG(Verbose, Desc, TEXT("ImportPackages: AddPackage"),
+			TEXT("Start loading imported package with id '0x%llX'"), ImportedPackageId.ValueForDebugging());
+			++AsyncLoadingThread.LoadingPackagesCounter;
 		}
 		else
 		{
-			UE_ASYNC_PACKAGE_LOG_VERBOSE(VeryVerbose, PackageDesc, TEXT("ImportPackages: UpdatePackage"),
-				TEXT("Imported package is already being loaded."));
+			UE_ASYNC_PACKAGE_LOG_VERBOSE(VeryVerbose, Desc, TEXT("ImportPackages: UpdatePackage"),
+				TEXT("Imported package with id '0x%llX' is already being loaded."), ImportedPackageId.ValueForDebugging());
 		}
 		ImportedPackage->AddRef();
 		Data.ImportedAsyncPackages[ImportedPackageIndex++] = ImportedPackage;
@@ -4748,9 +4749,6 @@ EAsyncPackageState::Type FAsyncLoadingThread2::ProcessLoadedPackagesFromGameThre
 			// this is to ensure we can re-enter FlushAsyncLoading from any of the callbacks
 			LoadedPackagesToProcess.RemoveAt(PackageIndex--);
 
-			// Incremented on the Async Thread, now decrement as we're done with this package				
-			ActiveAsyncPackagesCounter.Decrement();
-
 			TRACE_LOADTIME_END_LOAD_ASYNC_PACKAGE(Package);
 
 			check(Package->AsyncPackageLoadingState == EAsyncPackageLoadingState2::Finalize);
@@ -4800,7 +4798,7 @@ EAsyncPackageState::Type FAsyncLoadingThread2::ProcessLoadedPackagesFromGameThre
 			{
 				FailedPackageRequest.Callback->ExecuteIfBound(FailedPackageRequest.PackageName, nullptr, EAsyncLoadingResult::Failed);
 				RemovePendingRequests(TArrayView<int32>(&FailedPackageRequest.RequestID, 1));
-				--QueuedPackagesCounter;
+				--LoadingPackagesCounter;
 			}
 		}
 
@@ -4865,6 +4863,7 @@ EAsyncPackageState::Type FAsyncLoadingThread2::ProcessLoadedPackagesFromGameThre
 					CompletedPackages.RemoveAtSwap(PackageIndex--);
 					Package->ClearImportedPackages();
 					Package->ReleaseRef();
+					--LoadingPackagesCounter;
 				}
 			}
 
@@ -4952,7 +4951,7 @@ EAsyncPackageState::Type FAsyncLoadingThread2::TickAsyncLoadingFromGameThread(FA
 		if (Result != EAsyncPackageState::TimeOut)
 		{
 			// Flush deferred messages
-			if (ExistingAsyncPackagesCounter.GetValue() == 0)
+			if (!IsAsyncLoadingPackages())
 			{
 				bDidSomething = true;
 				FDeferredMessageLog::Flush();
@@ -6047,7 +6046,7 @@ void FAsyncLoadingThread2::QueueMissingPackage(FAsyncPackageDesc2& PackageDesc,
 	else
 	{
 		RemovePendingRequests(TArrayView<int32>(&PackageDesc.RequestID, 1));
-		--QueuedPackagesCounter;
+		--LoadingPackagesCounter;
 	}
 }
 
diff --git a/Engine/Source/Runtime/CoreUObject/Private/Serialization/AsyncLoadingThread.h b/Engine/Source/Runtime/CoreUObject/Private/Serialization/AsyncLoadingThread.h
--- a/Engine/Source/Runtime/CoreUObject/Private/Serialization/AsyncLoadingThread.h
+++ b/Engine/Source/Runtime/CoreUObject/Private/Serialization/AsyncLoadingThread.h
@@ -283,11 +283,24 @@ private:
 	/** Thread safe counter used to accumulate cycles spent on blocking. Using stats may generate to many stats messages. */
 	static FThreadSafeCounter BlockingCycles;
 #endif
+
+	IAsyncPackageLoader* IoStorePackageLoader = nullptr;
+
 public:
 
 	FAsyncLoadingThread(int32 InThreadIndex, IEDLBootNotificationManager& InEDLBootNotificationManager);
 	virtual ~FAsyncLoadingThread();
 
+	IAsyncPackageLoader* GetIoStorePackageLoader() const
+	{
+		return IoStorePackageLoader;
+	}
+
+	void SetIoStorePackageLoader(IAsyncPackageLoader* InIoStorePackageLoader)
+	{
+		IoStorePackageLoader = InIoStorePackageLoader;
+	}
+
 	//~ Begin FRunnable Interface.
 	virtual bool Init();
 	virtual uint32 Run();
diff --git a/Engine/Source/Runtime/CoreUObject/Private/Serialization/EditorPackageLoader.cpp b/Engine/Source/Runtime/CoreUObject/Private/Serialization/EditorPackageLoader.cpp
--- a/Engine/Source/Runtime/CoreUObject/Private/Serialization/EditorPackageLoader.cpp
+++ b/Engine/Source/Runtime/CoreUObject/Private/Serialization/EditorPackageLoader.cpp
@@ -20,6 +20,7 @@ public:
 	{
 		UncookedPackageLoader.Reset(new FAsyncLoadingThread(/** ThreadIndex = */ 0, InEDLBootNotificationManager));
 		CookedPackageLoader.Reset(MakeAsyncPackageLoader2(InIoDispatcher, UncookedPackageLoader.Get()));
+		UncookedPackageLoader->SetIoStorePackageLoader(CookedPackageLoader.Get());
 	}
 
 	virtual ~FEditorPackageLoader () { }
@@ -199,7 +200,7 @@ public:
 
 private:
 	TUniquePtr<IAsyncPackageLoader> CookedPackageLoader;
-	TUniquePtr<IAsyncPackageLoader> UncookedPackageLoader;
+	TUniquePtr<FAsyncLoadingThread> UncookedPackageLoader;
 	bool bHasInitializedCookedPackageLoader = false;
 };
 
