Improved the parsing and validation of default parameter meta-data in Python

- Class parameters using short paths now parse correctly.
- Default values are now validated to avoid false positives (eg, LatentInfo, WorldContext, etc).

#rnx

diff --git a/Engine/Plugins/Experimental/PythonScriptPlugin/Source/PythonScriptPlugin/Private/PyGenUtil.cpp b/Engine/Plugins/Experimental/PythonScriptPlugin/Source/PythonScriptPlugin/Private/PyGenUtil.cpp
--- a/Engine/Plugins/Experimental/PythonScriptPlugin/Source/PythonScriptPlugin/Private/PyGenUtil.cpp
+++ b/Engine/Plugins/Experimental/PythonScriptPlugin/Source/PythonScriptPlugin/Private/PyGenUtil.cpp
@@ -12,6 +12,7 @@
 #include "Misc/ScopeExit.h"
 #include "Misc/FileHelper.h"
 #include "Misc/PackageName.h"
+#include "Misc/OutputDeviceNull.h"
 #include "UObject/Class.h"
 #include "UObject/Stack.h"
 #include "UObject/Package.h"
@@ -887,7 +888,29 @@ void ExtractFunctionParams(const UFunction* InFunc, TArray<FGeneratedWrappedMeth
 		{
 			// This is either a default specified in Blueprint, or specified in C++ directly as meta-data
 			// (ie, because UHT can't parse the C++ default argument, eg, a default TMap argument)
-			GeneratedWrappedMethodParam.ParamDefaultValue = InFunc->GetMetaData(InParam->GetFName());
+			FString ParamDefaultValue = InFunc->GetMetaData(InParam->GetFName());
+
+			// If the parameter is a class then try and get the full name as the metadata might just be the short name
+			if (InParam->IsA<FClassProperty>() && !FPackageName::IsValidObjectPath(ParamDefaultValue))
+			{
+				if (const UClass* DefaultClass = FindObject<UClass>(ANY_PACKAGE, *ParamDefaultValue, true))
+				{
+					ParamDefaultValue = DefaultClass->GetPathName();
+				}
+			}
+
+			// Validate that the default is actually valid (ie, can be imported), as we may hit 
+			// false positives for this meta-data (eg, from LatentInfo, WorldContext, etc)
+			void* TempParam = FMemory_Alloca_Aligned(InParam->GetSize(), InParam->GetMinAlignment());
+			InParam->InitializeValue(TempParam);
+			{
+				FOutputDeviceNull NullOutputDevice;
+				if (PropertyAccessUtil::ImportDefaultPropertyValue(InParam, TempParam, *ParamDefaultValue, &NullOutputDevice))
+				{
+					GeneratedWrappedMethodParam.ParamDefaultValue = MoveTemp(ParamDefaultValue);
+				}
+			}
+			InParam->DestroyValue(TempParam);
 		}
 	};
 
diff --git a/Engine/Source/Runtime/CoreUObject/Private/UObject/PropertyAccessUtil.cpp b/Engine/Source/Runtime/CoreUObject/Private/UObject/PropertyAccessUtil.cpp
--- a/Engine/Source/Runtime/CoreUObject/Private/UObject/PropertyAccessUtil.cpp
+++ b/Engine/Source/Runtime/CoreUObject/Private/UObject/PropertyAccessUtil.cpp
@@ -2,6 +2,7 @@
 
 #include "UObject/PropertyAccessUtil.h"
 #include "UObject/EnumProperty.h"
+#include "UObject/TextProperty.h"
 #include "UObject/Object.h"
 #include "UObject/Class.h"
 #include "Misc/DefaultValueHelper.h"
@@ -642,14 +643,14 @@ FProperty* FindPropertyByName(const FName InPropName, const UStruct* InStruct)
 	return Prop;
 }
 
-void ImportDefaultPropertyValue(const FProperty* InProp, void* InPropValue, const FString& InDefaultValue)
+bool ImportDefaultPropertyValue(const FProperty* InProp, void* InPropValue, const FString& InDefaultValue, FOutputDevice* ErrorText)
 {
-	if (InDefaultValue.IsEmpty())
+	if (InDefaultValue.IsEmpty() && !(InProp->IsA<FStrProperty>() || InProp->IsA<FTextProperty>()))
 	{
-		return;
+		return false;
 	}
 
-	bool bDoImportText = true;
+	bool bImportedText = false;
 
 	// Certain struct types export using a non-standard default value, so we have to import them manually rather than use ImportText
 	if (const FStructProperty* StructProp = CastField<FStructProperty>(InProp))
@@ -657,34 +658,36 @@ void ImportDefaultPropertyValue(const FProperty* InProp, void* InPropValue, cons
 		if (StructProp->Struct == TBaseStructure<FVector>::Get())
 		{
 			FVector* Vector = (FVector*)InPropValue;
-			bDoImportText = !FDefaultValueHelper::ParseVector(InDefaultValue, *Vector);
+			bImportedText = FDefaultValueHelper::ParseVector(InDefaultValue, *Vector);
 		}
 		else if (StructProp->Struct == TBaseStructure<FVector2D>::Get())
 		{
 			FVector2D* Vector2D = (FVector2D*)InPropValue;
-			bDoImportText = !FDefaultValueHelper::ParseVector2D(InDefaultValue, *Vector2D);
+			bImportedText = FDefaultValueHelper::ParseVector2D(InDefaultValue, *Vector2D);
 		}
 		else if (StructProp->Struct == TBaseStructure<FRotator>::Get())
 		{
 			FRotator* Rotator = (FRotator*)InPropValue;
-			bDoImportText = !FDefaultValueHelper::ParseRotator(InDefaultValue, *Rotator);
+			bImportedText = FDefaultValueHelper::ParseRotator(InDefaultValue, *Rotator);
 		}
 		else if (StructProp->Struct == TBaseStructure<FColor>::Get())
 		{
 			FColor* Color = (FColor*)InPropValue;
-			bDoImportText = !FDefaultValueHelper::ParseColor(InDefaultValue, *Color);
+			bImportedText = FDefaultValueHelper::ParseColor(InDefaultValue, *Color);
 		}
 		else if (StructProp->Struct == TBaseStructure<FLinearColor>::Get())
 		{
 			FLinearColor* LinearColor = (FLinearColor*)InPropValue;
-			bDoImportText = !FDefaultValueHelper::ParseLinearColor(InDefaultValue, *LinearColor);
+			bImportedText = FDefaultValueHelper::ParseLinearColor(InDefaultValue, *LinearColor);
 		}
 	}
 
-	if (bDoImportText)
+	if (!bImportedText)
 	{
-		InProp->ImportText_Direct(*InDefaultValue, InPropValue, nullptr, PPF_None);
+		bImportedText = InProp->ImportText_Direct(*InDefaultValue, InPropValue, nullptr, PPF_None, ErrorText) != nullptr;
 	}
+
+	return bImportedText;
 }
 
 }
diff --git a/Engine/Source/Runtime/CoreUObject/Public/UObject/PropertyAccessUtil.h b/Engine/Source/Runtime/CoreUObject/Public/UObject/PropertyAccessUtil.h
--- a/Engine/Source/Runtime/CoreUObject/Public/UObject/PropertyAccessUtil.h
+++ b/Engine/Source/Runtime/CoreUObject/Public/UObject/PropertyAccessUtil.h
@@ -300,8 +300,10 @@ namespace PropertyAccessUtil
 	 * @param InProp The property to import the default value for.
 	 * @param InPropValue The destination location to receive the imported value.
 	 * @param InDefaultValue The default value to import.
+	 *
+	 * @return True if the default value could be imported, false otherwise.
 	 */
-	COREUOBJECT_API void ImportDefaultPropertyValue(const FProperty* InProp, void* InPropValue, const FString& InDefaultValue);
+	COREUOBJECT_API bool ImportDefaultPropertyValue(const FProperty* InProp, void* InPropValue, const FString& InDefaultValue, FOutputDevice* ErrorText = (FOutputDevice*)GWarn);
 
 	/**
 	 * Low-level function to determine if two properties are compatible with each other (for use by CopySinglePropertyValue/CopyCompletePropertyValue).
