AsyncLoading2: Fixed IsAsyncLoadingPackages returning false when there's still work left to do (introduced with 19718229)
#rnx

diff --git a/Engine/Source/Runtime/CoreUObject/Private/Serialization/AsyncLoading2.cpp b/Engine/Source/Runtime/CoreUObject/Private/Serialization/AsyncLoading2.cpp
--- a/Engine/Source/Runtime/CoreUObject/Private/Serialization/AsyncLoading2.cpp
+++ b/Engine/Source/Runtime/CoreUObject/Private/Serialization/AsyncLoading2.cpp
@@ -2048,9 +2048,11 @@ private:
 	FCriticalSection PendingRequestsCritical;
 
 	/** [ASYNC/GAME THREAD] Number of package load requests in the async loading queue */
-	TAtomic<uint32> QueuedPackagesCounter { 0 };
-	/** [ASYNC/GAME THREAD] Number of packages being loaded on the async thread and post loaded on the game thread. Excludes packages in the deferred delete queue*/
-	TAtomic<uint32> LoadingPackagesCounter { 0 };
+	TAtomic<int32> QueuedPackagesCounter { 0 };
+	/** [ASYNC/GAME THREAD] Number of packages being loaded on the async thread and post loaded on the game thread */
+	TAtomic<int32> LoadingPackagesCounter { 0 };
+	/** [ASYNC/GAME THREAD] While this is non-zero there's work left to do */
+	TAtomic<int32> PackagesWithRemainingWorkCounter{ 0 };
 
 	FThreadSafeCounter AsyncThreadReady;
 
@@ -2126,7 +2128,7 @@ public:
 	/** Returns true if packages are currently being loaded on the async thread */
 	inline virtual bool IsAsyncLoadingPackages() override
 	{
-		return QueuedPackagesCounter + LoadingPackagesCounter != 0;
+		return PackagesWithRemainingWorkCounter != 0;
 	}
 
 	/** Returns true this codes runs on the async loading thread */
@@ -2444,6 +2446,7 @@ private:
 		TRACE_CPUPROFILER_EVENT_SCOPE(DeleteAsyncPackage);
 		UE_ASYNC_PACKAGE_DEBUG(Package->Desc);
 		delete Package;
+		--PackagesWithRemainingWorkCounter;
 	}
 
 	/** Number of times we re-entered the async loading tick, mostly used by singlethreaded ticking. Debug purposes only. */
@@ -2573,6 +2576,7 @@ FAsyncPackage2* FAsyncLoadingThread2::FindOrInsertPackage(FAsyncPackageDesc2& De
 			Package = CreateAsyncPackage(Desc);
 			checkf(Package, TEXT("Failed to create async package %s"), *Desc.UPackageName.ToString());
 			Package->AddRef();
+			++LoadingPackagesCounter;
 			AsyncPackageLookup.Add(Desc.UPackageId, Package);
 			bInserted = true;
 		}
@@ -2629,7 +2633,6 @@ bool FAsyncLoadingThread2::CreateAsyncPackagesFromQueue(FAsyncLoadingThreadState
 				break;
 			}
 
-			++LoadingPackagesCounter;
 			--QueuedPackagesCounter;
 			++NumDequeued;
 		
@@ -2728,7 +2731,7 @@ bool FAsyncLoadingThread2::CreateAsyncPackagesFromQueue(FAsyncLoadingThreadState
 				{
 					UE_ASYNC_PACKAGE_LOG_VERBOSE(Verbose, PackageDesc, TEXT("CreateAsyncPackages: UpdatePackage"),
 						TEXT("Package is alreay being loaded."));
-					--LoadingPackagesCounter;
+					--PackagesWithRemainingWorkCounter;
 				}
 			}
 		}
@@ -3371,7 +3374,7 @@ void FAsyncPackage2::ImportPackagesRecursive(FIoBatch& IoBatch, IPackageStore& P
 		{
 			UE_ASYNC_PACKAGE_LOG(Verbose, Desc, TEXT("ImportPackages: AddPackage"),
 			TEXT("Start loading imported package with id '0x%llX'"), ImportedPackageId.ValueForDebugging());
-			++AsyncLoadingThread.LoadingPackagesCounter;
+			++AsyncLoadingThread.PackagesWithRemainingWorkCounter;
 		}
 		else
 		{
@@ -4749,6 +4752,9 @@ EAsyncPackageState::Type FAsyncLoadingThread2::ProcessLoadedPackagesFromGameThre
 			// this is to ensure we can re-enter FlushAsyncLoading from any of the callbacks
 			LoadedPackagesToProcess.RemoveAt(PackageIndex--);
 
+			// Incremented on the Async Thread, now decrement as we're done with this package
+			--LoadingPackagesCounter;
+
 			TRACE_LOADTIME_END_LOAD_ASYNC_PACKAGE(Package);
 
 			check(Package->AsyncPackageLoadingState == EAsyncPackageLoadingState2::Finalize);
@@ -4798,7 +4804,7 @@ EAsyncPackageState::Type FAsyncLoadingThread2::ProcessLoadedPackagesFromGameThre
 			{
 				FailedPackageRequest.Callback->ExecuteIfBound(FailedPackageRequest.PackageName, nullptr, EAsyncLoadingResult::Failed);
 				RemovePendingRequests(TArrayView<int32>(&FailedPackageRequest.RequestID, 1));
-				--LoadingPackagesCounter;
+				--PackagesWithRemainingWorkCounter;
 			}
 		}
 
@@ -4863,7 +4869,6 @@ EAsyncPackageState::Type FAsyncLoadingThread2::ProcessLoadedPackagesFromGameThre
 					CompletedPackages.RemoveAtSwap(PackageIndex--);
 					Package->ClearImportedPackages();
 					Package->ReleaseRef();
-					--LoadingPackagesCounter;
 				}
 			}
 
@@ -6011,6 +6016,7 @@ int32 FAsyncLoadingThread2::LoadPackage(const FPackagePath& InPackagePath, FName
 
 	PackageRequestQueue.Enqueue(FPackageRequest::Create(RequestId, InPackagePriority, PackageNameToLoad, PackageIdToLoad, InCustomName, MoveTemp(CompletionDelegate)));
 	++QueuedPackagesCounter;
+	++PackagesWithRemainingWorkCounter;
 
 	AltZenaphore.NotifyOne();
 
@@ -6046,7 +6052,7 @@ void FAsyncLoadingThread2::QueueMissingPackage(FAsyncPackageDesc2& PackageDesc,
 	else
 	{
 		RemovePendingRequests(TArrayView<int32>(&PackageDesc.RequestID, 1));
-		--LoadingPackagesCounter;
+		--PackagesWithRemainingWorkCounter;
 	}
 }
 
