FTopLevelAssetPath to represent the path of an object whose direct outer is its containing package, e.g. /Path/To/Package.AssetName.

diff --git a/Engine/Source/Runtime/CoreUObject/Private/UObject/TopLevelAssetPath.cpp b/Engine/Source/Runtime/CoreUObject/Private/UObject/TopLevelAssetPath.cpp
new file mode 100644
--- /dev/null
+++ b/Engine/Source/Runtime/CoreUObject/Private/UObject/TopLevelAssetPath.cpp
@@ -0,0 +1,202 @@
+// Copyright Epic Games, Inc. All Rights Reserved.
+
+#include "UObject/TopLevelAssetPath.h"
+
+#include "Containers/UnrealString.h"
+
+#include "Misc/AsciiSet.h"
+#include "Misc/RedirectCollector.h"
+#include "UObject/CoreRedirects.h"
+#include "UObject/ObjectRedirector.h"
+#include "UObject/UnrealType.h"
+#include "UObject/UObjectThreadContext.h"
+
+FTopLevelAssetPath::FTopLevelAssetPath(const UObject* InObject)
+{
+	if (InObject == nullptr)
+	{
+		PackageName = AssetName = FName();
+	}
+	else if (!InObject->GetOuter())
+	{
+		check(Cast<UPackage>(InObject) != nullptr);
+		PackageName = InObject->GetFName();
+		AssetName = FName();
+	}
+	else if (!InObject->IsAsset() || InObject->GetOuter()->GetOuter() != nullptr)
+	{
+		PackageName = AssetName = FName();
+	}
+	else
+	{
+		PackageName = InObject->GetOuter()->GetFName();
+		AssetName = InObject->GetFName();
+	}
+}
+
+void FTopLevelAssetPath::AppendString(FStringBuilderBase& Builder) const
+{
+	if (!IsNull())
+	{
+		Builder << PackageName;
+		if (!AssetName.IsNone())
+		{
+			Builder << '.' << AssetName;
+		}		
+	}
+}
+
+void FTopLevelAssetPath::AppendString(FString& Builder) const
+{
+	if (!IsNull())
+	{
+		PackageName.AppendString(Builder);
+		if( !AssetName.IsNone() )
+		{
+			Builder += TEXT(".");
+			AssetName.AppendString(Builder);
+		}
+	}
+}
+
+FString FTopLevelAssetPath::ToString() const
+{
+	TStringBuilder<256> Builder;
+	AppendString(Builder);
+	return FString(Builder);
+}
+
+void FTopLevelAssetPath::ToString(FString& OutString) const
+{
+	OutString.Reset();
+	AppendString(OutString);
+}
+
+bool FTopLevelAssetPath::TrySetPath(FName InPackageName, FName InAssetName)
+{
+	PackageName = InPackageName;
+	AssetName = InAssetName;
+	return !PackageName.IsNone();
+}
+
+bool FTopLevelAssetPath::TrySetPath(FWideStringView Path)
+{
+	if (Path.IsEmpty() || Path.Equals(TEXT("None"), ESearchCase::CaseSensitive))
+	{
+		// Empty path, just empty the pathname.
+		Reset();
+		return false;
+	}
+	else
+	{
+		if (Path[0] != '/')
+		{
+			// Possibly an ExportText path. Trim the ClassName.
+			Path = FPackageName::ExportTextPathToObjectPath(Path);
+
+			if (Path.IsEmpty() || Path[0] != '/')
+			{
+				Reset();
+				return false;
+			}
+		}
+
+		FAsciiSet Delim1(".");
+		FWideStringView PackageNameView = FAsciiSet::FindPrefixWithout(Path, Delim1);
+		if (PackageNameView.IsEmpty())
+		{
+			Reset();
+			return false;
+		}
+
+		FWideStringView AssetNameView = Path.Mid(PackageNameView.Len() + 1);
+		if (AssetNameView.IsEmpty())
+		{
+			// Reference to a package itself. Iffy, but supported for legacy usage of FSoftObjectPath.
+			PackageName = FName(PackageNameView);
+			AssetName = FName();
+			return true;
+		}
+
+		FAsciiSet Delim2("." SUBOBJECT_DELIMITER);
+		if (FAsciiSet::HasAny(AssetNameView, Delim2))
+		{
+			// Subobject path or is malformed and contains multiple '.' delimiters.
+			Reset();
+			return false;
+		}
+
+		PackageName = FName(PackageNameView);
+		AssetName = FName(AssetNameView);
+		return true;
+	}
+	return false;
+}
+
+bool FTopLevelAssetPath::TrySetPath(FAnsiStringView Path)
+{
+	TStringBuilder<FName::StringBufferSize> Wide;
+	Wide << Path;
+	return TrySetPath(Wide);
+}
+
+
+
+#if WITH_DEV_AUTOMATION_TESTS 
+
+#include "Misc/AutomationTest.h"
+
+// Combine import/export tests
+
+IMPLEMENT_SIMPLE_AUTOMATION_TEST(FTopLevelAssetPathTest, "System.Core.Misc.TopLevelAssetPath", EAutomationTestFlags::ApplicationContextMask | EAutomationTestFlags::EngineFilter);
+
+bool FTopLevelAssetPathTest::RunTest(const FString& Parameters)
+{
+	FName PackageName("/Path/To/Package");
+	FName AssetName("Asset");
+
+	FString AssetPathString(WriteToString<FName::StringBufferSize>(PackageName, '.', AssetName).ToView());
+
+	FTopLevelAssetPath EmptyPath;
+	TestEqual(TEXT("Empty path to string is empty string"), EmptyPath.ToString(), FString());
+
+	FTopLevelAssetPath PackagePath;
+	TestFalse(TEXT("TrySetPath(NAME_None, NAME_None) fails"), PackagePath.TrySetPath(NAME_None, NAME_None));
+	TestTrue(TEXT("TrySetPath(PackageName, NAME_None) succeeds"), PackagePath.TrySetPath(PackageName, NAME_None));
+	TestEqual(TEXT("PackagePath to string is PackageName"), PackagePath.ToString(), PackageName.ToString());
+
+	FTopLevelAssetPath AssetPath;
+	TestTrue(TEXT("TrySetPath(PackageName, AssetName) succeeds"), AssetPath.TrySetPath(PackageName, AssetName));
+	TestEqual(TEXT("AssetPath to string is PackageName.AssetName"), AssetPath.ToString(), AssetPathString);
+
+	FTopLevelAssetPath EmptyPathFromString;
+	TestFalse(TEXT("TrySetPath with empty string fails"), EmptyPathFromString.TrySetPath(TEXT("")));
+	TestEqual(TEXT("Empty path to string is empty string"), EmptyPathFromString.ToString(), FString());
+
+	FTopLevelAssetPath PackagePathFromString;
+	TestTrue(TEXT("TrySetPath(PackageName.ToString()) succeeds"), PackagePath.TrySetPath(PackageName.ToString()));
+	TestEqual(TEXT("PackagePath to string is PackageName"), PackagePath.ToString(), PackageName.ToString());
+
+	FTopLevelAssetPath AssetPathFromString;
+	TestTrue(TEXT("TrySetPath(AssetPath) succeeds"), PackagePath.TrySetPath(AssetPathString));
+	TestEqual(TEXT("AssetPathFromString to string is PackageName.AssetName"), PackagePath.ToString(), AssetPathString);
+
+	FTopLevelAssetPath FailedPath;
+	TestFalse(TEXT("TrySetPath with unrooted path string fails"), FailedPath.TrySetPath("UnrootedPackage/Subfolder"));
+	TestEqual(TEXT("Failed set to string is empty string"), FailedPath.ToString(), FString());
+
+	FTopLevelAssetPath SubObjectPath;
+	TestFalse(TEXT("TrySetPath with subobject path string fails"), SubObjectPath.TrySetPath("/Path/To/Package.Asset:Subobject"));
+	TestEqual(TEXT("Failed set to string is empty string"), SubObjectPath.ToString(), FString());
+
+	FTopLevelAssetPath MalformedPath;
+	TestFalse(TEXT("TrySetPath with malformed path string fails"), MalformedPath.TrySetPath("/Path/To/Package.Asset.Malformed"));
+	TestEqual(TEXT("Failed set to string is empty string"), MalformedPath.ToString(), FString());
+
+	// Test ExportText path with and without root 
+	// Test starting with . 
+
+	return true;
+}
+
+#endif // WITH_DEV_AUTOMATION_TESTS
diff --git a/Engine/Source/Runtime/CoreUObject/Public/UObject/TopLevelAssetPath.h b/Engine/Source/Runtime/CoreUObject/Public/UObject/TopLevelAssetPath.h
new file mode 100644
--- /dev/null
+++ b/Engine/Source/Runtime/CoreUObject/Public/UObject/TopLevelAssetPath.h
@@ -0,0 +1,167 @@
+// Copyright Epic Games, Inc. All Rights Reserved.
+
+#pragma once
+
+#include "CoreTypes.h"
+#include "Containers/StringView.h"
+#include "UObject/NameTypes.h"
+
+class UObject;
+class CORE_API FString;
+
+/* 
+ * A struct that can reference a top level asset such as '/Path/To/Package.AssetName'
+ * Stores two FNames internally to avoid
+ *  a) storing a concatenated FName that bloats global FName storage
+ *  b) storing an empty FString for a subobject path as FSoftObjectPath allows
+ * Can also be used to reference the package itself in which case the second name is NAME_None
+ * and the object resolves to the string `/Path/To/Package`
+*/
+struct COREUOBJECT_API FTopLevelAssetPath
+{
+	FTopLevelAssetPath() { }
+	FTopLevelAssetPath(TYPE_OF_NULLPTR) { }
+
+	/** Construct directly from components */
+	FTopLevelAssetPath(FName InPackageName, FName InAssetName)
+	{
+		TrySetPath(InPackageName, InAssetName);
+	}
+
+	UE_DEPRECATED(5.0, "FNames containing full asset paths have been replaced by FTopLevelAssetPath/FSoftLevelObjectPath."
+		"This function is only for temporary use interfacing with APIs that still produce an FName."
+		"Those APIS should be updated to use FTopLevelAssetPath or FSoftLevelObjectPath.")
+	FTopLevelAssetPath(FName InPath)
+	{
+		TrySetPath(InPath.ToString());
+	}
+
+	/** Construct from string / string view / raw string of a supported character type. */
+	explicit FTopLevelAssetPath(const FString& Path) { TrySetPath(FStringView(Path)); }
+	template<typename CharType>
+	explicit FTopLevelAssetPath(TStringView<CharType> Path) { TrySetPath(Path); }
+	template<typename CharType>
+	explicit FTopLevelAssetPath(const CharType* Path) { TrySetPath(TStringView<CharType>(Path)); }
+
+	/** Construct from an existing object in memory */
+	explicit FTopLevelAssetPath(const UObject* InObject);
+
+	/** Assign from the same types we can construct from */
+	FTopLevelAssetPath& operator=(const FString& Path) { TrySetPath(FStringView(Path)); return *this; }
+	template<typename CharType>
+	FTopLevelAssetPath& operator=(TStringView<CharType> Path) { TrySetPath(Path); return *this; }
+	template<typename CharType>
+	FTopLevelAssetPath& operator=(const CharType* Path) { TrySetPath(TStringView<CharType>(Path)); return *this; }
+	FTopLevelAssetPath& operator=(TYPE_OF_NULLPTR) { Reset(); return *this; }
+
+	/** Sets asset path of this reference based on components. */
+	bool TrySetPath(FName InPackageName, FName InAssetName);
+	/** Sets asset path of this reference based on a string path. Resets this object and returns false if the string is empty or does not represent a top level asset path. */
+	bool TrySetPath(FWideStringView Path);
+	bool TrySetPath(FAnsiStringView Path);
+	template<typename CharType>
+	bool TrySetPath(const CharType* Path) { return TrySetPath(TStringView<CharType>(Path)); }
+	bool TrySetPath(const FString& Path) { return TrySetPath(FStringView(Path)); }
+
+	/** Return the package name part e.g. /Path/To/Package as an FName. */
+	FName GetPackageName() const { return PackageName; }
+
+	/** Return the asset name part e.g. AssetName as an FName. */
+	FName GetAssetName() const { return AssetName; }
+
+	/** Append the full asset path (e.g. '/Path/To/Package.AssetName') to the string builder. */
+	void AppendString(FStringBuilderBase& Builder) const;
+	/** Append the full asset path (e.g. '/Path/To/Package.AssetName') to the string. */
+	void AppendString(FString& OutString) const;
+
+	/** Return the full asset path (e.g. '/Path/To/Package.AssetName') as a string. */
+	FString ToString() const;
+	/** Copy the full asset path (e.g. '/Path/To/Package.AssetName') into the provided string. */
+	void ToString(FString& OutString) const;
+
+	// Return the full asset path (e.g. '/Path/To/Package.AssetName') as an FName.
+	UE_DEPRECATED(5.1, "FNames containing full asset paths have been replaced by FTopLevelAssetPath/FSoftLevelObjectPath."
+		"This function is only for temporary use interfacing with APIs that still expect an FName."
+		"Those APIS should be updated to use FTopLevelAssetPath or FSoftLevelObjectPath.")
+	FName ToFName() const { return *ToString(); }
+
+	/** Check if this could possibly refer to a real object */
+	bool IsValid() const
+	{
+		return !PackageName.IsNone();
+	}
+
+	/** Checks to see if this is initialized to null */
+	bool IsNull() const
+	{
+		return PackageName.IsNone();
+	}
+
+	/** Resets reference to point to null */
+	void Reset()
+	{		
+		PackageName = AssetName = FName();
+	}
+	/** Compares two paths for non-case-sensitive equality. */
+	bool operator==(FTopLevelAssetPath const& Other) const
+	{
+		return PackageName == Other.PackageName && AssetName == Other.AssetName;
+	}
+
+	/** Compares two paths for non-case-sensitive inequality. */
+	bool operator!=(FTopLevelAssetPath const& Other) const
+	{
+		return !(*this == Other);
+	}
+
+	/** Serializes the internal path. Unlike FSoftObjectPath, does not handle any PIE or redirector fixups. */
+	friend FArchive& operator<<(FArchive& Ar, FTopLevelAssetPath& Path)
+	{
+		return Ar << Path.PackageName << Path.AssetName;
+	}
+
+	/** Serializes the internal path. Unlike FSoftObjectPath, does not handle any PIE or redirector fixups. */
+	friend void operator<<(FStructuredArchive::FSlot Slot, FTopLevelAssetPath& Path)
+	{
+		FStructuredArchive::FRecord Record = Slot.EnterRecord();
+		Record << SA_VALUE(TEXT("PackageName"), Path.PackageName) << SA_VALUE(TEXT("AssetName"), Path.AssetName);
+	}
+
+	/** Lexically compares two paths. */
+	int32 Compare(const FTopLevelAssetPath& Other) const
+	{
+		if (int32 Diff = PackageName.Compare(Other.PackageName))
+		{
+			return Diff;
+		}
+		return AssetName.Compare(Other.AssetName);
+	}
+
+	/** Compares two paths in a fast non-lexical order that is only valid for process lifetime. */
+	int32 CompareFast(const FTopLevelAssetPath& Other) const
+	{
+		if (int32 Diff = PackageName.CompareIndexes(Other.PackageName))
+		{
+			return Diff;
+		}
+		return AssetName.CompareIndexes(Other.AssetName);
+	}
+
+	friend uint32 GetTypeHash(FTopLevelAssetPath const& This)
+	{
+		return HashCombineFast(GetTypeHash(This.PackageName), GetTypeHash(This.AssetName));
+	}
+
+private:
+	/** Name of the package containing the asset e.g. /Path/To/Package */
+	FName PackageName;
+	/** Name of the asset within the package e.g. 'AssetName' */
+	FName AssetName;
+};
+
+
+inline FStringBuilderBase& operator<<(FStringBuilderBase& Builder, const FTopLevelAssetPath& Path)
+{
+	Path.AppendString(Builder);
+	return Builder;
+}
