Ensure ImportText allows the FOutputDevice to control the behavior rather than directly use UE_LOG

#rnx

diff --git a/Engine/Source/Runtime/CoreUObject/Private/UObject/EnumProperty.cpp b/Engine/Source/Runtime/CoreUObject/Private/UObject/EnumProperty.cpp
--- a/Engine/Source/Runtime/CoreUObject/Private/UObject/EnumProperty.cpp
+++ b/Engine/Source/Runtime/CoreUObject/Private/UObject/EnumProperty.cpp
@@ -371,7 +371,7 @@ const TCHAR* FEnumProperty::ImportText_Internal(const TCHAR* InBuffer, void* Con
 					SerializedObject = LoadContext->SerializedObject;
 				}
 			}
-			UE_LOG(LogClass, Warning, TEXT("In asset '%s', there is an enum property of type '%s' with an invalid value of '%s'"), *GetPathNameSafe(SerializedObject ? SerializedObject : FUObjectThreadContext::Get().ConstructedObject), *Enum->GetName(), *Temp);
+			ErrorText->Logf(ELogVerbosity::Warning, TEXT("In asset '%s', there is an enum property of type '%s' with an invalid value of '%s'"), *GetPathNameSafe(SerializedObject ? SerializedObject : FUObjectThreadContext::Get().ConstructedObject), *Enum->GetName(), *Temp);
 			return nullptr;
 		}
 	}
diff --git a/Engine/Source/Runtime/CoreUObject/Private/UObject/FieldPathProperty.cpp b/Engine/Source/Runtime/CoreUObject/Private/UObject/FieldPathProperty.cpp
--- a/Engine/Source/Runtime/CoreUObject/Private/UObject/FieldPathProperty.cpp
+++ b/Engine/Source/Runtime/CoreUObject/Private/UObject/FieldPathProperty.cpp
@@ -157,7 +157,7 @@ const TCHAR* FFieldPathProperty::ImportText_Internal( const TCHAR* Buffer, void*
 			FString UnquotedPathName;
 			if (!FParse::QuotedString(*PathName, UnquotedPathName))
 			{
-				UE_LOG(LogProperty, Warning, TEXT("FieldPathProperty: Bad quoted string: %s"), *PathName);
+				ErrorText->Logf(ELogVerbosity::Warning, TEXT("FieldPathProperty: Bad quoted string: %s"), *PathName);
 				return nullptr;
 			}
 			PathName = MoveTemp(UnquotedPathName);
diff --git a/Engine/Source/Runtime/CoreUObject/Private/UObject/PropertyArray.cpp b/Engine/Source/Runtime/CoreUObject/Private/UObject/PropertyArray.cpp
--- a/Engine/Source/Runtime/CoreUObject/Private/UObject/PropertyArray.cpp
+++ b/Engine/Source/Runtime/CoreUObject/Private/UObject/PropertyArray.cpp
@@ -671,7 +671,7 @@ const TCHAR* FArrayProperty::ImportTextInnerItem( const TCHAR* Buffer, const FPr
 			}
 			else if (Index >= Inner->ArrayDim)
 			{
-				UE_LOG(LogProperty, Warning, TEXT("%s is a fixed-sized array of %i values. Additional data after %i has been ignored during import."), *Inner->GetName(), Inner->ArrayDim, Inner->ArrayDim);
+				ErrorText->Logf(ELogVerbosity::Warning, TEXT("%s is a fixed-sized array of %i values. Additional data after %i has been ignored during import."), *Inner->GetName(), Inner->ArrayDim, Inner->ArrayDim);
 				break;
 			}
 		}
diff --git a/Engine/Source/Runtime/CoreUObject/Private/UObject/PropertyByte.cpp b/Engine/Source/Runtime/CoreUObject/Private/UObject/PropertyByte.cpp
--- a/Engine/Source/Runtime/CoreUObject/Private/UObject/PropertyByte.cpp
+++ b/Engine/Source/Runtime/CoreUObject/Private/UObject/PropertyByte.cpp
@@ -469,7 +469,7 @@ const TCHAR* FByteProperty::ImportText_Internal( const TCHAR* InBuffer, void* Co
 					SerializedObject = LoadContext->SerializedObject;
 				}
 			}
-			UE_LOG(LogClass, Warning, TEXT("In asset '%s', there is an enum property of type '%s' with an invalid value of '%s'"), *GetPathNameSafe(SerializedObject ? SerializedObject : FUObjectThreadContext::Get().ConstructedObject), *Enum->GetName(), *Temp);
+			ErrorText->Logf(ELogVerbosity::Warning, TEXT("In asset '%s', there is an enum property of type '%s' with an invalid value of '%s'"), *GetPathNameSafe(SerializedObject ? SerializedObject : FUObjectThreadContext::Get().ConstructedObject), *Enum->GetName(), *Temp);
 			return nullptr;
 		}
 	}
