Fixing code generation for TFieldPath property UFunction params.

diff --git a/Engine/Source/Programs/UnrealHeaderTool/Private/PropertyTypes.cpp b/Engine/Source/Programs/UnrealHeaderTool/Private/PropertyTypes.cpp
--- a/Engine/Source/Programs/UnrealHeaderTool/Private/PropertyTypes.cpp
+++ b/Engine/Source/Programs/UnrealHeaderTool/Private/PropertyTypes.cpp
@@ -2510,7 +2510,7 @@ struct FPropertyTypeTraitsFieldPath : public FPropertyTypeTraitsBase
 	{
 		const FPropertyBase& VarProperty = PropDef.GetPropertyBase();
 		ExtendedTypeText = FString::Printf(TEXT("TFieldPath<F%s>"), *VarProperty.FieldClassName.ToString());
-		return TEXT("STRUCT");
+		return TEXT("TFIELDPATH");
 	}
 };
 
diff --git a/Engine/Source/Runtime/CoreUObject/Private/UObject/FieldPathProperty.cpp b/Engine/Source/Runtime/CoreUObject/Private/UObject/FieldPathProperty.cpp
--- a/Engine/Source/Runtime/CoreUObject/Private/UObject/FieldPathProperty.cpp
+++ b/Engine/Source/Runtime/CoreUObject/Private/UObject/FieldPathProperty.cpp
@@ -187,7 +187,7 @@ FString FFieldPathProperty::GetCPPMacroType(FString& ExtendedTypeText) const
 {
 	check(PropertyClass);
 	ExtendedTypeText = FString::Printf(TEXT("TFieldPath<F%s>"), *PropertyClass->GetName());
-	return TEXT("STRUCT");
+	return TEXT("TFIELDPATH");
 }
 
 FString FFieldPathProperty::GetCPPTypeForwardDeclaration() const
diff --git a/Engine/Source/Runtime/CoreUObject/Public/UObject/ScriptMacros.h b/Engine/Source/Runtime/CoreUObject/Public/UObject/ScriptMacros.h
--- a/Engine/Source/Runtime/CoreUObject/Public/UObject/ScriptMacros.h
+++ b/Engine/Source/Runtime/CoreUObject/Public/UObject/ScriptMacros.h
@@ -89,6 +89,9 @@ enum {MAX_VARIABLE_SIZE = 0x0FFF };
 #define P_GET_SOFTCLASS(ObjectType,ParamName)		PARAM_PASSED_BY_VAL(ParamName, FSoftClassProperty, ObjectType)
 #define P_GET_SOFTCLASS_REF(ObjectType,ParamName)	PARAM_PASSED_BY_REF(ParamName, FSoftClassProperty, ObjectType)
 
+#define P_GET_TFIELDPATH(FieldType,ParamName)		PARAM_PASSED_BY_VAL(ParamName, FFieldPathProperty, FieldType)
+#define P_GET_TFIELDPATH_REF(FieldType,ParamName)	PARAM_PASSED_BY_REF(ParamName, FFieldPathProperty, FieldType)
+
 #define P_GET_ARRAY(ElementType,ParamName)			ElementType ParamName[(MAX_VARIABLE_SIZE/sizeof(ElementType))+1];		Stack.StepCompiledIn<FProperty>(ParamName);
 #define P_GET_ARRAY_REF(ElementType,ParamName)		ElementType ParamName##Temp[(MAX_VARIABLE_SIZE/sizeof(ElementType))+1]; ElementType* ParamName = Stack.StepCompiledInRef<FProperty, ElementType*>(ParamName##Temp);
 
